Merging MetalRHI & MetalShaderFormat changes from ZebraBranch, plus the code to enable or disable Mac support.

diff --git a/Engine/Config/BaseEngine.ini b/Engine/Config/BaseEngine.ini
--- a/Engine/Config/BaseEngine.ini
+++ b/Engine/Config/BaseEngine.ini
@@ -2345,6 +2345,9 @@ MinimumOSVersion=MSOS_Vista
 [/Script/LinuxTargetPlatform.LinuxTargetSettings]
 +TargetedRHIs=GLSL_430
 +TargetedRHIs=GLSL_150
+ 
+[/Script/MacTargetPlatform.MacTargetSettings]
++TargetedRHIs=GLSL_150
 
 [HMDPluginPriority]
 ; Since SteamVR also works with the Oculus Rift, give priority to the native Oculus plugin before trying SteamVR
diff --git a/Engine/Source/Developer/Apple/MetalShaderFormat/MetalShaderFormat.Build.cs b/Engine/Source/Developer/Apple/MetalShaderFormat/MetalShaderFormat.Build.cs
new file mode 100644
--- /dev/null
+++ b/Engine/Source/Developer/Apple/MetalShaderFormat/MetalShaderFormat.Build.cs
@@ -0,0 +1,25 @@
+// Copyright 1998-2015 Epic Games, Inc. All Rights Reserved.
+
+using UnrealBuildTool;
+
+public class MetalShaderFormat : ModuleRules
+{
+	public MetalShaderFormat(TargetInfo Target)
+	{
+		PrivateIncludePathModuleNames.Add("TargetPlatform");
+		PublicIncludePaths.Add("Runtime/Apple/MetalRHI/Public");
+
+		PrivateDependencyModuleNames.AddRange(
+			new string[] {
+				"Core",
+				"ShaderCore",
+				"ShaderCompilerCommon",
+				"ShaderPreprocessor"
+			}
+			);
+
+		AddThirdPartyPrivateStaticDependencies(Target, 
+			"HLSLCC"
+			);
+	}
+}
diff --git a/Engine/Source/Developer/Mac/MacTargetPlatform/Private/GenericMacTargetPlatform.h b/Engine/Source/Developer/Mac/MacTargetPlatform/Private/GenericMacTargetPlatform.h
--- a/Engine/Source/Developer/Mac/MacTargetPlatform/Private/GenericMacTargetPlatform.h
+++ b/Engine/Source/Developer/Mac/MacTargetPlatform/Private/GenericMacTargetPlatform.h
@@ -35,6 +35,22 @@ public:
 			FConfigCacheIni::LoadLocalIniFile(EngineSettings, TEXT("Engine"), true, *this->PlatformName());
 			TextureLODSettings = nullptr;
 			StaticMeshLODSettings.Initialize(EngineSettings);
+		
+			// Get the Target RHIs for this platform, we do not always want all those that are supported.
+			GConfig->GetArray(TEXT("/Script/MacTargetPlatform.MacTargetSettings"), TEXT("TargetedRHIs"), TargetedShaderFormats, GEngineIni);
+			
+			// Gather the list of Target RHIs and filter out any that may be invalid.
+			TArray<FName> PossibleShaderFormats;
+			GetAllPossibleShaderFormats(PossibleShaderFormats);
+			
+			for(int32 ShaderFormatIdx = TargetedShaderFormats.Num()-1; ShaderFormatIdx >= 0; ShaderFormatIdx--)
+			{
+				FString ShaderFormat = TargetedShaderFormats[ShaderFormatIdx];
+				if(PossibleShaderFormats.Contains(FName(*ShaderFormat)) == false)
+				{
+					TargetedShaderFormats.RemoveAt(ShaderFormatIdx);
+				}
+			}
 		#endif
 	}
 
@@ -100,12 +116,19 @@ return TSuper::SupportsFeature(Feature);
 		{
 			static FName NAME_GLSL_150_MAC(TEXT("GLSL_150_MAC"));
 			OutFormats.AddUnique(NAME_GLSL_150_MAC);
+			static FName NAME_SF_METAL_SM4(TEXT("SF_METAL_SM4"));
+			OutFormats.AddUnique(NAME_SF_METAL_SM4);
+			static FName NAME_SF_METAL_SM5(TEXT("SF_METAL_SM5"));
+			OutFormats.AddUnique(NAME_SF_METAL_SM5);
 		}
 	}
 
 	virtual void GetAllTargetedShaderFormats(TArray<FName>& OutFormats) const override
 	{
-		GetAllPossibleShaderFormats( OutFormats );
+		for(const FString& ShaderFormat : TargetedShaderFormats)
+		{
+			OutFormats.AddUnique(FName(*ShaderFormat));
+		}
 	}
 
 
@@ -212,6 +235,9 @@ private:
 
 	// Holds the static mesh LOD settings.
 	FStaticMeshLODSettings StaticMeshLODSettings;
+	
+	// List of shader formats specified as targets
+	TArray<FString> TargetedShaderFormats;
 #endif // WITH_ENGINE
 
 private:
diff --git a/Engine/Source/Editor/DetailCustomizations/Private/MacTargetSettingsDetails.cpp b/Engine/Source/Editor/DetailCustomizations/Private/MacTargetSettingsDetails.cpp
--- a/Engine/Source/Editor/DetailCustomizations/Private/MacTargetSettingsDetails.cpp
+++ b/Engine/Source/Editor/DetailCustomizations/Private/MacTargetSettingsDetails.cpp
@@ -88,6 +88,10 @@ static FString GetIconFilename(EMacImageScope::Type Scope)
 
 void FMacTargetSettingsDetails::CustomizeDetails( IDetailLayoutBuilder& DetailBuilder )
 {
+	// Setup the supported/targeted RHI property view
+	TargetShaderFormatsDetails = MakeShareable(new FMacTargetShaderFormatsPropertyDetails(&DetailBuilder));
+	TargetShaderFormatsDetails->CreateTargetShaderFormatsPropertyView();
+	
 	// Add the splash image customization
 	const FText EditorSplashDesc(LOCTEXT("EditorSplashLabel", "Editor Splash"));
 	IDetailCategoryBuilder& SplashCategoryBuilder = DetailBuilder.EditCategory(TEXT("Splash"));
@@ -212,4 +216,113 @@ bool FMacTargetSettingsDetails::HandlePostExternalIconCopy(const FString& InChos
 	return true;
 }
 
+FText GetFriendlyNameFromRHINameMac(const FString& InRHIName)
+{
+	FText FriendlyRHIName = LOCTEXT("UnknownRHI", "UnknownRHI");
+	if (InRHIName == TEXT("GLSL_150_MAC"))
+	{
+		FriendlyRHIName = LOCTEXT("OpenGL3", "OpenGL 3 (SM4)");
+	}
+	else if (InRHIName == TEXT("SF_METAL_SM4"))
+	{
+		FriendlyRHIName = LOCTEXT("MetalSM4", "Metal (SM4)");
+	}
+	else if (InRHIName == TEXT("SF_METAL_SM5"))
+	{
+		FriendlyRHIName = LOCTEXT("MetalSM5", "Metal (SM5, Experimental)");
+	}
+	
+	return FriendlyRHIName;
+}
+
+FMacTargetShaderFormatsPropertyDetails::FMacTargetShaderFormatsPropertyDetails(IDetailLayoutBuilder* InDetailBuilder)
+: DetailBuilder(InDetailBuilder)
+{
+	TargetShaderFormatsPropertyHandle = DetailBuilder->GetProperty("TargetedRHIs");
+	ensure(TargetShaderFormatsPropertyHandle.IsValid());
+}
+
+void FMacTargetShaderFormatsPropertyDetails::CreateTargetShaderFormatsPropertyView()
+{
+	DetailBuilder->HideProperty(TargetShaderFormatsPropertyHandle);
+	
+	// List of supported RHI's and selected targets
+	ITargetPlatform* TargetPlatform = FModuleManager::GetModuleChecked<ITargetPlatformModule>("MacTargetPlatform").GetTargetPlatform();
+	TArray<FName> ShaderFormats;
+	TargetPlatform->GetAllPossibleShaderFormats(ShaderFormats);
+	
+	IDetailCategoryBuilder& TargetedRHICategoryBuilder = DetailBuilder->EditCategory(TEXT("Targeted RHIs"));
+	
+	for (const FName& ShaderFormat : ShaderFormats)
+	{
+		FText FriendlyShaderFormatName = GetFriendlyNameFromRHINameMac(ShaderFormat.ToString());
+		
+		FDetailWidgetRow& TargetedRHIWidgetRow = TargetedRHICategoryBuilder.AddCustomRow(FriendlyShaderFormatName);
+		
+		TargetedRHIWidgetRow
+		.NameContent()
+		[
+			SNew(SHorizontalBox)
+			+ SHorizontalBox::Slot()
+			.Padding(FMargin(0, 1, 0, 1))
+			.FillWidth(1.0f)
+			[
+				SNew(STextBlock)
+				.Text(FriendlyShaderFormatName)
+				.Font(DetailBuilder->GetDetailFont())
+			 ]
+		 ]
+		.ValueContent()
+		[
+			SNew(SCheckBox)
+			.OnCheckStateChanged(this, &FMacTargetShaderFormatsPropertyDetails::OnTargetedRHIChanged, ShaderFormat)
+			.IsChecked(this, &FMacTargetShaderFormatsPropertyDetails::IsTargetedRHIChecked, ShaderFormat)
+		 ];
+	}
+}
+
+
+void FMacTargetShaderFormatsPropertyDetails::OnTargetedRHIChanged(ECheckBoxState InNewValue, FName InRHIName)
+{
+	TArray<void*> RawPtrs;
+	TargetShaderFormatsPropertyHandle->AccessRawData(RawPtrs);
+	
+	// Update the CVars with the selection
+	{
+		TargetShaderFormatsPropertyHandle->NotifyPreChange();
+		for (void* RawPtr : RawPtrs)
+		{
+			TArray<FString>& Array = *(TArray<FString>*)RawPtr;
+			if(InNewValue == ECheckBoxState::Checked)
+			{
+				Array.Add(InRHIName.ToString());
+			}
+			else
+			{
+				Array.Remove(InRHIName.ToString());
+			}
+		}
+		TargetShaderFormatsPropertyHandle->NotifyPostChange();
+	}
+}
+
+
+ECheckBoxState FMacTargetShaderFormatsPropertyDetails::IsTargetedRHIChecked(FName InRHIName) const
+{
+	ECheckBoxState CheckState = ECheckBoxState::Unchecked;
+	
+	TArray<void*> RawPtrs;
+	TargetShaderFormatsPropertyHandle->AccessRawData(RawPtrs);
+	
+	for(void* RawPtr : RawPtrs)
+	{
+		TArray<FString>& Array = *(TArray<FString>*)RawPtr;
+		if(Array.Contains(InRHIName.ToString()))
+		{
+			CheckState = ECheckBoxState::Checked;
+		}
+	}
+	return CheckState;
+}
+
 #undef LOCTEXT_NAMESPACE
diff --git a/Engine/Source/Editor/DetailCustomizations/Private/MacTargetSettingsDetails.h b/Engine/Source/Editor/DetailCustomizations/Private/MacTargetSettingsDetails.h
--- a/Engine/Source/Editor/DetailCustomizations/Private/MacTargetSettingsDetails.h
+++ b/Engine/Source/Editor/DetailCustomizations/Private/MacTargetSettingsDetails.h
@@ -2,6 +2,46 @@
 
 #pragma once
 
+/* FMacTargetShaderFormatsPropertyDetails
+ *****************************************************************************/
+
+/**
+ * Helper which implements details panel customizations for a device profiles parent property
+ */
+class FMacTargetShaderFormatsPropertyDetails
+: public TSharedFromThis<FMacTargetShaderFormatsPropertyDetails>
+{
+	
+public:
+	
+	/**
+	 * Constructor for the parent property details view
+	 *
+	 * @param InDetailsBuilder - Where we are adding our property view to
+	 */
+	FMacTargetShaderFormatsPropertyDetails(IDetailLayoutBuilder* InDetailBuilder);
+	
+	
+	/**
+	 * Create the UI to select which windows shader formats we are targetting
+	 */
+	void CreateTargetShaderFormatsPropertyView();
+	
+private:
+	
+	// Supported/Targeted RHI check boxes
+	void OnTargetedRHIChanged(ECheckBoxState InNewValue, FName InRHIName);
+	ECheckBoxState IsTargetedRHIChecked(FName InRHIName) const;
+	
+private:
+	
+	/** A handle to the detail view builder */
+	IDetailLayoutBuilder* DetailBuilder;
+	
+	/** Access to the Parent Property */
+	TSharedPtr<IPropertyHandle> TargetShaderFormatsPropertyHandle;
+};
+
 /**
  * Manages the Transform section of a details view                    
  */
@@ -23,4 +63,8 @@ private:
 
 	/** Delegate handler to set the path to start picking from */
 	bool HandlePostExternalIconCopy(const FString& InChosenImage);
+	
+private:
+	/** Reference to the target shader formats property view */
+	TSharedPtr<FMacTargetShaderFormatsPropertyDetails> TargetShaderFormatsDetails;
 };
\ No newline at end of file
diff --git a/Engine/Source/Programs/CrossCompilerTool/Private/CrossCompilerTool.cpp b/Engine/Source/Programs/CrossCompilerTool/Private/CrossCompilerTool.cpp
--- a/Engine/Source/Programs/CrossCompilerTool/Private/CrossCompilerTool.cpp
+++ b/Engine/Source/Programs/CrossCompilerTool/Private/CrossCompilerTool.cpp
@@ -33,6 +33,8 @@ namespace CCT
 		ILanguageSpec* Language = nullptr;
 		FCodeBackend* Backend = nullptr;
 
+		EHlslCompileTarget CompileTarget = RunInfo.Target;
+
 		uint32 Flags = HLSLCC_PackUniforms;
 		Flags |= RunInfo.bValidate ? 0 : HLSLCC_NoValidation;
 		Flags |= RunInfo.bRunCPP ? 0 : HLSLCC_NoPreprocess;
@@ -48,13 +50,14 @@ namespace CCT
 		FGlslLanguageSpec GlslLanguage(RunInfo.Target == HCT_FeatureLevelES2);
 		FGlslCodeBackend GlslBackend(Flags, RunInfo.Target);
 		FMetalLanguageSpec MetalLanguage;
-		FMetalCodeBackend MetalBackend(Flags, RunInfo.Target);
+		FMetalCodeBackend MetalBackend(Flags, CompileTarget);
 
 		switch (RunInfo.BackEnd)
 		{
 		case CCT::FRunInfo::BE_Metal:
 			Language = &MetalLanguage;
 			Backend = &MetalBackend;
+			CompileTarget = HCT_FeatureLevelES3_1;
 			break;
 
 		case CCT::FRunInfo::BE_OpenGL:
@@ -150,7 +153,7 @@ namespace CCT
 		ANSICHAR* ShaderSource = 0;
 		ANSICHAR* ErrorLog = 0;
 
-		FHlslCrossCompilerContext Context(Flags, RunInfo.Frequency, RunInfo.Target);
+		FHlslCrossCompilerContext Context(Flags, RunInfo.Frequency, CompileTarget);
 		if (Context.Init(TCHAR_TO_ANSI(*RunInfo.InputFile), Language))
 		{
 			Context.Run(
diff --git a/Engine/Source/Programs/CrossCompilerTool/Private/CrossCompilerUtils.cpp b/Engine/Source/Programs/CrossCompilerTool/Private/CrossCompilerUtils.cpp
--- a/Engine/Source/Programs/CrossCompilerTool/Private/CrossCompilerUtils.cpp
+++ b/Engine/Source/Programs/CrossCompilerTool/Private/CrossCompilerUtils.cpp
@@ -199,6 +199,30 @@ namespace CCT
 					OutBackEnd = BE_Metal;
 				}
 			}
+			else if (Switch == "metalmrt" || Switch == "metalsm4")
+			{
+				if (Target != HCT_InvalidTarget)
+				{
+					UE_LOG(LogCrossCompilerTool, Warning, TEXT("Ignoring extra command line argument -%s!"), *Switch);
+				}
+				else
+				{
+					Target = HCT_FeatureLevelSM4;
+					OutBackEnd = BE_Metal;
+				}
+			}
+			else if (Switch == "metalsm5")
+			{
+				if (Target != HCT_InvalidTarget)
+				{
+					UE_LOG(LogCrossCompilerTool, Warning, TEXT("Ignoring extra command line argument -metalsm5!"));
+				}
+				else
+				{
+					Target = HCT_FeatureLevelSM5;
+					OutBackEnd = BE_Metal;
+				}
+			}
 			else if (Switch == "gl3")
 			{
 				if (Target != HCT_InvalidTarget)
diff --git a/Engine/Source/Programs/UnrealBuildTool/Mac/UEBuildMac.cs b/Engine/Source/Programs/UnrealBuildTool/Mac/UEBuildMac.cs
--- a/Engine/Source/Programs/UnrealBuildTool/Mac/UEBuildMac.cs
+++ b/Engine/Source/Programs/UnrealBuildTool/Mac/UEBuildMac.cs
@@ -139,15 +139,26 @@ namespace UnrealBuildTool
 
 			InBuildTarget.GlobalCompileEnvironment.Config.Definitions.Add("WITH_TTS=0");
 			InBuildTarget.GlobalCompileEnvironment.Config.Definitions.Add("WITH_SPEECH_RECOGNITION=0");
-			if(!InBuildTarget.GlobalCompileEnvironment.Config.Definitions.Contains("WITH_DATABASE_SUPPORT=0") && !InBuildTarget.GlobalCompileEnvironment.Config.Definitions.Contains("WITH_DATABASE_SUPPORT=1"))
+			if (!InBuildTarget.GlobalCompileEnvironment.Config.Definitions.Contains("WITH_DATABASE_SUPPORT=0") && !InBuildTarget.GlobalCompileEnvironment.Config.Definitions.Contains("WITH_DATABASE_SUPPORT=1"))
 			{
 				InBuildTarget.GlobalCompileEnvironment.Config.Definitions.Add("WITH_DATABASE_SUPPORT=0");
 			}
+			// Needs OS X 10.11 for Metal
+			if (MacToolChain.MacOSSDKVersionFloat >= 10.11f && UEBuildConfiguration.bCompileAgainstEngine)
+			{
+				InBuildTarget.GlobalCompileEnvironment.Config.Definitions.Add("HAS_METAL=1");
+				InBuildTarget.ExtraModuleNames.Add("MetalRHI");
+			}
+			else
+			{
+				InBuildTarget.GlobalCompileEnvironment.Config.Definitions.Add("HAS_METAL=0");
+			}
+
 		}
 
 		/**
 		 *	Whether this platform should create debug information or not
-		 *	
+		 *
 		 *	@param	InPlatform			The UnrealTargetPlatform being built
 		 *	@param	InConfiguration		The UnrealTargetConfiguration being built
 		 *	
diff --git a/Engine/Source/Runtime/Apple/MetalRHI/MetalRHI.Build.cs b/Engine/Source/Runtime/Apple/MetalRHI/MetalRHI.Build.cs
new file mode 100644
--- /dev/null
+++ b/Engine/Source/Runtime/Apple/MetalRHI/MetalRHI.Build.cs
@@ -0,0 +1,29 @@
+// Copyright 1998-2015 Epic Games, Inc. All Rights Reserved.
+
+using UnrealBuildTool;
+
+public class MetalRHI : ModuleRules
+{	
+	public MetalRHI(TargetInfo Target)
+	{
+
+		PrivateDependencyModuleNames.AddRange(
+			new string[] {
+				"Core",
+				"Engine",
+				"RHI",
+				"RenderCore",
+				"ShaderCore",
+				"UtilityShaders"
+			}
+			);
+			
+		PublicWeakFrameworks.Add("Metal");
+
+		if (Target.Platform == UnrealTargetPlatform.Mac)
+		{
+			PublicFrameworks.Add("QuartzCore");
+			// Definitions.Add("METAL_STATISTICS=1");
+		}
+	}
+}
diff --git a/Engine/Source/Runtime/Apple/MetalRHI/Private/MetalBufferPools.cpp b/Engine/Source/Runtime/Apple/MetalRHI/Private/MetalBufferPools.cpp
new file mode 100644
--- /dev/null
+++ b/Engine/Source/Runtime/Apple/MetalRHI/Private/MetalBufferPools.cpp
@@ -0,0 +1,244 @@
+// Copyright 1998-2015 Epic Games, Inc. All Rights Reserved.
+
+#include "MetalRHIPrivate.h"
+
+#include "MetalBufferPools.h"
+
+FMetalPooledBuffer::FMetalPooledBuffer()
+: Buffer(nil)
+{
+	
+}
+
+FMetalPooledBuffer::FMetalPooledBuffer(FMetalPooledBuffer const& Other)
+: Buffer(nil)
+{
+	operator=(Other);
+}
+
+FMetalPooledBuffer::FMetalPooledBuffer(id<MTLBuffer> Buf)
+: Buffer(Buf)
+{
+	[Buffer retain];
+}
+
+FMetalPooledBuffer::~FMetalPooledBuffer()
+{
+	[Buffer release];
+	Buffer = nil;
+}
+
+FMetalPooledBuffer& FMetalPooledBuffer::operator=(FMetalPooledBuffer const& Other)
+{
+	if(this != &Other)
+	{
+		Buffer = [Other.Buffer retain];
+	}
+	return *this;
+}
+
+/** Get the pool bucket index from the size
+ * @param Size the number of bytes for the resource
+ * @returns The bucket index.
+ */
+uint32 FMetalBufferPoolPolicyData::GetPoolBucketIndex(CreationArguments Args)
+{
+	// Only BUFFER_CACHE_MODE cached memory is acceptable
+	// Auto & Private memory is also forbidden in this pool
+#if PLATFORM_MAC
+	check(Args.Storage == MTLStorageModeShared || Args.Storage == MTLStorageModeManaged);
+#endif
+	
+	uint32 Size = Args.Size;
+	
+	unsigned long Lower = 0;
+	unsigned long Upper = NumPoolBucketSizes;
+	unsigned long Middle;
+	
+	do
+	{
+		Middle = ( Upper + Lower ) >> 1;
+		if( Size <= BucketSizes[Middle-1] )
+		{
+			Upper = Middle;
+		}
+		else
+		{
+			Lower = Middle;
+		}
+	}
+	while( Upper - Lower > 1 );
+	
+	check( Size <= BucketSizes[Lower] );
+	check( (Lower == 0 ) || ( Size > BucketSizes[Lower-1] ) );
+	
+	// Managed buffers are placed into a different buffer than shared.
+	Lower += (NumPoolBucketSizes * Args.Storage);
+	
+	return Lower;
+}
+
+/** Get the pool bucket size from the index
+ * @param Bucket the bucket index
+ * @returns The bucket size.
+ */
+uint32 FMetalBufferPoolPolicyData::GetPoolBucketSize(uint32 Bucket)
+{
+	check(Bucket < NumPoolBuckets);
+	uint32 Index = Bucket;
+	if(Index >= NumPoolBucketSizes)
+	{
+		Index -= NumPoolBucketSizes;
+	}
+	checkf(Index < NumPoolBucketSizes, TEXT("%d %d"), Index, NumPoolBucketSizes);
+	return BucketSizes[Index];
+}
+
+/** Creates the resource
+ * @param Args The buffer size in bytes.
+ * @returns A suitably sized buffer or NULL on failure.
+ */
+FMetalPooledBuffer FMetalBufferPoolPolicyData::CreateResource(CreationArguments Args)
+{
+	// This manual construction is deliberate as newBuffer creates a resource with ref-cout == 1, so no need to retain
+	check(Args.Device);
+	FMetalPooledBuffer NewBuf;
+	uint32 BufferSize = GetPoolBucketSize(GetPoolBucketIndex(Args));
+	NewBuf.Buffer = [Args.Device newBufferWithLength: BufferSize options: BUFFER_CACHE_MODE
+#if METAL_API_1_1
+					 | (Args.Storage << MTLResourceStorageModeShift)
+#endif
+					 ];
+	return NewBuf;
+}
+
+/** Gets the arguments used to create resource
+ * @param Resource The buffer to get data for.
+ * @returns The arguments used to create the buffer.
+ */
+FMetalBufferPoolPolicyData::CreationArguments FMetalBufferPoolPolicyData::GetCreationArguments(FMetalPooledBuffer Resource)
+{
+#if PLATFORM_MAC
+	return FMetalBufferPoolPolicyData::CreationArguments(Resource.Buffer.device, Resource.Buffer.length, Resource.Buffer.storageMode);
+#else
+	return FMetalBufferPoolPolicyData::CreationArguments(Resource.Buffer.device, Resource.Buffer.length, MTLStorageModeShared);
+#endif
+}
+
+/** The bucket sizes */
+uint32 FMetalBufferPoolPolicyData::BucketSizes[NumPoolBucketSizes] = {
+	256, 512, 768, 1024, 2048, 4096, 8192, 16384,
+	32768, 65536, 131072, 262144, 524288, 1048576, 2097152,
+	4194304, 8388608, 16777216, 33554432, 67108864, 134217728, 268435456
+};
+
+/** Destructor */
+FMetalBufferPool::~FMetalBufferPool()
+{
+}
+
+void FMetalQueryBufferPool::Allocate(FMetalQueryResult& NewQuery)
+{
+    FMetalQueryBuffer* QB = IsValidRef(CurrentBuffer) ? CurrentBuffer.GetReference() : GetCurrentQueryBuffer();
+    if(Align(QB->WriteOffset, EQueryBufferAlignment) + EQueryResultMaxSize <= EQueryBufferMaxSize)
+    {
+		if(!QB->CommandBuffer)
+		{
+			QB->CommandBuffer = Context->GetCurrentCommandBuffer();
+			[QB->CommandBuffer retain];
+		}
+		NewQuery.SourceBuffer = QB;
+        NewQuery.Offset = Align(QB->WriteOffset, EQueryBufferAlignment);
+        FMemory::Memzero((((uint8*)[QB->Buffer contents]) + NewQuery.Offset), EQueryResultMaxSize);
+        QB->WriteOffset = Align(QB->WriteOffset, EQueryBufferAlignment) + EQueryResultMaxSize;
+	}
+	else
+	{
+		UE_LOG(LogRHI, Warning, TEXT("Performance: Resetting render command encoder as query buffer offset: %d exceeds the maximum allowed: %d."), QB->WriteOffset, EQueryBufferMaxSize);
+		Context->ResetRenderCommandEncoder();
+		Allocate(NewQuery);
+	}
+}
+
+FMetalQueryBuffer* FMetalQueryBufferPool::GetCurrentQueryBuffer()
+{
+    if(!IsValidRef(CurrentBuffer) || CurrentBuffer->WriteOffset > 0)
+    {
+        id<MTLBuffer> Buffer;
+        if(Buffers.Num())
+        {
+            Buffer = Buffers.Pop();
+        }
+        else
+        {
+            Buffer = [Context->GetDevice() newBufferWithLength:EQueryBufferMaxSize options: BUFFER_CACHE_MODE
+#if METAL_API_1_1
+					  | MTLResourceStorageModeShared
+#endif
+					  ];
+        }
+        CurrentBuffer = new FMetalQueryBuffer(Context, Buffer);
+    }
+
+	return CurrentBuffer.GetReference();
+}
+
+void FMetalQueryBufferPool::ReleaseQueryBuffer(id<MTLBuffer> Buffer)
+{
+    Buffers.Add(Buffer);
+}
+
+FRingBuffer::FRingBuffer(id<MTLDevice> Device, uint32 Size, uint32 InDefaultAlignment)
+{
+	DefaultAlignment = InDefaultAlignment;
+	Buffer = [Device newBufferWithLength:Size options:BUFFER_CACHE_MODE];
+	Offset = 0;
+	LastRead = Size;
+}
+
+uint32 FRingBuffer::Allocate(uint32 Size, uint32 Alignment)
+{
+	if (Alignment == 0)
+	{
+		Alignment = DefaultAlignment;
+	}
+	
+	if(LastRead <= Offset)
+	{
+		// align the offset
+		Offset = Align(Offset, Alignment);
+		
+		// wrap if needed
+		if (Offset + Size <= Buffer.length)
+		{
+			// get current location
+			uint32 ReturnOffset = Offset;
+			Offset += Size;
+			return ReturnOffset;
+		}
+		else // wrap
+		{
+			// 		NSLog(@"Wrapping at frame %lld [size = %d]", GFrameCounter, (uint32)Buffer.length);
+			Offset = 0;
+		}
+	}
+	
+	// align the offset
+	Offset = Align(Offset, Alignment);
+	if(Offset + Size < LastRead)
+	{
+		// get current location
+		uint32 ReturnOffset = Offset;
+		Offset += Size;
+		return ReturnOffset;
+	}
+	else
+	{
+		UE_LOG(LogMetal, Warning, TEXT("Wrapping offset %d into outstanding buffer region %d at frame %lld [size = %d]"), Offset, LastRead, GFrameCounter, (uint32)Buffer.length);
+		
+		// get current location
+		uint32 ReturnOffset = Offset;
+		Offset += Size;
+		return ReturnOffset;
+	}
+}
diff --git a/Engine/Source/Runtime/Apple/MetalRHI/Private/MetalBufferPools.h b/Engine/Source/Runtime/Apple/MetalRHI/Private/MetalBufferPools.h
new file mode 100644
--- /dev/null
+++ b/Engine/Source/Runtime/Apple/MetalRHI/Private/MetalBufferPools.h
@@ -0,0 +1,130 @@
+// Copyright 1998-2015 Epic Games, Inc. All Rights Reserved.
+
+
+#pragma once
+
+#import <Metal/Metal.h>
+#include "ResourcePool.h"
+
+#if !__IPHONE_9_0 && !__MAC_10_11
+typedef NS_ENUM(NSUInteger, MTLStorageMode)
+{
+	MTLStorageModeShared = 0
+};
+#endif
+
+struct FMetalPooledBuffer
+{
+	FMetalPooledBuffer();
+	FMetalPooledBuffer(FMetalPooledBuffer const& Other);
+	FMetalPooledBuffer(id<MTLBuffer> Buf);
+	~FMetalPooledBuffer();
+	
+	FMetalPooledBuffer& operator=(FMetalPooledBuffer const& Other);
+	
+	id<MTLBuffer> Buffer;
+};
+
+struct FMetalPooledBufferArgs
+{
+	FMetalPooledBufferArgs() : Device(nil), Size(0), Storage(MTLStorageModeShared) {}
+	
+	FMetalPooledBufferArgs(id<MTLDevice> InDevice, uint32 InSize, MTLStorageMode InStorage)
+	: Device(InDevice)
+	, Size(InSize)
+	, Storage(InStorage)
+	{
+	}
+	
+	id<MTLDevice> Device;
+	uint32 Size;
+	MTLStorageMode Storage;
+};
+
+class FMetalBufferPoolPolicyData
+{
+public:
+	/** Buffers are created with a simple byte size */
+	typedef FMetalPooledBufferArgs CreationArguments;
+	enum
+	{
+		NumSafeFrames = 2, /** Number of frames to leave buffers before reclaiming/reusing */
+		NumPoolBucketSizes = 22, /** Number of pool bucket sizes */
+		NumPoolBuckets = NumPoolBucketSizes * 2, /** Number of pool bucket sizes * 2 for Shared vs. Managed storage */
+		NumToDrainPerFrame = 65536, /** Max. number of resources to cull in a single frame */
+		CullAfterFramesNum = 30 * 60 /** Resources are culled if unused for more frames than this */
+	};
+	
+	/** Get the pool bucket index from the size
+	 * @param Size the number of bytes for the resource
+	 * @returns The bucket index.
+	 */
+	uint32 GetPoolBucketIndex(CreationArguments Args);
+	
+	/** Get the pool bucket size from the index
+	 * @param Bucket the bucket index
+	 * @returns The bucket size.
+	 */
+	uint32 GetPoolBucketSize(uint32 Bucket);
+	
+	/** Creates the resource
+	 * @param Args The buffer size in bytes.
+	 * @returns A suitably sized buffer or NULL on failure.
+	 */
+	FMetalPooledBuffer CreateResource(CreationArguments Args);
+	
+	/** Gets the arguments used to create resource
+	 * @param Resource The buffer to get data for.
+	 * @returns The arguments used to create the buffer.
+	 */
+	CreationArguments GetCreationArguments(FMetalPooledBuffer Resource);
+	
+private:
+	/** The bucket sizes */
+	static uint32 BucketSizes[NumPoolBucketSizes];
+};
+
+/** A pool for metal buffers with consistent usage, bucketed for efficiency. */
+class FMetalBufferPool : public TResourcePool<FMetalPooledBuffer, FMetalBufferPoolPolicyData, FMetalBufferPoolPolicyData::CreationArguments>
+{
+public:
+	/** Destructor */
+	virtual ~FMetalBufferPool();
+};
+
+struct FMetalQueryBufferPool
+{
+    enum
+    {
+        EQueryBufferAlignment = 8,
+        EQueryResultMaxSize = 8,
+        EQueryBufferMaxSize = 64 * 1024
+    };
+    
+    FMetalQueryBufferPool(class FMetalContext* InContext)
+    : Context(InContext)
+    {
+    }
+    
+    void Allocate(FMetalQueryResult& NewQuery);
+    FMetalQueryBuffer* GetCurrentQueryBuffer();
+    void ReleaseQueryBuffer(id<MTLBuffer> Buffer);
+    
+    FMetalQueryBufferRef CurrentBuffer;
+    TArray<id<MTLBuffer>> Buffers;
+	class FMetalContext* Context;
+};
+
+struct FRingBuffer
+{
+	FRingBuffer(id<MTLDevice> Device, uint32 Size, uint32 InDefaultAlignment);
+
+	uint32 GetOffset() { return Offset; }
+	void SetLastRead(uint32 Read) { LastRead = Read; }
+	uint32 Allocate(uint32 Size, uint32 Alignment);
+
+	uint32 DefaultAlignment;
+	id<MTLBuffer> Buffer;
+	uint32 Offset;
+	uint32 LastRead;
+};
diff --git a/Engine/Source/Runtime/Apple/MetalRHI/Private/MetalCommandEncoder.cpp b/Engine/Source/Runtime/Apple/MetalRHI/Private/MetalCommandEncoder.cpp
new file mode 100644
--- /dev/null
+++ b/Engine/Source/Runtime/Apple/MetalRHI/Private/MetalCommandEncoder.cpp
@@ -0,0 +1,817 @@
+// Copyright 1998-2015 Epic Games, Inc. All Rights Reserved.
+
+/*=============================================================================
+	MetalCommandEncoder.cpp: Metal command encoder wrapper.
+=============================================================================*/
+
+#include "MetalRHIPrivate.h"
+
+#include "MetalCommandEncoder.h"
+
+#pragma mark - Public C++ Boilerplate -
+
+FMetalCommandEncoder::FMetalCommandEncoder(void)
+: FrontFacingWinding(MTLWindingClockwise)
+, CullMode(MTLCullModeNone)
+#if METAL_API_1_1
+, DepthClipMode(MTLDepthClipModeClip)
+#endif
+, FillMode(MTLTriangleFillModeFill)
+, DepthStencilState(nil)
+, VisibilityMode(MTLVisibilityResultModeDisabled)
+, VisibilityOffset(0)
+, PipelineStatsBuffer(nil)
+, PipelineStatsOffset(0)
+, PipelineStatsMask(0)
+, RenderPassDesc(nil)
+, CommandBuffer(nil)
+, RenderCommandEncoder(nil)
+, RenderPipelineState(nil)
+, ComputeCommandEncoder(nil)
+, ComputePipelineState(nil)
+, BlitCommandEncoder(nil)
+, DebugGroups([NSMutableArray new])
+{
+	Viewport = {0.0, 0.0, 0.0, 0.0, FLT_MIN, FLT_MAX};
+	DepthBias[0] = DepthBias[1] = 0.0f;
+	DepthBias[2] = FLT_MAX;
+	ScissorRect = {0,0, 0,0};
+	FMemory::Memzero(BlendColor);
+	StencilRef[0] = StencilRef[1] = 0xff;
+	
+	FMemory::Memzero(ShaderBuffers);
+	FMemory::Memzero(ShaderTextures);
+	FMemory::Memzero(ShaderSamplers);
+	for(int32 i = 0; i < SF_NumFrequencies; i++)
+	{
+		for(int32 j = 0; j < ML_MaxSamplers; j++)
+		{
+			ShaderSamplers[i].MaxLods[j] = FLT_MAX;
+		}
+	}
+}
+
+FMetalCommandEncoder::~FMetalCommandEncoder(void)
+{
+	EndEncoding();
+	CommitCommandBuffer(false);
+	
+	if(ComputePipelineState)
+	{
+		[ComputePipelineState release];
+	}
+	if(RenderPipelineState)
+	{
+		[RenderPipelineState release];
+	}
+	if(DepthStencilState)
+	{
+		[DepthStencilState release];
+	}
+	[RenderPassDesc release];
+}
+	
+#pragma mark - Public Command Buffer Mutators -
+
+void FMetalCommandEncoder::StartCommandBuffer(id<MTLCommandBuffer> const InCommandBuffer)
+{
+	check(InCommandBuffer);
+	check(!CommandBuffer);
+	check(IsRenderCommandEncoderActive() == false && IsComputeCommandEncoderActive() == false && IsBlitCommandEncoderActive() == false);
+
+	[InCommandBuffer retain];
+	id<MTLCommandBuffer> OldCommandBuffer = CommandBuffer;
+	CommandBuffer = InCommandBuffer;
+	[OldCommandBuffer release];
+}
+	
+void FMetalCommandEncoder::CommitCommandBuffer(bool const bWait)
+{
+	check(CommandBuffer);
+	check(IsRenderCommandEncoderActive() == false && IsComputeCommandEncoderActive() == false && IsBlitCommandEncoderActive() == false);
+	
+	[CommandBuffer commit];
+	if(bWait)
+	{
+		[CommandBuffer waitUntilCompleted];
+	}
+	[CommandBuffer release];
+	CommandBuffer = nil;
+}
+
+#pragma mark - Public Command Encoder Accessors -
+	
+bool FMetalCommandEncoder::IsRenderCommandEncoderActive(void) const
+{
+	return RenderCommandEncoder != nil;
+}
+
+bool FMetalCommandEncoder::IsComputeCommandEncoderActive(void) const
+{
+	return ComputeCommandEncoder != nil;
+}
+
+bool FMetalCommandEncoder::IsBlitCommandEncoderActive(void) const
+{
+	return BlitCommandEncoder != nil;
+}
+
+id<MTLRenderCommandEncoder> FMetalCommandEncoder::GetRenderCommandEncoder(void) const
+{
+	return RenderCommandEncoder;
+}
+
+id<MTLComputeCommandEncoder> FMetalCommandEncoder::GetComputeCommandEncoder(void) const
+{
+	return ComputeCommandEncoder;
+}
+
+id<MTLBlitCommandEncoder> FMetalCommandEncoder::GetBlitCommandEncoder(void) const
+{
+	return BlitCommandEncoder;
+}
+	
+#pragma mark - Public Command Encoder Mutators -
+
+void FMetalCommandEncoder::BeginRenderCommandEncoding(MTLRenderPassDescriptor* const RenderPass)
+{
+	check(RenderPass);
+	check(CommandBuffer);
+	check(IsRenderCommandEncoderActive() == false && IsComputeCommandEncoderActive() == false && IsBlitCommandEncoderActive() == false);
+	
+	if(RenderPass != RenderPassDesc)
+	{
+		[RenderPass retain];
+		if(RenderPassDesc)
+		{
+			GetMetalDeviceContext().ReleaseObject(RenderPassDesc);
+		}
+		RenderPassDesc = RenderPass;
+	}
+	
+	RenderCommandEncoder = [[CommandBuffer renderCommandEncoderWithDescriptor:RenderPass] retain];
+	
+	if(GEmitDrawEvents)
+	{
+		NSString* Label = [DebugGroups count] > 0 ? [DebugGroups lastObject] : (NSString*)CFSTR("InitialPass");
+		[RenderCommandEncoder setLabel:Label];
+    }
+    
+    for(uint32 i = 0; i < SF_NumFrequencies; i++)
+    {
+        ShaderBuffers[i].Bound = 0;
+		ShaderTextures[i].Bound = 0;
+		ShaderSamplers[i].Bound = 0;
+    }
+}
+
+void FMetalCommandEncoder::RestoreRenderCommandEncoding(void)
+{
+	check(RenderPassDesc);
+	check(CommandBuffer);
+	check(IsRenderCommandEncoderActive() == false && IsComputeCommandEncoderActive() == false && IsBlitCommandEncoderActive() == false);
+	
+	BeginRenderCommandEncoding(RenderPassDesc);
+	RestoreRenderCommandEncodingState();
+}
+
+void FMetalCommandEncoder::RestoreRenderCommandEncodingState(void)
+{
+	check(RenderCommandEncoder);
+	
+	PushDebugGroup(@"RestoreRenderCommandEncodingState");
+	
+	[RenderCommandEncoder setRenderPipelineState:RenderPipelineState];
+	[RenderCommandEncoder setViewport:Viewport];
+	[RenderCommandEncoder setFrontFacingWinding:FrontFacingWinding];
+	[RenderCommandEncoder setCullMode:CullMode];
+#if METAL_API_1_1
+	if(GetMetalDeviceContext().SupportsFeature(EMetalFeaturesDepthClipMode))
+	{
+		[RenderCommandEncoder setDepthClipMode:DepthClipMode];
+	}
+#endif
+	[RenderCommandEncoder setDepthBias:DepthBias[0] slopeScale:DepthBias[1] clamp:DepthBias[2]];
+	
+	CGRect ViewportRect = CGRectMake(Viewport.originX, Viewport.originY, Viewport.width, Viewport.height);
+	CGRect Scissor = CGRectMake(ScissorRect.x, ScissorRect.y, ScissorRect.width, ScissorRect.height);
+	CGRect Unified = CGRectIntersection(ViewportRect, Scissor);
+	MTLScissorRect UnifiedScissorRect = {(NSUInteger)Unified.origin.x, (NSUInteger)Unified.origin.y, (NSUInteger)Unified.size.width, (NSUInteger)Unified.size.height};
+	
+	[RenderCommandEncoder setScissorRect:UnifiedScissorRect];
+	[RenderCommandEncoder setTriangleFillMode:FillMode];
+	[RenderCommandEncoder setBlendColorRed:BlendColor[0] green:BlendColor[1] blue:BlendColor[2] alpha:BlendColor[3]];
+	[RenderCommandEncoder setDepthStencilState:DepthStencilState];
+#if METAL_API_1_1
+	if(GetMetalDeviceContext().SupportsFeature(EMetalFeaturesSeparateStencil))
+	{
+		[RenderCommandEncoder setStencilFrontReferenceValue: StencilRef[0] backReferenceValue:StencilRef[1]];
+	}
+	else
+#endif
+	{
+		[RenderCommandEncoder setStencilReferenceValue: StencilRef[0]];
+	}
+	[RenderCommandEncoder setVisibilityResultMode:VisibilityMode offset:VisibilityOffset];
+	
+	for(uint32 i = 0; i < ML_MaxBuffers; i++)
+	{
+		if(ShaderBuffers[SF_Vertex].Buffers[i] != nil)
+		{
+            ShaderBuffers[SF_Vertex].Bound |= (1 << i);
+			[RenderCommandEncoder setVertexBuffer:ShaderBuffers[SF_Vertex].Buffers[i] offset:ShaderBuffers[SF_Vertex].Offsets[i] atIndex:i];
+		}
+		if(ShaderBuffers[SF_Pixel].Buffers[i] != nil)
+        {
+            ShaderBuffers[SF_Pixel].Bound |= (1 << i);
+			[RenderCommandEncoder setFragmentBuffer:ShaderBuffers[SF_Pixel].Buffers[i] offset:ShaderBuffers[SF_Pixel].Offsets[i] atIndex:i];
+		}
+	}
+	
+	for(uint32 i = 0; i < ML_MaxTextures; i++)
+	{
+		if(ShaderTextures[SF_Vertex].Textures[i])
+		{
+			ShaderTextures[SF_Vertex].Bound |= (1 << i);
+			[RenderCommandEncoder setVertexTexture:ShaderTextures[SF_Vertex].Textures[i] atIndex:i];
+		}
+		
+		if(ShaderTextures[SF_Pixel].Textures[i])
+		{
+			ShaderTextures[SF_Pixel].Bound |= (1 << i);
+			[RenderCommandEncoder setFragmentTexture:ShaderTextures[SF_Pixel].Textures[i] atIndex:i];
+		}
+	}
+	
+	for(uint32 i = 0; i < ML_MaxSamplers; i++)
+	{
+		if(ShaderSamplers[SF_Vertex].Samplers[i] != nil)
+		{
+			ShaderSamplers[SF_Vertex].Bound |= (1 << i);
+			[RenderCommandEncoder setVertexSamplerState:ShaderSamplers[SF_Vertex].Samplers[i] lodMinClamp:ShaderSamplers[SF_Vertex].MinLods[i] lodMaxClamp:ShaderSamplers[SF_Vertex].MaxLods[i] atIndex:i];
+		}
+		if(ShaderSamplers[SF_Pixel].Samplers[i] != nil)
+		{
+			ShaderSamplers[SF_Pixel].Bound |= (1 << i);
+			[RenderCommandEncoder setFragmentSamplerState:ShaderSamplers[SF_Pixel].Samplers[i] lodMinClamp:ShaderSamplers[SF_Pixel].MinLods[i] lodMaxClamp:ShaderSamplers[SF_Pixel].MaxLods[i] atIndex:i];
+		}
+	}
+	
+	PopDebugGroup();
+}
+
+void FMetalCommandEncoder::BeginComputeCommandEncoding(void)
+{
+	check(CommandBuffer);
+	check(IsRenderCommandEncoderActive() == false && IsComputeCommandEncoderActive() == false && IsBlitCommandEncoderActive() == false);
+	
+	ComputeCommandEncoder = [[CommandBuffer computeCommandEncoder] retain];
+	
+	if(GEmitDrawEvents)
+	{
+		NSString* Label = [DebugGroups count] > 0 ? [DebugGroups lastObject] : (NSString*)CFSTR("InitialPass");
+		[ComputeCommandEncoder setLabel:Label];
+	}
+	
+    for(uint32 i = 0; i < SF_NumFrequencies; i++)
+    {
+		ShaderBuffers[i].Bound = 0;
+		ShaderTextures[i].Bound = 0;
+		ShaderSamplers[i].Bound = 0;
+    }
+}
+
+void FMetalCommandEncoder::BeginBlitCommandEncoding(void)
+{
+	check(CommandBuffer);
+	check(IsRenderCommandEncoderActive() == false && IsComputeCommandEncoderActive() == false && IsBlitCommandEncoderActive() == false);
+	
+	BlitCommandEncoder = [[CommandBuffer blitCommandEncoder] retain];
+}
+
+void FMetalCommandEncoder::EndEncoding(void)
+{
+	if(IsRenderCommandEncoderActive())
+	{
+		[RenderCommandEncoder endEncoding];
+		[RenderCommandEncoder release];
+		RenderCommandEncoder = nil;
+	}
+	else if(IsComputeCommandEncoderActive())
+	{
+		[ComputeCommandEncoder endEncoding];
+		[ComputeCommandEncoder release];
+		ComputeCommandEncoder = nil;
+	}
+	else if(IsBlitCommandEncoderActive())
+	{
+		[BlitCommandEncoder endEncoding];
+		[BlitCommandEncoder release];
+		BlitCommandEncoder = nil;
+	}
+}
+
+#pragma mark - Public Debug Support -
+
+void FMetalCommandEncoder::InsertDebugSignpost(NSString* const String)
+{
+	[RenderCommandEncoder insertDebugSignpost:String];
+}
+
+void FMetalCommandEncoder::PushDebugGroup(NSString* const String)
+{
+	[DebugGroups addObject:String];
+	[RenderCommandEncoder pushDebugGroup:String];
+}
+
+void FMetalCommandEncoder::PopDebugGroup(void)
+{
+	[DebugGroups removeLastObject];
+	[RenderCommandEncoder popDebugGroup];
+}
+	
+#pragma mark - Public Render State Mutators -
+
+void FMetalCommandEncoder::SetRenderPipelineState(id<MTLRenderPipelineState> PipelineState)
+{
+	[PipelineState retain];
+	[RenderPipelineState release];
+	RenderPipelineState = PipelineState;
+	[RenderCommandEncoder setRenderPipelineState:RenderPipelineState];
+}
+
+void FMetalCommandEncoder::SetViewport(MTLViewport const& InViewport)
+{
+	Viewport = InViewport;
+	[RenderCommandEncoder setViewport:Viewport];
+}
+
+void FMetalCommandEncoder::SetFrontFacingWinding(MTLWinding const InFrontFacingWinding)
+{
+	FrontFacingWinding = InFrontFacingWinding;
+	[RenderCommandEncoder setFrontFacingWinding:FrontFacingWinding];
+}
+
+void FMetalCommandEncoder::SetCullMode(MTLCullMode const InCullMode)
+{
+	CullMode = InCullMode;
+	[RenderCommandEncoder setCullMode:InCullMode];
+}
+
+#if METAL_API_1_1
+void FMetalCommandEncoder::SetDepthClipMode(MTLDepthClipMode const InDepthClipMode)
+{
+	DepthClipMode = InDepthClipMode;
+	if(GetMetalDeviceContext().SupportsFeature(EMetalFeaturesDepthClipMode))
+	{
+		[RenderCommandEncoder setDepthClipMode:InDepthClipMode];
+	}
+	else
+	{
+		UE_LOG(LogMetal, Warning, TEXT("SetDepthClipMode unsupported on the current device."));
+	}
+}
+#endif
+
+void FMetalCommandEncoder::SetDepthBias(float const InDepthBias, float const InSlopeScale, float const InClamp)
+{
+	DepthBias[0] = InDepthBias;
+	DepthBias[1] = InSlopeScale;
+	DepthBias[2] = InClamp;
+	[RenderCommandEncoder setDepthBias:InDepthBias slopeScale:InSlopeScale clamp:InClamp];
+}
+
+void FMetalCommandEncoder::SetScissorRect(MTLScissorRect const& Rect)
+{
+	ScissorRect = Rect;
+	[RenderCommandEncoder setScissorRect:ScissorRect];
+}
+
+void FMetalCommandEncoder::SetTriangleFillMode(MTLTriangleFillMode const InFillMode)
+{
+	FillMode = InFillMode;
+	[RenderCommandEncoder setTriangleFillMode:FillMode];
+}
+
+void FMetalCommandEncoder::SetBlendColor(float const Red, float const Green, float const Blue, float const Alpha)
+{
+	BlendColor[0] = Red;
+	BlendColor[1] = Green;
+	BlendColor[2] = Blue;
+	BlendColor[3] = Alpha;
+	[RenderCommandEncoder setBlendColorRed:Red green:Green blue:Blue alpha:Alpha];
+}
+
+void FMetalCommandEncoder::SetDepthStencilState(id<MTLDepthStencilState> const InDepthStencilState)
+{
+	[InDepthStencilState retain];
+	[DepthStencilState release];
+	DepthStencilState = InDepthStencilState;
+	[RenderCommandEncoder setDepthStencilState:InDepthStencilState];
+}
+
+void FMetalCommandEncoder::SetStencilReferenceValue(uint32 const ReferenceValue)
+{
+	StencilRef[0] = ReferenceValue;
+	StencilRef[1] = ReferenceValue;
+	[RenderCommandEncoder setStencilReferenceValue:ReferenceValue];
+}
+
+void FMetalCommandEncoder::SetStencilReferenceValue(uint32 const FrontReferenceValue, uint32 const BackReferenceValue)
+{
+	StencilRef[0] = FrontReferenceValue;
+	StencilRef[1] = BackReferenceValue;
+#if METAL_API_1_1
+	if(GetMetalDeviceContext().SupportsFeature(EMetalFeaturesSeparateStencil))
+	{
+		[RenderCommandEncoder setStencilFrontReferenceValue: FrontReferenceValue backReferenceValue:BackReferenceValue];
+	}
+	else
+#endif
+	{
+		UE_LOG(LogMetal, Warning, TEXT("Attempting to set separate stencil ref on device which only supports unified stencil ref."));
+		[RenderCommandEncoder setStencilReferenceValue: FrontReferenceValue];
+	}
+}
+
+void FMetalCommandEncoder::SetVisibilityResultMode(MTLVisibilityResultMode const Mode, NSUInteger const Offset)
+{
+	VisibilityMode = Mode;
+	VisibilityOffset = Offset;
+	[RenderCommandEncoder setVisibilityResultMode: Mode offset:Offset];
+}
+	
+#pragma mark - Public Shader Resource Mutators -
+
+void FMetalCommandEncoder::SetShaderBuffer(EShaderFrequency Frequency, id<MTLBuffer> Buffer, NSUInteger const Offset, NSUInteger index)
+{
+	check(index < ML_MaxBuffers);
+    if(GetMetalDeviceContext().SupportsFeature(EMetalFeaturesSetBufferOffset) && (ShaderBuffers[Frequency].Bound & (1 << index)) && ShaderBuffers[Frequency].Buffers[index] == Buffer)
+    {
+		SetShaderBufferOffset(Frequency, Offset, index);
+    }
+    else
+    {
+		if(Buffer)
+		{
+			ShaderBuffers[Frequency].Bound |= (1 << index);
+		}
+		else
+		{
+			ShaderBuffers[Frequency].Bound &= ~(1 << index);
+		}
+        ShaderBuffers[Frequency].Buffers[index] = Buffer;
+        ShaderBuffers[Frequency].Offsets[index] = Offset;
+        switch (Frequency)
+        {
+            case SF_Vertex:
+                [RenderCommandEncoder setVertexBuffer:Buffer offset:Offset atIndex:index];
+                break;
+            case SF_Pixel:
+                [RenderCommandEncoder setFragmentBuffer:Buffer offset:Offset atIndex:index];
+                break;
+            case SF_Compute:
+                checkf(ComputeCommandEncoder != nil, TEXT("Attempted to use ComputeCommandEncoder before calling FMetalContext::ConditionalSwitchToCompute(void) to create the encoder"));
+                [ComputeCommandEncoder setBuffer:Buffer offset:Offset atIndex:index];
+                break;
+            default:
+                check(false);
+                break;
+        }
+    }
+}
+
+void FMetalCommandEncoder::SetShaderBufferOffset(EShaderFrequency Frequency, NSUInteger const Offset, NSUInteger index)
+{
+	check(index < ML_MaxBuffers);
+    checkf(ShaderBuffers[Frequency].Buffers[index] && (ShaderBuffers[Frequency].Bound & (1 << index)), TEXT("Buffer must already be bound"));
+#if METAL_API_1_1
+	if(GetMetalDeviceContext().SupportsFeature(EMetalFeaturesSetBufferOffset))
+	{
+		ShaderBuffers[Frequency].Offsets[index] = Offset;
+		switch (Frequency)
+		{
+			case SF_Vertex:
+				[RenderCommandEncoder setVertexBufferOffset:Offset atIndex:index];
+				break;
+			case SF_Pixel:
+				[RenderCommandEncoder setFragmentBufferOffset:Offset atIndex:index];
+				break;
+			case SF_Compute:
+				checkf(ComputeCommandEncoder != nil, TEXT("Attempted to use ComputeCommandEncoder before calling FMetalContext::ConditionalSwitchToCompute(void) to create the encoder"));
+				[ComputeCommandEncoder setBufferOffset:Offset atIndex:index];
+				break;
+			default:
+				check(false);
+				break;
+		}
+	}
+	else
+#endif
+	{
+		SetShaderBuffer(Frequency, ShaderBuffers[Frequency].Buffers[index], Offset, index);
+	}
+}
+
+void FMetalCommandEncoder::SetShaderBuffers(EShaderFrequency Frequency, id<MTLBuffer> const Buffers[], const NSUInteger Offsets[], NSRange const& Range)
+{
+	for (NSUInteger i = (NSUInteger)Range.location; i < (NSUInteger)Range.length; i++)
+	{
+        check(i < ML_MaxBuffers);
+		if(Buffers[i])
+		{
+			ShaderBuffers[Frequency].Bound |= (1 << i);
+		}
+		else
+		{
+			ShaderBuffers[Frequency].Bound &= ~(1 << i);
+		}
+		ShaderBuffers[Frequency].Buffers[i] = Buffers[i];
+		ShaderBuffers[Frequency].Offsets[i] = Offsets[i];
+	}
+	switch (Frequency)
+	{
+		case SF_Vertex:
+			[RenderCommandEncoder setVertexBuffers:Buffers offsets:Offsets withRange:Range];
+			break;
+		case SF_Pixel:
+			[RenderCommandEncoder setFragmentBuffers:Buffers offsets:Offsets withRange:Range];
+			break;
+		case SF_Compute:
+			checkf(ComputeCommandEncoder != nil, TEXT("Attempted to use ComputeCommandEncoder before calling FMetalContext::ConditionalSwitchToCompute(void) to create the encoder"));
+			[ComputeCommandEncoder setBuffers:Buffers offsets:Offsets withRange:Range];
+			break;
+		default:
+			check(false);
+			break;
+	}
+}
+
+void FMetalCommandEncoder::SetShaderTexture(EShaderFrequency Frequency, id<MTLTexture> Texture, NSUInteger index)
+{
+	check(index < ML_MaxTextures);
+	
+	if(Texture)
+	{
+		ShaderTextures[Frequency].Bound |= (1 << index);
+	}
+	else
+	{
+		ShaderTextures[Frequency].Bound &= ~(1 << index);
+	}
+	
+	ShaderTextures[Frequency].Textures[index] = Texture;
+	
+	switch (Frequency)
+	{
+		case SF_Vertex:
+			[RenderCommandEncoder setVertexTexture:Texture atIndex:index];
+			break;
+		case SF_Pixel:
+			[RenderCommandEncoder setFragmentTexture:Texture atIndex:index];
+			break;
+		case SF_Compute:
+			checkf(ComputeCommandEncoder != nil, TEXT("Attempted to use ComputeCommandEncoder before calling FMetalContext::ConditionalSwitchToCompute(void) to create the encoder"));
+			[ComputeCommandEncoder setTexture:Texture atIndex:index];
+			break;
+		default:
+			check(false);
+			break;
+	}
+}
+
+void FMetalCommandEncoder::SetShaderTextures(EShaderFrequency Frequency, const id<MTLTexture> Textures[], NSRange const& Range)
+{
+	for (NSUInteger i = (NSUInteger)Range.location; i < (NSUInteger)Range.length; i++)
+	{
+		check(i < ML_MaxTextures);
+		if(Textures[i])
+		{
+			ShaderTextures[Frequency].Bound |= (1 << i);
+		}
+		else
+		{
+			ShaderTextures[Frequency].Bound &= ~(1 << i);
+		}
+		
+		ShaderTextures[Frequency].Textures[i] = Textures[i];
+	}
+	
+	switch (Frequency)
+	{
+		case SF_Vertex:
+			[RenderCommandEncoder setVertexTextures:Textures withRange:Range];
+			break;
+		case SF_Pixel:
+			[RenderCommandEncoder setFragmentTextures:Textures withRange:Range];
+			break;
+		case SF_Compute:
+			checkf(ComputeCommandEncoder != nil, TEXT("Attempted to use ComputeCommandEncoder before calling FMetalContext::ConditionalSwitchToCompute(void) to create the encoder"));
+			[ComputeCommandEncoder setTextures:Textures withRange:Range];
+			break;
+		default:
+			check(false);
+			break;
+	}
+}
+
+void FMetalCommandEncoder::SetShaderSamplerState(EShaderFrequency Frequency, id<MTLSamplerState> Sampler, NSUInteger index)
+{
+	check(index < ML_MaxSamplers);
+	[Sampler retain];
+	[ShaderSamplers[Frequency].Samplers[index] release];
+	
+	if(Sampler)
+	{
+		ShaderSamplers[Frequency].Bound |= (1 << index);
+	}
+	else
+	{
+		ShaderSamplers[Frequency].Bound &= ~(1 << index);
+	}
+	
+	ShaderSamplers[Frequency].Samplers[index] = Sampler;
+	ShaderSamplers[Frequency].MinLods[index] = 0;
+	ShaderSamplers[Frequency].MaxLods[index] = FLT_MAX;
+	
+	switch (Frequency)
+	{
+		case SF_Vertex:
+			[RenderCommandEncoder setVertexSamplerState:Sampler atIndex:index];
+			break;
+		case SF_Pixel:
+			[RenderCommandEncoder setFragmentSamplerState:Sampler atIndex:index];
+			break;
+		case SF_Compute:
+			checkf(ComputeCommandEncoder != nil, TEXT("Attempted to use ComputeCommandEncoder before calling FMetalContext::ConditionalSwitchToCompute(void) to create the encoder"));
+			[ComputeCommandEncoder setSamplerState:Sampler atIndex:index];
+			break;
+		default:
+			check(false);
+			break;
+	}
+}
+
+void FMetalCommandEncoder::SetShaderSamplerStates(EShaderFrequency Frequency, const id<MTLSamplerState> Samplers[], NSRange const& Range)
+{
+	for (NSUInteger i = (NSUInteger)Range.location; i < (NSUInteger)Range.length; i++)
+	{
+		check(i < ML_MaxSamplers);
+		[Samplers[i] retain];
+		[ShaderSamplers[Frequency].Samplers[i] release];
+		
+		if(Samplers[i])
+		{
+			ShaderSamplers[Frequency].Bound |= (1 << i);
+		}
+		else
+		{
+			ShaderSamplers[Frequency].Bound &= ~(1 << i);
+		}
+		
+		ShaderSamplers[Frequency].Samplers[i] = Samplers[i];
+		ShaderSamplers[Frequency].MinLods[i] = 0;
+		ShaderSamplers[Frequency].MaxLods[i] = FLT_MAX;
+	}
+	
+	switch (Frequency)
+	{
+		case SF_Vertex:
+			[RenderCommandEncoder setVertexSamplerStates:Samplers withRange:Range];
+			break;
+		case SF_Pixel:
+			[RenderCommandEncoder setFragmentSamplerStates:Samplers withRange:Range];
+			break;
+		case SF_Compute:
+			checkf(ComputeCommandEncoder != nil, TEXT("Attempted to use ComputeCommandEncoder before calling FMetalContext::ConditionalSwitchToCompute(void) to create the encoder"));
+			[ComputeCommandEncoder setSamplerStates:Samplers withRange:Range];
+			break;
+		default:
+			check(false);
+			break;
+	}
+}
+
+void FMetalCommandEncoder::SetShaderSamplerState(EShaderFrequency Frequency, id<MTLSamplerState> Sampler, float lodMinClamp, float lodMaxClamp, NSUInteger index)
+{
+	check(index < ML_MaxSamplers);
+	[Sampler retain];
+	[ShaderSamplers[Frequency].Samplers[index] release];
+
+	if(Sampler)
+	{
+		ShaderSamplers[Frequency].Bound |= (1 << index);
+	}
+	else
+	{
+		ShaderSamplers[Frequency].Bound &= ~(1 << index);
+	}
+	
+	ShaderSamplers[Frequency].Samplers[index] = Sampler;
+	ShaderSamplers[Frequency].MinLods[index] = lodMinClamp;
+	ShaderSamplers[Frequency].MaxLods[index] = lodMaxClamp;
+	
+	switch (Frequency)
+	{
+		case SF_Vertex:
+			[RenderCommandEncoder setVertexSamplerState:Sampler lodMinClamp:lodMinClamp lodMaxClamp:lodMaxClamp atIndex:index];
+			break;
+		case SF_Pixel:
+			[RenderCommandEncoder setFragmentSamplerState:Sampler lodMinClamp:lodMinClamp lodMaxClamp:lodMaxClamp atIndex:index];
+			break;
+		case SF_Compute:
+			checkf(ComputeCommandEncoder != nil, TEXT("Attempted to use ComputeCommandEncoder before calling FMetalContext::ConditionalSwitchToCompute(void) to create the encoder"));
+			[ComputeCommandEncoder setSamplerState:Sampler lodMinClamp:lodMinClamp lodMaxClamp:lodMaxClamp atIndex:index];
+			break;
+		default:
+			check(false);
+			break;
+	}
+}
+
+void FMetalCommandEncoder::SetShaderSamplerStates(EShaderFrequency Frequency, const id<MTLSamplerState> Samplers[], const float LodMinClamps[], const float LodMaxClamps[], NSRange const& Range)
+{
+	for (NSUInteger i = (NSUInteger)Range.location; i < (NSUInteger)Range.length; i++)
+	{
+		check(i < ML_MaxSamplers);
+		[Samplers[i] retain];
+		
+		if(Samplers[i])
+		{
+			ShaderSamplers[Frequency].Bound |= (1 << i);
+		}
+		else
+		{
+			ShaderSamplers[Frequency].Bound &= ~(1 << i);
+		}
+		
+		[ShaderSamplers[Frequency].Samplers[i] release];
+		ShaderSamplers[Frequency].Samplers[i] = Samplers[i];
+		ShaderSamplers[Frequency].MinLods[i] = LodMinClamps[i];
+		ShaderSamplers[Frequency].MaxLods[i] = LodMaxClamps[i];
+	}
+	switch (Frequency)
+	{
+		case SF_Vertex:
+			[RenderCommandEncoder setVertexSamplerStates:Samplers lodMinClamps:LodMinClamps lodMaxClamps:LodMaxClamps withRange:Range];
+			break;
+		case SF_Pixel:
+			[RenderCommandEncoder setFragmentSamplerStates:Samplers lodMinClamps:LodMinClamps lodMaxClamps:LodMaxClamps withRange:Range];
+			break;
+		case SF_Compute:
+			checkf(ComputeCommandEncoder != nil, TEXT("Attempted to use ComputeCommandEncoder before calling FMetalContext::ConditionalSwitchToCompute(void) to create the encoder"));
+			[ComputeCommandEncoder setSamplerStates:Samplers lodMinClamps:LodMinClamps lodMaxClamps:LodMaxClamps withRange:Range];
+			break;
+		default:
+			check(false);
+			break;
+	}
+}
+
+#pragma mark - Public Compute State Mutators -
+
+void FMetalCommandEncoder::SetComputePipelineState(id<MTLComputePipelineState> const State)
+{
+	checkf(ComputeCommandEncoder != nil, TEXT("Attempted to use ComputeCommandEncoder before calling FMetalContext::ConditionalSwitchToCompute(void) to create the encoder"));
+	[State retain];
+	[ComputePipelineState release];
+	ComputePipelineState = State;
+	[ComputeCommandEncoder setComputePipelineState:State];
+}
+
+#pragma mark - Public Extension Accessors -
+
+#pragma mark - Public Extension Mutators -
+	
+#pragma mark - Public Support Functions -
+
+void FMetalCommandEncoder::UnbindObject(id const Object)
+{
+	for(uint32 i = 0; i < SF_NumFrequencies; i++)
+	{
+		for(uint32 Index = 0; Index < ML_MaxBuffers; Index++)
+		{
+			if(ShaderBuffers[i].Buffers[Index] == Object)
+			{
+				ShaderBuffers[i].Buffers[Index] = nullptr;
+				ShaderBuffers[i].Bound &= ~(1 << Index);
+			}
+		}
+		for(uint32 Index = 0; Index < ML_MaxTextures; Index++)
+		{
+			if(ShaderTextures[i].Textures[Index] == Object)
+			{
+				ShaderTextures[i].Textures[Index] = nullptr;
+				ShaderTextures[i].Bound &= ~(1 << Index);
+			}
+		}
+		for(uint32 Index = 0; Index < ML_MaxSamplers; Index++)
+		{
+			if(ShaderSamplers[i].Samplers[Index] == Object)
+			{
+				ShaderSamplers[i].Samplers[Index] = nullptr;
+				ShaderSamplers[i].Bound &= ~(1 << Index);
+			}
+		}
+	}
+}
diff --git a/Engine/Source/Runtime/Apple/MetalRHI/Private/MetalCommandEncoder.h b/Engine/Source/Runtime/Apple/MetalRHI/Private/MetalCommandEncoder.h
new file mode 100644
--- /dev/null
+++ b/Engine/Source/Runtime/Apple/MetalRHI/Private/MetalCommandEncoder.h
@@ -0,0 +1,385 @@
+// Copyright 1998-2015 Epic Games, Inc. All Rights Reserved.
+
+#pragma once
+
+#include <Metal/Metal.h>
+
+/**
+ * FMetalCommandEncoder:
+ *	Wraps the details of switching between different command encoders on the command-buffer, allowing for restoration of the render encoder if needed.
+ * 	UE4 expects the API to serialise commands in-order, but Metal expects applications to work with command-buffers directly so we need to implement 
+ *	the RHI semantics by switching between encoder types. This class hides the ugly details. Eventually it might be possible to move some of the operations
+ *	into pre- & post- command-buffers so that we avoid encoder switches but that will take changes to the RHI and high-level code too, so it won't happen soon.
+ */
+class FMetalCommandEncoder
+{
+public:
+#pragma mark - Public C++ Boilerplate -
+
+	/** Default constructor */
+	FMetalCommandEncoder(void);
+	
+	/** Destructor */
+	~FMetalCommandEncoder(void);
+	
+#pragma mark - Public Command Buffer Mutators -
+
+	/**
+	 * Start encoding to CommandBuffer. It is an error to call this with any outstanding command encoders or current command buffer.
+	 * Instead call EndEncoding & CommitCommandBuffer before calling this.
+	 * @param CommandBuffer The new command buffer to begin encoding to.
+	 */
+	void StartCommandBuffer(id<MTLCommandBuffer> const CommandBuffer);
+	
+	/**
+	 * Commit the existing command buffer if there is one & optionally waiting for completion, if there isn't a current command buffer this is a no-op.
+	 * @param bWait If true will wait for command buffer completion, otherwise the function returns immediately.
+ 	 */
+	void CommitCommandBuffer(bool const bWait);
+
+#pragma mark - Public Command Encoder Accessors -
+	
+	/** @returns True if and only if there is an active render command encoder, otherwise false. */
+	bool IsRenderCommandEncoderActive(void) const;
+	
+	/** @returns True if and only if there is an active compute command encoder, otherwise false. */
+	bool IsComputeCommandEncoderActive(void) const;
+	
+	/** @returns True if and only if there is an active blit command encoder, otherwise false. */
+	bool IsBlitCommandEncoderActive(void) const;
+
+	/** @returns The active render command encoder or nil if there isn't one. */
+	id<MTLRenderCommandEncoder> GetRenderCommandEncoder(void) const;
+	
+	/** @returns The active compute command encoder or nil if there isn't one. */
+	id<MTLComputeCommandEncoder> GetComputeCommandEncoder(void) const;
+	
+	/** @returns The active blit command encoder or nil if there isn't one. */
+	id<MTLBlitCommandEncoder> GetBlitCommandEncoder(void) const;
+	
+#pragma mark - Public Command Encoder Mutators -
+
+	/**
+ 	 * Begins encoding rendering commands into the current command buffer. No other encoder may be active.
+ 	 * @param RenderPass The descriptor for the new render pass.
+	 */
+	void BeginRenderCommandEncoding(MTLRenderPassDescriptor* const RenderPass);
+	
+	/** Restore the render command encoder from whatever state we were in. Any previous encoder must first be terminated with EndEncoding. */
+	void RestoreRenderCommandEncoding(void);
+	
+	/** Restores the render command state into a new render command encoder. */		
+	void RestoreRenderCommandEncodingState(void);
+	
+	/** Begins encoding compute commands into the current command buffer. No other encoder may be active. */
+	void BeginComputeCommandEncoding(void);
+	
+	/** Begins encoding blit commands into the current command buffer. No other encoder may be active. */
+	void BeginBlitCommandEncoding(void);
+	
+	/** Declare that all command generation from this encoder is complete, and detach from the MTLCommandBuffer if there is an encoder active or does nothing if there isn't. */
+	void EndEncoding(void);
+
+#pragma mark - Public Debug Support -
+	
+	/*
+	 * Inserts a debug string into the command buffer.  This does not change any API behavior, but can be useful when debugging.
+	 * @param string The name of the signpost. 
+	 */
+	void InsertDebugSignpost(NSString* const String);
+	
+	/*
+	 * Push a new named string onto a stack of string labels.
+	 * @param string The name of the debug group. 
+	 */
+	void PushDebugGroup(NSString* const String);
+	
+	/* Pop the latest named string off of the stack. */
+	void PopDebugGroup(void);
+	
+#pragma mark - Public Render State Mutators -
+
+	/*
+	 * Sets the current render pipeline state object.
+	 * @param PipelineState The pipeline state to set. Must not be nil.
+	 */
+	void SetRenderPipelineState(id<MTLRenderPipelineState> const PipelineState);
+	
+	/*
+	 * Set the viewport, which is used to transform vertexes from normalized device coordinates to window coordinates.  Fragments that lie outside of the viewport are clipped, and optionally clamped for fragments outside of znear/zfar.
+	 * @param Viewport The viewport dimensions to use.
+	 */
+	void SetViewport(MTLViewport const& Viewport);
+	
+	/*
+	 * The winding order of front-facing primitives.
+	 * @param FrontFacingWinding The front face winding.
+	 */
+	void SetFrontFacingWinding(MTLWinding const FrontFacingWinding);
+	
+	/*
+	 * Controls if primitives are culled when front facing, back facing, or not culled at all.
+	 * @param CullMode The cull mode.
+	 */
+	void SetCullMode(MTLCullMode const CullMode);
+
+#if METAL_API_1_1
+	/*
+	 * Controls what is done with fragments outside of the near or far planes.
+	 * @param DepthClipMode the clip mode.
+	 */
+	void SetDepthClipMode(MTLDepthClipMode const DepthClipMode);
+#endif
+	
+	/*
+	 * Depth Bias.
+	 * @param DepthBias The depth-bias value.
+	 * @param SlopeScale The slope-scale to apply.
+	 * @param Clamp The value to clamp to.
+	 */
+	void SetDepthBias(float const DepthBias, float const SlopeScale, float const Clamp);
+	
+	/*
+	 * Specifies a rectangle for a fragment scissor test.  All fragments outside of this rectangle are discarded.
+	 * @param Rect The scissor rect dimensions.
+	 */
+	void SetScissorRect(MTLScissorRect const& Rect);
+	
+	/*
+	 * Set how to rasterize triangle and triangle strip primitives.
+	 * @param FillMode The fill mode.
+	 */
+	void SetTriangleFillMode(MTLTriangleFillMode const FillMode);
+	
+	/*
+	 * Set the constant blend color used across all blending on all render targets
+	 * @param Red The value for the red channel in 0-1.
+	 * @param Green The value for the green channel in 0-1.
+	 * @param Blue The value for the blue channel in 0-1.
+	 * @param Alpha The value for the alpha channel in 0-1.
+	 */
+	void SetBlendColor(float const Red, float const Green, float const Blue, float const Alpha);
+	
+	/*
+	 * Set the DepthStencil state object.
+	 * @param DepthStencilState The depth-stencil state, must not be nil.
+	 */
+	void SetDepthStencilState(id<MTLDepthStencilState> const DepthStencilState);
+	
+	/*
+	 * Set the stencil reference value for both the back and front stencil buffers.
+	 * @param ReferenceValue The stencil ref value to use.
+	 */
+	void SetStencilReferenceValue(uint32 const ReferenceValue);
+
+	/*
+	 * Set the stencil reference value for the back and front stencil buffers independently.
+	 * @param FrontReferenceValue The front face stencil ref value.
+	 * @param BackReferenceValue The back face stencil ref value.
+	 */
+	void SetStencilReferenceValue(uint32 const FrontReferenceValue, uint32 const BackReferenceValue);
+	
+	/*
+	 * Monitor if samples pass the depth and stencil tests.
+	 * @param Mode Controls if the counter is disabled or moniters passing samples.
+	 * @param Offset The offset relative to the occlusion query buffer provided when the command encoder was created.  offset must be a multiple of 8.
+	 */
+	void SetVisibilityResultMode(MTLVisibilityResultMode const Mode, NSUInteger const Offset);
+	
+#pragma mark - Public Shader Resource Mutators -
+	
+	/*
+	 * Set a global buffer for all fragment shaders at the given bind point index.
+	 * @param Frequency The shader frequency to modify.
+	 * @param Buffer The buffer to bind or nil to clear.
+	 * @param Offset The offset in the buffer or 0 when Buffer is nil.
+	 * @param Index The index to modify.
+	 */
+	void SetShaderBuffer(EShaderFrequency const Frequency, id<MTLBuffer> const Buffer, NSUInteger const Offset, NSUInteger const Index);
+	
+	/*
+	 * Set the offset within the current global buffer for all fragment shaders at the given bind point index.
+	 * @param Frequency The shader frequency to modify.
+	 * @param Offset The offset in the buffer or 0 when Buffer is nil.
+	 * @param Index The index to modify.
+	 */
+	void SetShaderBufferOffset(EShaderFrequency const Frequency, NSUInteger const Offset, NSUInteger const Index);
+	
+	/*
+	 * Set an array of global buffers for all fragment shaders with the given bind point range.
+	 * @param Frequency The shader frequency to modify.
+	 * @param Buffer sThe buffers to bind or nil to clear.
+	 * @param Offset The offset in the buffer or 0 when Buffer is nil.
+	 * @param Range The start point and number of indices to modify.
+	 */
+	void SetShaderBuffers(EShaderFrequency const Frequency, const id<MTLBuffer> Buffers[], const NSUInteger Offset[], NSRange const& Range);
+	
+	/*
+	 * Set a global texture for all fragment shaders at the given bind point index.
+	 * @param Frequency The shader frequency to modify.
+	 * @param Texture The texture to bind or nil to clear.
+	 * @param Index The index to modify.
+	 */
+	void SetShaderTexture(EShaderFrequency const Frequency, id<MTLTexture> const Texture, NSUInteger const Index);
+	
+	/*
+	 * Set an array of global textures for all fragment shaders with the given bind point range.
+	 * @param Frequency The shader frequency to modify.
+	 * @param Textures The textures to bind or nil to clear.
+	 * @param Range The start point and number of indices to modify.
+	 */
+	void SetShaderTextures(EShaderFrequency const Frequency, const id<MTLTexture> Textures[], NSRange const& Range);
+	
+	/*
+	 * Set a global sampler for all fragment shaders at the given bind point index.
+	 * @param Frequency The shader frequency to modify.
+	 * @param Sampler The sampler state to bind or nil to clear.
+	 * @param Index The index to modify.
+	 */
+	void SetShaderSamplerState(EShaderFrequency const Frequency, id<MTLSamplerState> const Sampler, NSUInteger const Index);
+	
+	/*
+	 * Set an array of global samplers for all fragment shaders with the given bind point range.
+	 * @param Frequency The shader frequency to modify.
+	 * @param Samplers The sampler states to bind or nil to clear.
+	 * @param Range The start point and number of indices to modify.
+	 */
+	void SetShaderSamplerStates(EShaderFrequency const Frequency, const id<MTLSamplerState> Samplers[], NSRange const& Range);
+	
+	/*
+	 * Set a global sampler for all fragment shaders at the given bind point index.
+	 * @param Frequency The shader frequency to modify.
+	 * @param Sampler The sampler state to bind or nil to clear.
+	 * @param LodMinClamp The minimum lod value or 0.
+	 * @param LodMaxClamp The maximum lod value or FLT_MAX.
+	 * @param Index The index to modify.
+	 */
+	void SetShaderSamplerState(EShaderFrequency const Frequency, id<MTLSamplerState> const Sampler, float LodMinClamp, float LodMaxClamp, NSUInteger const Index);
+	
+	/*
+	 * Set an array of global samplers for all fragment shaders with the given bind point range.
+	 * @param Frequency The shader frequency to modify.
+	 * @param Samplers The sampler states to bind or nil to clear.
+	 * @param LodMinClamps The minimum lod values or 0.
+	 * @param LodMaxClamps The maximum lod values or FLT_MAX.
+	 * @param Range The start point and number of indices to modify.
+	 */
+	void SetShaderSamplerStates(EShaderFrequency const Frequency, const id<MTLSamplerState> Samplers[], const float LodMinClamps[], const float LodMaxClamps[], NSRange const& Range);
+	
+#pragma mark - Public Compute State Mutators -
+	
+	/*
+	 * Set the compute pipeline state that will be used.
+	 * @param State The state to set - must not be nil.
+	 */
+	void SetComputePipelineState(id<MTLComputePipelineState> const State);
+
+#pragma mark - Public Extension Accessors -
+	
+#pragma mark - Public Extension Mutators -
+	
+#pragma mark - Public Support Functions -
+
+	/*
+	 * Unbinds Object from the cached state so that it cannot be restored accidentally.
+	 * @param Object The object to remove from the state caching.
+	 */
+	void UnbindObject(id const Object);
+	
+private:
+#pragma mark - Private Per-Platform Defines -
+
+#if PLATFORM_IOS
+	#define METAL_MAX_TEXTURES 31
+	typedef uint32 FMetalTextureMask;
+#elif PLATFORM_MAC
+	#define METAL_MAX_TEXTURES 128
+	typedef __uint128_t FMetalTextureMask;
+#else
+	#error "Unsupported Platform!"
+#endif
+
+#pragma mark - Private Type Declarations -
+
+	/**
+	 * The sampler, buffer and texture resource limits as defined here:
+	 * https://developer.apple.com/library/ios/documentation/Miscellaneous/Conceptual/MetalProgrammingGuide/Render-Ctx/Render-Ctx.html
+	 */
+	enum EMetalLimits
+	{
+		ML_MaxSamplers = 16, /** Maximum number of samplers */
+		ML_MaxBuffers = 31, /** Maximum number of buffers */
+		ML_MaxTextures = METAL_MAX_TEXTURES /** Maximum number of textures - there are more textures available on Mac than iOS */
+	};
+	
+	/** A structure of arrays for the current buffer binding settings. */
+	struct FMetalBufferBindings
+	{
+		/** The bound buffers or nil. */
+		id<MTLBuffer> Buffers[ML_MaxBuffers];
+		/** The bound buffer offsets or 0. */
+		NSUInteger Offsets[ML_MaxBuffers];
+		/** A bitmask for which buffers were bound by the application where a bit value of 1 is bound and 0 is unbound. */
+        uint32 Bound;
+	};
+
+	/** A structure of arrays for the current texture binding settings. */
+	struct FMetalTextureBindings
+	{
+		/** The bound textures or nil. */
+		id<MTLTexture> Textures[ML_MaxTextures];
+		/** A bitmask for which textures were bound by the application where a bit value of 1 is bound and 0 is unbound. */
+		FMetalTextureMask Bound;
+	};
+
+	/** A structure of arrays for the current sampler binding settings. */
+	struct FMetalSamplerBindings
+	{
+		/** The bound sampler states or nil. */
+		id<MTLSamplerState> Samplers[ML_MaxSamplers];
+		/** The bound min. lods or 0. */
+		float MinLods[ML_MaxSamplers];
+		/** The bound max. lods or FLT_MAX. */
+		float MaxLods[ML_MaxSamplers];
+		/** A bitmask for which samplers were bound by the application where a bit value of 1 is bound and 0 is unbound. */
+		uint16 Bound;
+	};
+	
+#pragma mark - Private Member Variables -
+
+	FMetalBufferBindings ShaderBuffers[SF_NumFrequencies];
+	FMetalTextureBindings ShaderTextures[SF_NumFrequencies];
+	FMetalSamplerBindings ShaderSamplers[SF_NumFrequencies];
+	
+	MTLViewport Viewport;
+	MTLWinding FrontFacingWinding;
+	MTLCullMode CullMode;
+#if METAL_API_1_1
+	MTLDepthClipMode DepthClipMode;
+#endif
+	float DepthBias[3];
+	MTLScissorRect ScissorRect;
+	MTLTriangleFillMode FillMode;
+	float BlendColor[4];
+	
+	id<MTLDepthStencilState> DepthStencilState;
+	uint32 StencilRef[2];
+	
+	MTLVisibilityResultMode VisibilityMode;
+	NSUInteger VisibilityOffset;
+	
+	id<MTLBuffer> PipelineStatsBuffer;
+	NSUInteger PipelineStatsOffset;
+	NSUInteger PipelineStatsMask;
+	
+	MTLRenderPassDescriptor* RenderPassDesc;
+	
+	id<MTLCommandBuffer> CommandBuffer;
+	id<MTLRenderCommandEncoder> RenderCommandEncoder;
+	id<MTLRenderPipelineState> RenderPipelineState;
+	id<MTLComputeCommandEncoder> ComputeCommandEncoder;
+	id<MTLComputePipelineState> ComputePipelineState;
+	id<MTLBlitCommandEncoder> BlitCommandEncoder;
+	
+	NSMutableArray* DebugGroups;
+};
\ No newline at end of file
diff --git a/Engine/Source/Runtime/Apple/MetalRHI/Private/MetalCommandQueue.cpp b/Engine/Source/Runtime/Apple/MetalRHI/Private/MetalCommandQueue.cpp
new file mode 100644
--- /dev/null
+++ b/Engine/Source/Runtime/Apple/MetalRHI/Private/MetalCommandQueue.cpp
@@ -0,0 +1,91 @@
+// Copyright 1998-2015 Epic Games, Inc. All Rights Reserved.
+
+/*=============================================================================
+	MetalCommandQueue.cpp: Metal command queue wrapper.
+=============================================================================*/
+
+#include "MetalRHIPrivate.h"
+
+#include "MetalCommandQueue.h"
+#if METAL_STATISTICS
+#include "Runtime/Mac/NoRedist/MetalStatistics/Public/MetalStatistics.h"
+#include "ModuleManager.h"
+#endif
+
+#pragma mark - Public C++ Boilerplate -
+
+FMetalCommandQueue::FMetalCommandQueue(id<MTLDevice> Device, uint32 const MaxNumCommandBuffers /* = 0 */)
+: CommandQueue(nil)
+#if METAL_STATISTICS
+, Statistics(nullptr)
+#endif
+{
+	if(MaxNumCommandBuffers == 0)
+	{
+		CommandQueue = [Device newCommandQueue];
+	}
+	else
+	{
+		CommandQueue = [Device newCommandQueueWithMaxCommandBufferCount: MaxNumCommandBuffers];
+	}
+	check(CommandQueue);
+
+#if METAL_STATISTICS
+	IMetalStatisticsModule* StatsModule = FModuleManager::Get().LoadModulePtr<IMetalStatisticsModule>(TEXT("MetalStatistics"));
+	
+	if(StatsModule)
+	{
+		Statistics = StatsModule->CreateMetalStatistics(CommandQueue);
+		if(!Statistics->SupportsStatistics())
+		{
+			delete Statistics;
+			Statistics = nullptr;
+		}
+	}
+#endif
+}
+
+FMetalCommandQueue::~FMetalCommandQueue(void)
+{
+#if METAL_STATISTICS
+	delete Statistics;
+#endif
+	
+	[CommandQueue release];
+	CommandQueue = nil;
+}
+	
+#pragma mark - Public Command Buffer Mutators -
+
+id<MTLCommandBuffer> FMetalCommandQueue::CreateRetainedCommandBuffer(void)
+{
+	return [CommandQueue commandBuffer];
+}
+
+id<MTLCommandBuffer> FMetalCommandQueue::CreateUnretainedCommandBuffer(void)
+{
+	return [CommandQueue commandBufferWithUnretainedReferences];
+}
+
+#pragma mark - Public Command Queue Accessors -
+	
+id<MTLCommandQueue> FMetalCommandQueue::GetCommandQueue(void) const
+{
+	return CommandQueue;
+}
+
+#pragma mark - Public Debug Support -
+
+void FMetalCommandQueue::InsertDebugCaptureBoundary(void)
+{
+	[CommandQueue insertDebugCaptureBoundary];
+}
+
+#if METAL_STATISTICS
+#pragma mark - Public Statistics Extensions -
+
+IMetalStatistics* FMetalCommandQueue::GetStatistics(void)
+{
+	return Statistics;
+}
+#endif
diff --git a/Engine/Source/Runtime/Apple/MetalRHI/Private/MetalCommandQueue.h b/Engine/Source/Runtime/Apple/MetalRHI/Private/MetalCommandQueue.h
new file mode 100644
--- /dev/null
+++ b/Engine/Source/Runtime/Apple/MetalRHI/Private/MetalCommandQueue.h
@@ -0,0 +1,70 @@
+// Copyright 1998-2015 Epic Games, Inc. All Rights Reserved.
+
+#pragma once
+
+#include <Metal/Metal.h>
+
+/**
+ * FMetalCommandQueue:
+ */
+class FMetalCommandQueue
+{
+public:
+#pragma mark - Public C++ Boilerplate -
+
+	/**
+	 * Constructor
+	 * @param Device The Metal device to create on.
+	 * @param MaxNumCommandBuffers The maximum number of incomplete command-buffers, defaults to 0 which implies the system default.
+	 */
+	FMetalCommandQueue(id<MTLDevice> Device, uint32 const MaxNumCommandBuffers = 0);
+	
+	/** Destructor */
+	~FMetalCommandQueue(void);
+	
+#pragma mark - Public Command Buffer Mutators -
+
+	/**
+	 * Start encoding to CommandBuffer. It is an error to call this with any outstanding command encoders or current command buffer.
+	 * Instead call EndEncoding & CommitCommandBuffer before calling this.
+	 * @param CommandBuffer The new command buffer to begin encoding to.
+	 */
+	id<MTLCommandBuffer> CreateRetainedCommandBuffer(void);
+
+	/**
+	 * Start encoding to CommandBuffer. It is an error to call this with any outstanding command encoders or current command buffer.
+	 * Instead call EndEncoding & CommitCommandBuffer before calling this.
+	 * @param CommandBuffer The new command buffer to begin encoding to.
+	 */
+	id<MTLCommandBuffer> CreateUnretainedCommandBuffer(void);
+	
+	/**
+	 * Commit the supplied command buffer.
+	 * @param CommandBuffer The command buffer to commit, must be non-nil.
+ 	 */
+	void CommitCommandBuffer(id<MTLCommandBuffer> const CommandBuffer);
+
+#pragma mark - Public Command Queue Accessors -
+	
+	/** @returns The native command queue. */
+	id<MTLCommandQueue> GetCommandQueue(void) const;
+
+#pragma mark - Public Debug Support -
+
+	/** Inserts a boundary that marks the end of a frame for the debug capture tool. */
+	void InsertDebugCaptureBoundary(void);
+
+#if METAL_STATISTICS
+#pragma mark - Public Statistics Extensions -
+
+	/** @returns An object that provides Metal statistics information or nullptr. */
+	class IMetalStatistics* GetStatistics(void);
+#endif
+	
+private:
+#pragma mark - Private Member Variables -
+	id<MTLCommandQueue> CommandQueue;
+#if METAL_STATISTICS
+	class IMetalStatistics* Statistics;
+#endif
+};
\ No newline at end of file
