Merging MetalRHI & MetalShaderFormat changes from ZebraBranch, plus the code to enable or disable Mac support.

diff --git a/Engine/Source/Runtime/Apple/MetalRHI/Private/MetalCommands.cpp b/Engine/Source/Runtime/Apple/MetalRHI/Private/MetalCommands.cpp
new file mode 100644
--- /dev/null
+++ b/Engine/Source/Runtime/Apple/MetalRHI/Private/MetalCommands.cpp
@@ -0,0 +1,1076 @@
+// Copyright 1998-2015 Epic Games, Inc. All Rights Reserved.
+
+/*=============================================================================
+	MetalCommands.cpp: Metal RHI commands implementation.
+=============================================================================*/
+
+#include "MetalRHIPrivate.h"
+
+#include "GlobalShader.h"
+#include "OneColorShader.h"
+#include "RHICommandList.h"
+#include "RHIStaticStates.h"
+#include "ShaderParameterUtils.h"
+#include "SceneUtils.h"
+#include "ShaderCache.h"
+#include "MetalProfiler.h"
+
+static const bool GUsesInvertedZ = true;
+static FGlobalBoundShaderState GClearMRTBoundShaderState[8];
+
+/** Vertex declaration for just one FVector4 position. */
+class FVector4VertexDeclaration : public FRenderResource
+{
+public:
+	FVertexDeclarationRHIRef VertexDeclarationRHI;
+	virtual void InitRHI() override
+	{
+		FVertexDeclarationElementList Elements;
+		Elements.Add(FVertexElement(0, 0, VET_Float4, 0, sizeof(FVector4)));
+		VertexDeclarationRHI = RHICreateVertexDeclaration(Elements);
+	}
+	virtual void ReleaseRHI() override
+	{
+		VertexDeclarationRHI.SafeRelease();
+	}
+};
+static TGlobalResource<FVector4VertexDeclaration> GVector4VertexDeclaration;
+
+static MTLPrimitiveType TranslatePrimitiveType(uint32 PrimitiveType)
+{
+	switch (PrimitiveType)
+	{
+		case PT_TriangleList:	return MTLPrimitiveTypeTriangle;
+		case PT_TriangleStrip:	return MTLPrimitiveTypeTriangleStrip;
+		case PT_LineList:		return MTLPrimitiveTypeLine;
+		case PT_PointList:		return MTLPrimitiveTypePoint;
+		default:				UE_LOG(LogMetal, Fatal, TEXT("Unsupported primitive type %d"), (int32)PrimitiveType); return MTLPrimitiveTypeTriangle;
+	}
+}
+
+void FMetalRHICommandContext::RHISetStreamSource(uint32 StreamIndex,FVertexBufferRHIParamRef VertexBufferRHI,uint32 Stride,uint32 Offset)
+{
+	FMetalVertexBuffer* VertexBuffer = ResourceCast(VertexBufferRHI);
+	Context->GetCurrentState().SetVertexBuffer(UNREAL_TO_METAL_BUFFER_INDEX(StreamIndex), VertexBuffer ? VertexBuffer->Buffer : nil, Stride, Offset);
+}
+
+void FMetalDynamicRHI::RHISetStreamOutTargets(uint32 NumTargets, const FVertexBufferRHIParamRef* VertexBuffers, const uint32* Offsets)
+{
+	NOT_SUPPORTED("RHISetStreamOutTargets");
+
+}
+
+void FMetalRHICommandContext::RHISetRasterizerState(FRasterizerStateRHIParamRef NewStateRHI)
+{
+	FMetalRasterizerState* NewState = ResourceCast(NewStateRHI);
+
+	Context->GetCurrentState().SetRasterizerState(NewState);
+	
+	FShaderCache::SetRasterizerState(NewStateRHI);
+}
+
+void FMetalRHICommandContext::RHISetComputeShader(FComputeShaderRHIParamRef ComputeShaderRHI)
+{
+	FMetalComputeShader* ComputeShader = ResourceCast(ComputeShaderRHI);
+
+	// active compute mode
+	Context->GetCurrentState().ConditionalSwitchToCompute();
+	
+	// cache this for Dispatch
+	// sets this compute shader pipeline as the current (this resets all state, so we need to set all resources after calling this)
+	Context->GetCurrentState().SetComputeShader(ComputeShader);
+}
+
+void FMetalRHICommandContext::RHIDispatchComputeShader(uint32 ThreadGroupCountX, uint32 ThreadGroupCountY, uint32 ThreadGroupCountZ)
+{
+	RHI_PROFILE_DRAW_CALL_STATS(EMTLSamplePointBeforeCompute, EMTLSamplePointAfterCompute, 1,1);
+	ThreadGroupCountX = FMath::Max(ThreadGroupCountX, 1u);
+	ThreadGroupCountY = FMath::Max(ThreadGroupCountY, 1u);
+	ThreadGroupCountZ = FMath::Max(ThreadGroupCountZ, 1u);
+	Context->Dispatch(ThreadGroupCountX, ThreadGroupCountY, ThreadGroupCountZ);
+}
+
+void FMetalRHICommandContext::RHIDispatchIndirectComputeShader(FVertexBufferRHIParamRef ArgumentBufferRHI, uint32 ArgumentOffset)
+{ 
+	NOT_SUPPORTED("RHIDispatchIndirectComputeShader");
+}
+
+void FMetalRHICommandContext::RHISetViewport(uint32 MinX,uint32 MinY,float MinZ,uint32 MaxX,uint32 MaxY,float MaxZ)
+{
+	MTLViewport Viewport;
+	Viewport.originX = MinX;
+	Viewport.originY = MinY;
+	Viewport.width = MaxX - MinX;
+	Viewport.height = MaxY - MinY;
+	Viewport.znear = MinZ;
+	Viewport.zfar = MaxZ;
+	
+	Context->GetCommandEncoder().SetViewport(Viewport);
+
+	Context->GetCurrentState().SetViewport(Viewport);
+	
+	FShaderCache::SetViewport(MinX,MinY, MinZ, MaxX, MaxY, MaxZ);
+}
+
+void FMetalRHICommandContext::RHISetMultipleViewports(uint32 Count, const FViewportBounds* Data)
+{ 
+	NOT_SUPPORTED("RHISetMultipleViewports");
+}
+
+void FMetalRHICommandContext::RHISetScissorRect(bool bEnable,uint32 MinX,uint32 MinY,uint32 MaxX,uint32 MaxY)
+{
+	MTLScissorRect Scissor;
+	Scissor.x = MinX;
+	Scissor.y = MinY;
+	Scissor.width = MaxX - MinX;
+	Scissor.height = MaxY - MinY;
+
+	// metal doesn't support 0 sized scissor rect
+	if (bEnable == false || Scissor.width == 0 || Scissor.height == 0)
+	{
+		auto Viewport = Context->GetCurrentState().GetViewport();
+		Scissor.x = Viewport.originX;
+		Scissor.y = Viewport.originY;
+		Scissor.width = Viewport.width;
+		Scissor.height = Viewport.height;
+	}
+	Context->GetCommandEncoder().SetScissorRect(Scissor);
+}
+
+void FMetalRHICommandContext::RHISetBoundShaderState( FBoundShaderStateRHIParamRef BoundShaderStateRHI)
+{
+	FMetalBoundShaderState* BoundShaderState = ResourceCast(BoundShaderStateRHI);
+
+	Context->GetCurrentState().SetBoundShaderState(BoundShaderState);
+
+	BoundShaderStateHistory.Add(BoundShaderState);
+	
+	FShaderCache::SetBoundShaderState(BoundShaderStateRHI);
+}
+
+
+void FMetalRHICommandContext::RHISetUAVParameter(FComputeShaderRHIParamRef ComputeShaderRHI, uint32 UAVIndex, FUnorderedAccessViewRHIParamRef UAVRHI)
+{
+	FMetalUnorderedAccessView* UAV = ResourceCast(UAVRHI);
+
+	if (UAV)
+	{
+		UAV->Set(Context, UAVIndex);
+	}
+}
+
+void FMetalRHICommandContext::RHISetUAVParameter(FComputeShaderRHIParamRef ComputeShaderRHI,uint32 UAVIndex,FUnorderedAccessViewRHIParamRef UAVRHI, uint32 InitialCount)
+{
+	NOT_SUPPORTED("RHISetUAVParameter");
+}
+
+
+void FMetalRHICommandContext::RHISetShaderTexture(FVertexShaderRHIParamRef VertexShaderRHI, uint32 TextureIndex, FTextureRHIParamRef NewTextureRHI)
+{
+	FMetalSurface* Surface = GetMetalSurfaceFromRHITexture(NewTextureRHI);
+	if (Surface != nullptr)
+	{
+		Context->GetCommandEncoder().SetShaderTexture(SF_Vertex, Surface->Texture, TextureIndex);
+	}
+	else
+	{
+		Context->GetCommandEncoder().SetShaderTexture(SF_Vertex, nil, TextureIndex);
+	}
+	
+	FShaderCache::SetTexture(SF_Vertex, TextureIndex, NewTextureRHI);
+}
+
+void FMetalRHICommandContext::RHISetShaderTexture(FHullShaderRHIParamRef HullShader, uint32 TextureIndex, FTextureRHIParamRef NewTextureRHI)
+{
+	NOT_SUPPORTED("RHISetShaderTexture-Hull");
+}
+
+void FMetalRHICommandContext::RHISetShaderTexture(FDomainShaderRHIParamRef DomainShader, uint32 TextureIndex, FTextureRHIParamRef NewTextureRHI)
+{
+	NOT_SUPPORTED("RHISetShaderTexture-Domain");
+}
+
+void FMetalRHICommandContext::RHISetShaderTexture(FGeometryShaderRHIParamRef GeometryShader, uint32 TextureIndex, FTextureRHIParamRef NewTextureRHI)
+{
+	NOT_SUPPORTED("RHISetShaderTexture-Geometry");
+
+}
+
+void FMetalRHICommandContext::RHISetShaderTexture(FPixelShaderRHIParamRef PixelShader, uint32 TextureIndex, FTextureRHIParamRef NewTextureRHI)
+{
+	FMetalSurface* Surface = GetMetalSurfaceFromRHITexture(NewTextureRHI);
+	if (Surface != nullptr)
+	{
+		Context->GetCommandEncoder().SetShaderTexture(SF_Pixel, Surface->Texture, TextureIndex);
+	}
+	else
+	{
+		Context->GetCommandEncoder().SetShaderTexture(SF_Pixel, nil, TextureIndex);
+	}
+	
+	FShaderCache::SetTexture(SF_Pixel, TextureIndex, NewTextureRHI);
+}
+
+void FMetalRHICommandContext::RHISetShaderTexture(FComputeShaderRHIParamRef ComputeShader, uint32 TextureIndex, FTextureRHIParamRef NewTextureRHI)
+{
+	FMetalSurface* Surface = GetMetalSurfaceFromRHITexture(NewTextureRHI);
+	if (Surface != nullptr)
+	{
+		Context->GetCommandEncoder().SetShaderTexture(SF_Compute, Surface->Texture, TextureIndex);
+	}
+	else
+	{
+		Context->GetCommandEncoder().SetShaderTexture(SF_Compute, nil, TextureIndex);
+	}
+	
+	FShaderCache::SetTexture(SF_Compute, TextureIndex, NewTextureRHI);
+}
+
+
+void FMetalRHICommandContext::RHISetShaderResourceViewParameter(FVertexShaderRHIParamRef VertexShaderRHI, uint32 TextureIndex, FShaderResourceViewRHIParamRef SRVRHI)
+{
+	FMetalShaderResourceView* SRV = ResourceCast(SRVRHI);
+	if (SRV)
+	{
+		FRHITexture* Texture = SRV->SourceTexture.GetReference();
+		if (Texture)
+		{
+			FMetalSurface* Surface = SRV->TextureView;
+			if (Surface != nullptr)
+			{
+				Context->GetCommandEncoder().SetShaderTexture(SF_Vertex, Surface->Texture, TextureIndex);
+			}
+			else
+			{
+				Context->GetCommandEncoder().SetShaderTexture(SF_Vertex, nil, TextureIndex);
+			}
+		}
+		else
+		{
+			FMetalVertexBuffer* VB = SRV->SourceVertexBuffer.GetReference();
+			if (VB)
+			{
+				Context->GetCommandEncoder().SetShaderBuffer(SF_Vertex, VB->Buffer, 0, TextureIndex);
+			}
+		}
+	}
+	else
+	{
+		Context->GetCommandEncoder().SetShaderTexture(SF_Vertex, nil, TextureIndex);
+	}
+	FShaderCache::SetSRV(SF_Vertex, TextureIndex, SRVRHI);
+}
+
+void FMetalRHICommandContext::RHISetShaderResourceViewParameter(FHullShaderRHIParamRef HullShaderRHI,uint32 TextureIndex,FShaderResourceViewRHIParamRef SRVRHI)
+{
+	NOT_SUPPORTED("RHISetShaderResourceViewParameter");
+}
+
+void FMetalRHICommandContext::RHISetShaderResourceViewParameter(FDomainShaderRHIParamRef DomainShaderRHI,uint32 TextureIndex,FShaderResourceViewRHIParamRef SRVRHI)
+{
+	NOT_SUPPORTED("RHISetShaderResourceViewParameter");
+}
+
+void FMetalRHICommandContext::RHISetShaderResourceViewParameter(FGeometryShaderRHIParamRef GeometryShaderRHI,uint32 TextureIndex,FShaderResourceViewRHIParamRef SRVRHI)
+{
+	NOT_SUPPORTED("RHISetShaderResourceViewParameter");
+}
+
+void FMetalRHICommandContext::RHISetShaderResourceViewParameter(FPixelShaderRHIParamRef PixelShaderRHI,uint32 TextureIndex,FShaderResourceViewRHIParamRef SRVRHI)
+{
+	FMetalShaderResourceView* SRV = ResourceCast(SRVRHI);
+	if (SRV)
+	{
+		FRHITexture* Texture = SRV->SourceTexture.GetReference();
+		if (Texture)
+		{
+			FMetalSurface* Surface = SRV->TextureView;
+			if (Surface != nullptr)
+			{
+				Context->GetCommandEncoder().SetShaderTexture(SF_Pixel, Surface->Texture, TextureIndex);
+			}
+			else
+			{
+				Context->GetCommandEncoder().SetShaderTexture(SF_Pixel, nil, TextureIndex);
+			}
+		}
+		else
+		{
+			FMetalVertexBuffer* VB = SRV->SourceVertexBuffer.GetReference();
+			if (VB)
+			{
+				Context->GetCommandEncoder().SetShaderBuffer(SF_Pixel, VB->Buffer, 0, TextureIndex);
+			}
+		}
+	}
+	else
+	{
+		Context->GetCommandEncoder().SetShaderTexture(SF_Pixel, nil, TextureIndex);
+	}
+	FShaderCache::SetSRV(SF_Pixel, TextureIndex, SRVRHI);
+}
+
+void FMetalRHICommandContext::RHISetShaderResourceViewParameter(FComputeShaderRHIParamRef ComputeShaderRHI,uint32 TextureIndex,FShaderResourceViewRHIParamRef SRVRHI)
+{
+	FMetalShaderResourceView* SRV = ResourceCast(SRVRHI);
+	if (SRV)
+	{
+		FRHITexture* Texture = SRV->SourceTexture.GetReference();
+		if (Texture)
+		{
+			FMetalSurface* Surface = SRV->TextureView;
+			if (Surface != nullptr)
+			{
+				Context->GetCommandEncoder().SetShaderTexture(SF_Compute, Surface->Texture, TextureIndex);
+			}
+			else
+			{
+				Context->GetCommandEncoder().SetShaderTexture(SF_Compute, nil, TextureIndex);
+			}
+		}
+		else
+		{
+			FMetalVertexBuffer* VB = SRV->SourceVertexBuffer.GetReference();
+			if (VB)
+			{
+				Context->GetCommandEncoder().SetShaderBuffer(SF_Compute, VB->Buffer, 0, TextureIndex);
+			}
+		}
+	}
+	else
+	{
+		Context->GetCommandEncoder().SetShaderTexture(SF_Compute, nil, TextureIndex);
+	}
+	FShaderCache::SetSRV(SF_Compute, TextureIndex, SRVRHI);
+}
+
+
+void FMetalRHICommandContext::RHISetShaderSampler(FVertexShaderRHIParamRef VertexShaderRHI, uint32 SamplerIndex, FSamplerStateRHIParamRef NewStateRHI)
+{
+	FMetalSamplerState* NewState = ResourceCast(NewStateRHI);
+
+	Context->GetCommandEncoder().SetShaderSamplerState(SF_Vertex, NewState->State, SamplerIndex);
+	
+	FShaderCache::SetSamplerState(SF_Vertex, SamplerIndex, NewStateRHI);
+}
+
+void FMetalRHICommandContext::RHISetShaderSampler(FHullShaderRHIParamRef HullShader, uint32 SamplerIndex, FSamplerStateRHIParamRef NewStateRHI)
+{
+	FMetalSamplerState* NewState = ResourceCast(NewStateRHI);
+
+	NOT_SUPPORTED("RHISetSamplerState-Hull");
+}
+
+void FMetalRHICommandContext::RHISetShaderSampler(FDomainShaderRHIParamRef DomainShader, uint32 SamplerIndex, FSamplerStateRHIParamRef NewStateRHI)
+{
+	FMetalSamplerState* NewState = ResourceCast(NewStateRHI);
+
+	NOT_SUPPORTED("RHISetSamplerState-Domain");
+}
+
+void FMetalRHICommandContext::RHISetShaderSampler(FGeometryShaderRHIParamRef GeometryShader, uint32 SamplerIndex, FSamplerStateRHIParamRef NewStateRHI)
+{
+	FMetalSamplerState* NewState = ResourceCast(NewStateRHI);
+
+	NOT_SUPPORTED("RHISetSamplerState-Geometry");
+}
+
+void FMetalRHICommandContext::RHISetShaderSampler(FPixelShaderRHIParamRef PixelShader, uint32 SamplerIndex, FSamplerStateRHIParamRef NewStateRHI)
+{
+	FMetalSamplerState* NewState = ResourceCast(NewStateRHI);
+
+	Context->GetCommandEncoder().SetShaderSamplerState(SF_Pixel, NewState->State, SamplerIndex);
+	
+	FShaderCache::SetSamplerState(SF_Pixel, SamplerIndex, NewStateRHI);
+}
+
+void FMetalRHICommandContext::RHISetShaderSampler(FComputeShaderRHIParamRef ComputeShader, uint32 SamplerIndex, FSamplerStateRHIParamRef NewStateRHI)
+{
+	FMetalSamplerState* NewState = ResourceCast(NewStateRHI);
+
+	Context->GetCommandEncoder().SetShaderSamplerState(SF_Compute, NewState->State, SamplerIndex);
+	
+	FShaderCache::SetSamplerState(SF_Compute, SamplerIndex, NewStateRHI);
+}
+
+void FMetalRHICommandContext::RHISetShaderParameter(FVertexShaderRHIParamRef VertexShaderRHI, uint32 BufferIndex, uint32 BaseIndex, uint32 NumBytes, const void* NewValue)
+{
+	Context->GetCurrentState().GetShaderParameters(CrossCompiler::SHADER_STAGE_VERTEX).Set(BufferIndex, BaseIndex, NumBytes, NewValue);
+}
+
+void FMetalRHICommandContext::RHISetShaderParameter(FHullShaderRHIParamRef HullShaderRHI, uint32 BufferIndex, uint32 BaseIndex, uint32 NumBytes, const void* NewValue)
+{
+	NOT_SUPPORTED("RHISetShaderParameter-Hull");
+}
+
+void FMetalRHICommandContext::RHISetShaderParameter(FPixelShaderRHIParamRef PixelShaderRHI, uint32 BufferIndex, uint32 BaseIndex, uint32 NumBytes, const void* NewValue)
+{
+	Context->GetCurrentState().GetShaderParameters(CrossCompiler::SHADER_STAGE_PIXEL).Set(BufferIndex, BaseIndex, NumBytes, NewValue);
+}
+
+void FMetalRHICommandContext::RHISetShaderParameter(FDomainShaderRHIParamRef DomainShaderRHI, uint32 BufferIndex, uint32 BaseIndex, uint32 NumBytes, const void* NewValue)
+{
+	NOT_SUPPORTED("RHISetShaderParameter-Domain");
+}
+
+void FMetalRHICommandContext::RHISetShaderParameter(FGeometryShaderRHIParamRef GeometryShaderRHI, uint32 BufferIndex, uint32 BaseIndex, uint32 NumBytes, const void* NewValue)
+{
+	NOT_SUPPORTED("RHISetShaderParameter-Geometry");
+}
+
+void FMetalRHICommandContext::RHISetShaderParameter(FComputeShaderRHIParamRef ComputeShaderRHI,uint32 BufferIndex, uint32 BaseIndex, uint32 NumBytes, const void* NewValue)
+{
+	Context->GetCurrentState().GetShaderParameters(CrossCompiler::SHADER_STAGE_COMPUTE).Set(BufferIndex, BaseIndex, NumBytes, NewValue);
+}
+
+void FMetalRHICommandContext::RHISetShaderUniformBuffer(FVertexShaderRHIParamRef VertexShaderRHI, uint32 BufferIndex, FUniformBufferRHIParamRef BufferRHI)
+{
+	FMetalVertexShader* VertexShader = ResourceCast(VertexShaderRHI);
+	VertexShader->BoundUniformBuffers[BufferIndex] = BufferRHI;
+	VertexShader->DirtyUniformBuffers |= 1 << BufferIndex;
+
+	auto& Bindings = VertexShader->Bindings;
+	check(BufferIndex < Bindings.NumUniformBuffers);
+	if (Bindings.bHasRegularUniformBuffers)
+	{
+		auto* UB = (FMetalUniformBuffer*)BufferRHI;
+		Context->GetCommandEncoder().SetShaderBuffer(SF_Vertex, UB->Buffer, UB->Offset, BufferIndex);
+	}
+}
+
+void FMetalRHICommandContext::RHISetShaderUniformBuffer(FHullShaderRHIParamRef HullShader, uint32 BufferIndex, FUniformBufferRHIParamRef BufferRHI)
+{
+	NOT_SUPPORTED("RHISetShaderUniformBuffer-Hull");
+}
+
+void FMetalRHICommandContext::RHISetShaderUniformBuffer(FDomainShaderRHIParamRef DomainShader, uint32 BufferIndex, FUniformBufferRHIParamRef BufferRHI)
+{
+	NOT_SUPPORTED("RHISetShaderUniformBuffer-Domain");
+}
+
+void FMetalRHICommandContext::RHISetShaderUniformBuffer(FGeometryShaderRHIParamRef GeometryShader, uint32 BufferIndex, FUniformBufferRHIParamRef BufferRHI)
+{
+	NOT_SUPPORTED("RHISetShaderUniformBuffer-Geometry");
+}
+
+void FMetalRHICommandContext::RHISetShaderUniformBuffer(FPixelShaderRHIParamRef PixelShaderRHI, uint32 BufferIndex, FUniformBufferRHIParamRef BufferRHI)
+{
+	FMetalPixelShader* PixelShader = ResourceCast(PixelShaderRHI);
+	PixelShader->BoundUniformBuffers[BufferIndex] = BufferRHI;
+	PixelShader->DirtyUniformBuffers |= 1 << BufferIndex;
+
+	auto& Bindings = PixelShader->Bindings;
+	check(BufferIndex < Bindings.NumUniformBuffers);
+	if (Bindings.bHasRegularUniformBuffers)
+	{
+		auto* UB = (FMetalUniformBuffer*)BufferRHI;
+		Context->GetCommandEncoder().SetShaderBuffer(SF_Pixel, UB->Buffer, UB->Offset, BufferIndex);
+	}
+}
+
+void FMetalRHICommandContext::RHISetShaderUniformBuffer(FComputeShaderRHIParamRef ComputeShaderRHI, uint32 BufferIndex, FUniformBufferRHIParamRef BufferRHI)
+{
+	FMetalComputeShader* ComputeShader = ResourceCast(ComputeShaderRHI);
+	ComputeShader->BoundUniformBuffers[BufferIndex] = BufferRHI;
+	ComputeShader->DirtyUniformBuffers |= 1 << BufferIndex;
+
+	auto& Bindings = ComputeShader->Bindings;
+	check(BufferIndex < Bindings.NumUniformBuffers);
+	if (Bindings.bHasRegularUniformBuffers)
+	{
+		auto* UB = (FMetalUniformBuffer*)BufferRHI;
+		Context->GetCommandEncoder().SetShaderBuffer(SF_Compute, UB->Buffer, UB->Offset, BufferIndex);
+	}
+}
+
+
+void FMetalRHICommandContext::RHISetDepthStencilState(FDepthStencilStateRHIParamRef NewStateRHI, uint32 StencilRef)
+{
+	FMetalDepthStencilState* NewState = ResourceCast(NewStateRHI);
+
+	Context->GetCurrentState().SetDepthStencilState(NewState);
+	Context->GetCurrentState().SetStencilRef(StencilRef);
+	
+	FShaderCache::SetDepthStencilState(NewStateRHI);
+}
+
+void FMetalRHICommandContext::RHISetBlendState(FBlendStateRHIParamRef NewStateRHI, const FLinearColor& BlendFactor)
+{
+	FMetalBlendState* NewState = ResourceCast(NewStateRHI);
+	
+	Context->GetCurrentState().SetBlendState(NewState);
+	Context->GetCurrentState().SetBlendFactor(BlendFactor);
+	
+	FShaderCache::SetBlendState(NewStateRHI);
+}
+
+
+void FMetalRHICommandContext::RHISetRenderTargets(uint32 NumSimultaneousRenderTargets, const FRHIRenderTargetView* NewRenderTargets,
+	const FRHIDepthRenderTargetView* NewDepthStencilTargetRHI, uint32 NumUAVs, const FUnorderedAccessViewRHIParamRef* UAVs)
+{
+	FMetalContext* Manager = Context;
+	FRHIDepthRenderTargetView DepthView;
+	if (NewDepthStencilTargetRHI)
+	{
+		DepthView = *NewDepthStencilTargetRHI;
+	}
+	else
+	{
+		DepthView = FRHIDepthRenderTargetView(FTextureRHIParamRef(), ERenderTargetLoadAction::EClear, ERenderTargetStoreAction::ENoAction);
+	}
+
+	FRHISetRenderTargetsInfo Info(NumSimultaneousRenderTargets, NewRenderTargets, DepthView);
+	RHISetRenderTargetsAndClear(Info);
+
+	checkf(NumUAVs == 0, TEXT("Calling SetRenderTargets with UAVs is not supported in Metal yet"));
+}
+
+void FMetalDynamicRHI::RHIDiscardRenderTargets(bool Depth, bool Stencil, uint32 ColorBitMask)
+{
+#if PLATFORM_MAC // @todo zebra
+	NOT_SUPPORTED("RHIDiscardRenderTargets");
+#endif
+}
+
+void FMetalRHICommandContext::RHISetRenderTargetsAndClear(const FRHISetRenderTargetsInfo& RenderTargetsInfo)
+{
+	FMetalContext* Manager = Context;
+
+	Manager->SetRenderTargetsInfo(RenderTargetsInfo);
+
+	// Set the viewport to the full size of render target 0.
+	if (RenderTargetsInfo.ColorRenderTarget[0].Texture)
+	{
+		const FRHIRenderTargetView& RenderTargetView = RenderTargetsInfo.ColorRenderTarget[0];
+		FMetalSurface* RenderTarget = GetMetalSurfaceFromRHITexture(RenderTargetView.Texture);
+
+		uint32 Width = FMath::Max((uint32)(RenderTarget->Texture.width >> RenderTargetView.MipIndex), (uint32)1);
+		uint32 Height = FMath::Max((uint32)(RenderTarget->Texture.height >> RenderTargetView.MipIndex), (uint32)1);
+
+		RHISetViewport(0, 0, 0.0f, Width, Height, 1.0f);
+    }
+    
+    FShaderCache::SetRenderTargets(RenderTargetsInfo.NumColorRenderTargets, RenderTargetsInfo.ColorRenderTarget, &RenderTargetsInfo.DepthStencilRenderTarget);
+}
+
+
+void FMetalRHICommandContext::RHIDrawPrimitive(uint32 PrimitiveType, uint32 BaseVertexIndex, uint32 NumPrimitives, uint32 NumInstances)
+{
+	SCOPE_CYCLE_COUNTER(STAT_MetalDrawCallTime);
+	//checkf(NumInstances == 1, TEXT("Currently only 1 instance is supported"));
+	
+	RHI_DRAW_CALL_STATS(PrimitiveType,NumInstances*NumPrimitives);
+
+	// how many verts to render
+	uint32 NumVertices = GetVertexCountForPrimitiveCount(NumPrimitives, PrimitiveType);
+
+	// finalize any pending state
+	Context->PrepareToDraw(PrimitiveType);
+
+	uint32 VertexCount = GetVertexCountForPrimitiveCount(NumPrimitives,PrimitiveType);
+	RHI_PROFILE_DRAW_CALL_STATS(EMTLSamplePointBeforeDraw, EMTLSamplePointAfterDraw, NumPrimitives * NumInstances, VertexCount * NumInstances);
+	
+	// draw!
+	if(!FShaderCache::IsPredrawCall())
+	{
+		[Context->GetRenderContext() drawPrimitives:TranslatePrimitiveType(PrimitiveType)
+										vertexStart:BaseVertexIndex
+										vertexCount:NumVertices
+									  instanceCount:NumInstances];
+		
+		FShaderCache::LogDraw(0);
+	}
+}
+
+void FMetalRHICommandContext::RHIDrawPrimitiveIndirect(uint32 PrimitiveType, FVertexBufferRHIParamRef VertexBufferRHI, uint32 ArgumentOffset)
+{
+#if PLATFORM_IOS
+	NOT_SUPPORTED("RHIDrawPrimitiveIndirect");
+#else
+	SCOPE_CYCLE_COUNTER(STAT_MetalDrawCallTime);
+	RHI_DRAW_CALL_STATS(PrimitiveType,1);
+	
+	FMetalVertexBuffer* VertexBuffer = ResourceCast(VertexBufferRHI);
+	
+	// finalize any pending state
+	Context->PrepareToDraw(PrimitiveType);
+	
+	if(!FShaderCache::IsPredrawCall())
+	{
+		RHI_PROFILE_DRAW_CALL_STATS(EMTLSamplePointBeforeDraw, EMTLSamplePointAfterDraw, 1, 1);
+		
+		[Context->GetRenderContext() drawPrimitives:TranslatePrimitiveType(PrimitiveType)
+										indirectBuffer:VertexBuffer->Buffer
+								  indirectBufferOffset:ArgumentOffset];
+		
+		FShaderCache::LogDraw(0);
+	}
+#endif
+}
+
+
+void FMetalRHICommandContext::RHIDrawIndexedPrimitive(FIndexBufferRHIParamRef IndexBufferRHI, uint32 PrimitiveType, int32 BaseVertexIndex, uint32 FirstInstance,
+	uint32 NumVertices, uint32 StartIndex, uint32 NumPrimitives, uint32 NumInstances)
+{
+	SCOPE_CYCLE_COUNTER(STAT_MetalDrawCallTime);
+#if PLATFORM_IOS
+	//checkf(NumInstances == 1, TEXT("Currently only 1 instance is supported"));
+	checkf(BaseVertexIndex  == 0, TEXT("BaseVertexIndex must be 0, see GRHISupportsBaseVertexIndex"));
+	checkf(FirstInstance  == 0, TEXT("FirstInstance is currently unsupported on this RHI"));
+#endif
+
+	RHI_DRAW_CALL_STATS(PrimitiveType,NumInstances*NumPrimitives);
+
+	FMetalIndexBuffer* IndexBuffer = ResourceCast(IndexBufferRHI);
+
+	// finalize any pending state
+	Context->PrepareToDraw(PrimitiveType);
+	
+	uint32 NumIndices = GetVertexCountForPrimitiveCount(NumPrimitives, PrimitiveType);
+
+	RHI_PROFILE_DRAW_CALL_STATS(EMTLSamplePointBeforeDraw, EMTLSamplePointAfterDraw, NumPrimitives * NumInstances, NumVertices * NumInstances);
+	
+	if(!FShaderCache::IsPredrawCall())
+	{
+#if PLATFORM_IOS
+		[Context->GetRenderContext() drawIndexedPrimitives:TranslatePrimitiveType(PrimitiveType)
+												indexCount:NumIndices
+												 indexType:IndexBuffer->IndexType
+											   indexBuffer:IndexBuffer->Buffer
+										 indexBufferOffset:StartIndex * IndexBuffer->GetStride()
+											 instanceCount:NumInstances];
+#else
+		[Context->GetRenderContext() drawIndexedPrimitives:TranslatePrimitiveType(PrimitiveType)
+												indexCount:NumIndices
+												 indexType:IndexBuffer->IndexType
+											   indexBuffer:IndexBuffer->Buffer
+										 indexBufferOffset:StartIndex * IndexBuffer->GetStride()
+											 instanceCount:NumInstances
+											 baseVertex:BaseVertexIndex
+											 baseInstance:FirstInstance];
+#endif
+		FShaderCache::LogDraw(IndexBuffer->GetStride());
+	}
+}
+
+void FMetalRHICommandContext::RHIDrawIndexedIndirect(FIndexBufferRHIParamRef IndexBufferRHI, uint32 PrimitiveType, FStructuredBufferRHIParamRef VertexBufferRHI, int32 DrawArgumentsIndex, uint32 NumInstances)
+{
+#if PLATFORM_IOS
+	NOT_SUPPORTED("RHIDrawIndexedIndirect");
+#else
+	check(NumInstances > 1);
+
+	SCOPE_CYCLE_COUNTER(STAT_MetalDrawCallTime);
+	RHI_DRAW_CALL_STATS(PrimitiveType,1);
+	
+	FMetalIndexBuffer* IndexBuffer = ResourceCast(IndexBufferRHI);
+	FMetalStructuredBuffer* VertexBuffer = ResourceCast(VertexBufferRHI);
+	
+	// finalize any pending state
+	Context->PrepareToDraw(PrimitiveType);
+	
+	if(!FShaderCache::IsPredrawCall())
+	{
+		RHI_PROFILE_DRAW_CALL_STATS(EMTLSamplePointBeforeDraw, EMTLSamplePointAfterDraw, 1, 1);
+		
+		[Context->GetRenderContext() drawIndexedPrimitives:TranslatePrimitiveType(PrimitiveType)
+													 indexType:IndexBuffer->IndexType
+												   indexBuffer:IndexBuffer->Buffer
+											 indexBufferOffset:0
+												indirectBuffer:VertexBuffer->Buffer
+										  indirectBufferOffset:(DrawArgumentsIndex * 5 * sizeof(uint32))];
+		
+		FShaderCache::LogDraw(IndexBuffer->GetStride());
+	}
+#endif
+}
+
+void FMetalRHICommandContext::RHIDrawIndexedPrimitiveIndirect(uint32 PrimitiveType,FIndexBufferRHIParamRef IndexBufferRHI,FVertexBufferRHIParamRef VertexBufferRHI,uint32 ArgumentOffset)
+{
+#if PLATFORM_IOS
+	NOT_SUPPORTED("RHIDrawIndexedPrimitiveIndirect");
+#else
+	SCOPE_CYCLE_COUNTER(STAT_MetalDrawCallTime);
+	RHI_DRAW_CALL_STATS(PrimitiveType,1);
+	
+	FMetalIndexBuffer* IndexBuffer = ResourceCast(IndexBufferRHI);
+	FMetalVertexBuffer* VertexBuffer = ResourceCast(VertexBufferRHI);
+	
+	// finalize any pending state
+	Context->PrepareToDraw(PrimitiveType);
+	
+	if(!FShaderCache::IsPredrawCall())
+	{
+		RHI_PROFILE_DRAW_CALL_STATS(EMTLSamplePointBeforeDraw, EMTLSamplePointAfterDraw, 1, 1);
+		
+		[Context->GetRenderContext() drawIndexedPrimitives:TranslatePrimitiveType(PrimitiveType)
+													 indexType:IndexBuffer->IndexType
+												   indexBuffer:IndexBuffer->Buffer
+											 indexBufferOffset:0
+												indirectBuffer:VertexBuffer->Buffer
+										  indirectBufferOffset:ArgumentOffset];
+		
+		FShaderCache::LogDraw(IndexBuffer->GetStride());
+	}
+#endif
+}
+
+
+void FMetalRHICommandContext::RHIBeginDrawPrimitiveUP( uint32 PrimitiveType, uint32 NumPrimitives, uint32 NumVertices, uint32 VertexDataStride, void*& OutVertexData)
+{
+	SCOPE_CYCLE_COUNTER(STAT_MetalDrawCallTime);
+	checkSlow(PendingVertexBufferOffset == 0xFFFFFFFF);
+
+	// allocate space
+	PendingVertexBufferOffset = Context->AllocateFromRingBuffer(VertexDataStride * NumVertices);
+
+	// get the pointer to send back for writing
+	uint8* RingBufferBytes = (uint8*)[Context->GetRingBuffer() contents];
+	OutVertexData = RingBufferBytes + PendingVertexBufferOffset;
+	
+	RHI_PROFILE_DRAW_CALL_STATS(EMTLSamplePointBeforeDraw, EMTLSamplePointAfterDraw, NumPrimitives,NumVertices);
+	
+	PendingPrimitiveType = PrimitiveType;
+	PendingNumPrimitives = NumPrimitives;
+	PendingVertexDataStride = VertexDataStride;
+}
+
+
+void FMetalRHICommandContext::RHIEndDrawPrimitiveUP()
+{
+	SCOPE_CYCLE_COUNTER(STAT_MetalDrawCallTime);
+
+	RHI_DRAW_CALL_STATS(PendingPrimitiveType,PendingNumPrimitives);
+
+	// set the vertex buffer
+	Context->GetCurrentState().SetVertexBuffer(UNREAL_TO_METAL_BUFFER_INDEX(0), Context->GetRingBuffer(), PendingVertexDataStride, PendingVertexBufferOffset);
+	
+	// how many to draw
+	uint32 NumVertices = GetVertexCountForPrimitiveCount(PendingNumPrimitives, PendingPrimitiveType);
+
+	// last minute draw setup
+	Context->PrepareToDraw(PendingPrimitiveType);
+
+	if(!FShaderCache::IsPredrawCall())
+	{
+		[Context->GetRenderContext() drawPrimitives:TranslatePrimitiveType(PendingPrimitiveType)
+										vertexStart:0
+										vertexCount:NumVertices
+									  instanceCount:1];
+		
+		FShaderCache::LogDraw(0);
+	}
+	
+	// mark temp memory as usable
+	PendingVertexBufferOffset = 0xFFFFFFFF;
+}
+
+void FMetalRHICommandContext::RHIBeginDrawIndexedPrimitiveUP( uint32 PrimitiveType, uint32 NumPrimitives, uint32 NumVertices, uint32 VertexDataStride, void*& OutVertexData, uint32 MinVertexIndex, uint32 NumIndices, uint32 IndexDataStride, void*& OutIndexData)
+{
+	SCOPE_CYCLE_COUNTER(STAT_MetalDrawCallTime);
+	checkSlow(PendingVertexBufferOffset == 0xFFFFFFFF);
+	checkSlow(PendingIndexBufferOffset == 0xFFFFFFFF);
+
+	// allocate space
+	PendingVertexBufferOffset = Context->AllocateFromRingBuffer(VertexDataStride * NumVertices);
+	PendingIndexBufferOffset = Context->AllocateFromRingBuffer(IndexDataStride * NumIndices);
+	
+	// get the pointers to send back for writing
+	uint8* RingBufferBytes = (uint8*)[Context->GetRingBuffer() contents];
+	OutVertexData = RingBufferBytes + PendingVertexBufferOffset;
+	OutIndexData = RingBufferBytes + PendingIndexBufferOffset;
+	
+	RHI_PROFILE_DRAW_CALL_STATS(EMTLSamplePointBeforeDraw, EMTLSamplePointAfterDraw, NumPrimitives,NumVertices);
+	
+	PendingPrimitiveType = PrimitiveType;
+	PendingNumPrimitives = NumPrimitives;
+	PendingIndexDataStride = IndexDataStride;
+
+	PendingVertexDataStride = VertexDataStride;
+}
+
+void FMetalRHICommandContext::RHIEndDrawIndexedPrimitiveUP()
+{
+	SCOPE_CYCLE_COUNTER(STAT_MetalDrawCallTime);
+
+	RHI_DRAW_CALL_STATS(PendingPrimitiveType,PendingNumPrimitives);
+
+	// set the vertex buffer
+	Context->GetCurrentState().SetVertexBuffer(UNREAL_TO_METAL_BUFFER_INDEX(0), Context->GetRingBuffer(), PendingVertexDataStride, PendingVertexBufferOffset);
+
+	// how many to draw
+	uint32 NumIndices = GetVertexCountForPrimitiveCount(PendingNumPrimitives, PendingPrimitiveType);
+
+	// last minute draw setup
+	Context->PrepareToDraw(PendingPrimitiveType);
+	
+	if(!FShaderCache::IsPredrawCall())
+	{
+		[Context->GetRenderContext() drawIndexedPrimitives:TranslatePrimitiveType(PendingPrimitiveType)
+												indexCount:NumIndices
+												 indexType:(PendingIndexDataStride == 2) ? MTLIndexTypeUInt16 : MTLIndexTypeUInt32
+											   indexBuffer:Context->GetRingBuffer()
+										 indexBufferOffset:PendingIndexBufferOffset
+											 instanceCount:1];
+		
+		FShaderCache::LogDraw(PendingIndexDataStride);
+	}
+	
+	// mark temp memory as usable
+	PendingVertexBufferOffset = 0xFFFFFFFF;
+	PendingIndexBufferOffset = 0xFFFFFFFF;
+}
+
+
+void FMetalRHICommandContext::RHIClear(bool bClearColor,const FLinearColor& Color, bool bClearDepth,float Depth, bool bClearStencil,uint32 Stencil, FIntRect ExcludeRect)
+{
+	FMetalRHICommandContext::RHIClearMRT(bClearColor, 1, &Color, bClearDepth, Depth, bClearStencil, Stencil, ExcludeRect);
+}
+
+void FMetalRHICommandContext::RHIClearMRT(bool bClearColor,int32 NumClearColors,const FLinearColor* ClearColorArray,bool bClearDepth,float Depth,bool bClearStencil,uint32 Stencil, FIntRect ExcludeRect)
+{
+	// we don't support draw call clears before the RHI is initialized, reorder the code or make sure it's not a draw call clear
+	check(GIsRHIInitialized);
+
+	// Must specify enough clear colors for all active RTs
+	check(!bClearColor || NumClearColors >= Context->GetCurrentState().GetNumRenderTargets());
+
+	if (ExcludeRect.Min.X == 0 && ExcludeRect.Width() == Context->GetCurrentState().GetViewport().width && ExcludeRect.Min.Y == 0 && ExcludeRect.Height() == Context->GetCurrentState().GetViewport().height)
+	{
+		// no need to do anything
+		return;
+	}
+
+	RHIPushEvent(TEXT("MetalClearMRT"));
+
+	// Build new states
+	FBlendStateRHIRef BlendStateRHI;
+	
+	if (Context->GetCurrentState().GetNumRenderTargets() <= 1)
+	{
+		BlendStateRHI = (bClearColor && Context->GetCurrentState().GetHasValidRenderTarget())
+		? TStaticBlendState<>::GetRHI()
+		: TStaticBlendState<CW_NONE>::GetRHI();
+	}
+	else
+	{
+		BlendStateRHI = (bClearColor && Context->GetCurrentState().GetHasValidRenderTarget())
+			? TStaticBlendState<>::GetRHI()
+			: TStaticBlendStateWriteMask<CW_NONE,CW_NONE,CW_NONE,CW_NONE,CW_NONE,CW_NONE,CW_NONE,CW_NONE>::GetRHI();
+	}
+
+	FRasterizerStateRHIParamRef RasterizerStateRHI = TStaticRasterizerState<FM_Solid, CM_None>::GetRHI();
+	float BF[4] = { 0, 0, 0, 0 };
+
+	const FDepthStencilStateRHIParamRef DepthStencilStateRHI = 
+		(bClearDepth && bClearStencil)
+			? TStaticDepthStencilState<
+				true, CF_Always,
+				true,CF_Always,SO_Replace,SO_Replace,SO_Replace,
+				false,CF_Always,SO_Replace,SO_Replace,SO_Replace,
+				0xff,0xff
+				>::GetRHI()
+		: bClearDepth
+			? TStaticDepthStencilState<true, CF_Always>::GetRHI()
+		: bClearStencil
+			? TStaticDepthStencilState<
+				false, CF_Always,
+				true,CF_Always,SO_Replace,SO_Replace,SO_Replace,
+				false,CF_Always,SO_Replace,SO_Replace,SO_Replace,
+				0xff,0xff
+				>::GetRHI()
+		:     TStaticDepthStencilState<false, CF_Always>::GetRHI();
+
+	
+	FMetalStateCache const& StateCache = Context->GetCurrentState();
+	FLinearColor OriginalBlendFactor = StateCache.GetBlendFactor();
+	uint32 OriginalStencilRef = StateCache.GetStencilRef();
+	FBlendStateRHIRef OriginalBlend = StateCache.GetBlendState();
+	FDepthStencilStateRHIRef OriginalDepthStencil = StateCache.GetDepthStencilState();
+	FRasterizerStateRHIRef OriginalRasterizer = StateCache.GetRasterizerState();
+	FBoundShaderStateRHIRef OriginalBoundShaderState = StateCache.GetBoundShaderState();
+	
+	RHISetBlendState(BlendStateRHI, FLinearColor::Transparent);
+	RHISetDepthStencilState(DepthStencilStateRHI, Stencil);
+	RHISetRasterizerState(RasterizerStateRHI);
+
+	// Set the new shaders
+	auto ShaderMap = GetGlobalShaderMap(GMaxRHIFeatureLevel);
+	TShaderMapRef<TOneColorVS<true> > VertexShader(ShaderMap);
+
+	FOneColorPS* PixelShader = NULL;
+
+	// Set the shader to write to the appropriate number of render targets
+	if (Context->GetCurrentState().GetNumRenderTargets() <= 1)
+	{
+		TShaderMapRef<TOneColorPixelShaderMRT<1> > MRTPixelShader(ShaderMap);
+		PixelShader = *MRTPixelShader;
+	}
+	else if (Context->GetCurrentState().GetNumRenderTargets() == 2)
+	{
+		TShaderMapRef<TOneColorPixelShaderMRT<2> > MRTPixelShader(ShaderMap);
+		PixelShader = *MRTPixelShader;
+	}
+	else if (Context->GetCurrentState().GetNumRenderTargets() == 3)
+	{
+		TShaderMapRef<TOneColorPixelShaderMRT<3> > MRTPixelShader(ShaderMap);
+		PixelShader = *MRTPixelShader;
+	}
+	else if (Context->GetCurrentState().GetNumRenderTargets() == 4)
+	{
+		TShaderMapRef<TOneColorPixelShaderMRT<4> > MRTPixelShader(ShaderMap);
+		PixelShader = *MRTPixelShader;
+	}
+	else if (Context->GetCurrentState().GetNumRenderTargets() == 5)
+	{
+		TShaderMapRef<TOneColorPixelShaderMRT<5> > MRTPixelShader(ShaderMap);
+		PixelShader = *MRTPixelShader;
+	}
+	else if (Context->GetCurrentState().GetNumRenderTargets() == 6)
+	{
+		TShaderMapRef<TOneColorPixelShaderMRT<6> > MRTPixelShader(ShaderMap);
+		PixelShader = *MRTPixelShader;
+	}
+	else if (Context->GetCurrentState().GetNumRenderTargets() == 7)
+	{
+		TShaderMapRef<TOneColorPixelShaderMRT<7> > MRTPixelShader(ShaderMap);
+		PixelShader = *MRTPixelShader;
+	}
+	else if (Context->GetCurrentState().GetNumRenderTargets() == 8)
+	{
+		TShaderMapRef<TOneColorPixelShaderMRT<8> > MRTPixelShader(ShaderMap);
+		PixelShader = *MRTPixelShader;
+	}
+
+	{
+		FRHICommandList_RecursiveHazardous RHICmdList(this);
+		SetGlobalBoundShaderState(RHICmdList, GMaxRHIFeatureLevel, GClearMRTBoundShaderState[FMath::Max(Context->GetCurrentState().GetNumRenderTargets() - 1, 0)], GVector4VertexDeclaration.VertexDeclarationRHI, *VertexShader, PixelShader);
+		PixelShader->SetColors(RHICmdList, ClearColorArray, NumClearColors);
+
+		{
+			RHI_PROFILE_DRAW_CALL_STATS(EMTLSamplePointBeforeDraw, EMTLSamplePointAfterDraw, 0,0);
+			
+			const CGSize FrameBufferSize = Context->GetCurrentState().GetFrameBufferSize();
+			const bool bIgnoreExcludeRect = FrameBufferSize.width == Context->GetCurrentState().GetViewport().width && FrameBufferSize.height == Context->GetCurrentState().GetViewport().height;
+
+			// Draw a fullscreen quad
+			if (!bIgnoreExcludeRect && ExcludeRect.Width() > 0 && ExcludeRect.Height() > 0
+				&& ExcludeRect.Width() < Context->GetCurrentState().GetViewport().width && ExcludeRect.Height() < Context->GetCurrentState().GetViewport().height)
+			{
+				// with a hole in it (optimization in case the hardware has non constant clear performance)
+				FVector4 OuterVertices[4];
+				OuterVertices[0].Set(-1.0f, 1.0f, Depth, 1.0f);
+				OuterVertices[1].Set(1.0f, 1.0f, Depth, 1.0f);
+				OuterVertices[2].Set(1.0f, -1.0f, Depth, 1.0f);
+				OuterVertices[3].Set(-1.0f, -1.0f, Depth, 1.0f);
+
+				float InvViewWidth = 1.0f / Context->GetCurrentState().GetViewport().width;
+				float InvViewHeight = 1.0f / Context->GetCurrentState().GetViewport().height;
+				FVector4 FractionRect = FVector4(ExcludeRect.Min.X * InvViewWidth, ExcludeRect.Min.Y * InvViewHeight, (ExcludeRect.Max.X - 1) * InvViewWidth, (ExcludeRect.Max.Y - 1) * InvViewHeight);
+
+				FVector4 InnerVertices[4];
+				InnerVertices[0].Set(FMath::Lerp(-1.0f, 1.0f, FractionRect.X), FMath::Lerp(1.0f, -1.0f, FractionRect.Y), Depth, 1.0f);
+				InnerVertices[1].Set(FMath::Lerp(-1.0f, 1.0f, FractionRect.Z), FMath::Lerp(1.0f, -1.0f, FractionRect.Y), Depth, 1.0f);
+				InnerVertices[2].Set(FMath::Lerp(-1.0f, 1.0f, FractionRect.Z), FMath::Lerp(1.0f, -1.0f, FractionRect.W), Depth, 1.0f);
+				InnerVertices[3].Set(FMath::Lerp(-1.0f, 1.0f, FractionRect.X), FMath::Lerp(1.0f, -1.0f, FractionRect.W), Depth, 1.0f);
+
+				FVector4 Vertices[10];
+				Vertices[0] = OuterVertices[0];
+				Vertices[1] = InnerVertices[0];
+				Vertices[2] = OuterVertices[1];
+				Vertices[3] = InnerVertices[1];
+				Vertices[4] = OuterVertices[2];
+				Vertices[5] = InnerVertices[2];
+				Vertices[6] = OuterVertices[3];
+				Vertices[7] = InnerVertices[3];
+				Vertices[8] = OuterVertices[0];
+				Vertices[9] = InnerVertices[0];
+
+				DrawPrimitiveUP(RHICmdList, PT_TriangleStrip, 8, Vertices, sizeof(Vertices[0]));
+			}
+			else
+			{
+				// without a hole
+				FVector4 Vertices[4];
+				Vertices[0].Set(-1.0f, 1.0f, Depth, 1.0f);
+				Vertices[1].Set(1.0f, 1.0f, Depth, 1.0f);
+				Vertices[2].Set(-1.0f, -1.0f, Depth, 1.0f);
+				Vertices[3].Set(1.0f, -1.0f, Depth, 1.0f);
+				DrawPrimitiveUP(RHICmdList, PT_TriangleStrip, 2, Vertices, sizeof(Vertices[0]));
+			}
+		}
+		// Implicit flush. Always call flush when using a command list in RHI implementations before doing anything else. This is super hazardous.
+	}
+
+	RHISetBlendState(OriginalBlend, OriginalBlendFactor);
+	RHISetDepthStencilState(OriginalDepthStencil, Stencil);
+	RHISetRasterizerState(OriginalRasterizer);
+	RHISetBoundShaderState(OriginalBoundShaderState);
+
+	RHIPopEvent();
+}
+
+void FMetalDynamicRHI::RHISuspendRendering()
+{
+	// Not supported
+}
+
+void FMetalDynamicRHI::RHIResumeRendering()
+{
+	// Not supported
+}
+
+bool FMetalDynamicRHI::RHIIsRenderingSuspended()
+{
+	// Not supported
+	return false;
+}
+
+void FMetalDynamicRHI::RHIBlockUntilGPUIdle()
+{
+	NOT_SUPPORTED("RHIBlockUntilGPUIdle");
+}
+
+uint32 FMetalDynamicRHI::RHIGetGPUFrameCycles()
+{
+	return GGPUFrameTime;
+}
+
+void FMetalRHICommandContext::RHIAutomaticCacheFlushAfterComputeShader(bool bEnable)
+{
+	NOT_SUPPORTED("RHIAutomaticCacheFlushAfterComputeShader");
+}
+
+void FMetalRHICommandContext::RHIFlushComputeShaderCache()
+{
+	NOT_SUPPORTED("RHIFlushComputeShaderCache");
+}
+
+void FMetalDynamicRHI::RHIExecuteCommandList(FRHICommandList* RHICmdList)
+{
+	NOT_SUPPORTED("RHIExecuteCommandList");
+}
+
+void FMetalRHICommandContext::RHIEnableDepthBoundsTest(bool bEnable, float MinDepth, float MaxDepth)
+{
+    // not supported
+	NOT_SUPPORTED("RHIEnableDepthBoundsTest");
+}
+
+void FMetalRHICommandContext::RHISubmitCommandsHint()
+{
+    Context->SubmitCommandsHint();
+}
+
+IRHICommandContext* FMetalDynamicRHI::RHIGetDefaultContext()
+{
+	return this;
+}
+
+IRHICommandContextContainer* FMetalDynamicRHI::RHIGetCommandContextContainer()
+{
+	return nullptr;
+}
+
+
+
diff --git a/Engine/Source/Runtime/Apple/MetalRHI/Private/MetalContext.cpp b/Engine/Source/Runtime/Apple/MetalRHI/Private/MetalContext.cpp
new file mode 100644
--- /dev/null
+++ b/Engine/Source/Runtime/Apple/MetalRHI/Private/MetalContext.cpp
@@ -0,0 +1,838 @@
+// Copyright 1998-2015 Epic Games, Inc. All Rights Reserved.
+
+#include "MetalRHIPrivate.h"
+#if PLATFORM_IOS
+#include "IOSAppDelegate.h"
+#endif
+
+#include "MetalContext.h"
+#include "MetalProfiler.h"
+
+const uint32 RingBufferSize = 64 * 1024 * 1024;
+
+#if SHOULD_TRACK_OBJECTS
+TMap<id, int32> ClassCounts;
+#endif
+
+static void* LocalCreateAutoreleasePool()
+{
+	return [[NSAutoreleasePool alloc] init];
+}
+
+static void LocalReleaseAutoreleasePool(void *Pool)
+{
+	[(NSAutoreleasePool*)Pool release];
+}
+
+#if PLATFORM_MAC
+#if __MAC_OS_X_VERSION_MAX_ALLOWED < 101100
+MTL_EXTERN NSArray* MTLCopyAllDevices(void);
+#else
+MTL_EXTERN NSArray <id<MTLDevice>>* MTLCopyAllDevices(void);
+#endif
+
+static id<MTLDevice> GetMTLDevice()
+{
+	NSArray* DeviceList = MTLCopyAllDevices();
+	const int32 NumDevices = [DeviceList count];
+	
+	// return the first device if there’s only one in the list
+	if (NumDevices == 1)
+	{
+		return (id<MTLDevice>)[DeviceList objectAtIndex:0];
+	}
+	
+	const char* MetalDeviceToUse = getenv("MetalDeviceToUse");
+	const bool bForceIntelGPU = MetalDeviceToUse && strcmp(MetalDeviceToUse, "Intel") == 0;
+	const bool bForceDefault = FParse::Param(FCommandLine::Get(),TEXT("defaultgpu"));
+	
+	id<MTLDevice> DefaultDevice = MTLCreateSystemDefaultDevice();
+	if (bForceDefault)
+	{
+		return DefaultDevice;
+	}
+	
+#if WITH_EDITOR
+	const bool bForceHeadless = !IsRunningGame() && !bForceIntelGPU;
+#else
+	const bool bForceHeadless = false;
+#endif
+	
+	if (bForceIntelGPU || bForceHeadless)
+	{
+		// enumerate through the device list to find the low power device
+		for (int32 Index = 0; Index < NumDevices; Index++)
+		{
+			id<MTLDevice> Device = (id<MTLDevice>)[DeviceList objectAtIndex:Index];
+			const bool bIsLowPower = Device.lowPower;
+			const bool bIsHeadless = Device != DefaultDevice;
+			if (bIsLowPower && bForceIntelGPU)
+			{
+				return Device;
+			}
+			else if (bIsHeadless && !bIsLowPower && bForceHeadless)
+			{
+				return Device;
+			}
+		}
+	}
+	
+	return DefaultDevice;
+}
+
+static MTLPrimitiveTopologyClass TranslatePrimitiveTopology(uint32 PrimitiveType)
+{
+	switch (PrimitiveType)
+	{
+		case PT_TriangleList:
+		case PT_TriangleStrip:
+			return MTLPrimitiveTopologyClassTriangle;
+		case PT_LineList:
+			return MTLPrimitiveTopologyClassLine;
+		case PT_PointList:
+			return MTLPrimitiveTopologyClassPoint;
+		default:
+			UE_LOG(LogMetal, Fatal, TEXT("Unsupported primitive topology %d"), (int32)PrimitiveType);
+			return MTLPrimitiveTopologyClassTriangle;
+	}
+}
+#endif
+
+FMetalContext::FMetalContext()
+#if PLATFORM_IOS
+: Device([IOSAppDelegate GetDelegate].IOSView->MetalDevice)
+#else // @todo zebra
+: Device(GetMTLDevice())
+#endif
+, CommandQueue(Device)
+, StateCache(CommandEncoder)
+, CurrentCommandBuffer(nil)
+, RingBuffer(Device, RingBufferSize, BufferOffsetAlignment)
+, QueryBuffer(this)
+, SceneFrameCounter(0)
+, FrameCounter(0)
+, ResourceTableFrameCounter(INDEX_NONE)
+, Features(0)
+{
+#if PLATFORM_IOS
+	NSOperatingSystemVersion Vers = [[NSProcessInfo processInfo] operatingSystemVersion];
+	if(Vers.majorVersion >= 9)
+	{
+		Features = EMetalFeaturesSeparateStencil | EMetalFeaturesSetBufferOffset | EMetalFeaturesDepthClipMode | EMetalFeaturesResourceOptions;
+	}
+	else if(Vers.majorVersion == 8 && Vers.minorVersion >= 3)
+	{
+		Features = EMetalFeaturesSeparateStencil | EMetalFeaturesSetBufferOffset;
+	}
+#else // Assume that Mac & other platforms all support these from the start. They can diverge later.
+	Features = EMetalFeaturesSeparateStencil | EMetalFeaturesSetBufferOffset | EMetalFeaturesDepthClipMode | EMetalFeaturesResourceOptions;
+#endif
+	
+	// create a semaphore for multi-buffering the command buffer
+	CommandBufferSemaphore = dispatch_semaphore_create(FParse::Param(FCommandLine::Get(),TEXT("gpulockstep")) ? 1 : 3);
+	
+	AutoReleasePoolTLSSlot = FPlatformTLS::AllocTlsSlot();
+	
+	// Hook into the ios framepacer, if it's enabled for this platform.
+	FrameReadyEvent = NULL;
+#if PLATFORM_IOS
+	if( FIOSPlatformRHIFramePacer::IsEnabled() )
+	{
+		FrameReadyEvent = FPlatformProcess::GetSynchEventFromPool();
+		FIOSPlatformRHIFramePacer::InitWithEvent( FrameReadyEvent );
+	}
+#endif
+	
+	InitFrame();
+}
+
+FMetalContext::~FMetalContext()
+{
+	
+}
+
+id<MTLDevice> FMetalContext::GetDevice()
+{
+	return Device;
+}
+
+FMetalCommandQueue& FMetalContext::GetCommandQueue()
+{
+	return CommandQueue;
+}
+
+FMetalCommandEncoder& FMetalContext::GetCommandEncoder()
+{
+	return CommandEncoder;
+}
+
+id<MTLRenderCommandEncoder> FMetalContext::GetRenderContext()
+{
+	if(!CommandEncoder.IsRenderCommandEncoderActive())
+	{
+		UE_LOG(LogMetal, Warning, TEXT("Attempted to use GetRenderContext before calling SetRenderTarget to create the encoder"));
+		ConditionalSwitchToGraphics();
+		CommandEncoder.RestoreRenderCommandEncoding();
+	}
+	return CommandEncoder.GetRenderCommandEncoder();
+}
+
+id<MTLBlitCommandEncoder> FMetalContext::GetBlitContext()
+{
+	ConditionalSwitchToBlit();
+	return CommandEncoder.GetBlitCommandEncoder();
+}
+
+id<MTLCommandBuffer> FMetalContext::GetCurrentCommandBuffer()
+{
+	return CurrentCommandBuffer;
+}
+
+void FMetalContext::ReleaseObject(id Object)
+{
+	if (GIsRHIInitialized) // @todo zebra: there seems to be some race condition at exit when the framerate is very low
+	{
+		check(Object);
+		CommandEncoder.UnbindObject(Object);
+		check(!FreeList.Contains(Object));
+		FreeList.Add(Object);
+	}
+}
+
+FMetalBufferPool* FMetalContext::GetBufferPool()
+{
+	if(GIsRHIInitialized)
+	{
+		return &BufferPool;
+	}
+	return nullptr;
+}
+
+FMetalPooledBuffer FMetalContext::CreatePooledBuffer(FMetalPooledBufferArgs const& Args)
+{
+	FMetalPooledBuffer Buffer;
+	if(GIsRHIInitialized)
+	{
+		Buffer = BufferPool.CreatePooledResource(Args);
+		//[Buffer.Buffer setPurgeableState:MTLPurgeableStateNonVolatile];
+	}
+	else
+	{
+		FMetalBufferPoolPolicyData Policy;
+		uint32 BufferSize = Policy.GetPoolBucketSize(Policy.GetPoolBucketIndex(Args));
+		Buffer.Buffer = [Device newBufferWithLength:BufferSize options: BUFFER_CACHE_MODE
+#if METAL_API_1_1
+						 | (Args.Storage << MTLResourceStorageModeShift)
+#endif
+						 ];
+	}
+	return Buffer;
+}
+
+void FMetalContext::ReleasePooledBuffer(FMetalPooledBuffer Buf)
+{
+	if(GIsRHIInitialized)
+	{
+		CommandEncoder.UnbindObject(Buf.Buffer);
+		//[Buf.Buffer setPurgeableState:MTLPurgeableStateVolatile];
+		return BufferPool.ReleasePooledResource(Buf);
+	}
+}
+
+
+void FMetalContext::CreateAutoreleasePool()
+{
+	if (FPlatformTLS::GetTlsValue(AutoReleasePoolTLSSlot) == NULL)
+	{
+		FPlatformTLS::SetTlsValue(AutoReleasePoolTLSSlot, LocalCreateAutoreleasePool());
+	}
+}
+
+void FMetalContext::DrainAutoreleasePool()
+{
+	LocalReleaseAutoreleasePool(FPlatformTLS::GetTlsValue(AutoReleasePoolTLSSlot));
+	FPlatformTLS::SetTlsValue(AutoReleasePoolTLSSlot, NULL);
+}
+
+void FMetalContext::BeginFrame()
+{
+}
+
+void FMetalContext::EndFrame()
+{
+	while(DelayedFreeLists.Num())
+	{
+		FMetalDelayedFreeList* Pair = DelayedFreeLists[0];
+		if(Pair->Signal->Wait(0))
+		{
+			Pair->Signal->Reset();
+			FPlatformProcess::ReturnSynchEventToPool(Pair->Signal);
+			for( id Entry : Pair->FreeList )
+			{
+				[Entry release];
+			}
+			delete Pair;
+			DelayedFreeLists.RemoveAt(0, 1, false);
+		}
+		else
+		{
+			break;
+		}
+	}
+	if(FrameCounter != GFrameNumberRenderThread)
+	{
+		FrameCounter = GFrameNumberRenderThread;
+		BufferPool.DrainPool(false);
+	}
+	CommandQueue.InsertDebugCaptureBoundary();
+}
+
+void FMetalContext::BeginScene()
+{
+	// Increment the frame counter. INDEX_NONE is a special value meaning "uninitialized", so if
+	// we hit it just wrap around to zero.
+	SceneFrameCounter++;
+	if (SceneFrameCounter == INDEX_NONE)
+	{
+		SceneFrameCounter++;
+	}
+
+	static auto* ResourceTableCachingCvar = IConsoleManager::Get().FindTConsoleVariableDataInt(TEXT("rhi.ResourceTableCaching"));
+	if (ResourceTableCachingCvar == NULL || ResourceTableCachingCvar->GetValueOnAnyThread() == 1)
+	{
+		ResourceTableFrameCounter = SceneFrameCounter;
+	}
+}
+
+void FMetalContext::EndScene()
+{
+	ResourceTableFrameCounter = INDEX_NONE;
+}
+
+void FMetalContext::BeginDrawingViewport(FMetalViewport* Viewport)
+{
+#if PLATFORM_MAC
+	CAMetalLayer* CurrentLayer = (CAMetalLayer*)[Viewport->View layer];
+	StateCache.SetCurrentLayer(CurrentLayer);
+#endif
+}
+
+void FMetalContext::EndDrawingViewport(FMetalViewport* Viewport, bool bPresent)
+{
+	// commit the render context to the commandBuffer
+	if (CommandEncoder.IsRenderCommandEncoderActive() || CommandEncoder.IsComputeCommandEncoderActive() || CommandEncoder.IsBlitCommandEncoderActive())
+	{
+		CommandEncoder.EndEncoding();
+	}
+
+	// kick the whole buffer
+	uint32 RingBufferOffset = RingBuffer.GetOffset();
+	
+	FMetalDelayedFreeList* NewList = new FMetalDelayedFreeList;
+	NewList->Signal = FPlatformProcess::GetSynchEventFromPool(true);
+	NewList->FreeList = FreeList;
+	DelayedFreeLists.Add(NewList);
+	FreeList.Empty(FreeList.Num());
+
+	[CurrentCommandBuffer addCompletedHandler : ^ (id <MTLCommandBuffer> Buffer)
+	 {
+		RingBuffer.SetLastRead(RingBufferOffset);
+		dispatch_semaphore_signal(CommandBufferSemaphore);
+		NewList->Signal->Trigger();
+	 }];
+
+	// We may be limiting our framerate to the display link
+	if( FrameReadyEvent != nullptr )
+	{
+		FrameReadyEvent->Wait();
+	}
+
+	// enqueue a present if desired
+	if (StateCache.GetCurrentDrawable())
+	{
+		if (bPresent)
+		{
+			id<MTLDrawable> Drawable = StateCache.GetCurrentDrawable();
+			[CurrentCommandBuffer addScheduledHandler:^(id<MTLCommandBuffer>) {
+				[Drawable present];
+			 }];
+		}
+
+		StateCache.ResetCurrentDrawable();
+	}
+
+	// Commit before waiting to avoid leaving the gpu idle
+	CommandEncoder.CommitCommandBuffer(false);
+
+//#if WITH_EDITOR // @todo zebra
+//	if (!IsRunningGame())
+//	{
+//		[CurrentCommandBuffer waitUntilCompleted];
+//	}
+//#endif
+	
+	UNTRACK_OBJECT(CurrentCommandBuffer);
+	[CurrentCommandBuffer release];
+
+#if SHOULD_TRACK_OBJECTS
+	// print out outstanding objects
+	if ((GFrameCounter % 500) == 10)
+	{
+		for (auto It = ClassCounts.CreateIterator(); It; ++It)
+		{
+			NSLog(@"%@ has %d outstanding allocations", It.Key(), It.Value());
+		}
+	}
+#endif
+
+	// drain the pool
+	DrainAutoreleasePool();
+
+	InitFrame();
+
+	Viewport->GetBackBuffer()->Surface.Texture = nil;
+}
+
+void FMetalContext::InitFrame()
+{
+	// start an auto release pool (EndFrame will drain and remake)
+	CreateAutoreleasePool();
+
+	// create the command buffer for this frame
+	CreateCurrentCommandBuffer(true);
+
+	//	double Start = FPlatformTime::Seconds();
+	// block on the semaphore
+//	double Mid = FPlatformTime::Seconds();
+
+	// make sure first SetRenderTarget goes through
+	StateCache.SetHasValidRenderTarget(false);
+
+//	double End = FPlatformTime::Seconds();
+//	NSLog(@"Semaphore Block Time: %.2f   -- MakeDrawable Time: %.2f", 1000.0f*(Mid - Start), 1000.0f*(End - Mid));
+	
+
+	extern void InitFrame_UniformBufferPoolCleanup();
+	InitFrame_UniformBufferPoolCleanup();
+}
+
+void FMetalContext::CreateCurrentCommandBuffer(bool bWait)
+{
+	if (bWait)
+	{
+		dispatch_semaphore_wait(CommandBufferSemaphore, DISPATCH_TIME_FOREVER);
+	}
+	
+	CurrentCommandBuffer = CommandQueue.CreateUnretainedCommandBuffer();
+	[CurrentCommandBuffer retain];
+	TRACK_OBJECT(CurrentCommandBuffer);
+	
+	CommandEncoder.StartCommandBuffer(CurrentCommandBuffer);
+}
+
+void FMetalContext::SubmitCommandsHint()
+{
+    uint32 RingBufferOffset = RingBuffer.GetOffset();
+    [CurrentCommandBuffer addCompletedHandler : ^ (id <MTLCommandBuffer> Buffer)
+    {
+        RingBuffer.SetLastRead(RingBufferOffset);
+        dispatch_semaphore_signal(CommandBufferSemaphore);
+    }];
+    
+    // commit the render context to the commandBuffer
+	if (CommandEncoder.IsRenderCommandEncoderActive() || CommandEncoder.IsComputeCommandEncoderActive() || CommandEncoder.IsBlitCommandEncoderActive())
+	{
+		CommandEncoder.EndEncoding();
+	}
+    
+    // kick the whole buffer
+    // Commit to hand the commandbuffer off to the gpu
+    CommandEncoder.CommitCommandBuffer(false);
+    
+    //once a commandbuffer is commited it can't be added to again.
+	UNTRACK_OBJECT(CurrentCommandBuffer);
+	[CurrentCommandBuffer release];
+	
+    // create a new command buffer.
+    CreateCurrentCommandBuffer(false);
+}
+
+void FMetalContext::SubmitCommandBufferAndWait()
+{
+	uint32 RingBufferOffset = RingBuffer.GetOffset();
+	[CurrentCommandBuffer addCompletedHandler : ^ (id <MTLCommandBuffer> Buffer)
+	 {
+		RingBuffer.SetLastRead(RingBufferOffset);
+		dispatch_semaphore_signal(CommandBufferSemaphore);
+	 }];
+	
+	// commit the render context to the commandBuffer
+	if (CommandEncoder.IsRenderCommandEncoderActive() || CommandEncoder.IsComputeCommandEncoderActive() || CommandEncoder.IsBlitCommandEncoderActive())
+	{
+		CommandEncoder.EndEncoding();
+	}
+    
+    // kick the whole buffer
+    // Commit to hand the commandbuffer off to the gpu
+    CommandEncoder.CommitCommandBuffer(true);
+
+	//once a commandbuffer is commited it can't be added to again.
+	UNTRACK_OBJECT(CurrentCommandBuffer);
+	[CurrentCommandBuffer release];
+	
+	// create a new command buffer.
+	CreateCurrentCommandBuffer(true);
+}
+
+void FMetalContext::SubmitComputeCommandBufferAndWait()
+{
+}
+
+void FMetalContext::ResetRenderCommandEncoder()
+{
+	SubmitCommandsHint();
+	
+	ConditionalSwitchToGraphics();
+	CommandEncoder.RestoreRenderCommandEncoding();
+}
+
+void FMetalContext::PrepareToDraw(uint32 PrimitiveType)
+{
+	SCOPE_CYCLE_COUNTER(STAT_MetalPrepareDrawTime);
+
+	TRefCountPtr<FMetalBoundShaderState> CurrentBoundShaderState = StateCache.GetBoundShaderState();
+#if UE_BUILD_DEBUG || UE_BUILD_DEVELOPMENT
+	MTLVertexDescriptor* Layout = CurrentBoundShaderState->VertexDeclaration->Layout;
+	if(Layout && Layout.layouts)
+	{
+		for (uint32 i = 0; i < MaxMetalStreams; i++)
+		{
+			auto Attribute = [Layout.attributes objectAtIndexedSubscript:i];
+			if(Attribute)
+			{
+				auto BufferLayout = [Layout.layouts objectAtIndexedSubscript:Attribute.bufferIndex];
+				uint32 BufferLayoutStride = BufferLayout ? BufferLayout.stride : 0;
+				uint64 MetalSize = StateCache.GetVertexBuffer(Attribute.bufferIndex) ? (uint64)StateCache.GetVertexBuffer(Attribute.bufferIndex).length : 0llu;
+
+				if(!(BufferLayoutStride == StateCache.GetVertexStride(Attribute.bufferIndex) || (MetalSize > 0 && StateCache.GetVertexStride(Attribute.bufferIndex) == 0 && StateCache.GetVertexBuffer(Attribute.bufferIndex) && MetalSize == BufferLayoutStride)))
+				{
+					FString Report = FString::Printf(TEXT("Vertex Layout Mismatch: Index: %d, Buffer: %p, Len: %lld, Decl. Stride: %d, Stream Stride: %d"), Attribute.bufferIndex, StateCache.GetVertexBuffer(Attribute.bufferIndex), MetalSize, BufferLayoutStride, StateCache.GetVertexStride(Attribute.bufferIndex));
+#if PLATFORM_MAC // Disable this to get a soft error that can be debugged further in a GPU trace. // @todo zebra ios
+					if(!FShaderCache::IsPredrawCall())
+					{
+						UE_LOG(LogMetal, Fatal, TEXT("%s"), *Report);
+					}
+					else
+					{
+						UE_LOG(LogMetal, Warning, TEXT("%s"), *Report);
+					}
+#else
+					UE_LOG(LogMetal, Warning, TEXT("%s"), *Report);
+					//[CommandEncoder pushDebugGroup: Report.GetNSString()];
+					//[CommandEncoder popDebugGroup];
+#endif
+				}
+			}
+		}
+	}
+#endif
+	
+#if PLATFORM_MAC
+	StateCache.SetPrimitiveTopology(TranslatePrimitiveTopology(PrimitiveType));
+#endif
+	
+	// make sure the BSS has a valid pipeline state object
+	CurrentBoundShaderState->PrepareToDraw(this, StateCache.GetRenderPipelineDesc());
+	
+	if(!FShaderCache::IsPredrawCall())
+	{
+		CommitGraphicsResourceTables();
+		CommitNonComputeShaderConstants();
+		
+		if(CommandEncoder.IsBlitCommandEncoderActive() || CommandEncoder.IsComputeCommandEncoderActive())
+		{
+			CommandEncoder.EndEncoding();
+		}
+		
+		if(!CommandEncoder.IsRenderCommandEncoderActive())
+		{
+			CommandEncoder.RestoreRenderCommandEncoding();
+		}
+	}
+}
+
+void FMetalContext::SetRenderTargetsInfo(const FRHISetRenderTargetsInfo& RenderTargetsInfo)
+{
+	StateCache.SetRenderTargetsInfo(RenderTargetsInfo, QueryBuffer.GetCurrentQueryBuffer()->Buffer);
+}
+
+uint32 FMetalContext::AllocateFromRingBuffer(uint32 Size, uint32 Alignment)
+{
+	return RingBuffer.Allocate(Size, Alignment);
+}
+
+
+void FMetalContext::SetResource(uint32 ShaderStage, uint32 BindIndex, FRHITexture* RESTRICT TextureRHI)
+{
+	//	check(Surface->Texture != nil);
+	
+	// todo: this does multiple virtual function calls to find the right type to cast to
+	// this is due to multiple inheritance nastiness, NEEDS CLEANUP
+	FMetalSurface* Surface = GetMetalSurfaceFromRHITexture(TextureRHI);
+	
+	if (Surface != nullptr)
+	{
+		switch (ShaderStage)
+		{
+			case CrossCompiler::SHADER_STAGE_PIXEL:
+				CommandEncoder.SetShaderTexture(SF_Pixel, Surface->Texture, BindIndex);
+				FShaderCache::SetTexture(SF_Pixel, BindIndex, TextureRHI);
+				break;
+				
+			case CrossCompiler::SHADER_STAGE_VERTEX:
+				CommandEncoder.SetShaderTexture(SF_Vertex, Surface->Texture, BindIndex);
+				FShaderCache::SetTexture(SF_Vertex, BindIndex, TextureRHI);
+				break;
+				
+			case CrossCompiler::SHADER_STAGE_COMPUTE:
+				CommandEncoder.SetShaderTexture(SF_Compute, Surface->Texture, BindIndex);
+				FShaderCache::SetTexture(SF_Compute, BindIndex, TextureRHI);
+				break;
+				
+			default:
+				check(0);
+				break;
+		}
+	}
+	else
+	{
+		switch (ShaderStage)
+		{
+			case CrossCompiler::SHADER_STAGE_PIXEL:
+				CommandEncoder.SetShaderTexture(SF_Pixel, nil, BindIndex);
+				FShaderCache::SetTexture(SF_Pixel, BindIndex, nullptr);
+				break;
+				
+			case CrossCompiler::SHADER_STAGE_VERTEX:
+				CommandEncoder.SetShaderTexture(SF_Vertex, nil, BindIndex);
+				FShaderCache::SetTexture(SF_Vertex, BindIndex, nullptr);
+				break;
+				
+			case CrossCompiler::SHADER_STAGE_COMPUTE:
+				CommandEncoder.SetShaderTexture(SF_Compute, nil, BindIndex);
+				FShaderCache::SetTexture(SF_Compute, BindIndex, nullptr);
+				break;
+				
+			default:
+				check(0);
+				break;
+		}
+	}
+}
+
+void FMetalContext::SetResource(uint32 ShaderStage, uint32 BindIndex, FMetalSamplerState* RESTRICT SamplerState)
+{
+	check(SamplerState->State != nil);
+	switch (ShaderStage)
+	{
+		case CrossCompiler::SHADER_STAGE_PIXEL:
+			CommandEncoder.SetShaderSamplerState(SF_Pixel, SamplerState->State, BindIndex);
+			FShaderCache::SetSamplerState(SF_Pixel, BindIndex, SamplerState);
+			break;
+			
+		case CrossCompiler::SHADER_STAGE_VERTEX:
+			CommandEncoder.SetShaderSamplerState(SF_Vertex, SamplerState->State, BindIndex);
+			FShaderCache::SetSamplerState(SF_Vertex, BindIndex, SamplerState);
+			break;
+			
+		case CrossCompiler::SHADER_STAGE_COMPUTE:
+			CommandEncoder.SetShaderSamplerState(SF_Compute, SamplerState->State, BindIndex);
+			FShaderCache::SetSamplerState(SF_Compute, BindIndex, SamplerState);
+			break;
+			
+		default:
+			check(0);
+			break;
+	}
+}
+
+
+template <typename MetalResourceType>
+inline int32 FMetalContext::SetShaderResourcesFromBuffer(uint32 ShaderStage, FMetalUniformBuffer* RESTRICT Buffer, const uint32* RESTRICT ResourceMap, int32 BufferIndex)
+{
+	int32 NumSetCalls = 0;
+	uint32 BufferOffset = ResourceMap[BufferIndex];
+	if (BufferOffset > 0)
+	{
+		const uint32* RESTRICT ResourceInfos = &ResourceMap[BufferOffset];
+		uint32 ResourceInfo = *ResourceInfos++;
+		do
+		{
+			checkSlow(FRHIResourceTableEntry::GetUniformBufferIndex(ResourceInfo) == BufferIndex);
+			const uint16 ResourceIndex = FRHIResourceTableEntry::GetResourceIndex(ResourceInfo);
+			const uint8 BindIndex = FRHIResourceTableEntry::GetBindIndex(ResourceInfo);
+
+			// todo: could coalesce adjacent bound resources.
+			MetalResourceType* ResourcePtr = (MetalResourceType*)Buffer->RawResourceTable[ResourceIndex];
+			SetResource(ShaderStage, BindIndex, ResourcePtr);
+
+			NumSetCalls++;
+			ResourceInfo = *ResourceInfos++;
+		} while (FRHIResourceTableEntry::GetUniformBufferIndex(ResourceInfo) == BufferIndex);
+	}
+
+	return NumSetCalls;
+}
+
+template <class ShaderType>
+void FMetalContext::SetResourcesFromTables(ShaderType Shader, uint32 ShaderStage)
+{
+	checkSlow(Shader);
+
+	// Mask the dirty bits by those buffers from which the shader has bound resources.
+	uint32 DirtyBits = Shader->Bindings.ShaderResourceTable.ResourceTableBits & Shader->DirtyUniformBuffers;
+	uint32 NumSetCalls = 0;
+	while (DirtyBits)
+	{
+		// Scan for the lowest set bit, compute its index, clear it in the set of dirty bits.
+		const uint32 LowestBitMask = (DirtyBits)& (-(int32)DirtyBits);
+		const int32 BufferIndex = FMath::FloorLog2(LowestBitMask); // todo: This has a branch on zero, we know it could never be zero...
+		DirtyBits ^= LowestBitMask;
+		FMetalUniformBuffer* Buffer = (FMetalUniformBuffer*)Shader->BoundUniformBuffers[BufferIndex].GetReference();
+		check(Buffer);
+		check(BufferIndex < Shader->Bindings.ShaderResourceTable.ResourceTableLayoutHashes.Num());
+		check(Buffer->GetLayout().GetHash() == Shader->Bindings.ShaderResourceTable.ResourceTableLayoutHashes[BufferIndex]);
+		Buffer->CacheResources(ResourceTableFrameCounter);
+
+		// todo: could make this two pass: gather then set
+//		NumSetCalls += SetShaderResourcesFromBuffer<FMetalSurface>(ShaderStage, Buffer, Shader->Bindings.ShaderResourceTable.ShaderResourceViewMap.GetData(), BufferIndex);
+		if(Buffer->RawResourceTable.Num())
+		{
+			NumSetCalls += SetShaderResourcesFromBuffer<FRHITexture>(ShaderStage, Buffer, Shader->Bindings.ShaderResourceTable.TextureMap.GetData(), BufferIndex);
+			NumSetCalls += SetShaderResourcesFromBuffer<FMetalSamplerState>(ShaderStage, Buffer, Shader->Bindings.ShaderResourceTable.SamplerMap.GetData(), BufferIndex);
+		}
+	}
+	Shader->DirtyUniformBuffers = 0;
+//	SetTextureInTableCalls += NumSetCalls;
+}
+
+void FMetalContext::CommitGraphicsResourceTables()
+{
+	uint32 Start = FPlatformTime::Cycles();
+
+	TRefCountPtr<FMetalBoundShaderState> CurrentBoundShaderState = StateCache.GetBoundShaderState();
+	check(CurrentBoundShaderState);
+
+	SetResourcesFromTables(CurrentBoundShaderState->VertexShader, CrossCompiler::SHADER_STAGE_VERTEX);
+	if (IsValidRef(CurrentBoundShaderState->PixelShader))
+	{
+		SetResourcesFromTables(CurrentBoundShaderState->PixelShader, CrossCompiler::SHADER_STAGE_PIXEL);
+	}
+
+//	CommitResourceTableCycles += (FPlatformTime::Cycles() - Start);
+}
+
+void FMetalContext::CommitNonComputeShaderConstants()
+{
+	TRefCountPtr<FMetalBoundShaderState> CurrentBoundShaderState = StateCache.GetBoundShaderState();
+	StateCache.GetShaderParameters(CrossCompiler::SHADER_STAGE_VERTEX).CommitPackedUniformBuffers(CurrentBoundShaderState, nullptr, CrossCompiler::SHADER_STAGE_VERTEX, CurrentBoundShaderState->VertexShader->BoundUniformBuffers, CurrentBoundShaderState->VertexShader->UniformBuffersCopyInfo);
+	StateCache.GetShaderParameters(CrossCompiler::SHADER_STAGE_VERTEX).CommitPackedGlobals(this, CrossCompiler::SHADER_STAGE_VERTEX, CurrentBoundShaderState->VertexShader->Bindings);
+	
+	if (IsValidRef(CurrentBoundShaderState->PixelShader))
+	{
+		StateCache.GetShaderParameters(CrossCompiler::SHADER_STAGE_PIXEL).CommitPackedUniformBuffers(CurrentBoundShaderState, nullptr, CrossCompiler::SHADER_STAGE_PIXEL, CurrentBoundShaderState->PixelShader->BoundUniformBuffers, CurrentBoundShaderState->PixelShader->UniformBuffersCopyInfo);
+		StateCache.GetShaderParameters(CrossCompiler::SHADER_STAGE_PIXEL).CommitPackedGlobals(this, CrossCompiler::SHADER_STAGE_PIXEL, CurrentBoundShaderState->PixelShader->Bindings);
+	}
+}
+
+void FMetalContext::ConditionalSwitchToGraphics()
+{
+	StateCache.ConditionalSwitchToRender();
+}
+
+void FMetalContext::ConditionalSwitchToCompute()
+{
+	StateCache.ConditionalSwitchToCompute();
+}
+
+void FMetalContext::ConditionalSwitchToBlit()
+{
+	StateCache.ConditionalSwitchToBlit();
+}
+
+void FMetalContext::Dispatch(uint32 ThreadGroupCountX, uint32 ThreadGroupCountY, uint32 ThreadGroupCountZ)
+{
+	TRefCountPtr<FMetalComputeShader> CurrentComputeShader = StateCache.GetComputeShader();
+	check(CurrentComputeShader);
+
+	auto* ComputeShader = (FMetalComputeShader*)CurrentComputeShader;
+
+	SetResourcesFromTables(ComputeShader, CrossCompiler::SHADER_STAGE_COMPUTE);
+
+	TRefCountPtr<FMetalBoundShaderState> CurrentBoundShaderState = StateCache.GetBoundShaderState();
+	
+	StateCache.GetShaderParameters(CrossCompiler::SHADER_STAGE_COMPUTE).CommitPackedUniformBuffers(CurrentBoundShaderState, ComputeShader, CrossCompiler::SHADER_STAGE_COMPUTE, ComputeShader->BoundUniformBuffers, ComputeShader->UniformBuffersCopyInfo);
+	StateCache.GetShaderParameters(CrossCompiler::SHADER_STAGE_COMPUTE).CommitPackedGlobals(this, CrossCompiler::SHADER_STAGE_COMPUTE, ComputeShader->Bindings);
+	
+	MTLSize ThreadgroupCounts = MTLSizeMake(ComputeShader->NumThreadsX, ComputeShader->NumThreadsY, ComputeShader->NumThreadsZ);
+	check(ComputeShader->NumThreadsX > 0 && ComputeShader->NumThreadsY > 0 && ComputeShader->NumThreadsZ > 0);
+	MTLSize Threadgroups = MTLSizeMake(ThreadGroupCountX, ThreadGroupCountY, ThreadGroupCountZ);
+	//@todo-rco: setThreadgroupMemoryLength?
+	[CommandEncoder.GetComputeCommandEncoder() dispatchThreadgroups:Threadgroups threadsPerThreadgroup:ThreadgroupCounts];
+}
+
+void FMetalContext::StartTiming(class FMetalEventNode* EventNode)
+{
+	bool bRecreateRender = CommandEncoder.IsRenderCommandEncoderActive();
+	if (CommandEncoder.IsRenderCommandEncoderActive() || CommandEncoder.IsComputeCommandEncoderActive() || CommandEncoder.IsBlitCommandEncoderActive())
+	{
+		CommandEncoder.EndEncoding();
+	}
+	
+    if(EventNode && CurrentCommandBuffer)
+    {
+        EventNode->Start(CurrentCommandBuffer);
+    }
+    
+    // kick the whole buffer
+	// Commit to hand the commandbuffer off to the gpu
+	CommandEncoder.CommitCommandBuffer(false);
+	
+	//once a commandbuffer is commited it can't be added to again.
+	UNTRACK_OBJECT(CurrentCommandBuffer);
+	[CurrentCommandBuffer release];
+	
+	CreateCurrentCommandBuffer(false);
+	
+	if(bRecreateRender)
+	{
+		CommandEncoder.RestoreRenderCommandEncoding();
+	}
+}
+
+void FMetalContext::EndTiming(class FMetalEventNode* EventNode)
+{
+	bool bRecreateRender = CommandEncoder.IsRenderCommandEncoderActive();
+	if (CommandEncoder.IsRenderCommandEncoderActive() || CommandEncoder.IsComputeCommandEncoderActive() || CommandEncoder.IsBlitCommandEncoderActive())
+	{
+		CommandEncoder.EndEncoding();
+	}
+    
+	bool const bWait = EventNode->Wait();
+	EventNode->Stop(CurrentCommandBuffer);
+	
+	// kick the whole buffer
+	// Commit to hand the commandbuffer off to the gpu
+	CommandEncoder.CommitCommandBuffer(bWait);
+    
+    //once a commandbuffer is commited it can't be added to again.
+	UNTRACK_OBJECT(CurrentCommandBuffer);
+	[CurrentCommandBuffer release];
+
+    CreateCurrentCommandBuffer(false);
+	
+	if(bRecreateRender)
+	{
+		CommandEncoder.RestoreRenderCommandEncoding();
+	}
+}
\ No newline at end of file
diff --git a/Engine/Source/Runtime/Apple/MetalRHI/Private/MetalContext.h b/Engine/Source/Runtime/Apple/MetalRHI/Private/MetalContext.h
new file mode 100644
--- /dev/null
+++ b/Engine/Source/Runtime/Apple/MetalRHI/Private/MetalContext.h
@@ -0,0 +1,207 @@
+// Copyright 1998-2015 Epic Games, Inc. All Rights Reserved.
+
+
+#pragma once
+
+#import <Metal/Metal.h>
+#import <QuartzCore/CAMetalLayer.h>
+#include "MetalViewport.h"
+#include "MetalBufferPools.h"
+#include "MetalCommandEncoder.h"
+#include "MetalCommandQueue.h"
+#if PLATFORM_IOS
+#include "IOSView.h"
+#endif
+
+#define NUM_SAFE_FRAMES 4
+
+/**
+ * Enumeration of features which are present only on some OS/device combinations.
+ * These have to be checked at runtime as well as compile time to ensure backward compatibility.
+ */
+enum EMetalFeatures
+{
+	/** Support for separate front & back stencil ref. values */
+	EMetalFeaturesSeparateStencil = 1 << 0,
+	/** Support for specifying an update to the buffer offset only */
+	EMetalFeaturesSetBufferOffset = 1 << 1,
+	/** Support for specifying the depth clip mode */
+	EMetalFeaturesDepthClipMode = 1 << 2,
+	/** Support for specifying resource usage & memory options */
+	EMetalFeaturesResourceOptions = 1 << 3
+};
+
+class FMetalContext
+{
+public:
+	FMetalContext();
+	~FMetalContext();
+	
+	id<MTLDevice> GetDevice();
+	FMetalCommandQueue& GetCommandQueue();
+	FMetalCommandEncoder& GetCommandEncoder();
+	id<MTLRenderCommandEncoder> GetRenderContext();
+	id<MTLBlitCommandEncoder> GetBlitContext();
+	id<MTLCommandBuffer> GetCurrentCommandBuffer();
+	FMetalStateCache& GetCurrentState() { return StateCache; }
+	bool SupportsFeature(EMetalFeatures InFeature) { return ((Features & InFeature) != 0); }
+	
+	/** Return an auto-released command buffer, caller will need to retain it if it needs to live awhile */
+	id<MTLCommandBuffer> CreateCommandBuffer(bool bRetainReferences)
+	{
+		return bRetainReferences ? CommandQueue.CreateRetainedCommandBuffer() : CommandQueue.CreateUnretainedCommandBuffer();
+	}
+	
+	FMetalBufferPool* GetBufferPool();
+	FMetalPooledBuffer CreatePooledBuffer(FMetalPooledBufferArgs const& Args);
+	void ReleasePooledBuffer(FMetalPooledBuffer Buf);
+	void ReleaseObject(id Object);
+	
+	/**
+	 * Handle rendering thread starting/stopping
+	 */
+	void CreateAutoreleasePool();
+	void DrainAutoreleasePool();
+	
+	void BeginFrame();
+	void EndFrame();
+	
+	/** RHIBeginScene helper */
+	void BeginScene();
+	/** RHIEndScene helper */
+	void EndScene();
+
+	void BeginDrawingViewport(FMetalViewport* Viewport);
+	void EndDrawingViewport(FMetalViewport* Viewport, bool bPresent);
+
+	/**
+	 * Do anything necessary to prepare for any kind of draw call 
+	 * @param PrimitiveType The UE4 primitive type for the draw call, needed to compile the correct render pipeline.
+	 */
+	void PrepareToDraw(uint32 PrimitiveType);
+	
+	/**
+	 * Set the color, depth and stencil render targets, and then make the new command buffer/encoder
+	 */
+	void SetRenderTargetsInfo(const FRHISetRenderTargetsInfo& RenderTargetsInfo);
+	
+	/**
+	 * Allocate from a dynamic ring buffer - by default align to the allowed alignment for offset field when setting buffers
+	 */
+	uint32 AllocateFromRingBuffer(uint32 Size, uint32 Alignment=0);
+	id<MTLBuffer> GetRingBuffer()
+	{
+		return RingBuffer.Buffer;
+	}
+
+	FMetalQueryBufferPool& GetQueryBufferPool()
+	{
+		return QueryBuffer;
+	}
+
+    void SubmitCommandsHint();
+	void SubmitCommandBufferAndWait();
+	void SubmitComputeCommandBufferAndWait();
+	void ResetRenderCommandEncoder();
+
+	void Dispatch(uint32 ThreadGroupCountX, uint32 ThreadGroupCountY, uint32 ThreadGroupCountZ);
+
+	void StartTiming(class FMetalEventNode* EventNode);
+	void EndTiming(class FMetalEventNode* EventNode);
+
+protected:
+	void InitFrame();
+
+	/** Create & set the current command buffer, waiting on outstanding command buffers if required. */
+	void CreateCurrentCommandBuffer(bool bWait);
+
+	/**
+	 * Possibly switch from compute to graphics
+	 */
+	void ConditionalSwitchToGraphics();
+	
+	/**
+	 * Possibly switch from graphics to compute
+	 */
+	void ConditionalSwitchToCompute();
+	
+	/**
+	 * Switch to blitting
+	 */
+	void ConditionalSwitchToBlit();
+	
+	/** Apply the SRT before drawing */
+	void CommitGraphicsResourceTables();
+	
+	void CommitNonComputeShaderConstants();
+	
+private:
+	FORCEINLINE void SetResource(uint32 ShaderStage, uint32 BindIndex, FRHITexture* RESTRICT TextureRHI);
+	
+	FORCEINLINE void SetResource(uint32 ShaderStage, uint32 BindIndex, FMetalSamplerState* RESTRICT SamplerState);
+	
+	template <typename MetalResourceType>
+	inline int32 SetShaderResourcesFromBuffer(uint32 ShaderStage, FMetalUniformBuffer* RESTRICT Buffer, const uint32* RESTRICT ResourceMap, int32 BufferIndex);
+	
+	template <class ShaderType>
+	void SetResourcesFromTables(ShaderType Shader, uint32 ShaderStage);
+	
+protected:
+	/** Dynamic buffer pool */
+	FMetalBufferPool BufferPool;
+	
+	/** The underlying Metal device */
+	id<MTLDevice> Device;
+	
+	/** The wrapper around the device command-queue for creating & committing command buffers to */
+	FMetalCommandQueue CommandQueue;
+	
+	/** The wrapper for encoding commands into the current command buffer. */
+	FMetalCommandEncoder CommandEncoder;
+	
+	/** The cache of all tracked & accessible state. */
+	FMetalStateCache StateCache;
+	
+	/** The current command buffer that receives new commands. */
+	id<MTLCommandBuffer> CurrentCommandBuffer;
+	
+	/** A sempahore used to ensure that wait for previous frames to complete if more are in flight than we permit */
+	dispatch_semaphore_t CommandBufferSemaphore;
+	
+	/** Free lists for releasing objects only once it is safe to do so */
+	TSet<id> FreeList;
+	struct FMetalDelayedFreeList
+	{
+		FEvent* Signal;
+		TSet<id> FreeList;
+	};
+	TArray<FMetalDelayedFreeList*> DelayedFreeLists;
+	
+	/** A simple fixed-size ring buffer for dynamic data */
+	FRingBuffer RingBuffer;
+	
+	/** A pool of buffers for writing visibility query results. */
+	FMetalQueryBufferPool QueryBuffer;
+	
+	/** the slot to store a per-thread autorelease pool */
+	uint32 AutoReleasePoolTLSSlot;
+	
+	/** Internal frame counter, incremented on each call to RHIBeginScene. */
+	uint32 SceneFrameCounter;
+	
+	/** Internal frame counter, used to ensure that we only drain the buffer pool one after each frame within RHIEndFrame. */
+	uint32 FrameCounter;
+	
+	/**
+	 * Internal counter used for resource table caching.
+	 * INDEX_NONE means caching is not allowed.
+	 */
+	uint32 ResourceTableFrameCounter;
+	
+	/** Event for coordinating pausing of render thread to keep inline with the ios display link. */
+	FEvent* FrameReadyEvent;
+	
+	/** Bitfield of supported Metal features with varying availability depending on OS/device */
+	uint32 Features;
+};
+
diff --git a/Engine/Source/Runtime/Apple/MetalRHI/Private/MetalIndexBuffer.cpp b/Engine/Source/Runtime/Apple/MetalRHI/Private/MetalIndexBuffer.cpp
new file mode 100644
--- /dev/null
+++ b/Engine/Source/Runtime/Apple/MetalRHI/Private/MetalIndexBuffer.cpp
@@ -0,0 +1,114 @@
+// Copyright 1998-2015 Epic Games, Inc. All Rights Reserved.
+
+/*=============================================================================
+	MetalIndexBuffer.cpp: Metal Index buffer RHI implementation.
+=============================================================================*/
+
+#include "MetalRHIPrivate.h"
+
+
+/** Constructor */
+FMetalIndexBuffer::FMetalIndexBuffer(uint32 InStride, uint32 InSize, uint32 InUsage)
+	: FRHIIndexBuffer(InStride, InSize, InUsage)
+	, LockOffset(0)
+	, LockSize(0)
+	, IndexType((InStride == 2) ? MTLIndexTypeUInt16 : MTLIndexTypeUInt32)
+{
+	MTLStorageMode Mode = (InUsage & BUF_Volatile) ? MTLStorageModeShared : BUFFER_STORAGE_MODE;
+	FMetalPooledBuffer Buf = GetMetalDeviceContext().CreatePooledBuffer(FMetalPooledBufferArgs(GetMetalDeviceContext().GetDevice(), InSize, Mode));
+	Buffer = [Buf.Buffer retain];
+	TRACK_OBJECT(Buffer);
+}
+
+FMetalIndexBuffer::~FMetalIndexBuffer()
+{
+	SafeReleasePooledBuffer(Buffer);
+	[Buffer release];
+}
+
+void* FMetalIndexBuffer::Lock(EResourceLockMode LockMode, uint32 Offset, uint32 Size)
+{
+	check(LockOffset == 0 && LockSize == 0);
+	
+	// In order to properly synchronise the buffer access, when a dynamic buffer is locked for writing, discard the old buffer & create a new one. This prevents writing to a buffer while it is being read by the GPU & thus causing corruption. This matches the logic of other RHIs.
+	if ((GetUsage() & BUFFER_DYNAMIC_REALLOC) && LockMode == RLM_WriteOnly)
+	{
+		id<MTLBuffer> OldBuffer = Buffer;
+		MTLStorageMode Mode = (GetUsage() & BUF_Volatile) ? MTLStorageModeShared : BUFFER_STORAGE_MODE;
+		FMetalPooledBuffer Buf = GetMetalDeviceContext().CreatePooledBuffer(FMetalPooledBufferArgs(GetMetalDeviceContext().GetDevice(), GetSize(), Mode));
+		Buffer = [Buf.Buffer retain];
+		GetMetalDeviceContext().ReleasePooledBuffer(OldBuffer);
+		[OldBuffer release];
+	}
+	
+	if(LockMode != RLM_ReadOnly)
+	{
+		LockOffset = Offset;
+		LockSize = Size;
+	}
+#if PLATFORM_MAC
+	else if(Buffer.storageMode == MTLStorageModeManaged)
+	{
+		SCOPE_CYCLE_COUNTER(STAT_MetalBufferPageOffTime);
+		
+		// Synchronise the buffer with the CPU
+		id<MTLBlitCommandEncoder> Blitter = GetMetalDeviceContext().GetBlitContext();
+		[Blitter synchronizeResource:Buffer];
+		
+		//kick the current command buffer.
+		GetMetalDeviceContext().SubmitCommandBufferAndWait();
+	}
+#endif
+	
+	return ((uint8*)[Buffer contents]) + Offset;
+}
+
+void FMetalIndexBuffer::Unlock()
+{
+#if PLATFORM_MAC
+	if(LockSize && Buffer.storageMode == MTLStorageModeManaged)
+	{
+		[Buffer didModifyRange:NSMakeRange(LockOffset, LockSize)];
+	}
+#endif
+	LockOffset = 0;
+	LockSize = 0;
+}
+
+FIndexBufferRHIRef FMetalDynamicRHI::RHICreateIndexBuffer(uint32 Stride,uint32 Size, uint32 InUsage, FRHIResourceCreateInfo& CreateInfo)
+{
+	// make the RHI object, which will allocate memory
+	FMetalIndexBuffer* IndexBuffer = new FMetalIndexBuffer(Stride, Size, InUsage);
+	
+	if (CreateInfo.ResourceArray)
+	{
+		check(Size == CreateInfo.ResourceArray->GetResourceDataSize());
+
+		// make a buffer usable by CPU
+		void* Buffer = RHILockIndexBuffer(IndexBuffer, 0, Size, RLM_WriteOnly);
+
+		// copy the contents of the given data into the buffer
+		FMemory::Memcpy(Buffer, CreateInfo.ResourceArray->GetResourceData(), Size);
+
+		RHIUnlockIndexBuffer(IndexBuffer);
+
+		// Discard the resource array's contents.
+		CreateInfo.ResourceArray->Discard();
+	}
+
+	return IndexBuffer;
+}
+
+void* FMetalDynamicRHI::RHILockIndexBuffer(FIndexBufferRHIParamRef IndexBufferRHI, uint32 Offset, uint32 Size, EResourceLockMode LockMode)
+{
+	FMetalIndexBuffer* IndexBuffer = ResourceCast(IndexBufferRHI);
+
+	return (uint8*)IndexBuffer->Lock(LockMode, Offset, Size);
+}
+
+void FMetalDynamicRHI::RHIUnlockIndexBuffer(FIndexBufferRHIParamRef IndexBufferRHI)
+{
+	FMetalIndexBuffer* IndexBuffer = ResourceCast(IndexBufferRHI);
+
+	IndexBuffer->Unlock();
+}
diff --git a/Engine/Source/Runtime/Apple/MetalRHI/Private/MetalProfiler.cpp b/Engine/Source/Runtime/Apple/MetalRHI/Private/MetalProfiler.cpp
new file mode 100644
--- /dev/null
+++ b/Engine/Source/Runtime/Apple/MetalRHI/Private/MetalProfiler.cpp
@@ -0,0 +1,381 @@
+// Copyright 1998-2015 Epic Games, Inc. All Rights Reserved.
+
+#pragma once
+
+#include "MetalRHIPrivate.h"
+#include "MetalProfiler.h"
+
+DEFINE_STAT(STAT_MetalMakeDrawableTime);
+DEFINE_STAT(STAT_MetalDrawCallTime);
+DEFINE_STAT(STAT_MetalPrepareDrawTime);
+DEFINE_STAT(STAT_MetalUniformBufferCleanupTime);
+DEFINE_STAT(STAT_MetalFreeUniformBufferMemory);
+DEFINE_STAT(STAT_MetalNumFreeUniformBuffers);
+DEFINE_STAT(STAT_MetalPipelineStateTime);
+DEFINE_STAT(STAT_MetalBoundShaderStateTime);
+DEFINE_STAT(STAT_MetalVertexDeclarationTime);
+DEFINE_STAT(STAT_MetalBufferPageOffTime);
+DEFINE_STAT(STAT_MetalTexturePageOffTime);
+
+#if METAL_STATISTICS
+void FMetalEventNode::GetStats(FMetalPipelineStats& OutStats)
+{
+	for(auto Entry : DrawStats)
+	{
+		FMetalPipelineStats const& DrawStat = Entry->GetResult();
+		OutStats.RHIPrimitives += DrawStat.RHIPrimitives;
+		OutStats.RHIVertices += DrawStat.RHIVertices;
+		OutStats.VertexFunctionCost += DrawStat.VertexFunctionCost;
+		OutStats.FragmentFunctionCost += DrawStat.FragmentFunctionCost;
+		OutStats.DrawCallTime += DrawStat.DrawCallTime;
+		OutStats.InputVertices += DrawStat.InputVertices;
+		OutStats.InputPrimitives += DrawStat.InputPrimitives;
+		OutStats.VertexFunctionInvocations += DrawStat.VertexFunctionInvocations;
+		OutStats.ClipperInvocations += DrawStat.ClipperInvocations;
+		OutStats.ClipperPrimitives += DrawStat.ClipperPrimitives;
+		OutStats.FragmentFunctionInvocations += DrawStat.FragmentFunctionInvocations;
+	}
+	
+	OutStats.VertexFunctionCost = NumDraws > 0 ? OutStats.VertexFunctionCost / NumDraws : 0;
+	OutStats.FragmentFunctionCost = NumDraws > 0 ? OutStats.FragmentFunctionCost / NumDraws : 0;
+}
+
+/** Recursively dumps stats for each node with a depth first traversal. */
+static void DumpStatsEventNode(FMetalEventNode* Node, float RootResult, int32 Depth, int32& NumNodes, int32& NumDraws)
+{
+	NumNodes++;
+	if (Node->NumDraws > 0 || Node->Children.Num() > 0)
+	{
+		NumDraws += Node->NumDraws;
+		// Percent that this node was of the total frame time
+		const float Percent = Node->TimingResult * 100.0f / (RootResult * 1000.0f);
+		
+		const int32 EffectiveDepth = FMath::Max(Depth - 1, 0);
+		
+		FMetalPipelineStats Stats;
+		FMemory::Memzero(Stats);
+		Node->GetStats(Stats);
+		
+		// Print information about this node, padded to its depth in the tree
+		float DrawCallTime = FPlatformTime::ToMilliseconds(Stats.DrawCallTime);
+		UE_LOG(LogRHI, Warning, TEXT("%s%4.1f%%%5.2fms (%5.2fms)   %s %u draws %u (%u) prims %u (%u) verts %u vert invoke %4.1f%% vert cost %u clip invoke %u clip prims %u pixel invoke %4.1f%% pixel cost"),
+			   *FString(TEXT("")).LeftPad(EffectiveDepth * 3),
+			   Percent,
+			   Node->TimingResult,
+			   DrawCallTime,
+			   *Node->Name,
+			   Node->NumDraws,
+			   Stats.RHIPrimitives,
+			   Stats.InputPrimitives,
+			   Stats.RHIVertices,
+			   Stats.InputVertices,
+			   Stats.VertexFunctionInvocations,
+			   Stats.VertexFunctionCost,
+			   Stats.ClipperInvocations,
+			   Stats.ClipperPrimitives,
+			   Stats.FragmentFunctionInvocations,
+			   Stats.FragmentFunctionCost
+			   );
+		
+		float TotalChildTime = 0;
+		uint32 TotalChildDraws = 0;
+		for (int32 ChildIndex = 0; ChildIndex < Node->Children.Num(); ChildIndex++)
+		{
+			FMetalEventNode* ChildNode = (FMetalEventNode*)(Node->Children[ChildIndex].GetReference());
+			
+			int32 NumChildDraws = 0;
+			// Traverse children
+			DumpStatsEventNode(ChildNode, RootResult, Depth + 1, NumNodes, NumChildDraws);
+			NumDraws += NumChildDraws;
+			
+			TotalChildTime += ChildNode->TimingResult;
+			TotalChildDraws += NumChildDraws;
+		}
+		
+		const float UnaccountedTime = FMath::Max(Node->TimingResult - TotalChildTime, 0.0f);
+		const float UnaccountedPercent = UnaccountedTime * 100.0f / (RootResult * 1000.0f);
+		
+		// Add an 'Other Children' node if necessary to show time spent in the current node that is not in any of its children
+		if (Node->Children.Num() > 0 && TotalChildDraws > 0 && (UnaccountedPercent > 2.0f || UnaccountedTime > .2f))
+		{
+			UE_LOG(LogRHI, Warning, TEXT("%s%4.1f%%%5.2fms Unaccounted"),
+				   *FString(TEXT("")).LeftPad((EffectiveDepth + 1) * 3),
+				   UnaccountedPercent,
+				   UnaccountedTime);
+		}
+	}
+}
+#endif
+
+FMetalDrawProfiler::FMetalDrawProfiler(struct FMetalGPUProfiler* InProfiler, uint32 InStartPoint, uint32 InEndPoint, uint32 NumPrimitives, uint32 NumVertices)
+: Profiler(InProfiler)
+, StartPoint(InStartPoint)
+, EndPoint(InEndPoint)
+{
+	if(Profiler)
+	{
+		Profiler->StartGPUWork(InStartPoint, InEndPoint, NumPrimitives, NumVertices);
+	}
+}
+
+FMetalDrawProfiler::~FMetalDrawProfiler()
+{
+	if(Profiler)
+	{
+		Profiler->FinishGPUWork();
+	}
+}
+
+FMetalEventNode::~FMetalEventNode()
+{
+#if METAL_STATISTICS
+	for(auto Entry : DrawStats)
+	{
+		delete Entry;
+	}
+#endif
+}
+
+float FMetalEventNode::GetTiming()
+{
+	return FPlatformTime::ToSeconds(EndTime - StartTime);
+}
+
+void FMetalEventNode::StartTiming()
+{
+	StartTime = 0;
+	EndTime = 0;
+	
+	Context->StartTiming(this);
+}
+
+void FMetalEventNode::Start(id<MTLCommandBuffer> Buffer)
+{
+	check(Buffer);
+	[Buffer addCompletedHandler:^(id<MTLCommandBuffer>)
+	{
+		StartTime = mach_absolute_time();
+	}];
+}
+
+void FMetalEventNode::StopTiming()
+{
+	Context->EndTiming(this);
+}
+
+void FMetalEventNode::Stop(id<MTLCommandBuffer> Buffer)
+{
+	[Buffer addCompletedHandler:^(id<MTLCommandBuffer>)
+	{
+		EndTime = mach_absolute_time();
+     
+        if(bRoot)
+        {
+            GGPUFrameTime = FMath::TruncToInt( double(GetTiming()) / double(FPlatformTime::GetSecondsPerCycle()) );
+            if(!bFullProfiling)
+            {
+                delete this;
+            }
+        }
+	}];
+}
+
+#if METAL_STATISTICS
+void FMetalEventNode::StartDraw(bool bActiveStats, uint32 StartPoint, uint32 EndPoint, uint32 NumPrimitives, uint32 NumVertices)
+{
+	if(Context->GetCommandQueue().GetStatistics() && bActiveStats)
+	{
+		DrawStats.Add(Context->GetCommandQueue().GetStatistics()->CreateDrawStats(Context->GetCurrentCommandBuffer(), StartPoint, EndPoint, NumPrimitives, NumVertices));
+	}
+}
+
+void FMetalEventNode::StopDraw(void)
+{
+	if(DrawStats.Num())
+	{
+		DrawStats.Last()->End();
+	}
+}
+#endif
+	
+/** Start this frame of per tracking */
+void FMetalEventNodeFrame::StartFrame()
+{
+	RootNode->StartTiming();
+}
+
+/** End this frame of per tracking, but do not block yet */
+void FMetalEventNodeFrame::EndFrame()
+{
+	RootNode->StopTiming();
+}
+
+/** Calculates root timing base frequency (if needed by this RHI) */
+float FMetalEventNodeFrame::GetRootTimingResults()
+{
+    return RootNode->GetTiming();
+}
+
+void FMetalEventNodeFrame::LogDisjointQuery()
+{
+	
+}
+
+void FMetalGPUProfiler::Cleanup()
+{
+	
+}
+
+void FMetalGPUProfiler::PushEvent(const TCHAR* Name)
+{
+	FGPUProfiler::PushEvent(Name);
+}
+
+void FMetalGPUProfiler::PopEvent()
+{
+	FGPUProfiler::PopEvent();
+}
+
+TGlobalResource<FVector4VertexDeclaration> GMetalVector4VertexDeclaration;
+TGlobalResource<FTexture> GMetalLongTaskRT;
+
+static FGlobalBoundShaderState LongGPUTaskBoundShaderState;
+
+void FMetalGPUProfiler::BeginFrame()
+{
+    if(!CurrentEventNodeFrame)
+    {
+        // Start tracking the frame
+        CurrentEventNodeFrame = new FMetalEventNodeFrame(Context, GTriggerGPUProfile);
+        CurrentEventNodeFrame->StartFrame();
+        
+        if(GTriggerGPUProfile)
+        {
+            bTrackingEvents = true;
+            bLatchedGProfilingGPU = true;
+            
+            GTriggerGPUProfile = false;
+
+#if METAL_STATISTICS
+			if(Context->GetCommandQueue().GetStatistics())
+            {
+				Context->GetCommandQueue().GetStatistics()->BeginSamplingStatistics();
+				bActiveStats = true;
+            }
+#endif
+			
+			/*if (bLatchedGProfilingGPU)
+			{
+				// Issue a bunch of GPU work at the beginning of the frame, to make sure that we are GPU bound
+				// We can't isolate idle time from GPU timestamps
+				const auto FeatureLevel = GMaxRHIFeatureLevel;
+				
+				FRHICommandList_RecursiveHazardous RHICmdList(InRHI);
+				if(!IsValidRef(GMetalLongTaskRT.TextureRHI))
+				{
+					FRHIResourceCreateInfo Info;
+					GMetalLongTaskRT.TextureRHI = RHICreateTexture2D(1920, 1080, PF_B8G8R8A8, 1, 1, TexCreate_RenderTargetable, Info);
+				}
+				SetRenderTarget(RHICmdList, GMetalLongTaskRT.TextureRHI->GetTexture2D(), FTextureRHIRef());
+				RHICmdList.SetBlendState(TStaticBlendState<CW_RGBA, BO_Add, BF_One, BF_One>::GetRHI(), FLinearColor::Black);
+				RHICmdList.SetDepthStencilState(TStaticDepthStencilState<false, CF_Always>::GetRHI(), 0);
+				RHICmdList.SetRasterizerState(TStaticRasterizerState<FM_Solid, CM_None>::GetRHI());
+				
+				auto ShaderMap = GetGlobalShaderMap(FeatureLevel);
+				TShaderMapRef<TOneColorVS<true> > VertexShader(ShaderMap);
+				TShaderMapRef<FLongGPUTaskPS> PixelShader(ShaderMap);
+				
+				SetGlobalBoundShaderState(RHICmdList, FeatureLevel, LongGPUTaskBoundShaderState, GMetalVector4VertexDeclaration.VertexDeclarationRHI, *VertexShader, *PixelShader, 0);
+				
+				// Draw a fullscreen quad
+				FVector4 Vertices[4];
+				Vertices[0].Set( -1.0f,  1.0f, 0, 1.0f );
+				Vertices[1].Set(  1.0f,  1.0f, 0, 1.0f );
+				Vertices[2].Set( -1.0f, -1.0f, 0, 1.0f );
+				Vertices[3].Set(  1.0f, -1.0f, 0, 1.0f );
+				DrawPrimitiveUP(RHICmdList, PT_TriangleStrip, 2, Vertices, sizeof(Vertices[0]));
+				
+				RHICmdList.SubmitCommandsHint();
+				// RHICmdList flushes on destruction
+			}*/
+        }
+        
+        PushEvent(TEXT("FRAME"));
+    }
+    NumNestedFrames++;
+}
+
+void FMetalGPUProfiler::EndFrame()
+{
+    if(--NumNestedFrames == 0)
+    {
+        PopEvent();
+        
+        if(CurrentEventNodeFrame)
+        {
+            CurrentEventNodeFrame->EndFrame();
+			
+			if(bLatchedGProfilingGPU)
+			{
+#if METAL_STATISTICS
+				if(Context->GetCommandQueue().GetStatistics())
+				{
+					Context->GetCommandQueue().GetStatistics()->FinishSamplingStatistics();
+					bActiveStats = false;
+				}
+#endif
+				
+                bTrackingEvents = false;
+                bLatchedGProfilingGPU = false;
+            
+                UE_LOG(LogRHI, Warning, TEXT(""));
+                UE_LOG(LogRHI, Warning, TEXT(""));
+                CurrentEventNodeFrame->DumpEventTree();
+				
+#if METAL_STATISTICS
+				if(Context->GetCommandQueue().GetStatistics())
+				{
+					float RootResult = CurrentEventNodeFrame->GetRootTimingResults();
+					
+					UE_LOG(LogRHI, Warning, TEXT(""));
+					UE_LOG(LogRHI, Warning, TEXT(""));
+					UE_LOG(LogRHI, Warning, TEXT("Pipeline stats hierarchy, total GPU time %.2fms"), RootResult * 1000.0f);
+					
+					int32 NumNodes = 0;
+					int32 NumDraws = 0;
+					for (int32 BaseNodeIndex = 0; BaseNodeIndex < CurrentEventNodeFrame->EventTree.Num(); BaseNodeIndex++)
+					{
+						DumpStatsEventNode((FMetalEventNode*)(CurrentEventNodeFrame->EventTree[BaseNodeIndex].GetReference()), RootResult, 0, NumNodes, NumDraws);
+					}
+				}
+#endif
+            }
+			
+            delete CurrentEventNodeFrame;
+            CurrentEventNodeFrame = NULL;
+        }
+    }
+}
+
+void FMetalGPUProfiler::StartGPUWork(uint32 StartPoint, uint32 EndPoint, uint32 NumPrimitives, uint32 NumVertices)
+{
+	if(CurrentEventNode)
+	{
+		RegisterGPUWork(NumPrimitives, NumVertices);
+#if METAL_STATISTICS
+		FMetalEventNode* EventNode = (FMetalEventNode*)CurrentEventNode;
+		EventNode->StartDraw(bActiveStats, StartPoint, EndPoint, NumPrimitives, NumVertices);
+#endif
+	}
+}
+
+void FMetalGPUProfiler::FinishGPUWork(void)
+{
+#if METAL_STATISTICS
+	if(CurrentEventNode)
+	{
+		FMetalEventNode* EventNode = (FMetalEventNode*)CurrentEventNode;
+		EventNode->StopDraw();
+	}
+#endif
+}
diff --git a/Engine/Source/Runtime/Apple/MetalRHI/Private/MetalProfiler.h b/Engine/Source/Runtime/Apple/MetalRHI/Private/MetalProfiler.h
new file mode 100644
--- /dev/null
+++ b/Engine/Source/Runtime/Apple/MetalRHI/Private/MetalProfiler.h
@@ -0,0 +1,202 @@
+// Copyright 1998-2015 Epic Games, Inc. All Rights Reserved.
+
+#pragma once
+
+#include "MetalRHIPrivate.h"
+#include "MetalCommandQueue.h"
+#include "GPUProfiler.h"
+
+#if METAL_STATISTICS
+#include "Runtime/Mac/NoRedist/MetalStatistics/Public/MetalStatistics.h"
+#endif
+
+// Stats
+DECLARE_CYCLE_STAT_EXTERN(TEXT("MakeDrawable time"),STAT_MetalMakeDrawableTime,STATGROUP_MetalRHI, );
+DECLARE_CYCLE_STAT_EXTERN(TEXT("Draw call time"),STAT_MetalDrawCallTime,STATGROUP_MetalRHI, );
+DECLARE_CYCLE_STAT_EXTERN(TEXT("PrepareDraw time"),STAT_MetalPrepareDrawTime,STATGROUP_MetalRHI, );
+DECLARE_CYCLE_STAT_EXTERN(TEXT("PipelineState time"),STAT_MetalPipelineStateTime,STATGROUP_MetalRHI, );
+DECLARE_CYCLE_STAT_EXTERN(TEXT("BoundShaderState time"),STAT_MetalBoundShaderStateTime,STATGROUP_MetalRHI, );
+DECLARE_CYCLE_STAT_EXTERN(TEXT("VertexDeclaration time"),STAT_MetalVertexDeclarationTime,STATGROUP_MetalRHI, );
+DECLARE_CYCLE_STAT_EXTERN(TEXT("Uniform buffer pool cleanup time"), STAT_MetalUniformBufferCleanupTime, STATGROUP_MetalRHI, );
+DECLARE_CYCLE_STAT_EXTERN(TEXT("Buffer Page-Off time"), STAT_MetalBufferPageOffTime, STATGROUP_MetalRHI, );
+DECLARE_CYCLE_STAT_EXTERN(TEXT("Texture Page-Off time"), STAT_MetalTexturePageOffTime, STATGROUP_MetalRHI, );
+DECLARE_MEMORY_STAT_EXTERN(TEXT("Uniform buffer pool memory"), STAT_MetalFreeUniformBufferMemory, STATGROUP_MetalRHI, );
+DECLARE_DWORD_ACCUMULATOR_STAT_EXTERN(TEXT("Uniform buffer pool num free"), STAT_MetalNumFreeUniformBuffers, STATGROUP_MetalRHI, );
+
+#if METAL_STATISTICS
+#define RHI_PROFILE_DRAW_CALL_STATS(StartPoint, EndPoint, NumPrims, NumVerts) FMetalDrawProfiler GPUWork(Profiler, (uint32)StartPoint, (uint32)EndPoint, NumPrims, NumVerts)
+#else
+#define RHI_PROFILE_DRAW_CALL_STATS(StartPoint, EndPoint, NumPrims, NumVerts) FMetalDrawProfiler GPUWork(Profiler, (uint32)0, (uint32)0, NumPrims, NumVerts)
+#endif
+
+class FMetalDrawProfiler
+{
+public:
+	FMetalDrawProfiler(struct FMetalGPUProfiler* InProfiler, uint32 StartPoint, uint32 EndPoint, uint32 NumPrimitives = 0, uint32 NumVertices = 0);
+	~FMetalDrawProfiler();
+	
+private:
+	struct FMetalGPUProfiler* Profiler;
+	uint32 StartPoint;
+	uint32 EndPoint;
+};
+
+/** A single perf event node, which tracks information about a appBeginDrawEvent/appEndDrawEvent range. */
+class FMetalEventNode : public FGPUProfilerEventNode
+{
+public:
+	
+	FMetalEventNode(FMetalContext* InContext, const TCHAR* InName, FGPUProfilerEventNode* InParent, bool bIsRoot, bool bInFullProfiling)
+	: FGPUProfilerEventNode(InName, InParent)
+	, Context(InContext)
+	, StartTime(0)
+	, EndTime(0)
+	, bRoot(bIsRoot)
+    , bFullProfiling(bInFullProfiling)
+	{
+	}
+	
+	virtual ~FMetalEventNode();
+	
+	/**
+	 * Returns the time in ms that the GPU spent in this draw event.
+	 * This blocks the CPU if necessary, so can cause hitching.
+	 */
+	virtual float GetTiming() override;
+	
+	virtual void StartTiming() override;
+	
+	virtual void StopTiming() override;
+	
+	void Start(id<MTLCommandBuffer> Buffer);
+	void Stop(id<MTLCommandBuffer> Buffer);
+
+	bool Wait() const { return bRoot && bFullProfiling; }
+	bool IsRoot() const { return bRoot; }
+	
+	uint64 GetCycles() { return EndTime - StartTime; }
+	
+#if METAL_STATISTICS
+	void StartDraw(bool bActiveStats, uint32 StartPoint, uint32 EndPoint, uint32 NumPrimitives, uint32 NumVertices);
+	void StopDraw(void);
+	
+	void GetStats(FMetalPipelineStats& OutStats);
+#endif
+	
+private:
+	FMetalContext* Context;
+#if METAL_STATISTICS
+	TArray<IMetalDrawStats*> DrawStats;
+#endif
+	uint64 StartTime;
+	uint64 EndTime;
+	bool bRoot;
+    bool bFullProfiling;
+};
+
+/** An entire frame of perf event nodes, including ancillary timers. */
+class FMetalEventNodeFrame : public FGPUProfilerEventNodeFrame
+{
+public:
+	FMetalEventNodeFrame(FMetalContext* InContext, bool bInFullProfiling)
+	: RootNode(new FMetalEventNode(InContext, TEXT("Frame"), nullptr, true, bInFullProfiling))
+    , bFullProfiling(bInFullProfiling)
+	{
+	}
+	
+	virtual ~FMetalEventNodeFrame()
+	{
+        if(bFullProfiling)
+        {
+            delete RootNode;
+        }
+	}
+	
+	/** Start this frame of per tracking */
+	virtual void StartFrame() override;
+	
+	/** End this frame of per tracking, but do not block yet */
+	virtual void EndFrame() override;
+	
+	/** Calculates root timing base frequency (if needed by this RHI) */
+	virtual float GetRootTimingResults() override;
+	
+	virtual void LogDisjointQuery() override;
+	
+	FMetalEventNode* RootNode;
+    bool bFullProfiling;
+};
+
+// This class has multiple inheritance but really FGPUTiming is a static class
+class FMetalGPUTiming : public FGPUTiming
+{
+public:
+	
+	/**
+	 * Constructor.
+	 */
+	FMetalGPUTiming()
+	{
+		StaticInitialize(nullptr, PlatformStaticInitialize);
+	}
+	
+	
+private:
+	
+	/**
+	 * Initializes the static variables, if necessary.
+	 */
+	static void PlatformStaticInitialize(void* UserData)
+	{
+		// Are the static variables initialized?
+		if ( !GAreGlobalsInitialized )
+		{
+			GIsSupported = true;
+			GTimingFrequency = 1000 * 1000 * 1000;
+			GAreGlobalsInitialized = true;
+		}
+	}
+};
+
+/**
+ * Encapsulates GPU profiling logic and data.
+ * There's only one global instance of this struct so it should only contain global data, nothing specific to a frame.
+ */
+struct FMetalGPUProfiler : public FGPUProfiler
+{
+	/** GPU hitch profile histories */
+	TIndirectArray<FMetalEventNodeFrame> GPUHitchEventNodeFrames;
+	
+	FMetalGPUProfiler(FMetalContext* InContext)
+	:	FGPUProfiler()
+	,	Context(InContext)
+    ,   NumNestedFrames(0)
+    ,	bActiveStats(false)
+	{
+	}
+	
+	virtual ~FMetalGPUProfiler() {}
+	
+	virtual FGPUProfilerEventNode* CreateEventNode(const TCHAR* InName, FGPUProfilerEventNode* InParent) override
+	{
+		FMetalEventNode* EventNode = new FMetalEventNode(Context, InName, InParent, false, false);
+		return EventNode;
+	}
+	
+	void Cleanup();
+	
+	virtual void PushEvent(const TCHAR* Name) override;
+	virtual void PopEvent() override;
+	
+	void BeginFrame();
+	void EndFrame();
+	
+	void StartGPUWork(uint32 StartPoint, uint32 EndPoint, uint32 NumPrimitives = 0, uint32 NumVertices = 0);
+	void FinishGPUWork(void);
+	
+	FMetalGPUTiming TimingSupport;
+	FMetalContext* Context;
+    int32 NumNestedFrames;
+    bool bActiveStats;
+};
+
