Merging MetalRHI & MetalShaderFormat changes from ZebraBranch, plus the code to enable or disable Mac support.

diff --git a/Engine/Source/Runtime/Apple/MetalRHI/Private/MetalQuery.cpp b/Engine/Source/Runtime/Apple/MetalRHI/Private/MetalQuery.cpp
new file mode 100644
--- /dev/null
+++ b/Engine/Source/Runtime/Apple/MetalRHI/Private/MetalQuery.cpp
@@ -0,0 +1,191 @@
+// Copyright 1998-2015 Epic Games, Inc. All Rights Reserved.
+
+/*=============================================================================
+	MetalQuery.cpp: Metal query RHI implementation.
+=============================================================================*/
+
+#include "MetalRHIPrivate.h"
+
+FMetalQueryBuffer::FMetalQueryBuffer(FMetalContext* InContext, id<MTLBuffer> InBuffer)
+: Context(InContext)
+, Buffer(InBuffer)
+, WriteOffset(0)
+, bCompleted(false)
+, CommandBuffer(nil)
+{
+}
+
+FMetalQueryBuffer::~FMetalQueryBuffer()
+{
+	[CommandBuffer release];
+	if(Buffer && GIsRHIInitialized)
+    {
+        Context->GetQueryBufferPool().ReleaseQueryBuffer(Buffer);
+    }
+}
+
+bool FMetalQueryBuffer::Wait(uint64 Millis)
+{
+	if(!bCompleted)
+	{
+		bool bOK = false;
+		
+		if(CommandBuffer)
+		{
+			if(CommandBuffer.status < MTLCommandBufferStatusCompleted)
+			{
+				check(CommandBuffer.status >= MTLCommandBufferStatusCommitted && CommandBuffer.status <= MTLCommandBufferStatusError);
+				[CommandBuffer waitUntilCompleted];
+				check(CommandBuffer.status >= MTLCommandBufferStatusCommitted && CommandBuffer.status <= MTLCommandBufferStatusError);
+			}
+#if !UE_BUILD_SHIPPING
+			if(CommandBuffer.status == MTLCommandBufferStatusError)
+			{
+				UE_LOG(LogMetal, Error, TEXT("Command buffer %p failed!"), CommandBuffer);
+			}
+#endif
+			bCompleted = CommandBuffer.status >= MTLCommandBufferStatusCompleted;
+			[CommandBuffer release];
+			CommandBuffer = nil;
+		}
+		
+		return bOK;
+	}
+	return true;
+}
+
+void const* FMetalQueryBuffer::GetResult(uint32 Offset)
+{
+    check(Buffer);
+    void const* Result = (void const*)(((uint8*)[Buffer contents]) + Offset);
+    return Result;
+}
+
+bool FMetalQueryResult::Wait(uint64 Millis)
+{
+	if(IsValidRef(SourceBuffer))
+	{
+		return SourceBuffer->Wait(Millis);
+	}
+	return false;
+}
+
+void const* FMetalQueryResult::GetResult()
+{
+	if(IsValidRef(SourceBuffer))
+	{
+		return SourceBuffer->GetResult(Offset);
+	}
+	return nullptr;
+}
+
+FMetalRenderQuery::FMetalRenderQuery(ERenderQueryType InQueryType)
+: Type(InQueryType)
+{
+	Buffer.Offset = 0;
+}
+
+FMetalRenderQuery::~FMetalRenderQuery()
+{
+	Buffer.SourceBuffer.SafeRelease();
+	Buffer.Offset = 0;
+}
+
+void FMetalRenderQuery::Begin(FMetalContext* Context)
+{
+	if(Type == RQT_Occlusion)
+	{
+		// these only need to be 8 byte aligned (of course, a normal allocation oafter this will be 256 bytes, making large holes if we don't do all
+		// query allocations at once)
+		Context->GetQueryBufferPool().Allocate(Buffer);
+		
+		Result = 0;
+		bAvailable = false;
+		
+	#if PLATFORM_MAC
+		Context->GetCommandEncoder().SetVisibilityResultMode(MTLVisibilityResultModeCounting, Buffer.Offset);
+	#else
+		Context->GetCommandEncoder().SetVisibilityResultMode(MTLVisibilityResultModeBoolean, Buffer.Offset);
+	#endif
+	}
+}
+
+void FMetalRenderQuery::End(FMetalContext* Context)
+{
+	if(Type == RQT_Occlusion)
+	{
+		// switch back to non-occlusion rendering
+		Context->GetCommandEncoder().SetVisibilityResultMode(MTLVisibilityResultModeDisabled, 0);
+	}
+}
+
+FRenderQueryRHIRef FMetalDynamicRHI::RHICreateRenderQuery(ERenderQueryType QueryType)
+{
+	return new FMetalRenderQuery(QueryType);
+}
+
+bool FMetalDynamicRHI::RHIGetRenderQueryResult(FRenderQueryRHIParamRef QueryRHI,uint64& OutNumPixels,bool bWait)
+{
+	check(IsInRenderingThread());
+	FMetalRenderQuery* Query = ResourceCast(QueryRHI);
+
+    if(!Query->bAvailable)
+    {
+		SCOPE_CYCLE_COUNTER( STAT_RenderQueryResultTime );
+		
+		bool bOK = false;
+		
+		if (bWait)
+		{
+			uint32 IdleStart = FPlatformTime::Cycles();
+		
+			uint64 WaitMS = 500;
+			bOK = Query->Buffer.Wait(WaitMS);
+			
+			GRenderThreadIdle[ERenderThreadIdleTypes::WaitingForGPUQuery] += FPlatformTime::Cycles() - IdleStart;
+			GRenderThreadNumIdle[ERenderThreadIdleTypes::WaitingForGPUQuery]++;
+		}
+		
+        Query->Result = 0;
+		Query->bAvailable = true; // @todo Zebra Never wait for a failed signal again.
+		
+        if (bOK == false)
+        {
+			UE_LOG(LogMetal, Log, TEXT("Timed out while waiting for GPU to catch up. (0.5f s)"));
+			return false;
+        }
+		
+		if(Query->Type == RQT_Occlusion)
+		{
+			Query->Result = *((uint64 const*)Query->Buffer.GetResult());
+		}
+		Query->Buffer.SourceBuffer.SafeRelease();
+    }
+
+    // at this point, we are ready to read the value!
+    OutNumPixels = Query->Result;
+    return true;
+}
+
+// Occlusion/Timer queries.
+void FMetalRHICommandContext::RHIBeginRenderQuery(FRenderQueryRHIParamRef QueryRHI)
+{
+	FMetalRenderQuery* Query = ResourceCast(QueryRHI);
+	
+	Query->Begin(Context);
+}
+
+void FMetalRHICommandContext::RHIEndRenderQuery(FRenderQueryRHIParamRef QueryRHI)
+{
+	FMetalRenderQuery* Query = ResourceCast(QueryRHI);
+	
+	Query->End(Context);
+}
+
+void FMetalRHICommandContext::RHIBeginOcclusionQueryBatch()
+{
+}
+
+void FMetalRHICommandContext::RHIEndOcclusionQueryBatch()
+{
+}
diff --git a/Engine/Source/Runtime/Apple/MetalRHI/Private/MetalRHI.cpp b/Engine/Source/Runtime/Apple/MetalRHI/Private/MetalRHI.cpp
new file mode 100644
--- /dev/null
+++ b/Engine/Source/Runtime/Apple/MetalRHI/Private/MetalRHI.cpp
@@ -0,0 +1,450 @@
+// Copyright 1998-2015 Epic Games, Inc. All Rights Reserved.
+
+/*=============================================================================
+	MetalRHI.cpp: Metal device RHI implementation.
+=============================================================================*/
+
+#include "MetalRHIPrivate.h"
+#if PLATFORM_IOS
+#include "IOSAppDelegate.h"
+#endif
+#include "ShaderCache.h"
+#include "MetalProfiler.h"
+
+/** Set to 1 to enable GPU events in Xcode frame debugger */
+#define ENABLE_METAL_GPUEVENTS	(UE_BUILD_DEBUG | UE_BUILD_DEVELOPMENT)
+#define ENABLE_METAL_GPUPROFILE	(ENABLE_METAL_GPUEVENTS & 1)
+
+DEFINE_LOG_CATEGORY(LogMetal)
+
+bool FMetalDynamicRHIModule::IsSupported()
+{
+	return true;
+}
+
+FDynamicRHI* FMetalDynamicRHIModule::CreateRHI()
+{
+	return new FMetalDynamicRHI();
+}
+
+IMPLEMENT_MODULE(FMetalDynamicRHIModule, MetalRHI);
+
+FMetalDynamicRHI::FMetalDynamicRHI()
+: FMetalRHICommandContext(nullptr)
+{
+	// This should be called once at the start 
+	check( IsInGameThread() );
+	check( !GIsThreadedRendering );
+	
+	// @todo Zebra This is now supported on all GPUs in Mac Metal, but not on iOS.
+	// we cannot render to a volume texture without geometry shader support
+	GSupportsVolumeTextureRendering = false;
+
+	//@todo-rco: Query name from API
+	GRHIAdapterName = TEXT("Metal");
+	GRHIVendorId = 1; // non-zero to avoid asserts
+	
+#if PLATFORM_IOS
+	// get the device to ask about capabilities
+	id<MTLDevice> Device = [IOSAppDelegate GetDelegate].IOSView->MetalDevice;
+	// A8 can use 256 bits of MRTs
+	bool bCanUseWideMRTs = [Device supportsFeatureSet:MTLFeatureSet_iOS_GPUFamily2_v1];
+	bool bCanUseASTC = [Device supportsFeatureSet:MTLFeatureSet_iOS_GPUFamily2_v1] && !FParse::Param(FCommandLine::Get(),TEXT("noastc"));
+
+    bool bProjectSupportsMRTs = false;
+    GConfig->GetBool(TEXT("/Script/IOSRuntimeSettings.IOSRuntimeSettings"), TEXT("bSupportsMetalMRT"), bProjectSupportsMRTs, GEngineIni);
+    
+    // only allow GBuffers, etc on A8s (A7s are just not going to cut it)
+    if (bProjectSupportsMRTs && bCanUseWideMRTs && FParse::Param(FCommandLine::Get(),TEXT("metalmrt")))
+    {
+        GMaxRHIFeatureLevel = ERHIFeatureLevel::SM4;
+        GMaxRHIShaderPlatform = SP_METAL_MRT;
+    }
+    else
+    {
+        GMaxRHIFeatureLevel = ERHIFeatureLevel::ES3_1;
+        GMaxRHIShaderPlatform = SP_METAL;
+    }
+
+#else // @todo zebra
+    // get the device to ask about capabilities?
+    id<MTLDevice> Device = Context->GetDevice();
+    // A8 can use 256 bits of MRTs
+    bool bCanUseWideMRTs = true;
+    bool bCanUseASTC = false;
+	bool bSupportsD24S8 = false;
+
+	if(FParse::Param(FCommandLine::Get(),TEXT("metalsm5")))
+	{
+		GMaxRHIFeatureLevel = ERHIFeatureLevel::SM5;
+		GMaxRHIShaderPlatform = SP_METAL_SM5;
+	}
+	else
+	{
+		GMaxRHIFeatureLevel = ERHIFeatureLevel::SM4;
+		GMaxRHIShaderPlatform = SP_METAL_SM4;
+	}
+	
+	GRHIAdapterName = FString(Device.name);
+	
+	if(GRHIAdapterName.Contains("Nvidia"))
+	{
+		// Nvidia support layer indexing.
+		GSupportsVolumeTextureRendering = true;
+		GRHIVendorId = 0x10DE;
+	}
+	else if(GRHIAdapterName.Contains("ATi") || GRHIAdapterName.Contains("AMD"))
+	{
+		// AMD support layer indexing.
+		GSupportsVolumeTextureRendering = true;
+		GRHIVendorId = 0x1002;
+	}
+	else if(GRHIAdapterName.Contains("Intel"))
+	{
+		GSupportsVolumeTextureRendering = true;
+		GRHIVendorId = 0x8086;
+	}
+	
+	// Change the support depth format if we can
+	bSupportsD24S8 = Device.depth24Stencil8PixelFormatSupported;
+	
+#endif
+
+    
+
+	if (FPlatformMisc::IsDebuggerPresent() && UE_BUILD_DEBUG)
+	{
+#if PLATFORM_IOS // @todo zebra : needs a RENDER_API or whatever
+		// Enable GL debug markers if we're running in Xcode
+		extern int32 GEmitMeshDrawEvent;
+		GEmitMeshDrawEvent = 1;
+#endif
+		GEmitDrawEvents = true;
+	}
+	
+	GEmitDrawEvents |= ENABLE_METAL_GPUEVENTS;
+
+	GSupportsShaderFramebufferFetch = true;
+	GHardwareHiddenSurfaceRemoval = true;
+//	GRHIAdapterName = 
+//	GRHIVendorId = 
+	GSupportsRenderTargetFormat_PF_G8 = false;
+	GSupportsQuads = false;
+	GRHISupportsTextureStreaming = true;
+	GMaxShadowDepthBufferSizeX = 4096;
+	GMaxShadowDepthBufferSizeY = 4096;
+// 	GReadTexturePoolSizeFromIni = true;
+
+	GRHISupportsBaseVertexIndex = PLATFORM_MAC && !IsRHIDeviceAMD(); // Supported on OS X but not iOS - broken on AMD presently
+    GRHISupportsFirstInstance = PLATFORM_MAC; // Supported on OS X but not iOS.
+    
+	GRHIRequiresEarlyBackBufferRenderTarget = false;
+	GSupportsSeparateRenderTargetBlendState = (GMaxRHIFeatureLevel >= ERHIFeatureLevel::SM4);
+
+	GMaxTextureDimensions = 4096;
+	GMaxTextureMipCount = FPlatformMath::CeilLogTwo( GMaxTextureDimensions ) + 1;
+	GMaxTextureMipCount = FPlatformMath::Min<int32>( MAX_TEXTURE_MIP_COUNT, GMaxTextureMipCount );
+	GMaxCubeTextureDimensions = 4096;
+	GMaxTextureArrayLayers = 2048;
+
+	GShaderPlatformForFeatureLevel[ERHIFeatureLevel::ES2] = SP_NumPlatforms;
+	GShaderPlatformForFeatureLevel[ERHIFeatureLevel::ES3_1] = SP_METAL;
+#if PLATFORM_IOS
+	GShaderPlatformForFeatureLevel[ERHIFeatureLevel::SM4] = (GMaxRHIShaderPlatform == SP_METAL_MRT) ? SP_METAL_MRT : SP_NumPlatforms;
+	GShaderPlatformForFeatureLevel[ERHIFeatureLevel::SM5] = SP_NumPlatforms;
+#elif PLATFORM_MAC
+	GShaderPlatformForFeatureLevel[ERHIFeatureLevel::SM4] = SP_METAL_SM4;
+	GShaderPlatformForFeatureLevel[ERHIFeatureLevel::SM5] = (GMaxRHIShaderPlatform == SP_METAL_SM5) ? SP_METAL_SM5 : SP_NumPlatforms;
+#endif
+
+	// Initialize the platform pixel format map.
+	GPixelFormats[PF_Unknown			].PlatformFormat	= MTLPixelFormatInvalid;
+	GPixelFormats[PF_A32B32G32R32F		].PlatformFormat	= MTLPixelFormatRGBA32Float;
+	GPixelFormats[PF_B8G8R8A8			].PlatformFormat	= MTLPixelFormatBGRA8Unorm;
+	GPixelFormats[PF_G8					].PlatformFormat	= MTLPixelFormatR8Unorm;
+	GPixelFormats[PF_G16				].PlatformFormat	= MTLPixelFormatR16Unorm;
+#if PLATFORM_IOS
+    GPixelFormats[PF_DXT1				].PlatformFormat	= MTLPixelFormatInvalid;
+    GPixelFormats[PF_DXT3				].PlatformFormat	= MTLPixelFormatInvalid;
+    GPixelFormats[PF_DXT5				].PlatformFormat	= MTLPixelFormatInvalid;
+	GPixelFormats[PF_PVRTC2				].PlatformFormat	= MTLPixelFormatPVRTC_RGBA_2BPP;
+	GPixelFormats[PF_PVRTC2				].Supported			= true;
+	GPixelFormats[PF_PVRTC4				].PlatformFormat	= MTLPixelFormatPVRTC_RGBA_4BPP;
+	GPixelFormats[PF_PVRTC4				].Supported			= true;
+	GPixelFormats[PF_PVRTC4				].PlatformFormat	= MTLPixelFormatPVRTC_RGBA_4BPP;
+	GPixelFormats[PF_PVRTC4				].Supported			= true;
+	GPixelFormats[PF_ASTC_4x4			].PlatformFormat	= MTLPixelFormatASTC_4x4_LDR;
+	GPixelFormats[PF_ASTC_4x4			].Supported			= bCanUseASTC;
+	GPixelFormats[PF_ASTC_6x6			].PlatformFormat	= MTLPixelFormatASTC_6x6_LDR;
+	GPixelFormats[PF_ASTC_6x6			].Supported			= bCanUseASTC;
+	GPixelFormats[PF_ASTC_8x8			].PlatformFormat	= MTLPixelFormatASTC_8x8_LDR;
+	GPixelFormats[PF_ASTC_8x8			].Supported			= bCanUseASTC;
+	GPixelFormats[PF_ASTC_10x10			].PlatformFormat	= MTLPixelFormatASTC_10x10_LDR;
+	GPixelFormats[PF_ASTC_10x10			].Supported			= bCanUseASTC;
+	GPixelFormats[PF_ASTC_12x12			].PlatformFormat	= MTLPixelFormatASTC_12x12_LDR;
+	GPixelFormats[PF_ASTC_12x12			].Supported			= bCanUseASTC;
+#else // @todo zebra : srgb?
+    GPixelFormats[PF_DXT1				].PlatformFormat	= MTLPixelFormatBC1_RGBA;
+    GPixelFormats[PF_DXT3				].PlatformFormat	= MTLPixelFormatBC2_RGBA;
+    GPixelFormats[PF_DXT5				].PlatformFormat	= MTLPixelFormatBC3_RGBA;
+#endif
+	GPixelFormats[PF_UYVY				].PlatformFormat	= MTLPixelFormatInvalid;
+#if PLATFORM_IOS
+	GPixelFormats[PF_FloatRGB			].PlatformFormat	= MTLPixelFormatRGBA16Float;
+	GPixelFormats[PF_FloatRGB			].BlockBytes		= 8;
+#else
+	GPixelFormats[PF_FloatRGB			].PlatformFormat	= MTLPixelFormatRG11B10Float;
+	GPixelFormats[PF_FloatRGB			].BlockBytes		= 4;
+#endif
+	GPixelFormats[PF_FloatRGBA			].PlatformFormat	= MTLPixelFormatRGBA16Float;
+	GPixelFormats[PF_FloatRGBA			].BlockBytes		= 8;
+#if PLATFORM_IOS
+	GPixelFormats[PF_DepthStencil		].PlatformFormat	= MTLPixelFormatDepth32Float;
+	GPixelFormats[PF_DepthStencil		].BlockBytes		= 4;
+    GPixelFormats[PF_ShadowDepth		].PlatformFormat	= GPixelFormats[PF_DepthStencil].PlatformFormat; // all depth formats must be the same, for the pipeline state hash (see NUMBITS_DEPTH_ENABLED)
+#else
+	// Use Depth28_Stencil8 when it is available for consistency
+	if(bSupportsD24S8)
+	{
+		GPixelFormats[PF_DepthStencil	].PlatformFormat	= MTLPixelFormatDepth24Unorm_Stencil8;
+	}
+	else
+	{
+		GPixelFormats[PF_DepthStencil	].PlatformFormat	= MTLPixelFormatDepth32Float_Stencil8;
+	}
+    GPixelFormats[PF_DepthStencil		].BlockBytes		= 4;
+    GPixelFormats[PF_ShadowDepth		].PlatformFormat	= MTLPixelFormatDepth32Float;
+    GPixelFormats[PF_ShadowDepth		].BlockBytes		= 4;
+#endif
+    GPixelFormats[PF_X24_G8				].PlatformFormat	= MTLPixelFormatStencil8;
+    GPixelFormats[PF_X24_G8				].BlockBytes		= 1;
+	GPixelFormats[PF_R32_FLOAT			].PlatformFormat	= MTLPixelFormatR32Float;
+#if PLATFORM_IOS
+	GPixelFormats[PF_G16R16				].PlatformFormat	= MTLPixelFormatInvalid;//MTLPixelFormatRG16Unorm;
+	// we can't render to this in Metal, so mark it as unsupported (although we could texture from it, but we are only using it for render targets)
+	GPixelFormats[PF_G16R16				].Supported			= false;
+#else
+	GPixelFormats[PF_G16R16				].PlatformFormat	= MTLPixelFormatRG16Unorm;
+	GPixelFormats[PF_G16R16				].Supported			= true;
+#endif
+	GPixelFormats[PF_G16R16F			].PlatformFormat	= MTLPixelFormatRG16Float;
+	GPixelFormats[PF_G16R16F_FILTER		].PlatformFormat	= MTLPixelFormatRG16Float;
+	GPixelFormats[PF_G32R32F			].PlatformFormat	= MTLPixelFormatRG32Float;
+	GPixelFormats[PF_A2B10G10R10		].PlatformFormat    = MTLPixelFormatRGB10A2Unorm;
+	GPixelFormats[PF_A16B16G16R16		].PlatformFormat    = MTLPixelFormatRGBA16Unorm;
+#if PLATFORM_MAC
+	if(bSupportsD24S8)
+	{
+		GPixelFormats[PF_D24			].PlatformFormat	= MTLPixelFormatDepth24Unorm_Stencil8;
+	}
+	else
+	{
+		GPixelFormats[PF_D24			].PlatformFormat	= MTLPixelFormatDepth32Float;
+	}
+	GPixelFormats[PF_D24				].Supported			= true;
+#endif
+	GPixelFormats[PF_R16F				].PlatformFormat	= MTLPixelFormatR16Float;
+	GPixelFormats[PF_R16F_FILTER		].PlatformFormat	= MTLPixelFormatR16Float;
+#if PLATFORM_IOS
+	GPixelFormats[PF_BC5				].PlatformFormat	= MTLPixelFormatInvalid;
+#else
+    GPixelFormats[PF_BC4				].Supported			= true;
+    GPixelFormats[PF_BC4				].PlatformFormat	= MTLPixelFormatBC4_RUnorm;
+    GPixelFormats[PF_BC5				].Supported			= true;
+    GPixelFormats[PF_BC5				].PlatformFormat	= MTLPixelFormatBC5_RGUnorm;
+    GPixelFormats[PF_BC6H				].Supported			= true;
+    GPixelFormats[PF_BC6H               ].PlatformFormat	= MTLPixelFormatBC6H_RGBUfloat;
+    GPixelFormats[PF_BC7				].Supported			= true;
+    GPixelFormats[PF_BC7				].PlatformFormat	= MTLPixelFormatBC7_RGBAUnorm;
+#endif
+	GPixelFormats[PF_V8U8				].PlatformFormat	=
+	GPixelFormats[PF_A1					].PlatformFormat	= MTLPixelFormatInvalid;
+	GPixelFormats[PF_FloatR11G11B10		].PlatformFormat	= MTLPixelFormatRG11B10Float;
+	GPixelFormats[PF_FloatR11G11B10		].BlockBytes		= 4;
+	GPixelFormats[PF_A8					].PlatformFormat	= MTLPixelFormatA8Unorm;
+	GPixelFormats[PF_R32_UINT			].PlatformFormat	= MTLPixelFormatR32Uint;
+	GPixelFormats[PF_R32_SINT			].PlatformFormat	= MTLPixelFormatR32Sint;
+	GPixelFormats[PF_R16G16B16A16_UINT	].PlatformFormat	= MTLPixelFormatRGBA16Uint;
+	GPixelFormats[PF_R16G16B16A16_SINT	].PlatformFormat	= MTLPixelFormatRGBA16Sint;
+#if PLATFORM_IOS
+	GPixelFormats[PF_R5G6B5_UNORM		].PlatformFormat	= MTLPixelFormatB5G6R5Unorm;
+#else // @todo zebra
+    GPixelFormats[PF_R5G6B5_UNORM		].PlatformFormat	= MTLPixelFormatInvalid;
+#endif
+	GPixelFormats[PF_R8G8B8A8			].PlatformFormat	= MTLPixelFormatRGBA8Unorm;
+	GPixelFormats[PF_R8G8				].PlatformFormat	= MTLPixelFormatRG8Unorm;
+
+	GDynamicRHI = this;
+	GIsRHIInitialized = true;
+	
+#if PLATFORM_DESKTOP
+	static const auto CVar = IConsoleManager::Get().FindConsoleVariable(TEXT("r.Shaders.Optimize"));
+	
+	if (CVar->GetInt() == 0)
+	{
+		FShaderCache::InitShaderCache(SCO_NoShaderPreload, 128);
+	}
+	else
+	{
+		FShaderCache::InitShaderCache(SCO_Default, 128);
+	}
+#endif
+
+	// Notify all initialized FRenderResources that there's a valid RHI device to create their RHI resources for now.
+	for(TLinkedList<FRenderResource*>::TIterator ResourceIt(FRenderResource::GetResourceList());ResourceIt;ResourceIt.Next())
+	{
+		ResourceIt->InitDynamicRHI();
+	}
+	for(TLinkedList<FRenderResource*>::TIterator ResourceIt(FRenderResource::GetResourceList());ResourceIt;ResourceIt.Next())
+	{
+		ResourceIt->InitRHI();
+	}
+	
+    Profiler = nullptr;
+#if ENABLE_METAL_GPUPROFILE
+	Profiler = new FMetalGPUProfiler(Context);
+#endif
+}
+
+FMetalDynamicRHI::~FMetalDynamicRHI()
+{
+	check(IsInGameThread() && IsInRenderingThread());
+
+#if ENABLE_METAL_GPUPROFILE
+	delete Profiler;
+#endif
+	
+#if PLATFORM_DESKTOP
+	FShaderCache::ShutdownShaderCache();
+#endif
+	
+	GIsRHIInitialized = false;
+}
+
+uint64 FMetalDynamicRHI::RHICalcTexture2DPlatformSize(uint32 SizeX, uint32 SizeY, uint8 Format, uint32 NumMips, uint32 NumSamples, uint32 Flags, uint32& OutAlign)
+{
+	OutAlign = 0;
+	return CalcTextureSize(SizeX, SizeY, (EPixelFormat)Format, NumMips);
+}
+
+uint64 FMetalDynamicRHI::RHICalcTexture3DPlatformSize(uint32 SizeX, uint32 SizeY, uint32 SizeZ, uint8 Format, uint32 NumMips, uint32 Flags, uint32& OutAlign)
+{
+	OutAlign = 0;
+	return CalcTextureSize3D(SizeX, SizeY, SizeZ, (EPixelFormat)Format, NumMips);
+}
+
+uint64 FMetalDynamicRHI::RHICalcTextureCubePlatformSize(uint32 Size, uint8 Format, uint32 NumMips, uint32 Flags,	uint32& OutAlign)
+{
+	OutAlign = 0;
+	return CalcTextureSize(Size, Size, (EPixelFormat)Format, NumMips) * 6;
+}
+
+
+void FMetalDynamicRHI::Init()
+{
+}
+
+void FMetalRHICommandContext::RHIBeginFrame()
+{
+	// @todo zebra: GPUProfilingData, GNumDrawCallsRHI, GNumPrimitivesDrawnRHI
+#if ENABLE_METAL_GPUPROFILE
+	Profiler->BeginFrame();
+#endif
+	Context->BeginFrame();
+}
+
+void FMetalRHICommandContext::RHIEndFrame()
+{
+	// @todo zebra: GPUProfilingData.EndFrame();
+#if ENABLE_METAL_GPUPROFILE
+	Profiler->EndFrame();
+#endif
+	Context->EndFrame();
+}
+
+void FMetalRHICommandContext::RHIBeginScene()
+{
+	Context->BeginScene();
+}
+
+void FMetalRHICommandContext::RHIEndScene()
+{
+	Context->EndScene();
+}
+
+void FMetalRHICommandContext::RHIPushEvent(const TCHAR* Name)
+{
+#if ENABLE_METAL_GPUEVENTS
+#if ENABLE_METAL_GPUPROFILE
+	Profiler->PushEvent(Name);
+#endif
+	// @todo zebra : this was "[NSString stringWithTCHARString:Name]", an extension only on ios for some reason, it should be on Mac also
+	Context->GetCommandEncoder().PushDebugGroup([NSString stringWithCString:TCHAR_TO_UTF8(Name) encoding:NSUTF8StringEncoding]);
+#endif
+}
+
+void FMetalRHICommandContext::RHIPopEvent()
+{
+#if ENABLE_METAL_GPUEVENTS
+	Context->GetCommandEncoder().PopDebugGroup();
+#if ENABLE_METAL_GPUPROFILE
+	Profiler->PopEvent();
+#endif
+#endif
+}
+
+void FMetalDynamicRHI::RHIGetSupportedResolution( uint32 &Width, uint32 &Height )
+{
+
+}
+
+bool FMetalDynamicRHI::RHIGetAvailableResolutions(FScreenResolutionArray& Resolutions, bool bIgnoreRefreshRate)
+{
+
+	return false;
+}
+
+void FMetalDynamicRHI::RHIFlushResources()
+{
+
+}
+
+void FMetalDynamicRHI::RHIAcquireThreadOwnership()
+{
+	Context->CreateAutoreleasePool();
+}
+
+void FMetalDynamicRHI::RHIReleaseThreadOwnership()
+{
+	Context->DrainAutoreleasePool();
+}
+
+void FMetalDynamicRHI::RHIGpuTimeBegin(uint32 Hash, bool bCompute)
+{
+
+}
+
+void FMetalDynamicRHI::RHIGpuTimeEnd(uint32 Hash, bool bCompute)
+{
+
+}
+
+void* FMetalDynamicRHI::RHIGetNativeDevice()
+{
+	return (void*)Context->GetDevice();
+}
+
+void FMetalRHICommandContext::RHIBeginAsyncComputeJob_DrawThread(EAsyncComputePriority Priority)
+{
+	UE_LOG(LogRHI, Fatal, TEXT("%s not implemented yet"), ANSI_TO_TCHAR(__FUNCTION__));
+}
+
+void FMetalRHICommandContext::RHIEndAsyncComputeJob_DrawThread(uint32 FenceIndex)
+{
+	UE_LOG(LogRHI, Fatal, TEXT("%s not implemented yet"), ANSI_TO_TCHAR(__FUNCTION__));
+}
+
+void FMetalRHICommandContext::RHIGraphicsWaitOnAsyncComputeJob(uint32 FenceIndex)
+{
+	UE_LOG(LogRHI, Fatal, TEXT("%s not implemented yet"), ANSI_TO_TCHAR(__FUNCTION__));
+}
diff --git a/Engine/Source/Runtime/Apple/MetalRHI/Private/MetalRHIContext.cpp b/Engine/Source/Runtime/Apple/MetalRHI/Private/MetalRHIContext.cpp
new file mode 100644
--- /dev/null
+++ b/Engine/Source/Runtime/Apple/MetalRHI/Private/MetalRHIContext.cpp
@@ -0,0 +1,57 @@
+// Copyright 1998-2015 Epic Games, Inc. All Rights Reserved.
+
+#include "MetalRHIPrivate.h"
+
+TGlobalResource<TBoundShaderStateHistory<10000>> FMetalRHICommandContext::BoundShaderStateHistory;
+
+FMetalContext& GetMetalDeviceContext()
+{
+	FMetalRHICommandContext* Context = static_cast<FMetalRHICommandContext*>(RHIGetDefaultContext());
+	check(Context);
+	return Context->GetInternalContext();
+}
+
+void SafeReleaseMetalResource(id Object)
+{
+	if(GIsRHIInitialized && GDynamicRHI)
+	{
+		FMetalRHICommandContext* Context = static_cast<FMetalRHICommandContext*>(RHIGetDefaultContext());
+		if(Context)
+		{
+			Context->GetInternalContext().ReleaseObject(Object);
+			return;
+		}
+	}
+	[Object release];
+}
+
+void SafeReleasePooledBuffer(id<MTLBuffer> Buffer)
+{
+	if(GIsRHIInitialized && GDynamicRHI)
+	{
+		FMetalRHICommandContext* Context = static_cast<FMetalRHICommandContext*>(RHIGetDefaultContext());
+		if(Context)
+		{
+			Context->GetInternalContext().ReleasePooledBuffer(Buffer);
+		}
+	}
+}
+
+FMetalRHICommandContext::FMetalRHICommandContext(struct FMetalGPUProfiler* InProfiler)
+: Context(new FMetalContext)
+, Profiler(InProfiler)
+, PendingVertexBufferOffset(0xFFFFFFFF)
+, PendingVertexDataStride(0)
+, PendingIndexBufferOffset(0xFFFFFFFF)
+, PendingIndexDataStride(0)
+, PendingPrimitiveType(0)
+, PendingNumPrimitives(0)
+{
+	check(Context);
+}
+
+FMetalRHICommandContext::~FMetalRHICommandContext()
+{
+	delete Context;
+}
+
diff --git a/Engine/Source/Runtime/Apple/MetalRHI/Private/MetalRHIPrivate.h b/Engine/Source/Runtime/Apple/MetalRHI/Private/MetalRHIPrivate.h
new file mode 100644
--- /dev/null
+++ b/Engine/Source/Runtime/Apple/MetalRHI/Private/MetalRHIPrivate.h
@@ -0,0 +1,117 @@
+// Copyright 1998-2015 Epic Games, Inc. All Rights Reserved.
+
+/*=============================================================================
+	MetalRHIPrivate.h: Private Metal RHI definitions.
+=============================================================================*/
+
+#pragma once
+
+#include "Engine.h"
+
+// UE4 has a Max of 8 RTs, but we can spend less time looping with 6
+const uint32 MaxMetalRenderTargets = 6;
+
+// How many possible vertex streams are allowed
+const uint32 MaxMetalStreams = 31;
+
+// Requirement for vertex buffer offset field
+const uint32 BufferOffsetAlignment = 256;
+
+#define METAL_API_1_1 (__IPHONE_9_0 || __MAC_10_11)
+
+#if METAL_API_1_1
+#define BUFFER_CACHE_MODE MTLResourceCPUCacheModeWriteCombined
+#else
+#define BUFFER_CACHE_MODE MTLResourceOptionCPUCacheModeWriteCombined
+#endif
+#if PLATFORM_MAC
+#define BUFFER_MANAGED_MEM MTLResourceStorageModeManaged
+#define BUFFER_STORAGE_MODE MTLStorageModeManaged
+#define BUFFER_RESOURCE_STORAGE_MANAGED MTLResourceStorageModeManaged
+#define BUFFER_DYNAMIC_REALLOC BUF_AnyDynamic
+#else
+#define BUFFER_MANAGED_MEM 0
+#define BUFFER_STORAGE_MODE MTLStorageModeShared
+#define BUFFER_RESOURCE_STORAGE_MANAGED MTLResourceStorageModeShared
+#define BUFFER_DYNAMIC_REALLOC BUF_Volatile
+#endif
+
+#ifndef METAL_STATISTICS
+#define METAL_STATISTICS 0
+#endif
+
+#define SHOULD_TRACK_OBJECTS 0 // (UE_BUILD_DEBUG || UE_BUILD_DEVELOPMENT)
+
+#define UNREAL_TO_METAL_BUFFER_INDEX(Index) (30 - Index)
+
+// Dependencies
+#include "MetalRHI.h"
+#include "MetalGlobalUniformBuffer.h"
+#include "RHI.h"
+#import <Metal/Metal.h>
+#import <QuartzCore/CAMetalLayer.h>
+
+// Access the internal context for the device-owning DynamicRHI object
+FMetalContext& GetMetalDeviceContext();
+
+// Safely release a metal resource, correctly handling the case where the RHI has been destructed first
+void SafeReleaseMetalResource(id Object);
+
+// Safely release a pooled buffer, correctly handling the case where the RHI has been destructed first
+void SafeReleasePooledBuffer(id<MTLBuffer> Buffer);
+
+// Access the underlying surface object from any kind of texture
+FMetalSurface* GetMetalSurfaceFromRHITexture(FRHITexture* Texture);
+
+#define NOT_SUPPORTED(Func) UE_LOG(LogMetal, Fatal, TEXT("'%s' is not supported"), L##Func);
+
+#if SHOULD_TRACK_OBJECTS
+extern TMap<id, int32> ClassCounts;
+#define TRACK_OBJECT(Obj) ClassCounts.FindOrAdd([Obj class])++;
+#define UNTRACK_OBJECT(Obj) ClassCounts.FindOrAdd([Obj class])--;
+#else
+#define TRACK_OBJECT(Obj)
+#define UNTRACK_OBJECT(Obj)
+#endif
+
+
+
+FORCEINLINE int32 GetMetalCubeFace(ECubeFace Face)
+{
+	// According to Metal docs these should match now: https://developer.apple.com/library/prerelease/ios/documentation/Metal/Reference/MTLTexture_Ref/index.html#//apple_ref/c/tdef/MTLTextureType
+	switch (Face)
+	{
+		case CubeFace_PosX:;
+		default:			return 0;
+		case CubeFace_NegX:	return 1;
+		case CubeFace_PosY:	return 2;
+		case CubeFace_NegY:	return 3;
+		case CubeFace_PosZ:	return 4;
+		case CubeFace_NegZ:	return 5;
+	}
+}
+
+FORCEINLINE MTLLoadAction GetMetalRTLoadAction(ERenderTargetLoadAction LoadAction)
+{
+	switch(LoadAction)
+	{
+		case ERenderTargetLoadAction::ENoAction: return MTLLoadActionDontCare;
+		case ERenderTargetLoadAction::ELoad: return MTLLoadActionLoad;
+		case ERenderTargetLoadAction::EClear: return MTLLoadActionClear;
+		default: return MTLLoadActionDontCare;
+	}
+}
+
+FORCEINLINE MTLStoreAction GetMetalRTStoreAction(ERenderTargetStoreAction StoreAction)
+{
+	switch(StoreAction)
+	{
+		case ERenderTargetStoreAction::ENoAction: return MTLStoreActionDontCare;
+		case ERenderTargetStoreAction::EStore: return MTLStoreActionStore;
+		case ERenderTargetStoreAction::EMultisampleResolve: return MTLStoreActionMultisampleResolve;
+		default: return MTLStoreActionDontCare;
+	}
+}
+
+#include "MetalStateCache.h"
+#include "MetalContext.h"
diff --git a/Engine/Source/Runtime/Apple/MetalRHI/Private/MetalRenderPipelineDesc.cpp b/Engine/Source/Runtime/Apple/MetalRHI/Private/MetalRenderPipelineDesc.cpp
new file mode 100644
--- /dev/null
+++ b/Engine/Source/Runtime/Apple/MetalRHI/Private/MetalRenderPipelineDesc.cpp
@@ -0,0 +1,69 @@
+// Copyright 1998-2015 Epic Games, Inc. All Rights Reserved.
+
+#include "MetalRHIPrivate.h"
+#include "MetalRenderPipelineDesc.h"
+
+static_assert(Offset_End <= sizeof(FMetalRenderPipelineHash) * 8, "Too many bits used for the Hash");
+
+NSMutableDictionary* FMetalRenderPipelineDesc::MetalPipelineCache = [NSMutableDictionary new];
+uint32 FMetalRenderPipelineDesc::BlendBitOffsets[] = { Offset_BlendState0, Offset_BlendState1, Offset_BlendState2, Offset_BlendState3, Offset_BlendState4, Offset_BlendState5, };
+uint32 FMetalRenderPipelineDesc::RTBitOffsets[] = { Offset_RenderTargetFormat0, Offset_RenderTargetFormat1, Offset_RenderTargetFormat2, Offset_RenderTargetFormat3, Offset_RenderTargetFormat4, Offset_RenderTargetFormat5, };
+
+FMetalRenderPipelineDesc::FMetalRenderPipelineDesc()
+: PipelineDescriptor([[MTLRenderPipelineDescriptor alloc] init])
+, SampleCount(0)
+, Hash(0)
+{
+	for (int Index = 0; Index < MaxMetalRenderTargets; Index++)
+	{
+		[PipelineDescriptor.colorAttachments setObject:[[MTLRenderPipelineColorAttachmentDescriptor new] autorelease] atIndexedSubscript:Index];
+	}
+}
+
+id<MTLRenderPipelineState> FMetalRenderPipelineDesc::CreatePipelineStateForBoundShaderState(FMetalBoundShaderState* BSS) const
+{
+	SCOPE_CYCLE_COUNTER(STAT_MetalPipelineStateTime);
+	
+	MTLPixelFormat DepthFormat = PipelineDescriptor.depthAttachmentPixelFormat;
+	MTLPixelFormat StencilFormat = PipelineDescriptor.stencilAttachmentPixelFormat;
+	if(BSS->PixelShader && (BSS->PixelShader->Bindings.InOutMask & 0x8000) && PipelineDescriptor.depthAttachmentPixelFormat == MTLPixelFormatInvalid)
+	{
+		PipelineDescriptor.depthAttachmentPixelFormat = (MTLPixelFormat)GPixelFormats[PF_DepthStencil].PlatformFormat;
+		PipelineDescriptor.stencilAttachmentPixelFormat = (MTLPixelFormat)GPixelFormats[PF_DepthStencil].PlatformFormat;
+	}
+	
+	// set the bound shader state settings
+	PipelineDescriptor.vertexDescriptor = BSS->VertexDeclaration->Layout;
+	PipelineDescriptor.vertexFunction = BSS->VertexShader->Function;
+	PipelineDescriptor.fragmentFunction = BSS->PixelShader ? BSS->PixelShader->Function : nil;
+	PipelineDescriptor.sampleCount = SampleCount;
+
+	check(SampleCount > 0);
+
+	NSError* Error = nil;
+	id<MTLRenderPipelineState> PipelineState = (id<MTLRenderPipelineState>)[MetalPipelineCache objectForKey:PipelineDescriptor];
+	if(PipelineState == nil)
+	{
+		id<MTLDevice> Device = GetMetalDeviceContext().GetDevice();
+		PipelineState = [Device newRenderPipelineStateWithDescriptor:PipelineDescriptor error : &Error];
+		TRACK_OBJECT(PipelineState);
+		if(PipelineState)
+		{
+			[MetalPipelineCache setObject:PipelineState forKey:PipelineDescriptor];
+		}
+	}
+	
+	PipelineDescriptor.depthAttachmentPixelFormat = DepthFormat;
+	PipelineDescriptor.stencilAttachmentPixelFormat = StencilFormat;
+	
+	if (PipelineState == nil)
+	{
+		NSLog(@"Failed to generate a pipeline state object: %@", Error);
+		NSLog(@"Vertex shader: %@", BSS->VertexShader->GlslCodeNSString);
+		NSLog(@"Pixel shader: %@", BSS->PixelShader->GlslCodeNSString);
+		NSLog(@"Descriptor: %@", PipelineDescriptor);
+		return nil;
+	}
+
+	return PipelineState;
+}
diff --git a/Engine/Source/Runtime/Apple/MetalRHI/Private/MetalRenderPipelineDesc.h b/Engine/Source/Runtime/Apple/MetalRHI/Private/MetalRenderPipelineDesc.h
new file mode 100644
--- /dev/null
+++ b/Engine/Source/Runtime/Apple/MetalRHI/Private/MetalRenderPipelineDesc.h
@@ -0,0 +1,78 @@
+// Copyright 1998-2015 Epic Games, Inc. All Rights Reserved.
+
+#pragma once
+
+#include "MetalRHIPrivate.h"
+
+enum EMetalPipelineHashBits
+{
+	NumBits_BlendState = 5, //(x6=30),
+#if PLATFORM_MAC
+	NumBits_RenderTargetFormat = 5, //(x6=30),
+#else
+	NumBits_RenderTargetFormat = 4, //(x6=24),
+#endif
+	NumBits_DepthFormat = 3, //(x1=3),
+	NumBits_StencilFormat = 3, //(x1=3),
+	NumBits_SampleCount = 3, //(x1=3),
+#if PLATFORM_MAC
+	NumBits_PrimitiveTopology = 2, //(x1=2)
+#endif
+};
+
+enum EMetalPipelineHashOffsets
+{
+	Offset_BlendState0 = 0,
+	Offset_BlendState1 = Offset_BlendState0 + NumBits_BlendState,
+	Offset_BlendState2 = Offset_BlendState1 + NumBits_BlendState,
+	Offset_BlendState3 = Offset_BlendState2 + NumBits_BlendState,
+	Offset_BlendState4 = Offset_BlendState3 + NumBits_BlendState,
+	Offset_BlendState5 = Offset_BlendState4 + NumBits_BlendState,
+	Offset_RenderTargetFormat0 = Offset_BlendState5 + NumBits_BlendState,
+	Offset_RenderTargetFormat1 = Offset_RenderTargetFormat0 + NumBits_RenderTargetFormat,
+	Offset_RenderTargetFormat2 = Offset_RenderTargetFormat1 + NumBits_RenderTargetFormat,
+	Offset_RenderTargetFormat3 = Offset_RenderTargetFormat2 + NumBits_RenderTargetFormat,
+	Offset_RenderTargetFormat4 = Offset_RenderTargetFormat3 + NumBits_RenderTargetFormat,
+	Offset_RenderTargetFormat5 = Offset_RenderTargetFormat4 + NumBits_RenderTargetFormat,
+	Offset_DepthFormat = Offset_RenderTargetFormat5 + NumBits_RenderTargetFormat,
+	Offset_StencilFormat = Offset_DepthFormat + NumBits_DepthFormat,
+	Offset_SampleCount = Offset_StencilFormat + NumBits_StencilFormat,
+#if PLATFORM_MAC
+	Offset_PrimitiveTopology = Offset_SampleCount + NumBits_SampleCount,
+	Offset_End = Offset_PrimitiveTopology + NumBits_PrimitiveTopology
+#else
+	Offset_End = Offset_SampleCount + NumBits_SampleCount
+#endif
+};
+
+struct FMetalRenderPipelineDesc
+{
+	FMetalRenderPipelineDesc();
+
+	template<typename Type>
+	inline void SetHashValue(uint32 Offset, uint32 NumBits, Type Value)
+	{
+		FMetalRenderPipelineHash BitMask = ((((FMetalRenderPipelineHash)1ULL) << NumBits) - 1) << Offset;
+		Hash = (Hash & ~BitMask) | (((FMetalRenderPipelineHash)Value << Offset) & BitMask);
+	}
+
+	/**
+	 * @return the hash of the current pipeline state
+	 */
+	inline FMetalRenderPipelineHash GetHash() const { return Hash; }
+	
+	/**
+	 * @return an pipeline state object that matches the current state and the BSS
+	 */
+	id<MTLRenderPipelineState> CreatePipelineStateForBoundShaderState(FMetalBoundShaderState* BSS) const;
+
+	MTLRenderPipelineDescriptor* PipelineDescriptor;
+	uint32 SampleCount;
+
+	// running hash of the pipeline state
+	FMetalRenderPipelineHash Hash;
+
+	static NSMutableDictionary* MetalPipelineCache;
+	static uint32 BlendBitOffsets[6];
+	static uint32 RTBitOffsets[6];
+};
\ No newline at end of file
diff --git a/Engine/Source/Runtime/Apple/MetalRHI/Private/MetalRenderTarget.cpp b/Engine/Source/Runtime/Apple/MetalRHI/Private/MetalRenderTarget.cpp
new file mode 100644
--- /dev/null
+++ b/Engine/Source/Runtime/Apple/MetalRHI/Private/MetalRenderTarget.cpp
@@ -0,0 +1,333 @@
+// Copyright 1998-2015 Epic Games, Inc. All Rights Reserved.
+
+/*=============================================================================
+	MetalRenderTarget.cpp: Metal render target implementation.
+=============================================================================*/
+
+#include "MetalRHIPrivate.h"
+#include "ScreenRendering.h"
+#include "MetalProfiler.h"
+
+
+void FMetalRHICommandContext::RHICopyToResolveTarget(FTextureRHIParamRef SourceTextureRHI, FTextureRHIParamRef DestTextureRHI, bool bKeepOriginalSurface, const FResolveParams& ResolveParams)
+{
+	if(SourceTextureRHI != DestTextureRHI)
+	{
+		id<MTLBlitCommandEncoder> Blitter = Context->GetBlitContext();
+		
+		FMetalSurface* Source = GetMetalSurfaceFromRHITexture(SourceTextureRHI);
+		FMetalSurface* Destination = GetMetalSurfaceFromRHITexture(DestTextureRHI);
+		
+		switch (Source->Type)
+		{
+			case RRT_Texture2D:
+				break;
+			case RRT_TextureCube:
+				check(Source->SizeZ == 6); // Arrays might not work yet.
+				break;
+			default:
+				check(false); // Only Tex2D & Cube are tested to work so far!
+				break;
+		}
+		switch (Destination->Type)
+		{
+			case RRT_Texture2D:
+				break;
+			case RRT_TextureCube:
+				check(Destination->SizeZ == 6); // Arrays might not work yet.
+				break;
+			default:
+				check(false); // Only Tex2D & Cube are tested to work so far!
+				break;
+		}
+		
+		MTLOrigin Origin = MTLOriginMake(0, 0, 0);
+		MTLSize Size = MTLSizeMake(0, 0, 1);
+		if (ResolveParams.Rect.IsValid())
+		{
+			// Partial copy
+			Origin.x = ResolveParams.Rect.X1;
+			Origin.y = ResolveParams.Rect.Y1;
+			Size.width = ResolveParams.Rect.X2 - ResolveParams.Rect.X1;
+			Size.height = ResolveParams.Rect.Y2 - ResolveParams.Rect.Y1;
+		}
+		else
+		{
+			// Whole of source copy
+			Origin.x = 0;
+			Origin.y = 0;
+			
+			Size.width = FMath::Max<uint32>(1, Source->SizeX >> ResolveParams.MipIndex);
+			Size.height = FMath::Max<uint32>(1, Source->SizeY >> ResolveParams.MipIndex);
+		}
+		
+		const bool bSrcCubemap  = Source->bIsCubemap;
+		const bool bDestCubemap = Destination->bIsCubemap;
+		
+		uint32 DestIndex = ResolveParams.DestArrayIndex * (bDestCubemap ? 6 : 1) + (bDestCubemap ? uint32(ResolveParams.CubeFace) : 0);
+		uint32 SrcIndex  = ResolveParams.SourceArrayIndex * (bSrcCubemap ? 6 : 1) + (bSrcCubemap ? uint32(ResolveParams.CubeFace) : 0);
+		
+		if(Profiler)
+		{
+			Profiler->RegisterGPUWork();
+		}
+		
+		[Blitter copyFromTexture:Source->Texture sourceSlice:SrcIndex sourceLevel:ResolveParams.MipIndex sourceOrigin:Origin sourceSize:Size toTexture:Destination->Texture destinationSlice:DestIndex destinationLevel:ResolveParams.MipIndex destinationOrigin:Origin];
+	}
+}
+
+void FMetalDynamicRHI::RHIReadSurfaceData(FTextureRHIParamRef TextureRHI, FIntRect Rect, TArray<FColor>& OutData, FReadSurfaceDataFlags InFlags)
+{
+	FMetalSurface* Surface = GetMetalSurfaceFromRHITexture(TextureRHI);
+
+	// verify the input image format (but don't crash)
+	if (Surface->PixelFormat != PF_B8G8R8A8 && Surface->PixelFormat != PF_FloatRGBA)
+	{
+		UE_LOG(LogRHI, Log, TEXT("Trying to read an unsupported surface format."));
+	}
+	
+	// allocate output space
+	const uint32 SizeX = Rect.Width();
+	const uint32 SizeY = Rect.Height();
+	OutData.Empty();
+	OutData.AddUninitialized(SizeX * SizeY);
+
+	FColor* OutDataPtr = OutData.GetData();
+	MTLRegion Region = MTLRegionMake2D(Rect.Min.X, Rect.Min.Y, SizeX, SizeY);
+    
+    id<MTLTexture> Texture = Surface->Texture;
+    if(!Texture && (Surface->Flags & TexCreate_Presentable))
+    {
+        id<CAMetalDrawable> Drawable = (id<CAMetalDrawable>)GetMetalDeviceContext().GetCurrentState().GetCurrentDrawable();
+        Texture = Drawable ? Drawable.texture : nil;
+    }
+    if(!Texture)
+    {
+        UE_LOG(LogRHI, Error, TEXT("Trying to read from an uninitialised texture."));
+        return;
+    }
+
+	// function wants details about the destination, not the source
+	if(Surface->PixelFormat == PF_B8G8R8A8)
+	{
+		check(sizeof(FColor) == GPixelFormats[PF_B8G8R8A8].BlockBytes);
+		const uint32 Stride = sizeof(FColor) * SizeX;
+		const uint32 BytesPerImage = Stride  * SizeY;
+#if PLATFORM_MAC // @todo zebra
+		if(Texture.storageMode == MTLStorageModeManaged)
+		{
+			SCOPE_CYCLE_COUNTER(STAT_MetalTexturePageOffTime);
+			
+			// Synchronise the texture with the CPU
+			id<MTLBlitCommandEncoder> Blitter = Context->GetBlitContext();
+			[Blitter synchronizeTexture:Texture slice:0 level:0];
+			
+			//kick the current command buffer.
+			Context->SubmitCommandBufferAndWait();
+			
+			[Texture getBytes:OutDataPtr bytesPerRow:Stride bytesPerImage:BytesPerImage fromRegion:Region mipmapLevel:0 slice:0];
+		}
+		else
+#endif
+		{
+			FMetalPooledBuffer Buffer = Context->CreatePooledBuffer(FMetalPooledBufferArgs(Context->GetDevice(), BytesPerImage, MTLStorageModeShared));
+			{
+				// Synchronise the texture with the CPU
+				SCOPE_CYCLE_COUNTER(STAT_MetalTexturePageOffTime);
+				
+				id<MTLBlitCommandEncoder> Blitter = Context->GetBlitContext();
+				[Blitter copyFromTexture:Texture sourceSlice:0 sourceLevel:0 sourceOrigin:Region.origin sourceSize:Region.size toBuffer:Buffer.Buffer destinationOffset:0 destinationBytesPerRow:Stride destinationBytesPerImage:BytesPerImage];
+				
+				//kick the current command buffer.
+				Context->SubmitCommandBufferAndWait();
+				
+				FMemory::Memcpy(OutDataPtr, [Buffer.Buffer contents], BytesPerImage);
+			}
+			Context->ReleasePooledBuffer(Buffer);
+		}
+	}
+	else if(Surface->PixelFormat == PF_FloatRGBA)
+	{
+		bool bLinearToGamma = InFlags.GetLinearToGamma();
+		
+		// Determine minimal and maximal float value present in received data. Treat alpha separately.
+        // Metal uses 16-bit float for FloatRGBA, not 32
+
+		const uint32 Stride = GPixelFormats[Surface->PixelFormat].BlockBytes * SizeX;
+		const uint32 BytesPerImage = Stride  * SizeY;
+		int32 PixelComponentCount = 4 * SizeX * SizeY;
+		int32 FloatBGRADataSize = Stride * SizeY;
+		FMetalPooledBuffer Buffer = Context->CreatePooledBuffer(FMetalPooledBufferArgs(Context->GetDevice(), FloatBGRADataSize, MTLStorageModeShared));
+		{
+			// Synchronise the texture with the CPU
+			SCOPE_CYCLE_COUNTER(STAT_MetalTexturePageOffTime);
+			
+			id<MTLBlitCommandEncoder> Blitter = Context->GetBlitContext();
+			[Blitter copyFromTexture:Texture sourceSlice:0 sourceLevel:0 sourceOrigin:Region.origin sourceSize:Region.size toBuffer:Buffer.Buffer destinationOffset:0 destinationBytesPerRow:Stride destinationBytesPerImage:BytesPerImage];
+			
+			//kick the current command buffer.
+			Context->SubmitCommandBufferAndWait();
+		}
+		
+		FFloat16* FloatBGRAData = (FFloat16*)[Buffer.Buffer contents];
+
+		// Determine minimal and maximal float values present in received data. Treat each component separately.
+		float MinValue[4] = { 0.0f, 0.0f, 0.0f, 0.0f };
+		float MaxValue[4] = { 1.0f, 1.0f, 1.0f, 1.0f };
+		FFloat16* DataPtr = FloatBGRAData;
+		for( int32 PixelComponentIndex = 0; PixelComponentIndex < PixelComponentCount; ++PixelComponentIndex, ++DataPtr )
+		{
+			int32 ComponentIndex = PixelComponentIndex % 4;
+			MinValue[ComponentIndex] = FMath::Min<float>(*DataPtr,MinValue[ComponentIndex]);
+			MaxValue[ComponentIndex] = FMath::Max<float>(*DataPtr,MaxValue[ComponentIndex]);
+		}
+		
+		// Convert the data into BGRA8 buffer
+		DataPtr = FloatBGRAData;
+		float RescaleFactor[4] = { MaxValue[0] - MinValue[0], MaxValue[1] - MinValue[1], MaxValue[2] - MinValue[2], MaxValue[3] - MinValue[3] };
+		for( int32 PixelIndex = 0; PixelIndex < PixelComponentCount / 4; ++PixelIndex )
+		{
+			float B = (DataPtr[2] - MinValue[2]) / RescaleFactor[2];
+			float G = (DataPtr[1] - MinValue[1]) / RescaleFactor[1];
+			float R = (DataPtr[0] - MinValue[0]) / RescaleFactor[0];
+			float A = (DataPtr[3] - MinValue[3]) / RescaleFactor[3];
+			
+			FColor NormalizedColor = FLinearColor( R,G,B,A ).ToFColor(bLinearToGamma);
+			FMemory::Memcpy(OutDataPtr,&NormalizedColor,sizeof(FColor));
+			DataPtr += 4;
+			OutDataPtr++;
+		}
+		Context->ReleasePooledBuffer(Buffer);
+	}
+	else
+	{
+		NOT_SUPPORTED("RHIReadSurfaceData Format");
+	}
+}
+
+void FMetalDynamicRHI::RHIMapStagingSurface(FTextureRHIParamRef TextureRHI,void*& OutData,int32& OutWidth,int32& OutHeight)
+{
+    FMetalSurface* Surface = GetMetalSurfaceFromRHITexture(TextureRHI);
+    FMetalTexture2D* Texture = (FMetalTexture2D*)TextureRHI->GetTexture2D();
+    
+    uint32 Stride = 0;
+    OutWidth = Texture->GetSizeX();
+    OutHeight = Texture->GetSizeY();
+    OutData = Surface->Lock(0, 0, RLM_ReadOnly, Stride);
+}
+
+void FMetalDynamicRHI::RHIUnmapStagingSurface(FTextureRHIParamRef TextureRHI)
+{
+    FMetalSurface* Surface = GetMetalSurfaceFromRHITexture(TextureRHI);
+    
+    Surface->Unlock(0, 0);
+}
+
+void FMetalDynamicRHI::RHIReadSurfaceFloatData(FTextureRHIParamRef TextureRHI, FIntRect Rect, TArray<FFloat16Color>& OutData, ECubeFace CubeFace,int32 ArrayIndex,int32 MipIndex)
+{
+	FMetalSurface* Surface = GetMetalSurfaceFromRHITexture(TextureRHI);
+	
+    id<MTLTexture> Texture = Surface->Texture;
+    if(!Texture && (Surface->Flags & TexCreate_Presentable))
+    {
+        id<CAMetalDrawable> Drawable = (id<CAMetalDrawable>)GetMetalDeviceContext().GetCurrentState().GetCurrentDrawable();
+        Texture = Drawable ? Drawable.texture : nil;
+    }
+    if(!Texture)
+    {
+        UE_LOG(LogRHI, Error, TEXT("Trying to read from an uninitialised texture."));
+        return;
+    }
+    
+	// verify the input image format (but don't crash)
+	if (Surface->PixelFormat != PF_FloatRGBA)
+	{
+		UE_LOG(LogRHI, Log, TEXT("Trying to read non-FloatRGBA surface."));
+	}
+
+	if (TextureRHI->GetTextureCube())
+	{
+		// adjust index to account for cubemaps as texture arrays
+		ArrayIndex *= CubeFace_MAX;
+		ArrayIndex += GetMetalCubeFace(CubeFace);
+	}
+	
+	// allocate output space
+	const uint32 SizeX = Rect.Width();
+	const uint32 SizeY = Rect.Height();
+	OutData.Empty();
+	OutData.AddUninitialized(SizeX * SizeY);
+	
+	MTLRegion Region = MTLRegionMake2D(Rect.Min.X, Rect.Min.Y, SizeX, SizeY);
+	
+	// function wants details about the destination, not the source
+	const uint32 Stride = GPixelFormats[Surface->PixelFormat].BlockBytes * SizeX;
+	const uint32 BytesPerImage = Stride  * SizeY;
+	int32 FloatBGRADataSize = Stride * SizeY;
+	FMetalPooledBuffer Buffer = Context->CreatePooledBuffer(FMetalPooledBufferArgs(Context->GetDevice(), FloatBGRADataSize, MTLStorageModeShared));
+	{
+		// Synchronise the texture with the CPU
+		SCOPE_CYCLE_COUNTER(STAT_MetalTexturePageOffTime);
+		
+		id<MTLBlitCommandEncoder> Blitter = Context->GetBlitContext();
+		[Blitter copyFromTexture:Texture sourceSlice:ArrayIndex sourceLevel:MipIndex sourceOrigin:Region.origin sourceSize:Region.size toBuffer:Buffer.Buffer destinationOffset:0 destinationBytesPerRow:Stride destinationBytesPerImage:BytesPerImage];
+		
+		//kick the current command buffer.
+		Context->SubmitCommandBufferAndWait();
+	}
+	
+	FFloat16Color* FloatBGRAData = (FFloat16Color*)[Buffer.Buffer contents];
+	FFloat16Color* OutDataPtr = OutData.GetData();
+	FMemory::Memcpy(OutDataPtr, FloatBGRAData, FloatBGRADataSize);
+	
+	Context->ReleasePooledBuffer(Buffer);
+}
+
+void FMetalDynamicRHI::RHIRead3DSurfaceFloatData(FTextureRHIParamRef TextureRHI,FIntRect InRect,FIntPoint ZMinMax,TArray<FFloat16Color>& OutData)
+{
+	FMetalSurface* Surface = GetMetalSurfaceFromRHITexture(TextureRHI);
+	
+	id<MTLTexture> Texture = Surface->Texture;
+	if(!Texture)
+	{
+		UE_LOG(LogRHI, Error, TEXT("Trying to read from an uninitialised texture."));
+		return;
+	}
+	
+	// verify the input image format (but don't crash)
+	if (Surface->PixelFormat != PF_FloatRGBA)
+	{
+		UE_LOG(LogRHI, Log, TEXT("Trying to read non-FloatRGBA surface."));
+	}
+	
+	// allocate output space
+	const uint32 SizeX = InRect.Width();
+	const uint32 SizeY = InRect.Height();
+	const uint32 SizeZ = ZMinMax.Y - ZMinMax.X;
+	OutData.Empty();
+	OutData.AddUninitialized(SizeX * SizeY * SizeZ);
+	
+	MTLRegion Region = MTLRegionMake3D(InRect.Min.X, InRect.Min.Y, ZMinMax.X, SizeX, SizeY, SizeZ);
+	
+	// function wants details about the destination, not the source
+	const uint32 Stride = GPixelFormats[Surface->PixelFormat].BlockBytes * SizeX;
+	const uint32 BytesPerImage = Stride  * SizeY;
+	int32 FloatBGRADataSize = Stride * SizeY * SizeZ;
+	FMetalPooledBuffer Buffer = Context->CreatePooledBuffer(FMetalPooledBufferArgs(Context->GetDevice(), FloatBGRADataSize, MTLStorageModeShared));
+	{
+		// Synchronise the texture with the CPU
+		SCOPE_CYCLE_COUNTER(STAT_MetalTexturePageOffTime);
+		
+		id<MTLBlitCommandEncoder> Blitter = Context->GetBlitContext();
+		[Blitter copyFromTexture:Texture sourceSlice:0 sourceLevel:0 sourceOrigin:Region.origin sourceSize:Region.size toBuffer:Buffer.Buffer destinationOffset:0 destinationBytesPerRow:Stride destinationBytesPerImage:BytesPerImage];
+		
+		//kick the current command buffer.
+		Context->SubmitCommandBufferAndWait();
+	}
+	
+	FFloat16Color* FloatBGRAData = (FFloat16Color*)[Buffer.Buffer contents];
+	FFloat16Color* OutDataPtr = OutData.GetData();
+	FMemory::Memcpy(OutDataPtr, FloatBGRAData, FloatBGRADataSize);
+	
+	Context->ReleasePooledBuffer(Buffer);
+}
diff --git a/Engine/Source/Runtime/Apple/MetalRHI/Private/MetalStateCache.cpp b/Engine/Source/Runtime/Apple/MetalRHI/Private/MetalStateCache.cpp
new file mode 100644
--- /dev/null
+++ b/Engine/Source/Runtime/Apple/MetalRHI/Private/MetalStateCache.cpp
@@ -0,0 +1,670 @@
+// Copyright 1998-2015 Epic Games, Inc. All Rights Reserved.
+
+#include "MetalRHIPrivate.h"
+#include "MetalStateCache.h"
+
+static MTLTriangleFillMode TranslateFillMode(ERasterizerFillMode FillMode)
+{
+	switch (FillMode)
+	{
+		case FM_Wireframe:	return MTLTriangleFillModeLines;
+		case FM_Point:		return MTLTriangleFillModeFill;
+		default:			return MTLTriangleFillModeFill;
+	};
+}
+
+static MTLCullMode TranslateCullMode(ERasterizerCullMode CullMode)
+{
+	switch (CullMode)
+	{
+		case CM_CCW:	return MTLCullModeFront;
+		case CM_CW:		return MTLCullModeBack;
+		default:		return MTLCullModeNone;
+	}
+}
+
+FMetalStateCache::FMetalStateCache(FMetalCommandEncoder& InCommandEncoder)
+: CommandEncoder(InCommandEncoder)
+, BlendState(nullptr)
+, DepthStencilState(nullptr)
+, RasterizerState(nullptr)
+, BoundShaderState(nullptr)
+, StencilRef(0)
+, BlendFactor(FLinearColor::Transparent)
+, FrameBufferSize(CGSizeMake(0.0, 0.0))
+, CurrentDrawable(nil)
+#if PLATFORM_MAC
+, CurrentLayer(nil)
+#endif
+, bHasValidRenderTarget(false)
+{
+	Viewport.originX = Viewport.originY = Viewport.width = Viewport.height = Viewport.znear = Viewport.zfar = 0.0;
+	
+	FMemory::Memzero(VertexBuffers, sizeof(VertexBuffers));
+	FMemory::Memzero(VertexStrides, sizeof(VertexStrides));
+	
+	FMemory::Memzero(RenderTargetsInfo);
+	
+	//@todo-rco: What Size???
+	// make a buffer for each shader type
+	ShaderParameters[CrossCompiler::SHADER_STAGE_VERTEX].InitializeResources(1024*1024);
+	ShaderParameters[CrossCompiler::SHADER_STAGE_PIXEL].InitializeResources(1024*1024);
+	ShaderParameters[CrossCompiler::SHADER_STAGE_COMPUTE].InitializeResources(1024*1024);
+}
+
+FMetalStateCache::~FMetalStateCache()
+{
+	
+}
+
+void FMetalStateCache::SetBlendFactor(FLinearColor const& InBlendFactor)
+{
+	//if(BlendFactor != InBlendFactor) // @todo zebra
+	{
+		BlendFactor = InBlendFactor;
+		CommandEncoder.SetBlendColor(InBlendFactor.R, InBlendFactor.G, InBlendFactor.B, InBlendFactor.A);
+	}
+}
+
+void FMetalStateCache::SetStencilRef(uint32 const InStencilRef)
+{
+	//if(StencilRef != InStencilRef) // @todo zebra
+	{
+		StencilRef = InStencilRef;
+		CommandEncoder.SetStencilReferenceValue(InStencilRef);
+	}
+}
+
+void FMetalStateCache::SetBlendState(FMetalBlendState* InBlendState)
+{
+	//if(BlendState != InBlendState) // @todo zebra
+	{
+		BlendState = InBlendState;
+		if(InBlendState)
+		{
+			for(uint32 RenderTargetIndex = 0;RenderTargetIndex < MaxMetalRenderTargets; ++RenderTargetIndex)
+			{
+				MTLRenderPipelineColorAttachmentDescriptor* Blend = BlendState->RenderTargetStates[RenderTargetIndex].BlendState;
+				MTLRenderPipelineColorAttachmentDescriptor* Dest = [PipelineDesc.PipelineDescriptor.colorAttachments objectAtIndexedSubscript:RenderTargetIndex];
+
+				if(Blend && Dest)
+				{
+					// assign each property manually, would be nice if this was faster
+					Dest.blendingEnabled = Blend.blendingEnabled;
+					Dest.sourceRGBBlendFactor = Blend.sourceRGBBlendFactor;
+					Dest.destinationRGBBlendFactor = Blend.destinationRGBBlendFactor;
+					Dest.rgbBlendOperation = Blend.rgbBlendOperation;
+					Dest.sourceAlphaBlendFactor = Blend.sourceAlphaBlendFactor;
+					Dest.destinationAlphaBlendFactor = Blend.destinationAlphaBlendFactor;
+					Dest.alphaBlendOperation = Blend.alphaBlendOperation;
+					Dest.writeMask = Blend.writeMask;
+				}
+		
+				// set the hash bits for this RT
+				PipelineDesc.SetHashValue(FMetalRenderPipelineDesc::BlendBitOffsets[RenderTargetIndex], NumBits_BlendState, BlendState->RenderTargetStates[RenderTargetIndex].BlendStateKey);
+			}
+		}
+	}
+}
+
+void FMetalStateCache::SetDepthStencilState(FMetalDepthStencilState* InDepthStencilState)
+{
+	//if(DepthStencilState != InDepthStencilState) // @todo zebra
+	{
+		DepthStencilState = InDepthStencilState;
+		if(InDepthStencilState)
+		{
+			//activate the state
+			CommandEncoder.SetDepthStencilState(InDepthStencilState->State);
+		}
+	}
+}
+
+void FMetalStateCache::SetRasterizerState(FMetalRasterizerState* InRasterizerState)
+{
+	//if(RasterizerState != InRasterizerState) // @todo zebra
+	{
+		RasterizerState = InRasterizerState;
+		if(InRasterizerState)
+		{
+			CommandEncoder.SetFrontFacingWinding(MTLWindingCounterClockwise);
+			
+			CommandEncoder.SetCullMode(TranslateCullMode(RasterizerState->State.CullMode));
+			
+			// no clamping
+			CommandEncoder.SetDepthBias(RasterizerState->State.DepthBias, RasterizerState->State.SlopeScaleDepthBias, FLT_MAX);
+			
+			CommandEncoder.SetTriangleFillMode(TranslateFillMode(RasterizerState->State.FillMode));
+		}
+	}
+}
+
+void FMetalStateCache::SetBoundShaderState(FMetalBoundShaderState* InBoundShaderState)
+{
+	//if(BoundShaderState != InBoundShaderState) // @todo zebra
+	{
+		BoundShaderState = InBoundShaderState;
+	}
+}
+
+void FMetalStateCache::SetComputeShader(FMetalComputeShader* InComputeShader)
+{
+	//if(ComputeShader != InComputeShader) // @todo zebra
+	{
+		ComputeShader = InComputeShader;
+		
+		// set this compute shader pipeline as the current (this resets all state, so we need to set all resources after calling this)
+		CommandEncoder.SetComputePipelineState(ComputeShader->Kernel);
+	}
+}
+
+void FMetalStateCache::SetRenderTargetsInfo(FRHISetRenderTargetsInfo const& InRenderTargets, id<MTLBuffer> const QueryBuffer)
+{
+	ConditionalSwitchToRender();
+
+	// see if our new Info matches our previous Info
+	if (NeedsToSetRenderTarget(InRenderTargets))
+	{
+		// back this up for next frame
+		RenderTargetsInfo = InRenderTargets;
+		
+		// at this point, we need to fully set up an encoder/command buffer, so make a new one (autoreleased)
+		MTLRenderPassDescriptor* RenderPass = [MTLRenderPassDescriptor renderPassDescriptor];
+	
+		// if we need to do queries, write to the supplied query buffer
+		RenderPass.visibilityResultBuffer = QueryBuffer;
+	
+		// default to non-msaa
+	    int32 OldCount = PipelineDesc.SampleCount;
+		PipelineDesc.SampleCount = 0;
+	
+		bHasValidRenderTarget = false;
+		
+		uint8 ArrayTargets = 0;
+		uint8 BoundTargets = 0;
+		uint32 ArrayRenderLayers = UINT_MAX;
+		
+		bool bFramebufferSizeSet = false;
+		FrameBufferSize = CGSizeMake(0.f, 0.f);
+		
+		for (uint32 RenderTargetIndex = 0; RenderTargetIndex < MaxMetalRenderTargets; RenderTargetIndex++)
+		{
+			// default to invalid
+			uint8 FormatKey = 0;
+			// only try to set it if it was one that was set (ie less than RenderTargetsInfo.NumColorRenderTargets)
+			if (RenderTargetIndex < RenderTargetsInfo.NumColorRenderTargets && RenderTargetsInfo.ColorRenderTarget[RenderTargetIndex].Texture != nullptr)
+			{
+				const FRHIRenderTargetView& RenderTargetView = RenderTargetsInfo.ColorRenderTarget[RenderTargetIndex];
+				FMetalSurface& Surface = *GetMetalSurfaceFromRHITexture(RenderTargetView.Texture);
+				FormatKey = Surface.FormatKey;
+				
+				uint32 Width = FMath::Max((uint32)(Surface.SizeX >> RenderTargetView.MipIndex), (uint32)1);
+				uint32 Height = FMath::Max((uint32)(Surface.SizeY >> RenderTargetView.MipIndex), (uint32)1);
+				if(!bFramebufferSizeSet)
+				{
+					bFramebufferSizeSet = true;
+					FrameBufferSize.width = Width;
+					FrameBufferSize.height = Height;
+				}
+				else
+				{
+					FrameBufferSize.width = FMath::Min(FrameBufferSize.width, (CGFloat)Width);
+					FrameBufferSize.height = FMath::Min(FrameBufferSize.height, (CGFloat)Height);
+				}
+	
+				// if this is the back buffer, make sure we have a usable drawable
+				ConditionalUpdateBackBuffer(Surface);
+	
+				BoundTargets |= 1 << RenderTargetIndex;
+				
+	            if (Surface.Texture == nil)
+	            {
+	                PipelineDesc.SampleCount = OldCount;
+	                return;
+	            }
+	
+				// user code generally passes -1 as a default, but we need 0
+				uint32 ArraySliceIndex = RenderTargetView.ArraySliceIndex == 0xFFFFFFFF ? 0 : RenderTargetView.ArraySliceIndex;
+				if (Surface.bIsCubemap)
+				{
+					ArraySliceIndex = GetMetalCubeFace((ECubeFace)ArraySliceIndex);
+				}
+				
+				switch(Surface.Type)
+				{
+					case RRT_Texture2DArray:
+					case RRT_Texture3D:
+					case RRT_TextureCube:
+						if(RenderTargetView.ArraySliceIndex == 0xFFFFFFFF)
+						{
+							ArrayTargets |= (1 << RenderTargetIndex);
+							ArrayRenderLayers = FMath::Min(ArrayRenderLayers, Surface.GetNumFaces());
+						}
+						else
+						{
+							ArrayRenderLayers = 1;
+						}
+						break;
+					default:
+						ArrayRenderLayers = 1;
+						break;
+				}
+	
+				MTLRenderPassColorAttachmentDescriptor* ColorAttachment = [MTLRenderPassColorAttachmentDescriptor new];
+	
+				if (Surface.MSAATexture != nil)
+				{
+					// set up an MSAA attachment
+					ColorAttachment.texture = Surface.MSAATexture;
+					ColorAttachment.storeAction = MTLStoreActionMultisampleResolve;
+					ColorAttachment.resolveTexture = Surface.Texture;
+					PipelineDesc.SampleCount = Surface.MSAATexture.sampleCount;
+	
+					// only allow one MRT with msaa
+					checkf(RenderTargetsInfo.NumColorRenderTargets == 1, TEXT("Only expected one MRT when using MSAA"));
+				}
+				else
+				{
+					// set up non-MSAA attachment
+					ColorAttachment.texture = Surface.Texture;
+					ColorAttachment.storeAction = GetMetalRTStoreAction(RenderTargetView.StoreAction);
+					PipelineDesc.SampleCount = 1;
+				}
+				
+				ColorAttachment.level = RenderTargetView.MipIndex;
+				if(Surface.Type == RRT_Texture3D)
+				{
+					ColorAttachment.depthPlane = ArraySliceIndex;
+				}
+				else
+				{
+					ColorAttachment.slice = ArraySliceIndex;
+				}
+				
+				ColorAttachment.loadAction = GetMetalRTLoadAction(RenderTargetView.LoadAction);
+				const FClearValueBinding& ClearValue = RenderTargetsInfo.ColorRenderTarget[RenderTargetIndex].Texture->GetClearBinding();
+				if (ClearValue.ColorBinding == EClearBinding::EColorBound)
+				{
+					const FLinearColor& ClearColor = ClearValue.GetClearColor();
+					ColorAttachment.clearColor = MTLClearColorMake(ClearColor.R, ClearColor.G, ClearColor.B, ClearColor.A);
+				}
+
+				// assign the attachment to the slot
+				[RenderPass.colorAttachments setObject:ColorAttachment atIndexedSubscript:RenderTargetIndex];
+				[PipelineDesc.PipelineDescriptor.colorAttachments objectAtIndexedSubscript:RenderTargetIndex].pixelFormat = ColorAttachment.texture.pixelFormat;
+	
+				[ColorAttachment release];
+	
+				bHasValidRenderTarget = true;
+			}
+			else
+			{
+				[PipelineDesc.PipelineDescriptor.colorAttachments objectAtIndexedSubscript:RenderTargetIndex].pixelFormat = MTLPixelFormatInvalid;
+			}
+	
+			// update the hash no matter what case (null, unused, used)
+			PipelineDesc.SetHashValue(FMetalRenderPipelineDesc::RTBitOffsets[RenderTargetIndex], NumBits_RenderTargetFormat, FormatKey);
+		}
+		
+#if METAL_API_1_1 && PLATFORM_MAC
+		if(ArrayTargets)
+		{
+			if (ArrayTargets == BoundTargets)
+			{
+				RenderPass.renderTargetArrayLength = ArrayRenderLayers;
+			}
+			else
+			{
+				UE_LOG(LogMetal, Fatal, TEXT("All color render targets must be layered when performing mulit-layered rendering under Metal."));
+			}
+		}
+#endif
+	
+		// default to invalid
+		PipelineDesc.PipelineDescriptor.depthAttachmentPixelFormat = MTLPixelFormatInvalid;
+		PipelineDesc.PipelineDescriptor.stencilAttachmentPixelFormat = MTLPixelFormatInvalid;
+		
+		uint8 DepthFormatKey = 0;
+		uint8 StencilFormatKey = 0;
+	
+		// setup depth and/or stencil
+		if (RenderTargetsInfo.DepthStencilRenderTarget.Texture != nullptr)
+		{
+			FMetalSurface& Surface = *GetMetalSurfaceFromRHITexture(RenderTargetsInfo.DepthStencilRenderTarget.Texture);
+			
+			if(!bFramebufferSizeSet)
+			{
+				bFramebufferSizeSet = true;
+				FrameBufferSize.width = Surface.SizeX;
+				FrameBufferSize.height = Surface.SizeY;
+			}
+			else
+			{
+				FrameBufferSize.width = FMath::Min(FrameBufferSize.width, (CGFloat)Surface.SizeX);
+				FrameBufferSize.height = FMath::Min(FrameBufferSize.height, (CGFloat)Surface.SizeY);
+			}
+			
+			MTLPixelFormat DepthStencilFormat = Surface.Texture ? Surface.Texture.pixelFormat : MTLPixelFormatInvalid;
+			
+			bool bHasDepth = false;
+			bool bHasStencil = false;
+			switch(DepthStencilFormat)
+			{
+				case MTLPixelFormatDepth32Float:
+					bHasDepth = true;
+					break;
+				case MTLPixelFormatStencil8:
+					bHasStencil = true;
+					break;
+#if METAL_API_1_1
+				case MTLPixelFormatDepth32Float_Stencil8:
+					bHasDepth = true;
+					bHasStencil = true;
+					break;
+#if PLATFORM_MAC
+				case MTLPixelFormatDepth24Unorm_Stencil8:
+					bHasDepth = true;
+					bHasStencil = true;
+					break;
+#endif
+#endif
+				default:
+					break;
+			}
+
+			float DepthClearValue = 0.0f;
+			uint32 StencilClearValue = 0;
+			const FClearValueBinding& ClearValue = RenderTargetsInfo.DepthStencilRenderTarget.Texture->GetClearBinding();
+			if (ClearValue.ColorBinding == EClearBinding::EDepthStencilBound)
+			{
+				ClearValue.GetDepthStencil(DepthClearValue, StencilClearValue);
+			}
+
+			if (bHasDepth)
+			{
+				MTLRenderPassDepthAttachmentDescriptor* DepthAttachment = [[MTLRenderPassDepthAttachmentDescriptor alloc] init];
+				
+				DepthFormatKey = Surface.FormatKey;
+	
+				// set up the depth attachment
+				DepthAttachment.texture = Surface.Texture;
+				DepthAttachment.loadAction = GetMetalRTLoadAction(RenderTargetsInfo.DepthStencilRenderTarget.DepthLoadAction);
+#if PLATFORM_MAC
+				if(RenderTargetsInfo.DepthStencilRenderTarget.GetDepthStencilAccess().IsDepthWrite())
+				{
+					DepthAttachment.storeAction = GetMetalRTStoreAction(RenderTargetsInfo.DepthStencilRenderTarget.DepthStoreAction);
+				}
+				else
+				{
+					DepthAttachment.storeAction = MTLStoreActionDontCare;
+				}
+#else
+				DepthAttachment.storeAction = GetMetalRTStoreAction(RenderTargetsInfo.DepthStencilRenderTarget.DepthStoreAction);
+#endif
+				DepthAttachment.clearDepth = DepthClearValue;
+
+				PipelineDesc.PipelineDescriptor.depthAttachmentPixelFormat = DepthAttachment.texture.pixelFormat;
+				if (PipelineDesc.SampleCount == 0)
+				{
+					PipelineDesc.SampleCount = DepthAttachment.texture.sampleCount;
+				}
+	
+				// and assign it
+				RenderPass.depthAttachment = DepthAttachment;
+				[DepthAttachment release];
+			}
+	
+			if (bHasStencil)
+			{
+				MTLRenderPassStencilAttachmentDescriptor* StencilAttachment = [[MTLRenderPassStencilAttachmentDescriptor alloc] init];
+				
+				StencilFormatKey = Surface.FormatKey;
+	
+				// set up the stencil attachment
+				StencilAttachment.texture = Surface.StencilTexture;
+				StencilAttachment.loadAction = GetMetalRTLoadAction(RenderTargetsInfo.DepthStencilRenderTarget.StencilLoadAction);
+#if PLATFORM_MAC
+				if(RenderTargetsInfo.DepthStencilRenderTarget.GetDepthStencilAccess().IsStencilWrite())
+				{
+					StencilAttachment.storeAction = GetMetalRTStoreAction(RenderTargetsInfo.DepthStencilRenderTarget.GetStencilStoreAction());
+				}
+				else
+				{
+					StencilAttachment.storeAction = MTLStoreActionDontCare;
+				}
+#else
+				StencilAttachment.storeAction = GetMetalRTStoreAction(RenderTargetsInfo.DepthStencilRenderTarget.GetStencilStoreAction());
+#endif
+				StencilAttachment.clearStencil = StencilClearValue;
+
+				PipelineDesc.PipelineDescriptor.stencilAttachmentPixelFormat = StencilAttachment.texture.pixelFormat;
+				if (PipelineDesc.SampleCount == 0)
+				{
+					PipelineDesc.SampleCount = StencilAttachment.texture.sampleCount;
+				}
+	
+				// and assign it
+				RenderPass.stencilAttachment = StencilAttachment;
+				[StencilAttachment release];
+			}
+		}
+	
+		// update hash for the depth/stencil buffer & sample count
+		PipelineDesc.SetHashValue(Offset_DepthFormat, NumBits_DepthFormat, DepthFormatKey);
+		PipelineDesc.SetHashValue(Offset_StencilFormat, NumBits_StencilFormat, StencilFormatKey);
+		PipelineDesc.SetHashValue(Offset_SampleCount, NumBits_SampleCount, PipelineDesc.SampleCount);
+	
+		// commit pending commands on the old render target
+		if(CommandEncoder.IsRenderCommandEncoderActive())
+		{
+			CommandEncoder.EndEncoding();
+		}
+	
+		// make a new render context to use to render to the framebuffer
+		CommandEncoder.BeginRenderCommandEncoding(RenderPass);
+		
+		// Reset any existing state as that must be fully reinitialised by the caller.
+		DepthStencilState.SafeRelease();
+		RasterizerState.SafeRelease();
+		BlendState.SafeRelease();
+		BoundShaderState.SafeRelease();
+	}
+}
+
+void FMetalStateCache::SetHasValidRenderTarget(bool InHasValidRenderTarget)
+{
+	bHasValidRenderTarget = InHasValidRenderTarget;
+}
+
+void FMetalStateCache::SetViewport(const MTLViewport& InViewport)
+{
+	Viewport = InViewport;
+}
+
+void FMetalStateCache::SetVertexBuffer(uint32 const Index, id<MTLBuffer> Buffer, uint32 const Stride, uint32 const Offset)
+{
+	check(Index < MaxMetalStreams);
+	VertexBuffers[Index] = Buffer;
+	VertexStrides[Index] = Buffer ? Stride : 0;
+	if (Buffer != NULL)
+	{
+		CommandEncoder.SetShaderBuffer(SF_Vertex, Buffer, Offset, Index);
+	}
+}
+
+#if PLATFORM_MAC
+void FMetalStateCache::SetCurrentLayer(CAMetalLayer* NewLayer)
+{
+	check(NewLayer);
+	CurrentLayer = NewLayer;
+	check(CurrentLayer);
+}
+
+void FMetalStateCache::SetPrimitiveTopology(MTLPrimitiveTopologyClass PrimitiveType)
+{
+	PipelineDesc.SetHashValue(Offset_PrimitiveTopology, NumBits_PrimitiveTopology, PrimitiveType);
+	PipelineDesc.PipelineDescriptor.inputPrimitiveTopology = PrimitiveType;
+}
+#endif
+
+void FMetalStateCache::ConditionalSwitchToRender(void)
+{
+	// were we in blit or compute mode?
+	if(CommandEncoder.IsBlitCommandEncoderActive() || CommandEncoder.IsComputeCommandEncoderActive())
+	{
+		// stop the current encoding and cleanup
+		CommandEncoder.EndEncoding();
+	}
+	
+	if (!CommandEncoder.IsRenderCommandEncoderActive())
+	{
+		// Force a rebind of the render encoder state on the next SetRenderTarget call.
+		SetHasValidRenderTarget(false);
+	}	
+	// we can't start graphics encoding until a new SetRenderTargetsInfo is called because it needs the whole render pass
+	// we could cache the render pass if we want to support going back to previous render targets
+	// we will catch it by the fact that CommandEncoder will be nil until we call SetRenderTargetsInfo	
+}
+
+void FMetalStateCache::ConditionalSwitchToCompute(void)
+{
+	// if we were in rendering or blit mode, stop the encoding and start compute
+	if (CommandEncoder.IsRenderCommandEncoderActive() || CommandEncoder.IsBlitCommandEncoderActive())
+	{
+		// stop encoding graphics and clean up
+		CommandEncoder.EndEncoding();
+	}
+	if(!CommandEncoder.IsComputeCommandEncoderActive())
+	{
+		// Clear any previous compute shader
+		ComputeShader.SafeRelease();
+		
+		// start encoding for compute
+		CommandEncoder.BeginComputeCommandEncoding();
+	}
+}
+
+void FMetalStateCache::ConditionalSwitchToBlit(void)
+{
+	// if we were in rendering or compute mode, stop the encoding and start compute
+	if (CommandEncoder.IsRenderCommandEncoderActive() || CommandEncoder.IsComputeCommandEncoderActive())
+	{
+		// stop encoding graphics and clean up
+		CommandEncoder.EndEncoding();
+	}
+	if(!CommandEncoder.IsBlitCommandEncoderActive())
+	{
+		// start encoding for compute
+		CommandEncoder.BeginBlitCommandEncoding();
+	}
+}
+
+void FMetalStateCache::ResetCurrentDrawable()
+{
+	if (CurrentDrawable)
+	{
+		[CurrentDrawable release];
+		CurrentDrawable = nil;
+	}
+}
+
+void FMetalStateCache::ConditionalUpdateBackBuffer(FMetalSurface& Surface)
+{
+	// are we setting the back buffer? if so, make sure we have the drawable
+//	if (&Surface == &BackBuffer->Surface) // @todo zebra &Surface == &CurrentViewport->GetBackBuffer()->Surface, but perhaps Surface.Texture == nil is enough?
+	{
+		// update the back buffer texture the first time used this frame
+		if (Surface.Texture == nil)
+		{
+			SCOPE_CYCLE_COUNTER(STAT_MetalMakeDrawableTime);
+
+			uint32 IdleStart = FPlatformTime::Cycles();
+
+			if (CurrentDrawable == nil)
+			{
+				// make a drawable object for this frame
+#if PLATFORM_IOS
+				CurrentDrawable = [[[IOSAppDelegate GetDelegate].IOSView MakeDrawable] retain];
+#else // @todo zebra
+				check(CurrentLayer);
+				CurrentDrawable = CurrentLayer ? [[CurrentLayer nextDrawable] retain] : nil;
+				//check(CurrentDrawable);
+#endif
+			}
+
+			GRenderThreadIdle[ERenderThreadIdleTypes::WaitingForGPUPresent] += FPlatformTime::Cycles() - IdleStart;
+			GRenderThreadNumIdle[ERenderThreadIdleTypes::WaitingForGPUPresent]++;
+
+			// set the texture into the backbuffer
+            if (CurrentDrawable != nil)
+            {
+                Surface.Texture = CurrentDrawable.texture;
+            }
+            else
+            {
+                Surface.Texture = nil;
+            }
+		}
+	}
+}
+
+bool FMetalStateCache::NeedsToSetRenderTarget(const FRHISetRenderTargetsInfo& InRenderTargetsInfo) const
+{
+	// see if our new Info matches our previous Info
+	
+	// basic checks
+	bool bAllChecksPassed = GetHasValidRenderTarget() && CommandEncoder.IsRenderCommandEncoderActive() && InRenderTargetsInfo.NumColorRenderTargets == RenderTargetsInfo.NumColorRenderTargets &&
+		// handle the case where going from backbuffer + depth -> backbuffer + null, no need to reset RT and do a store/load
+		(InRenderTargetsInfo.DepthStencilRenderTarget.Texture == RenderTargetsInfo.DepthStencilRenderTarget.Texture || InRenderTargetsInfo.DepthStencilRenderTarget.Texture == nullptr);
+
+	// now check each color target if the basic tests passe
+	if (bAllChecksPassed)
+	{
+		for (int32 RenderTargetIndex = 0; RenderTargetIndex < InRenderTargetsInfo.NumColorRenderTargets; RenderTargetIndex++)
+		{
+			const FRHIRenderTargetView& RenderTargetView = InRenderTargetsInfo.ColorRenderTarget[RenderTargetIndex];
+			const FRHIRenderTargetView& PreviousRenderTargetView = RenderTargetsInfo.ColorRenderTarget[RenderTargetIndex];
+
+			// handle simple case of switching textures or mip/slice
+			if (RenderTargetView.Texture != PreviousRenderTargetView.Texture ||
+				RenderTargetView.MipIndex != PreviousRenderTargetView.MipIndex ||
+				RenderTargetView.ArraySliceIndex != PreviousRenderTargetView.ArraySliceIndex)
+			{
+				bAllChecksPassed = false;
+				break;
+			}
+			
+			// it's non-trivial when we need to switch based on load/store action:
+			// LoadAction - it only matters what we are switching to in the new one
+			//    If we switch to Load, no need to switch as we can re-use what we already have
+			//    If we switch to Clear, we have to always switch to a new RT to force the clear
+			//    If we switch to DontCare, there's definitely no need to switch
+			if (RenderTargetView.LoadAction == ERenderTargetLoadAction::EClear)
+			{
+				bAllChecksPassed = false;
+				break;
+			}
+			// StoreAction - this matters what the previous one was **In Spirit**
+			//    If we come from Store, we need to switch to a new RT to force the store
+			//    If we come from DontCare, then there's no need to switch
+			//    @todo metal: However, we basically only use Store now, and don't
+			//        care about intermediate results, only final, so we don't currently check the value
+//			if (PreviousRenderTargetView.StoreAction == ERenderTTargetStoreAction::EStore)
+//			{
+//				bAllChecksPassed = false;
+//				break;
+//			}
+		}
+	}
+
+	// if we are setting them to nothing, then this is probably end of frame, and we can't make a framebuffer
+	// with nothng, so just abort this (only need to check on single MRT case)
+	if (InRenderTargetsInfo.NumColorRenderTargets == 1 && InRenderTargetsInfo.ColorRenderTarget[0].Texture == nullptr &&
+		InRenderTargetsInfo.DepthStencilRenderTarget.Texture == nullptr)
+	{
+		bAllChecksPassed = true;
+	}
+
+	return bAllChecksPassed == false;
+}
+
diff --git a/Engine/Source/Runtime/Apple/MetalRHI/Private/MetalStateCache.h b/Engine/Source/Runtime/Apple/MetalRHI/Private/MetalStateCache.h
new file mode 100644
--- /dev/null
+++ b/Engine/Source/Runtime/Apple/MetalRHI/Private/MetalStateCache.h
@@ -0,0 +1,96 @@
+// Copyright 1998-2015 Epic Games, Inc. All Rights Reserved.
+
+#pragma once
+
+#include "MetalRHIPrivate.h"
+#include "MetalCommandEncoder.h"
+#include "MetalRenderPipelineDesc.h"
+
+class FMetalStateCache
+{
+public:
+	FMetalStateCache(FMetalCommandEncoder& CommandEncoder);
+	~FMetalStateCache();
+
+	void SetBlendFactor(FLinearColor const& InBlendFactor);
+	void SetStencilRef(uint32 const InStencilRef);
+	void SetBlendState(FMetalBlendState* InBlendState);
+	void SetDepthStencilState(FMetalDepthStencilState* InDepthStencilState);
+	void SetRasterizerState(FMetalRasterizerState* InRasterizerState);
+	void SetBoundShaderState(FMetalBoundShaderState* BoundShaderState);
+	void SetComputeShader(FMetalComputeShader* InComputeShader);
+	void SetRenderTargetsInfo(FRHISetRenderTargetsInfo const& InRenderTargets, id<MTLBuffer> const QueryBuffer);
+	void SetHasValidRenderTarget(bool const bIsValid);
+	void SetViewport(const MTLViewport& InViewport);
+	void SetVertexBuffer(uint32 const Index, id<MTLBuffer> Buffer, uint32 const Stride, uint32 const Offset);
+#if PLATFORM_MAC
+	void SetCurrentLayer(CAMetalLayer* NewLayer);
+	void SetPrimitiveTopology(MTLPrimitiveTopologyClass PrimitiveType);
+#endif
+
+	FMetalShaderParameterCache& GetShaderParameters(uint32 const Stage) { return ShaderParameters[Stage]; }
+	
+	/** Switch from blit/compute to render command encoding if needed. */
+	void ConditionalSwitchToRender(void);
+	
+	/** Switch from blit/render to compute command encoding if needed. */
+	void ConditionalSwitchToCompute(void);
+	
+	/** Switch from compute/render to blit command encoding if needed. */
+	void ConditionalSwitchToBlit(void);
+	
+	void ResetCurrentDrawable();
+	
+	FMetalRenderPipelineDesc const& GetRenderPipelineDesc() const { return PipelineDesc; }
+	FLinearColor const& GetBlendFactor() const { return BlendFactor; }
+	uint32 GetStencilRef() const { return StencilRef; }
+	FMetalBlendState* GetBlendState() const { return BlendState; }
+	FMetalDepthStencilState* GetDepthStencilState() const { return DepthStencilState; }
+	FMetalRasterizerState* GetRasterizerState() const { return RasterizerState; }
+	FMetalBoundShaderState* GetBoundShaderState() const { return BoundShaderState; }
+	FMetalComputeShader* GetComputeShader() const { return ComputeShader; }
+	CGSize GetFrameBufferSize() const { return FrameBufferSize; }
+	FRHISetRenderTargetsInfo const& GetRenderTargetsInfo() const { return RenderTargetsInfo; }
+	int32 GetNumRenderTargets() { return bHasValidRenderTarget ? RenderTargetsInfo.NumColorRenderTargets : -1; }
+	bool GetHasValidRenderTarget() const { return bHasValidRenderTarget; }
+	const MTLViewport& GetViewport() const { return Viewport; }
+	id<MTLBuffer> GetVertexBuffer(uint32 const Index) { check(Index < MaxMetalStreams); return VertexBuffers[Index]; }
+	uint32 GetVertexStride(uint32 const Index) { check(Index < MaxMetalStreams); return VertexStrides[Index]; }
+	id<CAMetalDrawable> GetCurrentDrawable() const { return CurrentDrawable; }
+#if PLATFORM_MAC
+	CAMetalLayer* GetCurrentLayer() const { return CurrentLayer; }
+#endif
+
+private:
+	void ConditionalUpdateBackBuffer(FMetalSurface& Surface);
+	bool NeedsToSetRenderTarget(const FRHISetRenderTargetsInfo& RenderTargetsInfo) const;
+	
+private:
+	FMetalCommandEncoder& CommandEncoder;
+	FMetalShaderParameterCache ShaderParameters[CrossCompiler::NUM_SHADER_STAGES];
+
+	FMetalRenderPipelineDesc PipelineDesc;
+
+	id<MTLBuffer> VertexBuffers[MaxMetalStreams];
+	uint32 VertexStrides[MaxMetalStreams];
+
+	TRefCountPtr<FMetalBlendState> BlendState;
+	TRefCountPtr<FMetalDepthStencilState> DepthStencilState;
+	TRefCountPtr<FMetalRasterizerState> RasterizerState;
+	TRefCountPtr<FMetalBoundShaderState> BoundShaderState;
+	TRefCountPtr<FMetalComputeShader> ComputeShader;
+	uint32 StencilRef;
+	
+	FLinearColor BlendFactor;
+	CGSize FrameBufferSize;
+	
+	id<CAMetalDrawable> CurrentDrawable;
+#if PLATFORM_MAC
+	CAMetalLayer* CurrentLayer;
+#endif
+
+	MTLViewport Viewport;
+	
+	FRHISetRenderTargetsInfo RenderTargetsInfo;
+	bool bHasValidRenderTarget;
+};
diff --git a/Engine/Source/Runtime/Apple/MetalRHI/Private/MetalStructuredBuffer.cpp b/Engine/Source/Runtime/Apple/MetalRHI/Private/MetalStructuredBuffer.cpp
new file mode 100644
--- /dev/null
+++ b/Engine/Source/Runtime/Apple/MetalRHI/Private/MetalStructuredBuffer.cpp
@@ -0,0 +1,102 @@
+// Copyright 1998-2015 Epic Games, Inc. All Rights Reserved.
+
+
+#include "MetalRHIPrivate.h"
+
+
+FMetalStructuredBuffer::FMetalStructuredBuffer(uint32 Stride, uint32 Size, FResourceArrayInterface* ResourceArray, uint32 Usage)
+	: FRHIStructuredBuffer(Stride, Size, Usage)
+	, LockOffset(0)
+	, LockSize(0)
+{
+	check((Size % Stride) == 0);
+
+	MTLStorageMode Mode = (GetUsage() & BUF_Volatile) ? MTLStorageModeShared : BUFFER_STORAGE_MODE;
+	FMetalPooledBuffer Buf = GetMetalDeviceContext().CreatePooledBuffer(FMetalPooledBufferArgs(GetMetalDeviceContext().GetDevice(), Size, Mode));
+	Buffer = [Buf.Buffer retain];
+
+	if (ResourceArray)
+	{
+		// copy any resources to the CPU address
+		void* LockedMemory = RHILockStructuredBuffer(this, 0, Size, RLM_WriteOnly);
+ 		FMemory::Memcpy(LockedMemory, ResourceArray->GetResourceData(), Size);
+		ResourceArray->Discard();
+	}
+
+	TRACK_OBJECT(Buffer);
+}
+
+FMetalStructuredBuffer::~FMetalStructuredBuffer()
+{
+	SafeReleasePooledBuffer(Buffer);
+	[Buffer release];
+}
+
+void* FMetalStructuredBuffer::Lock(EResourceLockMode LockMode, uint32 Offset, uint32 Size)
+{
+	check(LockSize == 0 && LockOffset == 0);
+	
+	// In order to properly synchronise the buffer access, when a dynamic buffer is locked for writing, discard the old buffer & create a new one. This prevents writing to a buffer while it is being read by the GPU & thus causing corruption. This matches the logic of other RHIs.
+	if (GetUsage() & BUF_AnyDynamic && LockMode == RLM_WriteOnly)
+	{
+		id<MTLBuffer> OldBuffer = Buffer;
+		MTLStorageMode Mode = (GetUsage() & BUF_Volatile) ? MTLStorageModeShared : BUFFER_STORAGE_MODE;
+		FMetalPooledBuffer Buf = GetMetalDeviceContext().CreatePooledBuffer(FMetalPooledBufferArgs(GetMetalDeviceContext().GetDevice(), GetSize(), Mode));
+		Buffer = [Buf.Buffer retain];
+		GetMetalDeviceContext().ReleasePooledBuffer(OldBuffer);
+		[OldBuffer release];
+	}
+	
+	if(LockMode != RLM_ReadOnly)
+	{
+		LockOffset = Offset;
+		LockSize = Size;
+	}
+#if PLATFORM_MAC
+	else if(Buffer.storageMode == MTLStorageModeManaged)
+	{
+		SCOPE_CYCLE_COUNTER(STAT_MetalBufferPageOffTime);
+		
+		// Synchronise the buffer with the CPU
+		id<MTLBlitCommandEncoder> Blitter = GetMetalDeviceContext().GetBlitContext();
+		[Blitter synchronizeResource:Buffer];
+		
+		//kick the current command buffer.
+		GetMetalDeviceContext().SubmitCommandBufferAndWait();
+	}
+#endif
+	
+	return ((uint8*)[Buffer contents]);
+}
+
+void FMetalStructuredBuffer::Unlock()
+{
+#if PLATFORM_MAC
+	if(LockSize && Buffer.storageMode == MTLStorageModeManaged)
+	{
+		[Buffer didModifyRange:NSMakeRange(LockOffset, LockSize)];
+	}
+#endif
+	LockOffset = 0;
+	LockSize = 0;
+}
+
+
+FStructuredBufferRHIRef FMetalDynamicRHI::RHICreateStructuredBuffer(uint32 Stride,uint32 Size,uint32 InUsage,FRHIResourceCreateInfo& CreateInfo)
+{
+	return new FMetalStructuredBuffer(Stride, Size, CreateInfo.ResourceArray, InUsage);
+}
+
+void* FMetalDynamicRHI::RHILockStructuredBuffer(FStructuredBufferRHIParamRef StructuredBufferRHI,uint32 Offset,uint32 Size,EResourceLockMode LockMode)
+{
+	FMetalStructuredBuffer* StructuredBuffer = ResourceCast(StructuredBufferRHI);
+	
+	// just return the memory plus the offset
+	return (uint8*)StructuredBuffer->Lock(LockMode, Offset, Size);
+}
+
+void FMetalDynamicRHI::RHIUnlockStructuredBuffer(FStructuredBufferRHIParamRef StructuredBufferRHI)
+{
+	FMetalStructuredBuffer* StructuredBuffer = ResourceCast(StructuredBufferRHI);
+	StructuredBuffer->Unlock();
+}
