Merging MetalRHI & MetalShaderFormat changes from ZebraBranch, plus the code to enable or disable Mac support.

diff --git a/Engine/Source/Runtime/Apple/MetalRHI/Private/MetalTexture.cpp b/Engine/Source/Runtime/Apple/MetalRHI/Private/MetalTexture.cpp
new file mode 100644
--- /dev/null
+++ b/Engine/Source/Runtime/Apple/MetalRHI/Private/MetalTexture.cpp
@@ -0,0 +1,1042 @@
+// Copyright 1998-2015 Epic Games, Inc. All Rights Reserved.
+
+/*=============================================================================
+	MetalVertexBuffer.cpp: Metal texture RHI implementation.
+=============================================================================*/
+
+#include "MetalRHIPrivate.h"
+
+uint8 FMetalSurface::NextKey = 1; // 0 is reserved for MTLPixelFormatInvalid
+
+/** Texture reference class. */
+class FMetalTextureReference : public FRHITextureReference
+{
+public:
+	explicit FMetalTextureReference(FLastRenderTimeContainer* InLastRenderTime)
+		: FRHITextureReference(InLastRenderTime)
+	{}
+
+	// IRefCountedObject interface.
+	virtual uint32 AddRef() const
+	{
+		return FRHIResource::AddRef();
+	}
+	virtual uint32 Release() const
+	{
+		return FRHIResource::Release();
+	}
+	virtual uint32 GetRefCount() const
+	{
+		return FRHIResource::GetRefCount();
+	}
+
+	void SetReferencedTexture(FRHITexture* InTexture)
+	{
+		FRHITextureReference::SetReferencedTexture(InTexture);
+	}
+};
+
+/** Given a pointer to a RHI texture that was created by the Metal RHI, returns a pointer to the FMetalTextureBase it encapsulates. */
+FMetalSurface* GetMetalSurfaceFromRHITexture(FRHITexture* Texture)
+{
+    if (!Texture)
+    {
+        return NULL;
+    }
+	if(Texture->GetTexture2D())
+	{
+		return &((FMetalTexture2D*)Texture)->Surface;
+	}
+	else if(Texture->GetTexture2DArray())
+	{
+		return &((FMetalTexture2DArray*)Texture)->Surface;
+	}
+	else if(Texture->GetTexture3D())
+	{
+		return &((FMetalTexture3D*)Texture)->Surface;
+	}
+	else if(Texture->GetTextureCube())
+	{
+		return &((FMetalTextureCube*)Texture)->Surface;
+	}
+	else if(Texture->GetTextureReference())
+	{
+		return GetMetalSurfaceFromRHITexture(static_cast<FMetalTextureReference*>(Texture)->GetReferencedTexture());
+	}
+	else
+	{
+		UE_LOG(LogMetal, Fatal, TEXT("Unknown RHI texture type"));
+		return &((FMetalTexture2D*)Texture)->Surface;
+	}
+}
+
+static bool IsRenderTarget(uint32 Flags)
+{
+	return (Flags & (TexCreate_RenderTargetable | TexCreate_ResolveTargetable | TexCreate_DepthStencilTargetable)) != 0;
+}
+
+#if METAL_API_1_1
+static MTLTextureUsage ConvertFlagsToUsage(uint32 Flags)
+{
+	MTLTextureUsage Usage = MTLTextureUsageUnknown;
+#if PLATFORM_MAC
+	if (Flags & (TexCreate_RenderTargetable|TexCreate_DepthStencilTargetable))
+	{
+		Usage |= MTLTextureUsageRenderTarget;
+		Usage |= MTLTextureUsageShaderRead;
+		Usage |= MTLTextureUsagePixelFormatView;
+	}
+	
+	if(Flags & (TexCreate_ShaderResource|TexCreate_ResolveTargetable))
+	{
+		Usage |= MTLTextureUsageShaderRead;
+		Usage |= MTLTextureUsagePixelFormatView;
+	}
+	
+	if (Flags & TexCreate_UAV)
+	{
+		Usage |= MTLTextureUsageShaderRead;
+		Usage |= MTLTextureUsageShaderWrite;
+		Usage |= MTLTextureUsagePixelFormatView;
+	}
+#endif
+	return Usage;
+}
+#endif
+
+static bool IsPixelFormatCompressed(EPixelFormat Format)
+{
+	switch (Format)
+	{
+		case PF_DXT1:
+		case PF_DXT3:
+		case PF_DXT5:
+		case PF_PVRTC2:
+		case PF_PVRTC4:
+		case PF_BC4:
+		case PF_BC5:
+		case PF_ATC_RGB:
+		case PF_ATC_RGBA_E:
+		case PF_ATC_RGBA_I:
+		case PF_ETC1:
+		case PF_ETC2_RGB:
+		case PF_ETC2_RGBA:
+		case PF_ASTC_4x4:
+		case PF_ASTC_6x6:
+		case PF_ASTC_8x8:
+		case PF_ASTC_10x10:
+		case PF_ASTC_12x12:
+		case PF_BC6H:
+		case PF_BC7:
+			return true;
+		default:
+			return false;
+	}
+}
+
+FMetalSurface::FMetalSurface(FMetalSurface const& Source, NSRange MipRange)
+: Type(Source.Type)
+, PixelFormat(Source.PixelFormat)
+, Texture(nil)
+, MSAATexture(nil)
+, StencilTexture(nil)
+, SizeX(Source.SizeX)
+, SizeY(Source.SizeY)
+, SizeZ(Source.SizeZ)
+, bIsCubemap(Source.bIsCubemap)
+, Flags(Source.Flags)
+, WriteLock(0)
+, TotalTextureSize(0)
+, IB(nullptr)
+{
+	MTLPixelFormat MetalFormat = (MTLPixelFormat)GPixelFormats[PixelFormat].PlatformFormat;
+
+    NSRange Slices = NSMakeRange(0, FMath::Min(SizeZ, 1u));
+    // @todo Zebra Temporary workaround for absence of X24_G8 or equivalent to GL_STENCIL_INDEX so that the stencil part of a texture may be sampled
+	// For now, if we find ourselves *requiring* this we'd need to lazily blit the stencil data out to a separate texture. radr://21813831
+#if METAL_API_1_1
+    if(Source.PixelFormat != PF_DepthStencil)
+    {
+        Texture = [Source.Texture newTextureViewWithPixelFormat:MetalFormat textureType:Source.Texture.textureType levels:MipRange slices:Slices];
+    }
+    else
+#endif
+    {
+        Texture = [Source.Texture retain];
+    }
+	if(Source.MSAATexture)
+	{
+		MSAATexture = Texture;
+	}
+	if(Source.StencilTexture)
+	{
+#if PLATFORM_IOS
+		check(false); // @todo Zebra Must handle separate stencil texture SRV!
+#endif
+		StencilTexture = Texture;
+	}
+	
+	const uint32 BlockSizeX = GPixelFormats[PixelFormat].BlockSizeX;
+	const uint32 BlockSizeY = GPixelFormats[PixelFormat].BlockSizeY;
+	const uint32 BlockBytes = GPixelFormats[PixelFormat].BlockBytes;
+	SizeX = FMath::Max(SizeX >> MipRange.location, BlockSizeX);
+	SizeY = FMath::Max(SizeY >> MipRange.location, BlockSizeY);
+	SizeZ = (Type != RRT_Texture3D) ? SizeZ : FMath::Max(SizeZ >> MipRange.location, 1u);
+	
+	FMemory::Memzero(LockedMemory, sizeof(LockedMemory));
+}
+
+FMetalSurface::FMetalSurface(FMetalSurface const& Source, NSRange const MipRange, EPixelFormat Format)
+: Type(Source.Type)
+, PixelFormat(Format)
+, Texture(nil)
+, MSAATexture(nil)
+, StencilTexture(nil)
+, SizeX(Source.SizeX)
+, SizeY(Source.SizeY)
+, SizeZ(Source.SizeZ)
+, bIsCubemap(Source.bIsCubemap)
+, Flags(Source.Flags)
+, WriteLock(0)
+, TotalTextureSize(0)
+, IB(nullptr)
+{
+	check(!Source.MSAATexture);
+	
+	MTLPixelFormat MetalFormat = (MTLPixelFormat)GPixelFormats[PixelFormat].PlatformFormat;
+	
+	NSRange Slices = NSMakeRange(0, FMath::Min(SizeZ, 1u));
+    // @todo Zebra Temporary workaround for absence of X24_G8 or equivalent to GL_STENCIL_INDEX so that the stencil part of a texture may be sampled
+	// For now, if we find ourselves *requiring* this we'd need to lazily blit the stencil data out to a separate texture. radr://21813831
+#if METAL_API_1_1
+    if(Source.PixelFormat != PF_DepthStencil)
+    {
+        Texture = [Source.Texture newTextureViewWithPixelFormat:MetalFormat textureType:Source.Texture.textureType levels:MipRange slices:Slices];
+    }
+    else
+#endif
+    {
+        Texture = [Source.Texture retain];
+    }
+	if(Source.MSAATexture)
+	{
+		MSAATexture = Texture;
+	}
+	if(Source.StencilTexture)
+	{
+#if PLATFORM_IOS
+		check(false); // @todo Zebra Must handle separate stencil texture SRV!
+#endif
+		StencilTexture = Texture;
+	}
+	
+	const uint32 BlockSizeX = GPixelFormats[PixelFormat].BlockSizeX;
+	const uint32 BlockSizeY = GPixelFormats[PixelFormat].BlockSizeY;
+	const uint32 BlockBytes = GPixelFormats[PixelFormat].BlockBytes;
+	SizeX = FMath::Max(SizeX >> MipRange.location, BlockSizeX);
+	SizeY = FMath::Max(SizeY >> MipRange.location, BlockSizeY);
+	SizeZ = (Type != RRT_Texture3D) ? SizeZ : FMath::Max(SizeZ >> MipRange.location, 1u);
+	
+	FMemory::Memzero(LockedMemory, sizeof(LockedMemory));
+}
+
+
+FMetalSurface::FMetalSurface(ERHIResourceType ResourceType, EPixelFormat Format, uint32 InSizeX, uint32 InSizeY, uint32 InSizeZ, uint32 NumSamples, bool bArray, uint32 ArraySize, uint32 NumMips, uint32 InFlags, FResourceBulkDataInterface* BulkData)
+	: Type(ResourceType)
+	, PixelFormat(Format)
+	, Texture(nil)
+    , MSAATexture(nil)
+	, StencilTexture(nil)
+	, SizeX(InSizeX)
+	, SizeY(InSizeY)
+	, SizeZ(InSizeZ)
+	, bIsCubemap(false)
+	, Flags(InFlags)
+    , WriteLock(0)
+	, TotalTextureSize(0)
+	, IB(nullptr)
+{
+	// get a unique key for this surface's format
+	static TMap<uint32, uint8> PixelFormatKeyMap;
+	if (PixelFormatKeyMap.Num() == 0)
+	{
+		// Add depth stencil formats fist, so we don't have to use 5 bits for them in the pipeline hash
+		PixelFormatKeyMap.Add(MTLPixelFormatDepth32Float, NextKey++);
+		PixelFormatKeyMap.Add(MTLPixelFormatStencil8, NextKey++);
+#if METAL_API_1_1
+		PixelFormatKeyMap.Add(MTLPixelFormatDepth32Float_Stencil8, NextKey++);
+#if PLATFORM_MAC
+		PixelFormatKeyMap.Add(MTLPixelFormatDepth24Unorm_Stencil8, NextKey++);
+#endif
+#endif
+	}
+	
+	// Dirty, filty hacks to pass through via an IOSurfaceRef which is required if in-game movie playback is to render at more than 5fps
+#if PLATFORM_MAC
+	bool const bIOSurfaceData = (BulkData && BulkData->GetResourceBulkDataSize() == ~0u);
+	if(bIOSurfaceData)
+	{
+		// This won't ever be SRGB as IOSurface doesn't support it...
+		Flags &= ~TexCreate_SRGB;
+	}
+#endif
+
+	MTLPixelFormat MTLFormat = (MTLPixelFormat)GPixelFormats[Format].PlatformFormat;
+#if PLATFORM_MAC
+	if (Flags & TexCreate_SRGB)
+	{
+		switch (MTLFormat)
+		{
+			case MTLPixelFormatBC1_RGBA:
+				MTLFormat = MTLPixelFormatBC1_RGBA_sRGB;
+				break;
+			case MTLPixelFormatBC2_RGBA:
+				MTLFormat = MTLPixelFormatBC2_RGBA_sRGB;
+				break;
+			case MTLPixelFormatBC3_RGBA:
+				MTLFormat = MTLPixelFormatBC3_RGBA_sRGB;
+				break;
+			case MTLPixelFormatRGBA8Unorm:
+				MTLFormat = MTLPixelFormatRGBA8Unorm_sRGB;
+				break;
+			case MTLPixelFormatBGRA8Unorm:
+				MTLFormat = MTLPixelFormatBGRA8Unorm_sRGB;
+				break;
+			default:
+				break;
+		}
+	}
+#endif
+	uint8* Key = PixelFormatKeyMap.Find(MTLFormat);
+	if (Key == NULL)
+	{
+		Key = &PixelFormatKeyMap.Add(MTLFormat, NextKey++);
+
+#if PLATFORM_MAC
+		// only giving 5 bits to the key
+		checkf(NextKey < 32, TEXT("Too many unique pixel formats to fit into the PipelineStateHash"));
+#else
+		// only giving 4 bits to the key
+		checkf(NextKey < 16, TEXT("Too many unique pixel formats to fit into the PipelineStateHash"));
+#endif
+	}
+	// set the key
+	FormatKey = *Key;
+
+	FMemory::Memzero(LockedMemory, sizeof(LockedMemory));
+
+
+	// the special back buffer surface will be updated in GetMetalDeviceContext().BeginDrawingViewport - no need to set the texture here
+	if (Flags & TexCreate_Presentable)
+	{
+		return;
+	}
+
+	bool bIsRenderTarget = IsRenderTarget(Flags);
+	MTLTextureDescriptor* Desc;
+	
+	if (ResourceType == RRT_TextureCube)
+	{
+		Desc = [MTLTextureDescriptor textureCubeDescriptorWithPixelFormat:MTLFormat
+																	size:SizeX
+															   mipmapped:(NumMips > 1)];
+		bIsCubemap = true;
+	}
+	else if (ResourceType == RRT_Texture3D)
+	{
+		Desc = [MTLTextureDescriptor new];
+		Desc.textureType = MTLTextureType3D;
+		Desc.width = SizeX;
+		Desc.height = SizeY;
+		Desc.depth = SizeZ;
+		Desc.pixelFormat = MTLFormat;
+		Desc.arrayLength = 1;
+		Desc.mipmapLevelCount = 1;
+		Desc.sampleCount = 1;
+	}
+	else
+	{
+		Desc = [MTLTextureDescriptor texture2DDescriptorWithPixelFormat:MTLFormat
+																 width:SizeX
+																height:SizeY
+															 mipmapped:(NumMips > 1)];
+		Desc.depth = SizeZ;
+	}
+
+	// flesh out the descriptor
+	if (bArray)
+	{
+		Desc.arrayLength = ArraySize;
+		if (bIsCubemap)
+		{
+#if PLATFORM_MAC
+			Desc.textureType = MTLTextureTypeCubeArray;
+#else
+			check(0);
+#endif
+		}
+	}
+	Desc.mipmapLevelCount = NumMips;
+	
+#if METAL_API_1_1 // @todo
+	if(GetMetalDeviceContext().SupportsFeature(EMetalFeaturesResourceOptions))
+	{
+		Desc.usage = ConvertFlagsToUsage(Flags);
+		
+		if((Flags & TexCreate_CPUReadback) && !(Flags & (TexCreate_RenderTargetable|TexCreate_DepthStencilTargetable|TexCreate_FastVRAM)))
+		{
+			Desc.cpuCacheMode = MTLCPUCacheModeDefaultCache;
+	#if PLATFORM_MAC
+			Desc.storageMode = MTLStorageModeManaged;
+			Desc.resourceOptions = MTLResourceCPUCacheModeDefaultCache|MTLResourceStorageModeManaged;
+	#else
+			Desc.storageMode = MTLStorageModeShared;
+			Desc.resourceOptions = MTLResourceCPUCacheModeDefaultCache|MTLResourceStorageModeShared;
+	#endif
+		}
+		else if(((Flags & (TexCreate_NoTiling)) && !(Flags & (TexCreate_FastVRAM|TexCreate_DepthStencilTargetable|TexCreate_RenderTargetable))))
+		{
+	#if PLATFORM_MAC
+			Desc.cpuCacheMode = MTLCPUCacheModeWriteCombined;
+			Desc.storageMode = MTLStorageModeManaged;
+			Desc.resourceOptions = MTLResourceCPUCacheModeWriteCombined|MTLResourceStorageModeManaged;
+	#else
+			Desc.cpuCacheMode = MTLCPUCacheModeDefaultCache;
+			Desc.storageMode = MTLStorageModeShared;
+			Desc.resourceOptions = MTLResourceCPUCacheModeDefaultCache|MTLResourceStorageModeShared;
+	#endif
+		}
+		else
+		{
+			check(!(Flags & TexCreate_CPUReadback));
+	#if PLATFORM_MAC
+			Desc.cpuCacheMode = MTLCPUCacheModeWriteCombined;
+			Desc.storageMode = MTLStorageModePrivate;
+			Desc.resourceOptions = MTLResourceCPUCacheModeWriteCombined|MTLResourceStorageModePrivate;
+	#else
+			Desc.cpuCacheMode = MTLCPUCacheModeDefaultCache;
+			Desc.storageMode = MTLStorageModeShared;
+			Desc.resourceOptions = MTLResourceCPUCacheModeDefaultCache|MTLResourceStorageModeShared;
+	#endif
+		}
+	}
+#endif
+
+#if PLATFORM_MAC
+	// Dirty, filty hacks to pass through via an IOSurfaceRef
+	if(bIOSurfaceData)
+	{
+		checkf(NumMips == 1&& ArraySize == 1, TEXT("Only handling bulk data with 1 mip and 1 array length"));
+		IB = (CVImageBufferRef)BulkData->GetResourceBulkData();
+		CVPixelBufferRetain(IB);
+		{
+			Texture = [GetMetalDeviceContext().GetDevice() newTextureWithDescriptor:Desc iosurface:CVPixelBufferGetIOSurface(IB) plane:0];
+			
+			if (Texture == nil)
+			{
+				NSLog(@"Failed to create texture, desc  %@", Desc);
+			}
+		}
+		TRACK_OBJECT(Texture);
+		
+		BulkData->Discard();
+	}
+	else
+#endif
+	{
+		Texture = [GetMetalDeviceContext().GetDevice() newTextureWithDescriptor:Desc];
+		if (Texture == nil)
+		{
+			NSLog(@"Failed to create texture, desc  %@", Desc);
+		}
+		TRACK_OBJECT(Texture);
+		
+		// upload existing bulkdata
+		if (BulkData)
+		{
+			UE_LOG(LogIOS, Display, TEXT("Got a bulk data texture, with %d mips"), NumMips);
+			checkf(NumMips == 1&& ArraySize == 1, TEXT("Only handling bulk data with 1 mip and 1 array length"));
+			uint32 Stride;
+
+			// lock, copy, unlock
+			void* LockedData = Lock(0, 0, RLM_WriteOnly, Stride);
+			FMemory::Memcpy(LockedData, BulkData->GetResourceBulkData(), BulkData->GetResourceBulkDataSize());
+			Unlock(0, 0);
+
+			// bulk data can be unloaded now
+			BulkData->Discard();
+		}
+	}
+	
+	// calculate size of the texture
+	TotalTextureSize = GetMemorySize();
+
+	if (!FParse::Param(FCommandLine::Get(), TEXT("nomsaa")))
+	{
+		if (NumSamples > 1)
+		{
+			check(bIsRenderTarget);
+			Desc.textureType = MTLTextureType2DMultisample;
+	
+			// allow commandline to override
+			FParse::Value(FCommandLine::Get(), TEXT("msaa="), NumSamples);
+			Desc.sampleCount = NumSamples;
+
+			MSAATexture = [GetMetalDeviceContext().GetDevice() newTextureWithDescriptor:Desc];
+			if (Format == PF_DepthStencil)
+			{
+				[Texture release];
+				Texture = MSAATexture;
+
+				// we don't have the resolve texture, so we just update the memory size with the MSAA size
+				TotalTextureSize = TotalTextureSize * NumSamples;
+			}
+			else
+			{
+				// an MSAA render target takes NumSamples more space, in addition to the resolve texture
+				TotalTextureSize += TotalTextureSize * NumSamples;
+			}
+        
+			NSLog(@"Creating %dx MSAA %d x %d %s surface", (int32)Desc.sampleCount, SizeX, SizeY, (Flags & TexCreate_RenderTargetable) ? "Color" : "Depth");
+			if (MSAATexture == nil)
+			{
+				NSLog(@"Failed to create texture, desc  %@", Desc);
+			}
+			TRACK_OBJECT(MSAATexture);
+		}
+	}
+
+	// create a stencil buffer if needed
+	if (Format == PF_DepthStencil)
+	{
+#if PLATFORM_IOS
+		Desc.textureType = MTLTextureType2D;
+		Desc.pixelFormat = MTLPixelFormatStencil8;
+		StencilTexture = [GetMetalDeviceContext().GetDevice() newTextureWithDescriptor:Desc];
+
+		// 1 byte per texel
+		TotalTextureSize += SizeX * SizeY;
+#else
+        StencilTexture = Texture;
+        
+        // 1 byte per texel
+        TotalTextureSize += SizeX * SizeY;
+#endif
+	}
+
+	// track memory usage
+	if (bIsRenderTarget)
+	{
+		GCurrentRendertargetMemorySize += Align(TotalTextureSize, 1024) / 1024;
+	}
+	else
+	{
+		GCurrentTextureMemorySize += Align(TotalTextureSize, 1024) / 1024;
+	}
+
+#if STATS
+	if (ResourceType == RRT_TextureCube)
+	{
+		if (bIsRenderTarget)
+		{
+			INC_MEMORY_STAT_BY(STAT_RenderTargetMemoryCube, TotalTextureSize);
+		}
+		else
+		{
+			INC_MEMORY_STAT_BY(STAT_TextureMemoryCube, TotalTextureSize);
+		}
+	}
+	else if (ResourceType == RRT_Texture3D)
+	{
+		if (bIsRenderTarget)
+		{
+			INC_MEMORY_STAT_BY(STAT_RenderTargetMemory3D, TotalTextureSize);
+		}
+		else
+		{
+			INC_MEMORY_STAT_BY(STAT_TextureMemory3D, TotalTextureSize);
+		}
+	}
+	else
+	{
+		if (bIsRenderTarget)
+		{
+			INC_MEMORY_STAT_BY(STAT_RenderTargetMemory2D, TotalTextureSize);
+		}
+		else
+		{
+			INC_MEMORY_STAT_BY(STAT_TextureMemory2D, TotalTextureSize);
+		}
+	}
+#endif
+}
+
+
+FMetalSurface::~FMetalSurface()
+{
+	// track memory usage
+	if (IsRenderTarget(Flags))
+	{
+		GCurrentRendertargetMemorySize -= Align(TotalTextureSize, 1024) / 1024;
+	}
+	else
+	{
+		GCurrentTextureMemorySize -= Align(TotalTextureSize, 1024) / 1024;
+	}
+
+	if (MSAATexture != nil)
+    {
+		SafeReleaseMetalResource(MSAATexture);
+    }
+    
+	if (!(Flags & TexCreate_Presentable) && MSAATexture != Texture)
+	{
+		SafeReleaseMetalResource(Texture);
+	}
+	
+	if(IB)
+	{
+		CVPixelBufferRelease(IB);
+	}
+
+	IB = nullptr;
+	MSAATexture = nil;
+	Texture = nil;
+	StencilTexture = nil;
+	for(uint32 i = 0; i < 16; i++)
+	{
+		if(LockedMemory[i])
+		{
+			SafeReleaseMetalResource(LockedMemory[i]);
+			LockedMemory[i] = nullptr;
+		}
+	}
+}
+
+
+void* FMetalSurface::Lock(uint32 MipIndex, uint32 ArrayIndex, EResourceLockMode LockMode, uint32& DestStride)
+{
+	// get size and stride
+	const uint32 MipBytes = GetMipSize(MipIndex, &DestStride, false);
+	
+	// allocate some temporary memory
+	if(!LockedMemory[MipIndex])
+	{
+#if METAL_API_1_1
+		NSUInteger ResMode = MTLResourceStorageModeShared | (GetMetalDeviceContext().SupportsFeature(EMetalFeaturesResourceOptions) ? MTLResourceCPUCacheModeWriteCombined : 0);
+#else
+		NSUInteger ResMode = MTLStorageModeShared;
+#endif
+		LockedMemory[MipIndex] = [GetMetalDeviceContext().GetDevice() newBufferWithLength:MipBytes options:ResMode];
+	}
+	
+    switch(LockMode)
+    {
+        case RLM_ReadOnly:
+        {
+			SCOPE_CYCLE_COUNTER(STAT_MetalTexturePageOffTime);
+			const uint32 MipSizeX = SizeX >> MipIndex;
+			const uint32 MipSizeY = SizeY >> MipIndex;
+			const uint32 MipSizeZ = FMath::Max(SizeZ >> MipIndex, 1u);
+
+			MTLRegion Region;
+			if (SizeZ <= 1 || bIsCubemap)
+			{
+				// upload the texture to the texture slice
+				Region = MTLRegionMake2D(0, 0, FMath::Max<uint32>(SizeX >> MipIndex, 1), FMath::Max<uint32>(SizeY >> MipIndex, 1));
+			}
+			else
+			{
+				// upload the texture to the texture slice
+				Region = MTLRegionMake3D(0, 0, 0, FMath::Max<uint32>(SizeX >> MipIndex, 1), FMath::Max<uint32>(SizeY >> MipIndex, 1), FMath::Max<uint32>(SizeZ >> MipIndex, 1));
+			}
+
+			id<MTLBlitCommandEncoder> Blitter = GetMetalDeviceContext().GetBlitContext();
+			[Blitter copyFromTexture:Texture sourceSlice:ArrayIndex sourceLevel:MipIndex sourceOrigin:Region.origin sourceSize:Region.size toBuffer:LockedMemory[MipIndex] destinationOffset:0 destinationBytesPerRow:DestStride destinationBytesPerImage:MipBytes];
+			
+			//kick the current command buffer.
+			GetMetalDeviceContext().SubmitCommandBufferAndWait();
+			
+            break;
+        }
+        case RLM_WriteOnly:
+        {
+            WriteLock |= 1 << MipIndex;
+            break;
+        }
+        default:
+            check(false);
+            break;
+    }
+
+	return [LockedMemory[MipIndex] contents];
+}
+
+void FMetalSurface::Unlock(uint32 MipIndex, uint32 ArrayIndex)
+{
+    if(WriteLock & (1 << MipIndex))
+    {
+        WriteLock &= ~(1 << MipIndex);
+        uint32 Stride;
+        uint32 BytesPerImage = GetMipSize(MipIndex, &Stride, true);
+
+		MTLRegion Region;
+        if (SizeZ <= 1 || bIsCubemap)
+		{
+            // upload the texture to the texture slice
+            Region = MTLRegionMake2D(0, 0, FMath::Max<uint32>(SizeX >> MipIndex, 1), FMath::Max<uint32>(SizeY >> MipIndex, 1));
+        }
+        else
+        {
+            // upload the texture to the texture slice
+            Region = MTLRegionMake3D(0, 0, 0, FMath::Max<uint32>(SizeX >> MipIndex, 1), FMath::Max<uint32>(SizeY >> MipIndex, 1), FMath::Max<uint32>(SizeZ >> MipIndex, 1));
+        }
+#if PLATFORM_MAC
+		if(Texture.storageMode == MTLStorageModePrivate)
+		{
+			id<MTLBlitCommandEncoder> Blitter = GetMetalDeviceContext().GetBlitContext();
+			
+			[Blitter copyFromBuffer:LockedMemory[MipIndex] sourceOffset:0 sourceBytesPerRow:Stride sourceBytesPerImage:BytesPerImage sourceSize:Region.size toTexture:Texture destinationSlice:ArrayIndex destinationLevel:MipIndex destinationOrigin:Region.origin];
+		}
+		else
+#else
+		if (Texture.pixelFormat >= MTLPixelFormatPVRTC_RGB_2BPP && Texture.pixelFormat <= MTLPixelFormatETC2_RGB8A1_sRGB) // @todo zebra
+		{
+			Stride = 0;
+			BytesPerImage = 0;
+		}
+#endif
+		{
+			[Texture replaceRegion:Region mipmapLevel:MipIndex slice:ArrayIndex withBytes:[LockedMemory[MipIndex] contents] bytesPerRow:Stride bytesPerImage:BytesPerImage];
+		}
+		
+		GetMetalDeviceContext().ReleaseObject(LockedMemory[MipIndex]);
+		LockedMemory[MipIndex] = nullptr;
+    }
+}
+
+uint32 FMetalSurface::GetMipSize(uint32 MipIndex, uint32* Stride, bool bSingleLayer)
+{
+	// Calculate the dimensions of the mip-map.
+	const uint32 BlockSizeX = GPixelFormats[PixelFormat].BlockSizeX;
+	const uint32 BlockSizeY = GPixelFormats[PixelFormat].BlockSizeY;
+	const uint32 BlockBytes = GPixelFormats[PixelFormat].BlockBytes;
+	const uint32 MipSizeX = FMath::Max(SizeX >> MipIndex, BlockSizeX);
+	const uint32 MipSizeY = FMath::Max(SizeY >> MipIndex, BlockSizeY);
+	const uint32 MipSizeZ = bSingleLayer ? 1 : FMath::Max(SizeZ >> MipIndex, 1u);
+	uint32 NumBlocksX = (MipSizeX + BlockSizeX - 1) / BlockSizeX;
+	uint32 NumBlocksY = (MipSizeY + BlockSizeY - 1) / BlockSizeY;
+	if (PixelFormat == PF_PVRTC2 || PixelFormat == PF_PVRTC4)
+	{
+		// PVRTC has minimum 2 blocks width and height
+		NumBlocksX = FMath::Max<uint32>(NumBlocksX, 2);
+		NumBlocksY = FMath::Max<uint32>(NumBlocksY, 2);
+	}
+
+	const uint32 MipBytes = NumBlocksX * NumBlocksY * BlockBytes * MipSizeZ;
+
+	if (Stride)
+	{
+		*Stride = NumBlocksX * BlockBytes;
+	}
+
+	return MipBytes;
+}
+
+uint32 FMetalSurface::GetMemorySize()
+{
+	// if already calculated, no need to do it again
+	if (TotalTextureSize != 0)
+	{
+		return TotalTextureSize;
+	}
+
+	if (Texture == nil)
+	{
+		return 0;
+	}
+
+	uint32 TotalSize = 0;
+	for (uint32 MipIndex = 0; MipIndex < [Texture mipmapLevelCount]; MipIndex++)
+	{
+		TotalSize += GetMipSize(MipIndex, NULL, false);
+	}
+
+	return TotalSize;
+}
+
+uint32 FMetalSurface::GetNumFaces()
+{
+	switch (Type)
+	{
+		case RRT_Texture2DArray:
+		case RRT_Texture3D:
+		case RRT_TextureCube:
+			return SizeZ * Texture.arrayLength;
+			
+		case RRT_Texture:
+		case RRT_Texture2D:
+		default:
+			return 1;
+	}
+}
+
+/*-----------------------------------------------------------------------------
+	Texture allocator support.
+-----------------------------------------------------------------------------*/
+
+void FMetalDynamicRHI::RHIGetTextureMemoryStats(FTextureMemoryStats& OutStats)
+{
+	OutStats.DedicatedVideoMemory = 0;
+	OutStats.DedicatedSystemMemory = 0;
+	OutStats.SharedSystemMemory = 0;
+	OutStats.TotalGraphicsMemory = 0;
+
+	OutStats.AllocatedMemorySize = int64(GCurrentTextureMemorySize) * 1024;
+	OutStats.TexturePoolSize = GTexturePoolSize;
+	OutStats.PendingMemoryAdjustment = 0;
+}
+
+bool FMetalDynamicRHI::RHIGetTextureMemoryVisualizeData( FColor* /*TextureData*/, int32 /*SizeX*/, int32 /*SizeY*/, int32 /*Pitch*/, int32 /*PixelSize*/ )
+{
+	NOT_SUPPORTED("RHIGetTextureMemoryVisualizeData");
+	return false;
+}
+
+uint32 FMetalDynamicRHI::RHIComputeMemorySize(FTextureRHIParamRef TextureRHI)
+{
+	if(!TextureRHI)
+	{
+		return 0;
+	}
+
+	return GetMetalSurfaceFromRHITexture(TextureRHI)->GetMemorySize();
+}
+
+/*-----------------------------------------------------------------------------
+	2D texture support.
+-----------------------------------------------------------------------------*/
+
+FTexture2DRHIRef FMetalDynamicRHI::RHICreateTexture2D(uint32 SizeX, uint32 SizeY, uint8 Format, uint32 NumMips, uint32 NumSamples, uint32 Flags, FRHIResourceCreateInfo& CreateInfo)
+{
+	return new FMetalTexture2D((EPixelFormat)Format, SizeX, SizeY, NumMips, NumSamples, Flags, CreateInfo.BulkData, CreateInfo.ClearValueBinding);
+}
+
+FTexture2DRHIRef FMetalDynamicRHI::RHIAsyncCreateTexture2D(uint32 SizeX,uint32 SizeY,uint8 Format,uint32 NumMips,uint32 Flags,void** InitialMipData,uint32 NumInitialMips)
+{
+	UE_LOG(LogMetal, Fatal, TEXT("RHIAsyncCreateTexture2D is not supported"));
+	return FTexture2DRHIRef();
+}
+
+void FMetalDynamicRHI::RHICopySharedMips(FTexture2DRHIParamRef DestTexture2D,FTexture2DRHIParamRef SrcTexture2D)
+{
+	NOT_SUPPORTED("RHICopySharedMips");
+}
+
+FTexture2DArrayRHIRef FMetalDynamicRHI::RHICreateTexture2DArray(uint32 SizeX, uint32 SizeY, uint32 SizeZ, uint8 Format, uint32 NumMips, uint32 Flags, FRHIResourceCreateInfo& CreateInfo)
+{
+	return new FMetalTexture2DArray((EPixelFormat)Format, SizeX, SizeY, SizeZ, NumMips, Flags, CreateInfo.BulkData, CreateInfo.ClearValueBinding);
+}
+
+FTexture3DRHIRef FMetalDynamicRHI::RHICreateTexture3D(uint32 SizeX, uint32 SizeY, uint32 SizeZ, uint8 Format, uint32 NumMips, uint32 Flags, FRHIResourceCreateInfo& CreateInfo)
+{
+	return new FMetalTexture3D((EPixelFormat)Format, SizeX, SizeY, SizeZ, NumMips, Flags, CreateInfo.BulkData, CreateInfo.ClearValueBinding);
+}
+
+void FMetalDynamicRHI::RHIGetResourceInfo(FTextureRHIParamRef Ref, FRHIResourceInfo& OutInfo)
+{
+    // @todo Needed for visualisation!!
+	// NOT_SUPPORTED("RHIGetResourceInfo");
+}
+
+void FMetalDynamicRHI::RHIGenerateMips(FTextureRHIParamRef SourceSurfaceRHI)
+{
+	
+	NOT_SUPPORTED("RHIGenerateMips");
+}
+
+FTexture2DRHIRef FMetalDynamicRHI::RHIAsyncReallocateTexture2D(FTexture2DRHIParamRef OldTextureRHI, int32 NewMipCount, int32 NewSizeX, int32 NewSizeY, FThreadSafeCounter* RequestStatus)
+{
+	FMetalTexture2D* OldTexture = ResourceCast(OldTextureRHI);
+
+	FMetalTexture2D* NewTexture = new FMetalTexture2D(OldTexture->GetFormat(), NewSizeX, NewSizeY, NewMipCount, OldTexture->GetNumSamples(), OldTexture->GetFlags(), NULL, OldTextureRHI->GetClearBinding());
+
+	// @todo: gather these all up over a frame
+	id<MTLCommandBuffer> CommandBuffer = GetMetalDeviceContext().CreateCommandBuffer(false/*bRetainReferences*/);
+
+	// create a blitter object
+	id<MTLBlitCommandEncoder> Blitter = [CommandBuffer blitCommandEncoder];
+
+
+	// figure out what mips to schedule
+	const uint32 NumSharedMips = FMath::Min(OldTexture->GetNumMips(), NewTexture->GetNumMips());
+	const uint32 SourceMipOffset = OldTexture->GetNumMips() - NumSharedMips;
+	const uint32 DestMipOffset = NewTexture->GetNumMips() - NumSharedMips;
+
+	// only handling straight 2D textures here
+	uint32 SliceIndex = 0;
+	MTLOrigin Origin = MTLOriginMake(0,0,0);
+	
+	id<MTLTexture> Tex = OldTexture->Surface.Texture;
+	[Tex retain];
+
+	for (uint32 MipIndex = 0; MipIndex < NumSharedMips; ++MipIndex)
+	{
+		const uint32 MipSizeX = FMath::Max<uint32>(1, NewSizeX >> (MipIndex + DestMipOffset));
+		const uint32 MipSizeY = FMath::Max<uint32>(1, NewSizeY >> (MipIndex + DestMipOffset));
+
+		// set up the copy
+		[Blitter copyFromTexture:OldTexture->Surface.Texture 
+					 sourceSlice:SliceIndex
+					 sourceLevel:MipIndex + SourceMipOffset
+					sourceOrigin:Origin
+					  sourceSize:MTLSizeMake(MipSizeX, MipSizeY, 1)
+					   toTexture:NewTexture->Surface.Texture
+			    destinationSlice:SliceIndex
+			    destinationLevel:MipIndex + DestMipOffset
+			   destinationOrigin:Origin];
+	}
+
+	// when done, decrement the counter to indicate it's safe
+	[CommandBuffer addCompletedHandler:^(id <MTLCommandBuffer> Buffer)
+	{
+		[Tex release];
+		RequestStatus->Decrement();
+	}];
+
+	// kick it off!
+	[Blitter endEncoding];
+	[CommandBuffer commit];
+
+	return NewTexture;
+}
+
+ETextureReallocationStatus FMetalDynamicRHI::RHIFinalizeAsyncReallocateTexture2D( FTexture2DRHIParamRef Texture2D, bool bBlockUntilCompleted )
+{
+	return TexRealloc_Succeeded;
+}
+
+ETextureReallocationStatus FMetalDynamicRHI::RHICancelAsyncReallocateTexture2D( FTexture2DRHIParamRef Texture2D, bool bBlockUntilCompleted )
+{
+	return TexRealloc_Failed;
+}
+
+
+void* FMetalDynamicRHI::RHILockTexture2D(FTexture2DRHIParamRef TextureRHI,uint32 MipIndex,EResourceLockMode LockMode,uint32& DestStride,bool bLockWithinMiptail)
+{
+	FMetalTexture2D* Texture = ResourceCast(TextureRHI);
+	return Texture->Surface.Lock(MipIndex, 0, LockMode, DestStride);
+}
+
+void FMetalDynamicRHI::RHIUnlockTexture2D(FTexture2DRHIParamRef TextureRHI,uint32 MipIndex,bool bLockWithinMiptail)
+{
+	FMetalTexture2D* Texture = ResourceCast(TextureRHI);
+	Texture->Surface.Unlock(MipIndex, 0);
+}
+
+void* FMetalDynamicRHI::RHILockTexture2DArray(FTexture2DArrayRHIParamRef TextureRHI,uint32 TextureIndex,uint32 MipIndex,EResourceLockMode LockMode,uint32& DestStride,bool bLockWithinMiptail)
+{
+	FMetalTexture2DArray* Texture = ResourceCast(TextureRHI);
+	return Texture->Surface.Lock(MipIndex, TextureIndex, LockMode, DestStride);
+}
+
+void FMetalDynamicRHI::RHIUnlockTexture2DArray(FTexture2DArrayRHIParamRef TextureRHI,uint32 TextureIndex,uint32 MipIndex,bool bLockWithinMiptail)
+{
+	FMetalTexture2DArray* Texture = ResourceCast(TextureRHI);
+	Texture->Surface.Unlock(MipIndex, TextureIndex);
+}
+
+void FMetalDynamicRHI::RHIUpdateTexture2D(FTexture2DRHIParamRef TextureRHI, uint32 MipIndex, const struct FUpdateTextureRegion2D& UpdateRegion, uint32 SourcePitch, const uint8* SourceData)
+{	
+	FMetalTexture2D* Texture = ResourceCast(TextureRHI);
+	
+	id<MTLTexture> Tex = Texture->Surface.Texture;
+	
+	MTLRegion Region = MTLRegionMake2D(UpdateRegion.DestX, UpdateRegion.DestY, UpdateRegion.Width, UpdateRegion.Height);
+	
+	[Tex replaceRegion:Region mipmapLevel:MipIndex slice:0 withBytes:SourceData bytesPerRow:SourcePitch bytesPerImage:0];
+}
+
+void FMetalDynamicRHI::RHIUpdateTexture3D(FTexture3DRHIParamRef TextureRHI,uint32 MipIndex,const FUpdateTextureRegion3D& UpdateRegion,uint32 SourceRowPitch,uint32 SourceDepthPitch,const uint8* SourceData)
+{	
+	FMetalTexture3D* Texture = ResourceCast(TextureRHI);
+
+	id<MTLTexture> Tex = Texture->Surface.Texture;
+
+	MTLRegion Region = MTLRegionMake3D(UpdateRegion.DestX, UpdateRegion.DestY, UpdateRegion.DestZ, UpdateRegion.Width, UpdateRegion.Height, UpdateRegion.Depth);
+	
+	[Tex replaceRegion:Region mipmapLevel:MipIndex slice:0 withBytes:SourceData bytesPerRow:SourceRowPitch bytesPerImage:SourceDepthPitch];
+}
+
+
+/*-----------------------------------------------------------------------------
+	Cubemap texture support.
+-----------------------------------------------------------------------------*/
+FTextureCubeRHIRef FMetalDynamicRHI::RHICreateTextureCube(uint32 Size, uint8 Format, uint32 NumMips, uint32 Flags, FRHIResourceCreateInfo& CreateInfo)
+{
+	return new FMetalTextureCube((EPixelFormat)Format, Size, false, 1, NumMips, Flags, CreateInfo.BulkData, CreateInfo.ClearValueBinding);
+}
+
+FTextureCubeRHIRef FMetalDynamicRHI::RHICreateTextureCubeArray(uint32 Size, uint32 ArraySize, uint8 Format, uint32 NumMips, uint32 Flags, FRHIResourceCreateInfo& CreateInfo)
+{
+	return new FMetalTextureCube((EPixelFormat)Format, Size, true, ArraySize, NumMips, Flags, CreateInfo.BulkData, CreateInfo.ClearValueBinding);
+}
+
+void* FMetalDynamicRHI::RHILockTextureCubeFace(FTextureCubeRHIParamRef TextureCubeRHI,uint32 FaceIndex,uint32 ArrayIndex,uint32 MipIndex,EResourceLockMode LockMode,uint32& DestStride,bool bLockWithinMiptail)
+{
+	FMetalTextureCube* TextureCube = ResourceCast(TextureCubeRHI);
+	uint32 MetalFace = GetMetalCubeFace((ECubeFace)FaceIndex);
+	return TextureCube->Surface.Lock(MipIndex, MetalFace + (6 * ArrayIndex), LockMode, DestStride);
+}
+
+void FMetalDynamicRHI::RHIUnlockTextureCubeFace(FTextureCubeRHIParamRef TextureCubeRHI,uint32 FaceIndex,uint32 ArrayIndex,uint32 MipIndex,bool bLockWithinMiptail)
+{
+	FMetalTextureCube* TextureCube = ResourceCast(TextureCubeRHI);
+	uint32 MetalFace = GetMetalCubeFace((ECubeFace)FaceIndex);
+	TextureCube->Surface.Unlock(MipIndex, MetalFace + (ArrayIndex * 6));
+}
+
+
+
+
+FTextureReferenceRHIRef FMetalDynamicRHI::RHICreateTextureReference(FLastRenderTimeContainer* InLastRenderTime)
+{
+	return new FMetalTextureReference(InLastRenderTime);
+}
+
+void FMetalRHICommandContext::RHIUpdateTextureReference(FTextureReferenceRHIParamRef TextureRefRHI, FTextureRHIParamRef NewTextureRHI)
+{
+	FMetalTextureReference* TextureRef = (FMetalTextureReference*)TextureRefRHI;
+	if (TextureRef)
+	{
+		TextureRef->SetReferencedTexture(NewTextureRHI);
+	}
+}
+
+
+void FMetalDynamicRHI::RHIBindDebugLabelName(FTextureRHIParamRef TextureRHI, const TCHAR* Name)
+{
+	FMetalSurface* Surf = GetMetalSurfaceFromRHITexture(TextureRHI);
+	if(Surf->Texture)
+	{
+		Surf->Texture.label = FString(Name).GetNSString();
+	}
+	if(Surf->MSAATexture)
+	{
+		Surf->MSAATexture.label = FString(Name).GetNSString();
+	}
+	if(Surf->StencilTexture)
+	{
+		Surf->StencilTexture.label = FString(Name).GetNSString();
+	}
+}
+
+void FMetalDynamicRHI::RHIVirtualTextureSetFirstMipInMemory(FTexture2DRHIParamRef TextureRHI, uint32 FirstMip)
+{
+	NOT_SUPPORTED("RHIVirtualTextureSetFirstMipInMemory");
+}
+
+void FMetalDynamicRHI::RHIVirtualTextureSetFirstMipVisible(FTexture2DRHIParamRef TextureRHI, uint32 FirstMip)
+{
+	NOT_SUPPORTED("RHIVirtualTextureSetFirstMipVisible");
+}
diff --git a/Engine/Source/Runtime/Apple/MetalRHI/Private/MetalUAV.cpp b/Engine/Source/Runtime/Apple/MetalRHI/Private/MetalUAV.cpp
new file mode 100644
--- /dev/null
+++ b/Engine/Source/Runtime/Apple/MetalRHI/Private/MetalUAV.cpp
@@ -0,0 +1,204 @@
+// Copyright 1998-2015 Epic Games, Inc. All Rights Reserved.
+
+
+#include "MetalRHIPrivate.h"
+#include "ShaderCache.h"
+
+FMetalShaderResourceView::~FMetalShaderResourceView()
+{
+	if(TextureView)
+	{
+		if(TextureView->Texture)
+		{
+			[TextureView->Texture release];
+			TextureView->Texture = nil;
+		}
+		delete TextureView;
+		TextureView = nullptr;
+	}
+	FShaderCache::RemoveSRV(this);
+	SourceVertexBuffer = NULL;
+	SourceTexture = NULL;
+}
+
+void FMetalUnorderedAccessView::Set(FMetalContext* Context, uint32 ResourceIndex)
+{
+	// figure out which one of the resources we need to set
+	FMetalStructuredBuffer* StructuredBuffer = SourceStructuredBuffer.GetReference();
+	FMetalVertexBuffer* VertexBuffer = SourceVertexBuffer.GetReference();
+	FRHITexture* Texture = SourceTexture.GetReference();
+	if (StructuredBuffer)
+	{
+		Context->GetCommandEncoder().SetShaderBuffer(SF_Compute, StructuredBuffer->Buffer, 0, ResourceIndex);
+	}
+	else if (VertexBuffer)
+	{
+		Context->GetCommandEncoder().SetShaderBuffer(SF_Compute, VertexBuffer->Buffer, 0, ResourceIndex);
+	}
+	else if (Texture)
+	{
+		FMetalSurface* Surface = GetMetalSurfaceFromRHITexture(Texture);
+		if (Surface != nullptr)
+		{
+			Context->GetCommandEncoder().SetShaderTexture(SF_Compute, Surface->Texture, ResourceIndex);
+		}
+		else
+		{
+			Context->GetCommandEncoder().SetShaderTexture(SF_Compute, nil, ResourceIndex);
+		}
+	}
+
+}
+
+
+FUnorderedAccessViewRHIRef FMetalDynamicRHI::RHICreateUnorderedAccessView(FStructuredBufferRHIParamRef StructuredBufferRHI, bool bUseUAVCounter, bool bAppendBuffer)
+{
+	FMetalStructuredBuffer* StructuredBuffer = ResourceCast(StructuredBufferRHI);
+
+	// create the UAV buffer to point to the structured buffer's memory
+	FMetalUnorderedAccessView* UAV = new FMetalUnorderedAccessView;
+	UAV->SourceStructuredBuffer = StructuredBuffer;
+
+	return UAV;
+}
+
+FUnorderedAccessViewRHIRef FMetalDynamicRHI::RHICreateUnorderedAccessView(FTextureRHIParamRef TextureRHI, uint32 MipLevel)
+{
+	FMetalSurface* Surface = GetMetalSurfaceFromRHITexture(TextureRHI);
+
+	// create the UAV buffer to point to the structured buffer's memory
+	FMetalUnorderedAccessView* UAV = new FMetalUnorderedAccessView;
+	UAV->SourceTexture = (FRHITexture*)TextureRHI;
+
+	return UAV;
+}
+
+FUnorderedAccessViewRHIRef FMetalDynamicRHI::RHICreateUnorderedAccessView(FVertexBufferRHIParamRef VertexBufferRHI, uint8 Format)
+{
+	FMetalVertexBuffer* VertexBuffer = ResourceCast(VertexBufferRHI);
+
+	// create the UAV buffer to point to the structured buffer's memory
+	FMetalUnorderedAccessView* UAV = new FMetalUnorderedAccessView;
+	UAV->SourceVertexBuffer = VertexBuffer;
+
+	return UAV;
+}
+
+FShaderResourceViewRHIRef FMetalDynamicRHI::RHICreateShaderResourceView(FStructuredBufferRHIParamRef StructuredBufferRHI)
+{
+	FMetalStructuredBuffer* StructuredBuffer = ResourceCast(StructuredBufferRHI);
+
+	FMetalShaderResourceView* SRV = new FMetalShaderResourceView;
+	SRV->TextureView = nullptr;
+	UE_LOG(LogRHI, Fatal,TEXT("Metal RHI doesn't support RHICreateShaderResourceView with FStructuredBufferRHIParamRef yet!"));
+	return SRV;
+}
+
+FShaderResourceViewRHIRef FMetalDynamicRHI::RHICreateShaderResourceView(FVertexBufferRHIParamRef VertexBufferRHI, uint32 Stride, uint8 Format)
+{
+	FMetalVertexBuffer* VertexBuffer = ResourceCast(VertexBufferRHI);
+
+	FMetalShaderResourceView* SRV = new FMetalShaderResourceView;
+	SRV->SourceVertexBuffer = VertexBuffer;
+	SRV->TextureView = nullptr;
+	
+	FShaderCache::LogSRV(SRV, VertexBufferRHI, Stride, Format);
+	
+	return SRV;
+}
+
+FShaderResourceViewRHIRef FMetalDynamicRHI::RHICreateShaderResourceView(FTexture2DRHIParamRef Texture2DRHI, uint8 MipLevel)
+{
+    // @todo Zebra Temporary workaround for absence of X24_G8 or equivalent to GL_STENCIL_INDEX so that the stencil part of a texture may be sampled
+    // For now, if we find ourselves *requiring* this we'd need to lazily blit the stencil data out to a separate texture. radr://21813831
+    if(Texture2DRHI->GetFormat() != PF_DepthStencil)
+    {
+        FMetalShaderResourceView* SRV = new FMetalShaderResourceView;
+        SRV->SourceTexture = (FRHITexture*)Texture2DRHI;
+        
+        FMetalSurface* Surface = GetMetalSurfaceFromRHITexture(Texture2DRHI);
+        SRV->TextureView = Surface ? new FMetalSurface(*Surface, NSMakeRange(MipLevel, 1)) : nullptr;
+        
+        SRV->MipLevel = MipLevel;
+        SRV->NumMips = 1;
+        SRV->Format = PF_Unknown;
+        
+        FShaderCache::LogSRV(SRV, Texture2DRHI, MipLevel, SRV->NumMips, SRV->Format);
+        
+        return SRV;
+    }
+    else
+    {
+        UE_LOG(LogMetal, Warning, TEXT("FMetalDynamicRHI::RHICreateShaderResourceView: Cannot create a shader resource view for textures with pixel format: %d"), (uint32)Texture2DRHI->GetFormat());
+        return nullptr;
+    }
+}
+
+FShaderResourceViewRHIRef FMetalDynamicRHI::RHICreateShaderResourceView(FTexture2DRHIParamRef Texture2DRHI, uint8 MipLevel, uint8 NumMipLevels, uint8 Format)
+{
+	 // @todo Zebra Temporary workaround for absence of X24_G8 or equivalent to GL_STENCIL_INDEX so that the stencil part of a texture may be sampled
+    // For now, if we find ourselves *requiring* this we'd need to lazily blit the stencil data out to a separate texture. radr://21813831
+    if(Texture2DRHI->GetFormat() != PF_DepthStencil)
+    {
+        FMetalShaderResourceView* SRV = new FMetalShaderResourceView;
+        SRV->SourceTexture = (FRHITexture*)Texture2DRHI;
+        
+        FMetalSurface* Surface = GetMetalSurfaceFromRHITexture(Texture2DRHI);
+        SRV->TextureView = Surface ? new FMetalSurface(*Surface, NSMakeRange(MipLevel, NumMipLevels), (EPixelFormat)Format) : nullptr;
+        
+        SRV->MipLevel = MipLevel;
+        SRV->NumMips = NumMipLevels;
+        SRV->Format = Format;
+        
+        FShaderCache::LogSRV(SRV, Texture2DRHI, MipLevel, NumMipLevels, Format);
+        return SRV;
+    }
+    else
+    {
+        UE_LOG(LogMetal, Warning, TEXT("FMetalDynamicRHI::RHICreateShaderResourceView: Cannot create a shader resource view for textures with pixel format: %d"), (uint32)Texture2DRHI->GetFormat());
+        return nullptr;
+    }
+}
+
+FShaderResourceViewRHIRef FMetalDynamicRHI::RHICreateShaderResourceView(FTexture3DRHIParamRef Texture3DRHI, uint8 MipLevel)
+{
+	FMetalShaderResourceView* SRV = new FMetalShaderResourceView;
+	SRV->SourceTexture = (FRHITexture*)Texture3DRHI;
+	return SRV;
+}
+
+FShaderResourceViewRHIRef FMetalDynamicRHI::RHICreateShaderResourceView(FTexture2DArrayRHIParamRef Texture2DArrayRHI, uint8 MipLevel)
+{
+	FMetalShaderResourceView* SRV = new FMetalShaderResourceView;
+	SRV->SourceTexture = (FRHITexture*)Texture2DArrayRHI;
+	return SRV;
+}
+
+FShaderResourceViewRHIRef FMetalDynamicRHI::RHICreateShaderResourceView(FTextureCubeRHIParamRef TextureCubeRHI, uint8 MipLevel)
+{
+	FMetalShaderResourceView* SRV = new FMetalShaderResourceView;
+	SRV->SourceTexture = (FRHITexture*)TextureCubeRHI;
+	return SRV;
+}
+
+void FMetalRHICommandContext::RHIClearUAV(FUnorderedAccessViewRHIParamRef UnorderedAccessViewRHI, const uint32* Values)
+{
+	FMetalUnorderedAccessView* UnorderedAccessView = ResourceCast(UnorderedAccessViewRHI);
+	if (UnorderedAccessView->SourceStructuredBuffer)
+	{
+	}
+	else if (UnorderedAccessView->SourceTexture)
+	{
+	}
+	else
+	{
+		check(UnorderedAccessView->SourceVertexBuffer);
+		
+		// Fill the buffer via a blit encoder - I hope that is sufficient.
+		id<MTLBlitCommandEncoder> Blitter = Context->GetBlitContext();
+		[Blitter fillBuffer:UnorderedAccessView->SourceVertexBuffer->Buffer range:NSMakeRange(0, UnorderedAccessView->SourceVertexBuffer->GetSize()) value:Values[0]];
+		
+		// If there are problems you may need to add calls to restore the render command encoder at this point
+		// but we don't generally want to do that.
+	}
+}
diff --git a/Engine/Source/Runtime/Apple/MetalRHI/Private/MetalVertexBuffer.cpp b/Engine/Source/Runtime/Apple/MetalRHI/Private/MetalVertexBuffer.cpp
new file mode 100644
--- /dev/null
+++ b/Engine/Source/Runtime/Apple/MetalRHI/Private/MetalVertexBuffer.cpp
@@ -0,0 +1,151 @@
+// Copyright 1998-2015 Epic Games, Inc. All Rights Reserved.
+
+/*=============================================================================
+	MetalVertexBuffer.cpp: Metal vertex buffer RHI implementation.
+=============================================================================*/
+
+#include "MetalRHIPrivate.h"
+
+
+FMetalVertexBuffer::FMetalVertexBuffer(uint32 InSize, uint32 InUsage)
+	: FRHIVertexBuffer(InSize, InUsage)
+	, LockOffset(0)
+	, LockSize(0)
+	, ZeroStrideElementSize((InUsage & BUF_ZeroStride) ? InSize : 0)
+{
+	checkf(InSize <= 256 * 1024 * 1024, TEXT("Metal doesn't support buffers > 256 MB"));
+	// Zero-stride buffers must be separate in order to wrap appropriately
+	if (InUsage & BUF_Volatile && !(InUsage & BUF_ZeroStride))
+	{
+		FMetalPooledBuffer Buf = GetMetalDeviceContext().CreatePooledBuffer(FMetalPooledBufferArgs(GetMetalDeviceContext().GetDevice(), InSize, MTLStorageModeShared));
+		Buffer = [Buf.Buffer retain];
+	}
+	else
+	{
+		if(!(InUsage & BUF_ZeroStride))
+		{
+			FMetalPooledBuffer Buf = GetMetalDeviceContext().CreatePooledBuffer(FMetalPooledBufferArgs(GetMetalDeviceContext().GetDevice(), InSize, BUFFER_STORAGE_MODE));
+			Buffer = [Buf.Buffer retain];
+		}
+		else
+		{
+			Buffer = [GetMetalDeviceContext().GetDevice() newBufferWithLength:InSize options:BUFFER_CACHE_MODE|BUFFER_MANAGED_MEM];
+		}
+	}
+	TRACK_OBJECT(Buffer);
+}
+
+FMetalVertexBuffer::~FMetalVertexBuffer()
+{
+	if(!(GetUsage() & BUF_ZeroStride))
+	{
+		SafeReleasePooledBuffer(Buffer);
+		[Buffer release];
+	}
+	else
+	{
+		SafeReleaseMetalResource(Buffer);
+	}
+}
+
+
+void* FMetalVertexBuffer::Lock(EResourceLockMode LockMode, uint32 Offset, uint32 Size)
+{
+	check(LockSize == 0 && LockOffset == 0);
+	
+	// In order to properly synchronise the buffer access, when a dynamic buffer is locked for writing, discard the old buffer & create a new one. This prevents writing to a buffer while it is being read by the GPU & thus causing corruption. This matches the logic of other RHIs.
+	if ((GetUsage() & BUFFER_DYNAMIC_REALLOC) && LockMode == RLM_WriteOnly)
+	{
+		id<MTLBuffer> OldBuffer = Buffer;
+		if(!(GetUsage() & BUF_ZeroStride))
+		{
+			MTLStorageMode Mode = (GetUsage() & BUF_Volatile) ? MTLStorageModeShared : BUFFER_STORAGE_MODE;
+			FMetalPooledBuffer Buf = GetMetalDeviceContext().CreatePooledBuffer(FMetalPooledBufferArgs(GetMetalDeviceContext().GetDevice(), GetSize(), Mode));
+			Buffer = [Buf.Buffer retain];
+			GetMetalDeviceContext().ReleasePooledBuffer(OldBuffer);
+			[OldBuffer release];
+		}
+		else
+		{
+			Buffer = [GetMetalDeviceContext().GetDevice() newBufferWithLength:Buffer.length options:BUFFER_CACHE_MODE|BUFFER_MANAGED_MEM];
+			GetMetalDeviceContext().ReleaseObject(OldBuffer);
+		}
+	}
+	
+	if(LockMode != RLM_ReadOnly)
+	{
+		LockSize = Size;
+		LockOffset = Offset;
+	}
+#if PLATFORM_MAC
+	else if(Buffer.storageMode == MTLStorageModeManaged)
+	{
+		SCOPE_CYCLE_COUNTER(STAT_MetalBufferPageOffTime);
+		
+		// Synchronise the buffer with the CPU
+		id<MTLBlitCommandEncoder> Blitter = GetMetalDeviceContext().GetBlitContext();
+		[Blitter synchronizeResource:Buffer];
+		
+		//kick the current command buffer.
+		GetMetalDeviceContext().SubmitCommandBufferAndWait();
+	}
+#endif
+	
+	return ((uint8*)[Buffer contents]) + Offset;
+}
+
+void FMetalVertexBuffer::Unlock()
+{
+#if PLATFORM_MAC
+	if(LockSize && Buffer.storageMode == MTLStorageModeManaged)
+	{
+		[Buffer didModifyRange:NSMakeRange(LockOffset, LockSize)];
+	}
+#endif
+	LockSize = 0;
+	LockOffset = 0;
+}
+
+FVertexBufferRHIRef FMetalDynamicRHI::RHICreateVertexBuffer(uint32 Size, uint32 InUsage, FRHIResourceCreateInfo& CreateInfo)
+{
+	// make the RHI object, which will allocate memory
+	FMetalVertexBuffer* VertexBuffer = new FMetalVertexBuffer(Size, InUsage);
+
+	if (CreateInfo.ResourceArray)
+	{
+		check(Size == CreateInfo.ResourceArray->GetResourceDataSize());
+
+		// make a buffer usable by CPU
+		void* Buffer = RHILockVertexBuffer(VertexBuffer, 0, Size, RLM_WriteOnly);
+
+		// copy the contents of the given data into the buffer
+		FMemory::Memcpy(Buffer, CreateInfo.ResourceArray->GetResourceData(), Size);
+
+		RHIUnlockVertexBuffer(VertexBuffer);
+
+		// Discard the resource array's contents.
+		CreateInfo.ResourceArray->Discard();
+	}
+
+	return VertexBuffer;
+}
+
+void* FMetalDynamicRHI::RHILockVertexBuffer(FVertexBufferRHIParamRef VertexBufferRHI, uint32 Offset, uint32 Size, EResourceLockMode LockMode)
+{
+	FMetalVertexBuffer* VertexBuffer = ResourceCast(VertexBufferRHI);
+
+	// default to vertex buffer memory
+	return (uint8*)VertexBuffer->Lock(LockMode, Offset, Size);
+}
+
+void FMetalDynamicRHI::RHIUnlockVertexBuffer(FVertexBufferRHIParamRef VertexBufferRHI)
+{
+	FMetalVertexBuffer* VertexBuffer = ResourceCast(VertexBufferRHI);
+
+	VertexBuffer->Unlock();
+}
+
+void FMetalDynamicRHI::RHICopyVertexBuffer(FVertexBufferRHIParamRef SourceBufferRHI,FVertexBufferRHIParamRef DestBufferRHI)
+{
+	NOT_SUPPORTED("RHICopyVertexBuffer");
+}
diff --git a/Engine/Source/Runtime/Apple/MetalRHI/Private/MetalViewport.cpp b/Engine/Source/Runtime/Apple/MetalRHI/Private/MetalViewport.cpp
new file mode 100644
--- /dev/null
+++ b/Engine/Source/Runtime/Apple/MetalRHI/Private/MetalViewport.cpp
@@ -0,0 +1,181 @@
+// Copyright 1998-2015 Epic Games, Inc. All Rights Reserved.
+
+/*=============================================================================
+	MetalViewport.cpp: Metal viewport RHI implementation.
+=============================================================================*/
+
+#include "MetalRHIPrivate.h"
+#if PLATFORM_MAC
+#include "CocoaWindow.h"
+#include "CocoaThread.h"
+#endif
+
+#if PLATFORM_MAC
+
+@implementation FMetalView
+
+- (id)initWithFrame:(NSRect)frameRect
+{
+	self = [super initWithFrame:frameRect];
+	if (self)
+	{
+	}
+	return self;
+}
+
+- (BOOL)isOpaque
+{
+	return YES;
+}
+
+- (BOOL)mouseDownCanMoveWindow
+{
+	return YES;
+}
+
+@end
+
+#endif
+
+FMetalViewport::FMetalViewport(void* WindowHandle, uint32 InSizeX,uint32 InSizeY,bool bInIsFullscreen)
+{
+#if PLATFORM_MAC
+	FCocoaWindow* Window = (FCocoaWindow*)WindowHandle;
+	const NSRect ContentRect = NSMakeRect(0, 0, InSizeX, InSizeY);
+	View = [[FMetalView alloc] initWithFrame:ContentRect];
+	[View setAutoresizingMask:NSViewWidthSizable | NSViewHeightSizable];
+	[View setWantsLayer:YES];
+
+	CAMetalLayer* Layer = [CAMetalLayer new];
+
+	CGFloat bgColor[] = { 0.0, 0.0, 0.0, 0.0 };
+	Layer.edgeAntialiasingMask = 0;
+	Layer.masksToBounds = YES;
+	Layer.backgroundColor = CGColorCreate(CGColorSpaceCreateDeviceRGB(), bgColor);
+	Layer.presentsWithTransaction = NO;
+	Layer.anchorPoint = CGPointMake(0.5, 0.5);
+	Layer.frame = ContentRect;
+	Layer.magnificationFilter = kCAFilterNearest;
+	Layer.minificationFilter = kCAFilterNearest;
+//	Layer.drawsAsynchronously = YES; // @todo zebra: do we need and/or want this on Mac?
+
+	[Layer setDevice:GetMetalDeviceContext().GetDevice()];
+	[Layer setPixelFormat:MTLPixelFormatBGRA8Unorm];
+	[Layer setFramebufferOnly:NO];
+	[Layer removeAllAnimations];
+
+	[View setLayer:Layer];
+
+	[Window setContentView:View];
+	[[Window standardWindowButton:NSWindowCloseButton] setAction:@selector(performClose:)];
+
+	Resize(InSizeX, InSizeY, bInIsFullscreen);
+
+	// @todo zebra: Sometimes ConditionalUpdateBackBuffer() is called before BeginDrawingViewport. It doesn't seem right, we have to investigate.
+	if (!GetMetalDeviceContext().GetCurrentState().GetCurrentLayer())
+	{
+		GetMetalDeviceContext().GetCurrentState().SetCurrentLayer(Layer);
+	}
+#else
+	Resize(InSizeX, InSizeY, bInIsFullscreen);
+#endif
+}
+
+FMetalViewport::~FMetalViewport()
+{
+	BackBuffer.SafeRelease();
+	check(!IsValidRef(BackBuffer));
+}
+
+void FMetalViewport::Resize(uint32 InSizeX, uint32 InSizeY, bool bInIsFullscreen)
+{
+	BackBuffer.SafeRelease();	// when the rest of the engine releases it, its framebuffers will be released too (those the engine knows about)
+
+	FRHIResourceCreateInfo CreateInfo;
+	BackBuffer = (FMetalTexture2D*)(FTexture2DRHIParamRef)RHICreateTexture2D(InSizeX, InSizeY, PF_B8G8R8A8, 1, 1, TexCreate_RenderTargetable | TexCreate_Presentable, CreateInfo);
+#if PLATFORM_MAC // @todo zebra: ios
+	((CAMetalLayer*)View.layer).drawableSize = CGSizeMake(InSizeX, InSizeY);
+#else
+	IOSAppDelegate* AppDelegate = [IOSAppDelegate GetDelegate];
+	FIOSView* GLView = AppDelegate.IOSView;
+	[GLView UpdateRenderWidth:InSizeX andHeight:InSizeY];
+	
+	// check the size of the window
+	float ScalingFactor = [[IOSAppDelegate GetDelegate].IOSView contentScaleFactor];
+	CGRect ViewFrame = [[IOSAppDelegate GetDelegate].IOSView frame];
+	check(ScalingFactor * ViewFrame.size.width == InSizeX && ScalingFactor * ViewFrame.size.height == InSizeY);
+#endif
+	GetMetalDeviceContext().GetCurrentState().ResetCurrentDrawable();
+}
+
+/*=============================================================================
+ *	The following RHI functions must be called from the main thread.
+ *=============================================================================*/
+FViewportRHIRef FMetalDynamicRHI::RHICreateViewport(void* WindowHandle,uint32 SizeX,uint32 SizeY,bool bIsFullscreen,EPixelFormat PreferredPixelFormat /* ignored */)
+{
+	check( IsInGameThread() );
+	return new FMetalViewport(WindowHandle, SizeX, SizeY, bIsFullscreen);
+}
+
+void FMetalDynamicRHI::RHIResizeViewport(FViewportRHIParamRef ViewportRHI,uint32 SizeX,uint32 SizeY,bool bIsFullscreen)
+{
+	check( IsInGameThread() );
+
+	FMetalViewport* Viewport = ResourceCast(ViewportRHI);
+	Viewport->Resize(SizeX, SizeY, bIsFullscreen);
+}
+
+void FMetalDynamicRHI::RHITick( float DeltaTime )
+{
+	check( IsInGameThread() );
+}
+
+/*=============================================================================
+ *	Viewport functions.
+ *=============================================================================*/
+
+void FMetalRHICommandContext::RHIBeginDrawingViewport(FViewportRHIParamRef ViewportRHI, FTextureRHIParamRef RenderTargetRHI)
+{
+	FMetalViewport* Viewport = ResourceCast(ViewportRHI);
+
+	Context->BeginDrawingViewport(Viewport);
+
+	// Set the render target and viewport.
+	if (RenderTargetRHI)
+	{
+		FRHIRenderTargetView RTV(RenderTargetRHI);
+		RHISetRenderTargets(1, &RTV, nullptr, 0, NULL);
+	}
+	else
+	{
+		FRHIRenderTargetView RTV(Viewport->GetBackBuffer());
+		RHISetRenderTargets(1, &RTV, nullptr, 0, NULL);
+	}
+}
+
+void FMetalRHICommandContext::RHIEndDrawingViewport(FViewportRHIParamRef ViewportRHI,bool bPresent,bool bLockToVsync)
+{
+	FMetalViewport* Viewport = ResourceCast(ViewportRHI);
+	Context->EndDrawingViewport(Viewport, bPresent);
+#if PLATFORM_MAC
+	FCocoaWindow* Window = (FCocoaWindow*)[Viewport->View window];
+	MainThreadCall(^{
+		[Window startRendering];
+	}, NSDefaultRunLoopMode, false);
+#endif
+}
+
+bool FMetalDynamicRHI::RHIIsDrawingViewport()
+{
+	return true;
+}
+
+FTexture2DRHIRef FMetalDynamicRHI::RHIGetViewportBackBuffer(FViewportRHIParamRef ViewportRHI)
+{
+	FMetalViewport* Viewport = ResourceCast(ViewportRHI);
+	return Viewport->GetBackBuffer();
+}
+
+void FMetalDynamicRHI::RHIAdvanceFrameForGetViewportBackBuffer()
+{
+}
diff --git a/Engine/Source/Runtime/Apple/MetalRHI/Public/MetalRHIContext.h b/Engine/Source/Runtime/Apple/MetalRHI/Public/MetalRHIContext.h
new file mode 100644
--- /dev/null
+++ b/Engine/Source/Runtime/Apple/MetalRHI/Public/MetalRHIContext.h
@@ -0,0 +1,337 @@
+// Copyright 1998-2015 Epic Games, Inc. All Rights Reserved.
+
+#pragma once
+
+// Metal RHI public headers.
+#include <Metal/Metal.h>
+#include "MetalState.h"
+#include "MetalResources.h"
+#include "MetalViewport.h"
+
+class FMetalContext;
+
+/** The interface RHI command context. Sometimes the RHI handles these. On platforms that can processes command lists in parallel, it is a separate object. */
+class FMetalRHICommandContext : public IRHICommandContext
+{
+public:
+	FMetalRHICommandContext(struct FMetalGPUProfiler* InProfiler);
+	virtual ~FMetalRHICommandContext();
+
+	/** Get the internal context */
+	FORCEINLINE FMetalContext& GetInternalContext() const { return *Context; }
+	
+	template<typename TRHIType>
+	static FORCEINLINE typename TMetalResourceTraits<TRHIType>::TConcreteType* ResourceCast(TRHIType* Resource)
+	{
+		return static_cast<typename TMetalResourceTraits<TRHIType>::TConcreteType*>(Resource);
+	}
+	
+	/////// RHI Context Methods
+	
+	
+	/**
+	 *Sets the current compute shader.  Mostly for compliance with platforms
+	 *that require shader setting before resource binding.
+	 */
+	virtual void RHISetComputeShader(FComputeShaderRHIParamRef ComputeShader) final override;
+	
+	virtual void RHIDispatchComputeShader(uint32 ThreadGroupCountX, uint32 ThreadGroupCountY, uint32 ThreadGroupCountZ) final override;
+	
+	virtual void RHIDispatchIndirectComputeShader(FVertexBufferRHIParamRef ArgumentBuffer, uint32 ArgumentOffset) final override;
+	
+	virtual void RHIAutomaticCacheFlushAfterComputeShader(bool bEnable) final override;
+	
+	virtual void RHIFlushComputeShaderCache() final override;
+	
+	// Useful when used with geometry shader (emit polygons to different viewports), otherwise SetViewPort() is simpler
+	// @param Count >0
+	// @param Data must not be 0
+	virtual void RHISetMultipleViewports(uint32 Count, const FViewportBounds* Data) final override;
+	
+	/** Clears a UAV to the multi-component value provided. */
+	virtual void RHIClearUAV(FUnorderedAccessViewRHIParamRef UnorderedAccessViewRHI, const uint32* Values) final override;
+	
+	/**
+	 * Resolves from one texture to another.
+	 * @param SourceTexture - texture to resolve from, 0 is silenty ignored
+	 * @param DestTexture - texture to resolve to, 0 is silenty ignored
+	 * @param bKeepOriginalSurface - true if the original surface will still be used after this function so must remain valid
+	 * @param ResolveParams - optional resolve params
+	 */
+	virtual void RHICopyToResolveTarget(FTextureRHIParamRef SourceTexture, FTextureRHIParamRef DestTexture, bool bKeepOriginalSurface, const FResolveParams& ResolveParams) final override;
+	
+	virtual void RHIBeginRenderQuery(FRenderQueryRHIParamRef RenderQuery) final override;
+	
+	virtual void RHIEndRenderQuery(FRenderQueryRHIParamRef RenderQuery) final override;
+	
+	virtual void RHIBeginOcclusionQueryBatch() final override;
+	
+	virtual void RHIEndOcclusionQueryBatch() final override;
+	
+	virtual void RHISubmitCommandsHint() final override;
+	
+	// This method is queued with an RHIThread, otherwise it will flush after it is queued; without an RHI thread there is no benefit to queuing this frame advance commands
+	virtual void RHIBeginDrawingViewport(FViewportRHIParamRef Viewport, FTextureRHIParamRef RenderTargetRHI) final override;
+	
+	// This method is queued with an RHIThread, otherwise it will flush after it is queued; without an RHI thread there is no benefit to queuing this frame advance commands
+	virtual void RHIEndDrawingViewport(FViewportRHIParamRef Viewport, bool bPresent, bool bLockToVsync) final override;
+	
+	// This method is queued with an RHIThread, otherwise it will flush after it is queued; without an RHI thread there is no benefit to queuing this frame advance commands
+	virtual void RHIBeginFrame() final override;
+	
+	// This method is queued with an RHIThread, otherwise it will flush after it is queued; without an RHI thread there is no benefit to queuing this frame advance commands
+	virtual void RHIEndFrame() final override;
+	
+	/**
+		* Signals the beginning of scene rendering. The RHI makes certain caching assumptions between
+		* calls to BeginScene/EndScene. Currently the only restriction is that you can't update texture
+		* references.
+		*/
+	// This method is queued with an RHIThread, otherwise it will flush after it is queued; without an RHI thread there is no benefit to queuing this frame advance commands
+	virtual void RHIBeginScene() final override;
+	
+	/**
+		* Signals the end of scene rendering. See RHIBeginScene.
+		*/
+	// This method is queued with an RHIThread, otherwise it will flush after it is queued; without an RHI thread there is no benefit to queuing this frame advance commands
+	virtual void RHIEndScene() final override;
+	
+	virtual void RHISetStreamSource(uint32 StreamIndex, FVertexBufferRHIParamRef VertexBuffer, uint32 Stride, uint32 Offset) final override;
+	
+	virtual void RHISetRasterizerState(FRasterizerStateRHIParamRef NewState) final override;
+	
+	// @param MinX including like Win32 RECT
+	// @param MinY including like Win32 RECT
+	// @param MaxX excluding like Win32 RECT
+	// @param MaxY excluding like Win32 RECT
+	virtual void RHISetViewport(uint32 MinX, uint32 MinY, float MinZ, uint32 MaxX, uint32 MaxY, float MaxZ) final override;
+	
+	// @param MinX including like Win32 RECT
+	// @param MinY including like Win32 RECT
+	// @param MaxX excluding like Win32 RECT
+	// @param MaxY excluding like Win32 RECT
+	virtual void RHISetScissorRect(bool bEnable, uint32 MinX, uint32 MinY, uint32 MaxX, uint32 MaxY) final override;
+	
+	/**
+	 * Set bound shader state. This will set the vertex decl/shader, and pixel shader
+	 * @param BoundShaderState - state resource
+	 */
+	virtual void RHISetBoundShaderState(FBoundShaderStateRHIParamRef BoundShaderState) final override;
+	
+	/** Set the shader resource view of a surface.  This is used for binding TextureMS parameter types that need a multi sampled view. */
+	virtual void RHISetShaderTexture(FVertexShaderRHIParamRef VertexShader, uint32 TextureIndex, FTextureRHIParamRef NewTexture) final override;
+	
+	/** Set the shader resource view of a surface.  This is used for binding TextureMS parameter types that need a multi sampled view. */
+	virtual void RHISetShaderTexture(FHullShaderRHIParamRef HullShader, uint32 TextureIndex, FTextureRHIParamRef NewTexture) final override;
+	
+	/** Set the shader resource view of a surface.  This is used for binding TextureMS parameter types that need a multi sampled view. */
+	virtual void RHISetShaderTexture(FDomainShaderRHIParamRef DomainShader, uint32 TextureIndex, FTextureRHIParamRef NewTexture) final override;
+	
+	/** Set the shader resource view of a surface.  This is used for binding TextureMS parameter types that need a multi sampled view. */
+	virtual void RHISetShaderTexture(FGeometryShaderRHIParamRef GeometryShader, uint32 TextureIndex, FTextureRHIParamRef NewTexture) final override;
+	
+	/** Set the shader resource view of a surface.  This is used for binding TextureMS parameter types that need a multi sampled view. */
+	virtual void RHISetShaderTexture(FPixelShaderRHIParamRef PixelShader, uint32 TextureIndex, FTextureRHIParamRef NewTexture) final override;
+	
+	/** Set the shader resource view of a surface.  This is used for binding TextureMS parameter types that need a multi sampled view. */
+	virtual void RHISetShaderTexture(FComputeShaderRHIParamRef PixelShader, uint32 TextureIndex, FTextureRHIParamRef NewTexture) final override;
+	
+	/**
+	 * Sets sampler state.
+	 * @param GeometryShader	The geometry shader to set the sampler for.
+	 * @param SamplerIndex		The index of the sampler.
+	 * @param NewState			The new sampler state.
+	 */
+	virtual void RHISetShaderSampler(FComputeShaderRHIParamRef ComputeShader, uint32 SamplerIndex, FSamplerStateRHIParamRef NewState) final override;
+	
+	/**
+	 * Sets sampler state.
+	 * @param GeometryShader	The geometry shader to set the sampler for.
+	 * @param SamplerIndex		The index of the sampler.
+	 * @param NewState			The new sampler state.
+	 */
+	virtual void RHISetShaderSampler(FVertexShaderRHIParamRef VertexShader, uint32 SamplerIndex, FSamplerStateRHIParamRef NewState) final override;
+	
+	/**
+	 * Sets sampler state.
+	 * @param GeometryShader	The geometry shader to set the sampler for.
+	 * @param SamplerIndex		The index of the sampler.
+	 * @param NewState			The new sampler state.
+	 */
+	virtual void RHISetShaderSampler(FGeometryShaderRHIParamRef GeometryShader, uint32 SamplerIndex, FSamplerStateRHIParamRef NewState) final override;
+	
+	/**
+	 * Sets sampler state.
+	 * @param GeometryShader	The geometry shader to set the sampler for.
+	 * @param SamplerIndex		The index of the sampler.
+	 * @param NewState			The new sampler state.
+	 */
+	virtual void RHISetShaderSampler(FDomainShaderRHIParamRef DomainShader, uint32 SamplerIndex, FSamplerStateRHIParamRef NewState) final override;
+	
+	/**
+	 * Sets sampler state.
+	 * @param GeometryShader	The geometry shader to set the sampler for.
+	 * @param SamplerIndex		The index of the sampler.
+	 * @param NewState			The new sampler state.
+	 */
+	virtual void RHISetShaderSampler(FHullShaderRHIParamRef HullShader, uint32 SamplerIndex, FSamplerStateRHIParamRef NewState) final override;
+	
+	/**
+	 * Sets sampler state.
+	 * @param GeometryShader	The geometry shader to set the sampler for.
+	 * @param SamplerIndex		The index of the sampler.
+	 * @param NewState			The new sampler state.
+	 */
+	virtual void RHISetShaderSampler(FPixelShaderRHIParamRef PixelShader, uint32 SamplerIndex, FSamplerStateRHIParamRef NewState) final override;
+	
+	/** Sets a compute shader UAV parameter. */
+	virtual void RHISetUAVParameter(FComputeShaderRHIParamRef ComputeShader, uint32 UAVIndex, FUnorderedAccessViewRHIParamRef UAV) final override;
+	
+	/** Sets a compute shader UAV parameter and initial count */
+	virtual void RHISetUAVParameter(FComputeShaderRHIParamRef ComputeShader, uint32 UAVIndex, FUnorderedAccessViewRHIParamRef UAV, uint32 InitialCount) final override;
+	
+	virtual void RHISetShaderResourceViewParameter(FPixelShaderRHIParamRef PixelShader, uint32 SamplerIndex, FShaderResourceViewRHIParamRef SRV) final override;
+	
+	virtual void RHISetShaderResourceViewParameter(FVertexShaderRHIParamRef VertexShader, uint32 SamplerIndex, FShaderResourceViewRHIParamRef SRV) final override;
+	
+	virtual void RHISetShaderResourceViewParameter(FComputeShaderRHIParamRef ComputeShader, uint32 SamplerIndex, FShaderResourceViewRHIParamRef SRV) final override;
+	
+	virtual void RHISetShaderResourceViewParameter(FHullShaderRHIParamRef HullShader, uint32 SamplerIndex, FShaderResourceViewRHIParamRef SRV) final override;
+	
+	virtual void RHISetShaderResourceViewParameter(FDomainShaderRHIParamRef DomainShader, uint32 SamplerIndex, FShaderResourceViewRHIParamRef SRV) final override;
+	
+	virtual void RHISetShaderResourceViewParameter(FGeometryShaderRHIParamRef GeometryShader, uint32 SamplerIndex, FShaderResourceViewRHIParamRef SRV) final override;
+	
+	virtual void RHISetShaderUniformBuffer(FVertexShaderRHIParamRef VertexShader, uint32 BufferIndex, FUniformBufferRHIParamRef Buffer) final override;
+	
+	virtual void RHISetShaderUniformBuffer(FHullShaderRHIParamRef HullShader, uint32 BufferIndex, FUniformBufferRHIParamRef Buffer) final override;
+	
+	virtual void RHISetShaderUniformBuffer(FDomainShaderRHIParamRef DomainShader, uint32 BufferIndex, FUniformBufferRHIParamRef Buffer) final override;
+	
+	virtual void RHISetShaderUniformBuffer(FGeometryShaderRHIParamRef GeometryShader, uint32 BufferIndex, FUniformBufferRHIParamRef Buffer) final override;
+	
+	virtual void RHISetShaderUniformBuffer(FPixelShaderRHIParamRef PixelShader, uint32 BufferIndex, FUniformBufferRHIParamRef Buffer) final override;
+	
+	virtual void RHISetShaderUniformBuffer(FComputeShaderRHIParamRef ComputeShader, uint32 BufferIndex, FUniformBufferRHIParamRef Buffer) final override;
+	
+	virtual void RHISetShaderParameter(FVertexShaderRHIParamRef VertexShader, uint32 BufferIndex, uint32 BaseIndex, uint32 NumBytes, const void* NewValue) final override;
+	
+	virtual void RHISetShaderParameter(FPixelShaderRHIParamRef PixelShader, uint32 BufferIndex, uint32 BaseIndex, uint32 NumBytes, const void* NewValue) final override;
+	
+	virtual void RHISetShaderParameter(FHullShaderRHIParamRef HullShader, uint32 BufferIndex, uint32 BaseIndex, uint32 NumBytes, const void* NewValue) final override;
+	
+	virtual void RHISetShaderParameter(FDomainShaderRHIParamRef DomainShader, uint32 BufferIndex, uint32 BaseIndex, uint32 NumBytes, const void* NewValue) final override;
+	
+	virtual void RHISetShaderParameter(FGeometryShaderRHIParamRef GeometryShader, uint32 BufferIndex, uint32 BaseIndex, uint32 NumBytes, const void* NewValue) final override;
+	
+	virtual void RHISetShaderParameter(FComputeShaderRHIParamRef ComputeShader, uint32 BufferIndex, uint32 BaseIndex, uint32 NumBytes, const void* NewValue) final override;
+	
+	virtual void RHISetDepthStencilState(FDepthStencilStateRHIParamRef NewState, uint32 StencilRef) final override;
+	
+	// Allows to set the blend state, parameter can be created with RHICreateBlendState()
+	virtual void RHISetBlendState(FBlendStateRHIParamRef NewState, const FLinearColor& BlendFactor) final override;
+	
+	virtual void RHISetRenderTargets(uint32 NumSimultaneousRenderTargets, const FRHIRenderTargetView* NewRenderTargets, const FRHIDepthRenderTargetView* NewDepthStencilTarget, uint32 NumUAVs, const FUnorderedAccessViewRHIParamRef* UAVs) final override;
+	
+	virtual void RHISetRenderTargetsAndClear(const FRHISetRenderTargetsInfo& RenderTargetsInfo) final override;
+	
+	virtual void RHIDrawPrimitive(uint32 PrimitiveType, uint32 BaseVertexIndex, uint32 NumPrimitives, uint32 NumInstances) final override;
+	
+	virtual void RHIDrawPrimitiveIndirect(uint32 PrimitiveType, FVertexBufferRHIParamRef ArgumentBuffer, uint32 ArgumentOffset) final override;
+	
+	virtual void RHIDrawIndexedIndirect(FIndexBufferRHIParamRef IndexBufferRHI, uint32 PrimitiveType, FStructuredBufferRHIParamRef ArgumentsBufferRHI, int32 DrawArgumentsIndex, uint32 NumInstances) final override;
+	
+	// @param NumPrimitives need to be >0
+	virtual void RHIDrawIndexedPrimitive(FIndexBufferRHIParamRef IndexBuffer, uint32 PrimitiveType, int32 BaseVertexIndex, uint32 FirstInstance, uint32 NumVertices, uint32 StartIndex, uint32 NumPrimitives, uint32 NumInstances) final override;
+	
+	virtual void RHIDrawIndexedPrimitiveIndirect(uint32 PrimitiveType, FIndexBufferRHIParamRef IndexBuffer, FVertexBufferRHIParamRef ArgumentBuffer, uint32 ArgumentOffset) final override;
+	
+	/**
+	 * Preallocate memory or get a direct command stream pointer to fill up for immediate rendering . This avoids memcpys below in DrawPrimitiveUP
+	 * @param PrimitiveType The type (triangles, lineloop, etc) of primitive to draw
+	 * @param NumPrimitives The number of primitives in the VertexData buffer
+	 * @param NumVertices The number of vertices to be written
+	 * @param VertexDataStride Size of each vertex
+	 * @param OutVertexData Reference to the allocated vertex memory
+	 */
+	virtual void RHIBeginDrawPrimitiveUP(uint32 PrimitiveType, uint32 NumPrimitives, uint32 NumVertices, uint32 VertexDataStride, void*& OutVertexData) final override;
+	
+	/**
+	 * Draw a primitive using the vertex data populated since RHIBeginDrawPrimitiveUP and clean up any memory as needed
+	 */
+	virtual void RHIEndDrawPrimitiveUP() final override;
+	
+	/**
+	 * Preallocate memory or get a direct command stream pointer to fill up for immediate rendering . This avoids memcpys below in DrawIndexedPrimitiveUP
+	 * @param PrimitiveType The type (triangles, lineloop, etc) of primitive to draw
+	 * @param NumPrimitives The number of primitives in the VertexData buffer
+	 * @param NumVertices The number of vertices to be written
+	 * @param VertexDataStride Size of each vertex
+	 * @param OutVertexData Reference to the allocated vertex memory
+	 * @param MinVertexIndex The lowest vertex index used by the index buffer
+	 * @param NumIndices Number of indices to be written
+	 * @param IndexDataStride Size of each index (either 2 or 4 bytes)
+	 * @param OutIndexData Reference to the allocated index memory
+	 */
+	virtual void RHIBeginDrawIndexedPrimitiveUP(uint32 PrimitiveType, uint32 NumPrimitives, uint32 NumVertices, uint32 VertexDataStride, void*& OutVertexData, uint32 MinVertexIndex, uint32 NumIndices, uint32 IndexDataStride, void*& OutIndexData) final override;
+	
+	/**
+	 * Draw a primitive using the vertex and index data populated since RHIBeginDrawIndexedPrimitiveUP and clean up any memory as needed
+	 */
+	virtual void RHIEndDrawIndexedPrimitiveUP() final override;
+	
+	/*
+	 * This method clears all MRT's, but to only one color value
+	 * @param ExcludeRect within the viewport in pixels, is only a hint to optimize - if a fast clear can be done this is preferred
+	 */
+	virtual void RHIClear(bool bClearColor, const FLinearColor& Color, bool bClearDepth, float Depth, bool bClearStencil, uint32 Stencil, FIntRect ExcludeRect) final override;
+	
+	/*
+	 * This method clears all MRT's to potentially different color values
+	 * @param ExcludeRect within the viewport in pixels, is only a hint to optimize - if a fast clear can be done this is preferred
+	 */
+	virtual void RHIClearMRT(bool bClearColor, int32 NumClearColors, const FLinearColor* ColorArray, bool bClearDepth, float Depth, bool bClearStencil, uint32 Stencil, FIntRect ExcludeRect) final override;
+	
+	/**
+	 * Enabled/Disables Depth Bounds Testing with the given min/max depth.
+	 * @param bEnable	Enable(non-zero)/disable(zero) the depth bounds test
+	 * @param MinDepth	The minimum depth for depth bounds test
+	 * @param MaxDepth	The maximum depth for depth bounds test.
+	 *					The valid values for fMinDepth and fMaxDepth are such that 0 <= fMinDepth <= fMaxDepth <= 1
+	 */
+	virtual void RHIEnableDepthBoundsTest(bool bEnable, float MinDepth, float MaxDepth) final override;
+	
+	virtual void RHIPushEvent(const TCHAR* Name) final override;
+	
+	virtual void RHIPopEvent() final override;
+	
+	virtual void RHIUpdateTextureReference(FTextureReferenceRHIParamRef TextureRef, FTextureRHIParamRef NewTexture) final override;
+	
+	/** Start AsyncCompute command stream recording (no effect if not supported) */
+	virtual void RHIBeginAsyncComputeJob_DrawThread(EAsyncComputePriority Priority) final override;
+	
+	/** End AsyncCompute command stream recording (no effect if not supported) */
+	virtual void RHIEndAsyncComputeJob_DrawThread(uint32 FenceIndex) final override;
+	
+	/** Wait for AsyncCompute command stream to finish (no effect if not supported) */
+	virtual void RHIGraphicsWaitOnAsyncComputeJob(uint32 FenceIndex) final override;
+	
+protected:
+	static TGlobalResource<TBoundShaderStateHistory<10000>> BoundShaderStateHistory;
+	
+	/** Context implementation details. */
+	FMetalContext* Context;
+	
+	/** Profiling implementation details. */
+	struct FMetalGPUProfiler* Profiler;
+	
+	/** Some local variables to track the pending primitive information uised in RHIEnd*UP functions */
+	uint32 PendingVertexBufferOffset;
+	uint32 PendingVertexDataStride;
+	
+	uint32 PendingIndexBufferOffset;
+	uint32 PendingIndexDataStride;
+	
+	uint32 PendingPrimitiveType;
+	uint32 PendingNumPrimitives;
+};
diff --git a/Engine/Source/Runtime/Apple/MetalRHI/Public/MetalResources.h b/Engine/Source/Runtime/Apple/MetalRHI/Public/MetalResources.h
new file mode 100644
--- /dev/null
+++ b/Engine/Source/Runtime/Apple/MetalRHI/Public/MetalResources.h
@@ -0,0 +1,762 @@
+// Copyright 1998-2015 Epic Games, Inc. All Rights Reserved.
+
+/*=============================================================================
+	MetalResources.h: Metal resource RHI definitions.
+=============================================================================*/
+
+#pragma once
+
+#include "BoundShaderStateCache.h"
+#include "MetalShaderResources.h"
+#include "ShaderCache.h"
+
+class FMetalContext;
+
+/** This represents a vertex declaration that hasn't been combined with a specific shader to create a bound shader. */
+class FMetalVertexDeclaration : public FRHIVertexDeclaration
+{
+public:
+
+	/** Initialization constructor. */
+	FMetalVertexDeclaration(const FVertexDeclarationElementList& InElements);
+	~FMetalVertexDeclaration();
+	
+	/** Cached element info array (offset, stream index, etc) */
+	FVertexDeclarationElementList Elements;
+
+	/** This is the layout for the vertex elements */
+	MTLVertexDescriptor* Layout;
+
+protected:
+	void GenerateLayout(const FVertexDeclarationElementList& Elements);
+
+};
+
+
+
+/** This represents a vertex shader that hasn't been combined with a specific declaration to create a bound shader. */
+template<typename BaseResourceType, int32 ShaderType>
+class TMetalBaseShader : public BaseResourceType, public IRefCountedObject
+{
+public:
+	enum { StaticFrequency = ShaderType };
+
+	/** Initialization constructor. */
+	TMetalBaseShader()
+	{
+	}
+	TMetalBaseShader(const TArray<uint8>& InCode);
+
+	/** Destructor */
+	virtual ~TMetalBaseShader();
+
+
+	// IRefCountedObject interface.
+	virtual uint32 AddRef() const
+	{
+		return FRHIResource::AddRef();
+	}
+	virtual uint32 Release() const
+	{
+		return FRHIResource::Release();
+	}
+	virtual uint32 GetRefCount() const
+	{
+		return FRHIResource::GetRefCount();
+	}
+	
+	// this is the compiler shader
+	id<MTLFunction> Function;
+
+	/** External bindings for this shader. */
+	FMetalShaderBindings Bindings;
+	TArray< TRefCountPtr<FRHIUniformBuffer> > BoundUniformBuffers;
+
+	/** Bitfield for which uniform buffers are dirty */
+	uint64 DirtyUniformBuffers;
+
+	// List of memory copies from RHIUniformBuffer to packed uniforms
+	TArray<CrossCompiler::FUniformBufferCopyInfo> UniformBuffersCopyInfo;
+	
+	TArray<ANSICHAR> GlslCode;
+	NSString* GlslCodeNSString;
+	const ANSICHAR*  GlslCodeString; // make it easier in VS to see shader code in debug mode; points to begin of GlslCode
+};
+
+typedef TMetalBaseShader<FRHIVertexShader, SF_Vertex> FMetalVertexShader;
+typedef TMetalBaseShader<FRHIPixelShader, SF_Pixel> FMetalPixelShader;
+typedef TMetalBaseShader<FRHIHullShader, SF_Hull> FMetalHullShader;
+typedef TMetalBaseShader<FRHIDomainShader, SF_Domain> FMetalDomainShader;
+typedef TMetalBaseShader<FRHIGeometryShader, SF_Geometry> FMetalGeometryShader;
+
+class FMetalComputeShader : public TMetalBaseShader<FRHIComputeShader, SF_Compute>
+{
+public:
+	FMetalComputeShader(const TArray<uint8>& InCode);
+	
+	// the state object for a compute shader
+	id <MTLComputePipelineState> Kernel;
+	
+	// thread group counts
+	int32 NumThreadsX;
+	int32 NumThreadsY;
+	int32 NumThreadsZ;
+};
+
+#if PLATFORM_MAC
+typedef __uint128_t FMetalRenderPipelineHash;
+#else
+typedef uint64 FMetalRenderPipelineHash;
+#endif
+
+/**
+ * Combined shader state and vertex definition for rendering geometry. 
+ * Each unique instance consists of a vertex decl, vertex shader, and pixel shader.
+ */
+class FMetalBoundShaderState : public FRHIBoundShaderState
+{
+public:
+
+	FCachedBoundShaderStateLink CacheLink;
+
+	/** Cached vertex structure */
+	TRefCountPtr<FMetalVertexDeclaration> VertexDeclaration;
+
+	/** Cached shaders */
+	TRefCountPtr<FMetalVertexShader> VertexShader;
+	TRefCountPtr<FMetalPixelShader> PixelShader;
+	TRefCountPtr<FMetalHullShader> HullShader;
+	TRefCountPtr<FMetalDomainShader> DomainShader;
+	TRefCountPtr<FMetalGeometryShader> GeometryShader;
+
+	/** Initialization constructor. */
+	FMetalBoundShaderState(
+		FVertexDeclarationRHIParamRef InVertexDeclarationRHI,
+		FVertexShaderRHIParamRef InVertexShaderRHI,
+		FPixelShaderRHIParamRef InPixelShaderRHI,
+		FHullShaderRHIParamRef InHullShaderRHI,
+		FDomainShaderRHIParamRef InDomainShaderRHI,
+		FGeometryShaderRHIParamRef InGeometryShaderRHI);
+
+	/**
+	 *Destructor
+	 */
+	~FMetalBoundShaderState();
+	
+	/**
+	 * Prepare a pipeline state object for the current state right before drawing
+	 */
+	void PrepareToDraw(FMetalContext* Context, const struct FMetalRenderPipelineDesc& RenderPipelineDesc);
+
+protected:
+	TMap<FMetalRenderPipelineHash, id<MTLRenderPipelineState> > PipelineStates;
+};
+
+
+/** Texture/RT wrapper. */
+class FMetalSurface
+{
+public:
+
+	/** 
+	 * Constructor that will create Texture and Color/DepthBuffers as needed
+	 */
+	FMetalSurface(ERHIResourceType ResourceType, EPixelFormat Format, uint32 SizeX, uint32 SizeY, uint32 SizeZ, uint32 NumSamples, bool bArray, uint32 ArraySize, uint32 NumMips, uint32 Flags, FResourceBulkDataInterface* BulkData);
+
+	FMetalSurface(FMetalSurface const& Source, NSRange MipRange);
+	
+	FMetalSurface(FMetalSurface const& Source, NSRange MipRange, EPixelFormat Format);
+	
+	/**
+	 * Destructor
+	 */
+	~FMetalSurface();
+
+	/**
+	 * Locks one of the texture's mip-maps.
+	 * @param ArrayIndex Index of the texture array/face in the form Index*6+Face
+	 * @return A pointer to the specified texture data.
+	 */
+	void* Lock(uint32 MipIndex, uint32 ArrayIndex, EResourceLockMode LockMode, uint32& DestStride);
+
+	/** Unlocks a previously locked mip-map.
+	 * @param ArrayIndex Index of the texture array/face in the form Index*6+Face
+	 */
+	void Unlock(uint32 MipIndex, uint32 ArrayIndex);
+
+	/**
+	 * Returns how much memory a single mip uses, and optionally returns the stride
+	 */
+	uint32 GetMipSize(uint32 MipIndex, uint32* Stride, bool bSingleLayer);
+
+	/**
+	 * Returns how much memory is used by the surface
+	 */
+	uint32 GetMemorySize();
+
+	/** Returns the number of faces for the texture */
+	uint32 GetNumFaces();
+
+	ERHIResourceType Type;
+	EPixelFormat PixelFormat;
+	uint8 FormatKey;
+	id<MTLTexture> Texture;
+	id<MTLTexture> MSAATexture;
+	id<MTLTexture> StencilTexture;
+	uint32 SizeX, SizeY, SizeZ;
+	bool bIsCubemap;
+	
+	uint32 Flags;
+	// one per mip
+	id<MTLBuffer> LockedMemory[16];
+    uint32 WriteLock;
+
+	// how much memory is allocated for this texture
+	uint64 TotalTextureSize;
+
+private:
+	// The movie playback IOSurface wrapper to avoid page-off
+	CVImageBufferRef IB;
+	
+	// next format for the pixel format mapping
+	static uint8 NextKey;
+};
+
+class FMetalTexture2D : public FRHITexture2D
+{
+public:
+	/** The surface info */
+	FMetalSurface Surface;
+
+	// Constructor, just calls base and Surface constructor
+	FMetalTexture2D(EPixelFormat Format, uint32 SizeX, uint32 SizeY, uint32 NumMips, uint32 NumSamples, uint32 Flags, FResourceBulkDataInterface* BulkData, const FClearValueBinding& InClearValue)
+		: FRHITexture2D(SizeX, SizeY, NumMips, NumSamples, Format, Flags, InClearValue)
+		, Surface(RRT_Texture2D, Format, SizeX, SizeY, 1, NumSamples, /*bArray=*/ false, 1, NumMips, Flags, BulkData)
+	{
+		FShaderCache* ShaderCache = FShaderCache::GetShaderCache();
+		if ( ShaderCache )
+		{
+			FShaderTextureKey Tex;
+			Tex.Format = (EPixelFormat)Format;
+			Tex.Flags = Flags;
+			Tex.MipLevels = NumMips;
+			Tex.Samples = NumSamples;
+			Tex.X = SizeX;
+			Tex.Y = SizeY;
+			Tex.Type = SCTT_Texture2D;
+			
+			FShaderCache::LogTexture(Tex, this);
+		}
+	}
+	
+	virtual ~FMetalTexture2D()
+	{
+		FShaderCache::RemoveTexture(this);
+	}
+};
+
+class FMetalTexture2DArray : public FRHITexture2DArray
+{
+public:
+	/** The surface info */
+	FMetalSurface Surface;
+
+	// Constructor, just calls base and Surface constructor
+	FMetalTexture2DArray(EPixelFormat Format, uint32 SizeX, uint32 SizeY, uint32 ArraySize, uint32 NumMips, uint32 Flags, FResourceBulkDataInterface* BulkData, const FClearValueBinding& InClearValue)
+		: FRHITexture2DArray(SizeX, SizeY, ArraySize, NumMips, Format, Flags, InClearValue)
+		, Surface(RRT_Texture2DArray, Format, SizeX, SizeY, 1, /*NumSamples=*/1, /*bArray=*/ true, ArraySize, NumMips, Flags, BulkData)
+	{
+		FShaderCache* ShaderCache = FShaderCache::GetShaderCache();
+		if ( ShaderCache )
+		{
+			FShaderTextureKey Tex;
+			Tex.Format = (EPixelFormat)Format;
+			Tex.Flags = Flags;
+			Tex.MipLevels = NumMips;
+			Tex.Samples = 1;
+			Tex.X = SizeX;
+			Tex.Y = SizeY;
+			Tex.Z = ArraySize;
+			Tex.Type = SCTT_Texture2DArray;
+			
+			FShaderCache::LogTexture(Tex, this);
+		}
+	}
+	
+	virtual ~FMetalTexture2DArray()
+	{
+		FShaderCache::RemoveTexture(this);
+	}
+};
+
+class FMetalTexture3D : public FRHITexture3D
+{
+public:
+	/** The surface info */
+	FMetalSurface Surface;
+
+	// Constructor, just calls base and Surface constructor
+	FMetalTexture3D(EPixelFormat Format, uint32 SizeX, uint32 SizeY, uint32 SizeZ, uint32 NumMips, uint32 Flags, FResourceBulkDataInterface* BulkData, const FClearValueBinding& InClearValue)
+		: FRHITexture3D(SizeX, SizeY, SizeZ, NumMips, Format, Flags, InClearValue)
+		, Surface(RRT_Texture3D, Format, SizeX, SizeY, SizeZ, /*NumSamples=*/1, /*bArray=*/ false, 1, NumMips, Flags, BulkData)
+	{
+		FShaderCache* ShaderCache = FShaderCache::GetShaderCache();
+		if ( ShaderCache )
+		{
+			FShaderTextureKey Tex;
+			Tex.Format = (EPixelFormat)Format;
+			Tex.Flags = Flags;
+			Tex.MipLevels = NumMips;
+			Tex.Samples = 1;
+			Tex.X = SizeX;
+			Tex.Y = SizeY;
+			Tex.Z = SizeZ;
+			Tex.Type = SCTT_Texture3D;
+			
+			FShaderCache::LogTexture(Tex, this);
+		}
+	}
+	
+	virtual ~FMetalTexture3D()
+	{
+		FShaderCache::RemoveTexture(this);
+	}
+};
+
+class FMetalTextureCube : public FRHITextureCube
+{
+public:
+	/** The surface info */
+	FMetalSurface Surface;
+
+	// Constructor, just calls base and Surface constructor
+	FMetalTextureCube(EPixelFormat Format, uint32 Size, bool bArray, uint32 ArraySize, uint32 NumMips, uint32 Flags, FResourceBulkDataInterface* BulkData, const FClearValueBinding& InClearValue)
+		: FRHITextureCube(Size, NumMips, Format, Flags, InClearValue)
+		, Surface(RRT_TextureCube, Format, Size, Size, 6, /*NumSamples=*/1, bArray, ArraySize, NumMips, Flags, BulkData)
+	{
+		FShaderCache* ShaderCache = FShaderCache::GetShaderCache();
+		if ( ShaderCache )
+		{
+			FShaderTextureKey Tex;
+			Tex.Format = (EPixelFormat)Format;
+			Tex.Flags = Flags;
+			Tex.MipLevels = NumMips;
+			Tex.Samples = 1;
+			Tex.X = Size;
+			Tex.Y = Size;
+			Tex.Z = ArraySize;
+			Tex.Type = bArray ? SCTT_TextureCube : SCTT_TextureCubeArray;
+			
+			FShaderCache::LogTexture(Tex, this);
+		}
+	}
+	
+	virtual ~FMetalTextureCube()
+	{
+		FShaderCache::RemoveTexture(this);
+	}
+};
+
+struct FMetalQueryBuffer : public FRHIResource
+{
+    FMetalQueryBuffer(FMetalContext* InContext, id<MTLBuffer> InBuffer);
+    
+    virtual ~FMetalQueryBuffer();
+    
+    bool Wait(uint64 Millis);
+    void const* GetResult(uint32 Offset);
+	
+	FMetalContext* Context;
+    id<MTLBuffer> Buffer;
+	uint32 WriteOffset;
+	bool bCompleted;
+	id<MTLCommandBuffer> CommandBuffer;
+};
+typedef TRefCountPtr<FMetalQueryBuffer> FMetalQueryBufferRef;
+
+struct FMetalQueryResult
+{
+    bool Wait(uint64 Millis);
+    void const* GetResult();
+    
+    FMetalQueryBufferRef SourceBuffer;
+    uint32 Offset;
+};
+
+/** Metal occlusion query */
+class FMetalRenderQuery : public FRHIRenderQuery
+{
+public:
+
+	/** Initialization constructor. */
+	FMetalRenderQuery(ERenderQueryType InQueryType);
+
+	~FMetalRenderQuery();
+
+	/**
+	 * Kick off an occlusion test 
+	 */
+	void Begin(FMetalContext* Context);
+
+	/**
+	 * Finish up an occlusion test 
+	 */
+	void End(FMetalContext* Context);
+	
+	// The type of query
+	ERenderQueryType Type;
+
+	// Query buffer allocation details as the buffer is already set on the command-encoder
+    FMetalQueryResult Buffer;
+
+    // Query result.
+    uint64 Result;
+    
+    // Result availability - if not set the first call to acquite it will read the buffer & cache
+    bool bAvailable;
+};
+
+/** Index buffer resource class that stores stride information. */
+class FMetalIndexBuffer : public FRHIIndexBuffer
+{
+public:
+
+	/** Constructor */
+	FMetalIndexBuffer(uint32 InStride, uint32 InSize, uint32 InUsage);
+	~FMetalIndexBuffer();
+
+	/**
+	 * Prepare a CPU accessible buffer for uploading to GPU memory
+	 */
+	void* Lock(EResourceLockMode LockMode, uint32 Offset, uint32 Size=0);
+
+	/**
+	 * Prepare a CPU accessible buffer for uploading to GPU memory
+	 */
+	void Unlock();
+	
+	// balsa buffer memory
+	id<MTLBuffer> Buffer;
+
+	// offet into the buffer (for lock usage)
+	uint32 LockOffset;
+	
+	// Lock size
+	uint32 LockSize;
+	
+	// 16- or 32-bit
+	MTLIndexType IndexType;
+};
+
+
+/** Vertex buffer resource class that stores usage type. */
+class FMetalVertexBuffer : public FRHIVertexBuffer
+{
+public:
+
+	/** Constructor */
+	FMetalVertexBuffer(uint32 InSize, uint32 InUsage);
+	~FMetalVertexBuffer();
+
+	/**
+	 * Prepare a CPU accessible buffer for uploading to GPU memory
+	 */
+	void* Lock(EResourceLockMode LockMode, uint32 Offset, uint32 Size=0);
+
+	/**
+	 * Prepare a CPU accessible buffer for uploading to GPU memory
+	 */
+	void Unlock();
+	
+	// balsa buffer memory
+	id<MTLBuffer> Buffer;
+
+	// offset into the buffer (for lock usage)
+	uint32 LockOffset;
+	
+	// Sizeof outstanding lock.
+	uint32 LockSize;
+
+	// if the buffer is a zero stride buffer, we need to duplicate and grow on demand, this is the size of one element to copy
+	uint32 ZeroStrideElementSize;
+};
+
+class FMetalUniformBuffer : public FRHIUniformBuffer
+{
+public:
+
+	// Constructor
+	FMetalUniformBuffer(const void* Contents, const FRHIUniformBufferLayout& Layout, EUniformBufferUsage Usage);
+
+	// Destructor 
+	~FMetalUniformBuffer();
+	
+
+	/** Cache resources if needed. */
+	inline void CacheResources(uint32 InFrameCounter)
+	{
+		if (InFrameCounter == INDEX_NONE || LastCachedFrame != InFrameCounter)
+		{
+			CacheResourcesInternal();
+			LastCachedFrame = InFrameCounter;
+		}
+	}
+
+	bool IsConstantBuffer() const
+	{
+		return Buffer.length > 0;
+	}
+
+
+	/** The buffer containing the contents - either a fresh buffer or the ring buffer */
+	id<MTLBuffer> Buffer;
+	
+	/** This offset is used when passing to setVertexBuffer, etc */
+	uint32 Offset;
+
+	uint32 Size; // @todo zebra: HACK! This should be removed and the code that uses it should be changed to use GetSize() instead once we fix the problem with FRHIUniformBufferLayout being released too early
+
+	/** Resource table containing RHI references. */
+	TArray<TRefCountPtr<FRHIResource> > ResourceTable;
+
+	/** Raw resource table, cached once per frame. */
+	TArray<void*> RawResourceTable;
+
+	/** The frame in which RawResourceTable was last cached. */
+	uint32 LastCachedFrame;
+
+private:
+	/** Actually cache resources. */
+	void CacheResourcesInternal();
+
+};
+
+
+class FMetalStructuredBuffer : public FRHIStructuredBuffer
+{
+public:
+	// Constructor
+	FMetalStructuredBuffer(uint32 Stride, uint32 Size, FResourceArrayInterface* ResourceArray, uint32 InUsage);
+
+	// Destructor
+	~FMetalStructuredBuffer();
+
+	/**
+	 * Prepare a CPU accessible buffer for uploading to GPU memory
+	 */
+	void* Lock(EResourceLockMode LockMode, uint32 Offset, uint32 Size=0);
+	
+	/**
+	 * Prepare a CPU accessible buffer for uploading to GPU memory
+	 */
+	void Unlock();
+	
+	// offset into the buffer (for lock usage)
+	uint32 LockOffset;
+	
+	// Sizeof outstanding lock.
+	uint32 LockSize;
+	
+	// the actual buffer
+	id<MTLBuffer> Buffer;
+};
+
+
+
+class FMetalUnorderedAccessView : public FRHIUnorderedAccessView
+{
+public:
+
+	/** Set it into the compute context */
+	void Set(FMetalContext* Context, uint32 ResourceIndex);
+	
+	// the potential resources to refer to with the UAV object
+	TRefCountPtr<FMetalStructuredBuffer> SourceStructuredBuffer;
+	TRefCountPtr<FMetalVertexBuffer> SourceVertexBuffer;
+	TRefCountPtr<FRHITexture> SourceTexture;
+};
+
+
+class FMetalShaderResourceView : public FRHIShaderResourceView
+{
+public:
+
+	// The vertex buffer this SRV comes from (can be null)
+	TRefCountPtr<FMetalVertexBuffer> SourceVertexBuffer;
+
+	// The texture that this SRV come from
+	TRefCountPtr<FRHITexture> SourceTexture;
+	
+	FMetalSurface* TextureView;
+	uint8 MipLevel;
+	uint8 NumMips;
+	uint8 Format;
+
+	~FMetalShaderResourceView();
+};
+
+class FMetalShaderParameterCache
+{
+public:
+	/** Constructor. */
+	FMetalShaderParameterCache();
+
+	/** Destructor. */
+	~FMetalShaderParameterCache();
+
+	void InitializeResources(int32 UniformArraySize);
+
+	/**
+	 * Marks all uniform arrays as dirty.
+	 */
+	void MarkAllDirty();
+
+	/**
+	 * Sets values directly into the packed uniform array
+	 */
+	void Set(uint32 BufferIndex, uint32 ByteOffset, uint32 NumBytes, const void* NewValues);
+
+	/**
+	 * Commit shader parameters to the currently bound program.
+	 */
+	void CommitPackedGlobals(FMetalContext* Context, int32 Stage, const FMetalShaderBindings& Bindings);
+	void CommitPackedUniformBuffers(TRefCountPtr<FMetalBoundShaderState> BoundShaderState, FMetalComputeShader* ComputeShader, int32 Stage, const TArray< TRefCountPtr<FRHIUniformBuffer> >& UniformBuffers, const TArray<CrossCompiler::FUniformBufferCopyInfo>& UniformBuffersCopyInfo);
+
+private:
+	/** CPU memory block for storing uniform values. */
+	uint8* PackedGlobalUniforms[CrossCompiler::PACKED_TYPEINDEX_MAX];
+
+	struct FRange
+	{
+		uint32	LowVector;
+		uint32	HighVector;
+	};
+	/** Dirty ranges for each uniform array. */
+	FRange	PackedGlobalUniformDirty[CrossCompiler::PACKED_TYPEINDEX_MAX];
+
+	/** Scratch CPU memory block for uploading packed uniforms. */
+	uint8* PackedUniformsScratch[CrossCompiler::PACKED_TYPEINDEX_MAX];
+
+	int32 GlobalUniformArraySize;
+};
+
+template<class T>
+struct TMetalResourceTraits
+{
+};
+template<>
+struct TMetalResourceTraits<FRHIVertexDeclaration>
+{
+	typedef FMetalVertexDeclaration TConcreteType;
+};
+template<>
+struct TMetalResourceTraits<FRHIVertexShader>
+{
+	typedef FMetalVertexShader TConcreteType;
+};
+template<>
+struct TMetalResourceTraits<FRHIGeometryShader>
+{
+	typedef FMetalGeometryShader TConcreteType;
+};
+template<>
+struct TMetalResourceTraits<FRHIHullShader>
+{
+	typedef FMetalHullShader TConcreteType;
+};
+template<>
+struct TMetalResourceTraits<FRHIDomainShader>
+{
+	typedef FMetalDomainShader TConcreteType;
+};
+template<>
+struct TMetalResourceTraits<FRHIPixelShader>
+{
+	typedef FMetalPixelShader TConcreteType;
+};
+template<>
+struct TMetalResourceTraits<FRHIComputeShader>
+{
+	typedef FMetalComputeShader TConcreteType;
+};
+template<>
+struct TMetalResourceTraits<FRHIBoundShaderState>
+{
+	typedef FMetalBoundShaderState TConcreteType;
+};
+template<>
+struct TMetalResourceTraits<FRHITexture3D>
+{
+	typedef FMetalTexture3D TConcreteType;
+};
+template<>
+struct TMetalResourceTraits<FRHITexture2D>
+{
+	typedef FMetalTexture2D TConcreteType;
+};
+template<>
+struct TMetalResourceTraits<FRHITexture2DArray>
+{
+	typedef FMetalTexture2DArray TConcreteType;
+};
+template<>
+struct TMetalResourceTraits<FRHITextureCube>
+{
+	typedef FMetalTextureCube TConcreteType;
+};
+template<>
+struct TMetalResourceTraits<FRHIRenderQuery>
+{
+	typedef FMetalRenderQuery TConcreteType;
+};
+template<>
+struct TMetalResourceTraits<FRHIUniformBuffer>
+{
+	typedef FMetalUniformBuffer TConcreteType;
+};
+template<>
+struct TMetalResourceTraits<FRHIIndexBuffer>
+{
+	typedef FMetalIndexBuffer TConcreteType;
+};
+template<>
+struct TMetalResourceTraits<FRHIStructuredBuffer>
+{
+	typedef FMetalStructuredBuffer TConcreteType;
+};
+template<>
+struct TMetalResourceTraits<FRHIVertexBuffer>
+{
+	typedef FMetalVertexBuffer TConcreteType;
+};
+template<>
+struct TMetalResourceTraits<FRHIShaderResourceView>
+{
+	typedef FMetalShaderResourceView TConcreteType;
+};
+template<>
+struct TMetalResourceTraits<FRHIUnorderedAccessView>
+{
+	typedef FMetalUnorderedAccessView TConcreteType;
+};
+
+template<>
+struct TMetalResourceTraits<FRHISamplerState>
+{
+	typedef FMetalSamplerState TConcreteType;
+};
+template<>
+struct TMetalResourceTraits<FRHIRasterizerState>
+{
+	typedef FMetalRasterizerState TConcreteType;
+};
+template<>
+struct TMetalResourceTraits<FRHIDepthStencilState>
+{
+	typedef FMetalDepthStencilState TConcreteType;
+};
+template<>
+struct TMetalResourceTraits<FRHIBlendState>
+{
+	typedef FMetalBlendState TConcreteType;
+};
diff --git a/Engine/Source/Runtime/Apple/MetalRHI/Public/MetalViewport.h b/Engine/Source/Runtime/Apple/MetalRHI/Public/MetalViewport.h
new file mode 100644
--- /dev/null
+++ b/Engine/Source/Runtime/Apple/MetalRHI/Public/MetalViewport.h
@@ -0,0 +1,39 @@
+// Copyright 1998-2015 Epic Games, Inc. All Rights Reserved.
+
+/*=============================================================================
+	MetalViewport.h: Metal viewport RHI definitions.
+=============================================================================*/
+
+#pragma once
+
+#if PLATFORM_MAC
+#include "CocoaTextView.h"
+@interface FMetalView : FCocoaTextView
+@end
+#endif
+
+class FMetalViewport : public FRHIViewport
+{
+public:
+
+	FMetalViewport(void* WindowHandle, uint32 InSizeX, uint32 InSizeY, bool bInIsFullscreen);
+	~FMetalViewport();
+
+	void Resize(uint32 InSizeX, uint32 InSizeY, bool bInIsFullscreen);
+
+	FMetalTexture2D* GetBackBuffer() const { return BackBuffer; }
+
+#if PLATFORM_MAC
+	FMetalView* View;
+#endif
+
+private:
+
+	TRefCountPtr<FMetalTexture2D> BackBuffer;
+};
+
+template<>
+struct TMetalResourceTraits<FRHIViewport>
+{
+	typedef FMetalViewport TConcreteType;
+};
diff --git a/Engine/Source/Runtime/Core/Private/Mac/MacPlatformMisc.cpp b/Engine/Source/Runtime/Core/Private/Mac/MacPlatformMisc.cpp
--- a/Engine/Source/Runtime/Core/Private/Mac/MacPlatformMisc.cpp
+++ b/Engine/Source/Runtime/Core/Private/Mac/MacPlatformMisc.cpp
@@ -1004,12 +1004,26 @@ void FMacPlatformMisc::NormalizePath(FString& InPath)
 	SCOPED_AUTORELEASE_POOL;
 	if (InPath.Len() > 1)
 	{
+#if WITH_EDITOR
 		const bool bAppendSlash = InPath[InPath.Len() - 1] == '/'; // NSString will remove the trailing slash, if present, so we need to restore it after conversion
 		InPath = [[InPath.GetNSString() stringByStandardizingPath] stringByResolvingSymlinksInPath];
 		if (bAppendSlash)
 		{
 			InPath += TEXT("/");
 		}
+#else
+		InPath.ReplaceInline(TEXT("\\"), TEXT("/"));
+		// This replacement addresses a "bug" where some callers
+		// pass in paths that are badly composed with multiple
+		// subdir separators.
+		InPath.ReplaceInline(TEXT("//"), TEXT("/"));
+		if (!InPath.IsEmpty() && InPath[InPath.Len() - 1] == TEXT('/'))
+		{
+			InPath.LeftChop(1);
+		}
+		// Remove redundant current-dir references.
+		InPath.ReplaceInline(TEXT("/./"), TEXT("/"));
+#endif
 	}
 }
 
@@ -1857,3 +1871,29 @@ void NewReportEnsure( const TCHAR* ErrorMessage )
 	bReentranceGuard = false;
 	EnsureLock.Unlock();
 }
+typedef NSArray* (*MTLCopyAllDevices)(void);
+
+bool FMacPlatformMisc::HasPlatformFeature(const TCHAR* FeatureName)
+{
+	if (FCString::Stricmp(FeatureName, TEXT("Metal")) == 0 && !FParse::Param(FCommandLine::Get(),TEXT("opengl")))
+	{
+		// Find out if there are any Metal devices on the system - some Mac's have none
+		void* DLLHandle = FPlatformProcess::GetDllHandle(TEXT("/System/Library/Frameworks/Metal.framework/Metal"));
+		if(DLLHandle)
+		{
+			// Use the copy all function because we don't want to invoke a GPU switch at this point on dual-GPU Macbooks
+			MTLCopyAllDevices CopyDevicesPtr = (MTLCopyAllDevices)FPlatformProcess::GetDllExport(DLLHandle, TEXT("MTLCopyAllDevices"));
+			if(CopyDevicesPtr)
+			{
+				SCOPED_AUTORELEASE_POOL;
+				NSArray* MetalDevices = CopyDevicesPtr();
+				[MetalDevices autorelease];
+				FPlatformProcess::FreeDllHandle(DLLHandle);
+				return MetalDevices && [MetalDevices count] > 0;
+			}
+		}
+	}
+	
+	return FGenericPlatformMisc::HasPlatformFeature(FeatureName);
+}
+
diff --git a/Engine/Source/Runtime/Core/Public/Mac/MacPlatformMisc.h b/Engine/Source/Runtime/Core/Public/Mac/MacPlatformMisc.h
--- a/Engine/Source/Runtime/Core/Public/Mac/MacPlatformMisc.h
+++ b/Engine/Source/Runtime/Core/Public/Mac/MacPlatformMisc.h
@@ -100,6 +100,7 @@ struct CORE_API FMacPlatformMisc : public FGenericPlatformMisc
 	static void NormalizePath(FString& InPath);
 	static FString GetPrimaryGPUBrand();
 	static void GetOSVersions( FString& out_OSVersionLabel, FString& out_OSSubVersionLabel );
+	static bool HasPlatformFeature(const TCHAR* FeatureName);
 	static bool GetDiskTotalAndFreeSpace(const FString& InPath, uint64& TotalNumberOfBytes, uint64& NumberOfFreeBytes);
 
 
diff --git a/Engine/Source/Runtime/RHI/Private/Apple/AppleDynamicRHI.cpp b/Engine/Source/Runtime/RHI/Private/Apple/AppleDynamicRHI.cpp
--- a/Engine/Source/Runtime/RHI/Private/Apple/AppleDynamicRHI.cpp
+++ b/Engine/Source/Runtime/RHI/Private/Apple/AppleDynamicRHI.cpp
@@ -9,13 +9,11 @@ FDynamicRHI* PlatformCreateDynamicRHI()
 	IDynamicRHIModule* DynamicRHIModule = NULL;
 
 	// Load the dynamic RHI module.
-#if PLATFORM_IOS
 	if (FPlatformMisc::HasPlatformFeature(TEXT("Metal")))
 	{
 		DynamicRHIModule = &FModuleManager::LoadModuleChecked<IDynamicRHIModule>(TEXT("MetalRHI"));
 	}
 	else
-#endif
 	{
 		DynamicRHIModule = &FModuleManager::LoadModuleChecked<IDynamicRHIModule>(TEXT("OpenGLDrv"));
 	}
