Build fix for Mac and iOS.

diff --git a/Engine/Source/Runtime/Apple/MetalRHI/Private/MetalShaders.cpp b/Engine/Source/Runtime/Apple/MetalRHI/Private/MetalShaders.cpp
--- a/Engine/Source/Runtime/Apple/MetalRHI/Private/MetalShaders.cpp
+++ b/Engine/Source/Runtime/Apple/MetalRHI/Private/MetalShaders.cpp
@@ -49,12 +49,12 @@ static FMetalCompiledShaderCache& GetMetalCompiledShaderCache()
 
 /** Initialization constructor. */
 template<typename BaseResourceType, int32 ShaderType>
-TMetalBaseShader<BaseResourceType, ShaderType>::TMetalBaseShader(const TArray<uint8>& InShaderCode)
+TMetalBaseShader<BaseResourceType, ShaderType>::TMetalBaseShader(const TArray<uint8>& InCode)
 	: DirtyUniformBuffers(0)
 {
-	FMemoryReader Ar(InShaderCode, true);
+	FMemoryReader Ar(InCode, true);
 	
-	Ar.SetLimitSize(FShaderCodeReader(InShaderCode).GetActualShaderCodeSize());
+	Ar.SetLimitSize(FShaderCodeReader(InCode).GetActualShaderCodeSize());
 
 	// was the shader already compiled offline?
 	uint8 OfflineCompiledFlag;
