fixed Metal shader code

code_review:Ben.Marsh

diff --git a/Engine/Source/Runtime/Apple/MetalRHI/Private/MetalShaders.cpp b/Engine/Source/Runtime/Apple/MetalRHI/Private/MetalShaders.cpp
--- a/Engine/Source/Runtime/Apple/MetalRHI/Private/MetalShaders.cpp
+++ b/Engine/Source/Runtime/Apple/MetalRHI/Private/MetalShaders.cpp
@@ -49,12 +49,14 @@ static FMetalCompiledShaderCache& GetMetalCompiledShaderCache()
 
 /** Initialization constructor. */
 template<typename BaseResourceType, int32 ShaderType>
-TMetalBaseShader<BaseResourceType, ShaderType>::TMetalBaseShader(const TArray<uint8>& InCode)
+TMetalBaseShader<BaseResourceType, ShaderType>::TMetalBaseShader(const TArray<uint8>& InShaderCode)
 	: DirtyUniformBuffers(0)
 {
-	FMemoryReader Ar(InCode, true);
+	FShaderCodeReader ShaderCode(InShaderCode);
+
+	FMemoryReader Ar(InShaderCode, true);
 	
-	Ar.SetLimitSize(FShaderCodeReader(InCode).GetActualShaderCodeSize());
+	Ar.SetLimitSize(ShaderCode.GetActualShaderCodeSize());
 
 	// was the shader already compiled offline?
 	uint8 OfflineCompiledFlag;
@@ -67,21 +69,21 @@ TMetalBaseShader<BaseResourceType, ShaderType>::TMetalBaseShader(const TArray<ui
 
 	// remember where the header ended and code (precompiled or source) begins
 	int32 CodeOffset = Ar.Tell();
-	const ANSICHAR* SourceCode = (ANSICHAR*)InCode.GetData() + CodeOffset;
+	const ANSICHAR* SourceCode = (ANSICHAR*)InShaderCode.GetData() + CodeOffset;
 
 	uint32 CodeLength, CodeCRC;
 	if (OfflineCompiledFlag)
 	{
-		CodeLength = InCode.Num() - CodeOffset;
+		CodeLength = ShaderCode.GetActualShaderCodeSize() - CodeOffset;
 
 		// CRC the compiled code
-		CodeCRC = FCrc::MemCrc_DEPRECATED(InCode.GetData() + CodeOffset, InCode.Num() - CodeOffset);
+		CodeCRC = FCrc::MemCrc_DEPRECATED(InShaderCode.GetData() + CodeOffset, CodeLength);
 	}
 	else
 	{
 		UE_LOG(LogMetal, Display, TEXT("Loaded a non-offline compiled shader (will be slower to load)"));
 
-		CodeLength = InCode.Num() - CodeOffset - 1;
+		CodeLength = ShaderCode.GetActualShaderCodeSize() - CodeOffset - 1;
 
 		// CRC the source
 		CodeCRC = FCrc::MemCrc_DEPRECATED(SourceCode, CodeLength);
@@ -97,10 +99,10 @@ TMetalBaseShader<BaseResourceType, ShaderType>::TMetalBaseShader(const TArray<ui
 		if (OfflineCompiledFlag)
 		{
 			// allow GCD to copy the data into its own buffer
-			//		dispatch_data_t GCDBuffer = dispatch_data_create(InCode.GetTypedData() + CodeOffset, InCode.Num() - CodeOffset, nil, DISPATCH_DATA_DESTRUCTOR_DEFAULT);
-			uint32 BufferSize = InCode.Num() - CodeOffset;
+			//		dispatch_data_t GCDBuffer = dispatch_data_create(InShaderCode.GetTypedData() + CodeOffset, ShaderCode.GetActualShaderCodeSize() - CodeOffset, nil, DISPATCH_DATA_DESTRUCTOR_DEFAULT);
+			uint32 BufferSize = ShaderCode.GetActualShaderCodeSize() - CodeOffset;
 			void* Buffer = FMemory::Malloc( BufferSize );
-			FMemory::Memcpy( Buffer, InCode.GetData() + CodeOffset, BufferSize );
+			FMemory::Memcpy( Buffer, InShaderCode.GetData() + CodeOffset, BufferSize );
 			dispatch_data_t GCDBuffer = dispatch_data_create(Buffer, BufferSize, dispatch_get_main_queue(), ^(void) { FMemory::Free(Buffer); } );
 
 			// load up the already compiled shader
