Disable r.AllowPointLightCubemapShadows  on Mac Metal as they aren't presently supported for Metal but are for OpenGL (so it can't go in an ini).

diff --git a/Engine/Source/Runtime/Apple/MetalRHI/Private/MetalRHI.cpp b/Engine/Source/Runtime/Apple/MetalRHI/Private/MetalRHI.cpp
--- a/Engine/Source/Runtime/Apple/MetalRHI/Private/MetalRHI.cpp
+++ b/Engine/Source/Runtime/Apple/MetalRHI/Private/MetalRHI.cpp
@@ -108,6 +108,13 @@ FMetalDynamicRHI::FMetalDynamicRHI()
 	// Change the support depth format if we can
 	bSupportsD24S8 = Device.depth24Stencil8PixelFormatSupported;
 	
+	// Disable point light cubemap shadows on Mac Metal as currently they aren't supported.
+	static auto CVarCubemapShadows = IConsoleManager::Get().FindConsoleVariable(TEXT("r.AllowPointLightCubemapShadows"));
+	if(CVarCubemapShadows && CVarCubemapShadows->GetInt() != 0)
+	{
+		CVarCubemapShadows->Set(0);
+	}
+	
 #endif
 
     
