Fix resetting of render command encoder on visibility buffer wrap - this requires assigning a new visibility buffer so we must create a new render-pass descriptor and restore all other previous state, not just restore the render command encoder as it was.

diff --git a/Engine/Source/Runtime/Apple/MetalRHI/Private/MetalContext.cpp b/Engine/Source/Runtime/Apple/MetalRHI/Private/MetalContext.cpp
--- a/Engine/Source/Runtime/Apple/MetalRHI/Private/MetalContext.cpp
+++ b/Engine/Source/Runtime/Apple/MetalRHI/Private/MetalContext.cpp
@@ -493,7 +493,8 @@ void FMetalContext::ResetRenderCommandEncoder()
 	SubmitCommandsHint();
 	
 	ConditionalSwitchToGraphics();
-	CommandEncoder.RestoreRenderCommandEncoding();
+	StateCache.SetRenderTargetsInfo(StateCache.GetRenderTargetsInfo(), QueryBuffer.GetCurrentQueryBuffer()->Buffer);
+	CommandEncoder.RestoreRenderCommandEncodingState();
 }
 
 void FMetalContext::PrepareToDraw(uint32 PrimitiveType)
@@ -520,7 +521,7 @@ void FMetalContext::PrepareToDraw(uint32 PrimitiveType)
 #if PLATFORM_MAC // Disable this to get a soft error that can be debugged further in a GPU trace. // @todo zebra ios
 					if(!FShaderCache::IsPredrawCall())
 					{
-						UE_LOG(LogMetal, Fatal, TEXT("%s"), *Report);
+						UE_LOG(LogMetal, Warning, TEXT("%s"), *Report);
 					}
 					else
 					{
diff --git a/Engine/Source/Runtime/Apple/MetalRHI/Private/MetalStateCache.cpp b/Engine/Source/Runtime/Apple/MetalRHI/Private/MetalStateCache.cpp
--- a/Engine/Source/Runtime/Apple/MetalRHI/Private/MetalStateCache.cpp
+++ b/Engine/Source/Runtime/Apple/MetalRHI/Private/MetalStateCache.cpp
@@ -25,6 +25,7 @@ static MTLCullMode TranslateCullMode(ERasterizerCullMode CullMode)
 
 FMetalStateCache::FMetalStateCache(FMetalCommandEncoder& InCommandEncoder)
 : CommandEncoder(InCommandEncoder)
+, VisibilityResults(nil)
 , BlendState(nullptr)
 , DepthStencilState(nullptr)
 , RasterizerState(nullptr)
@@ -163,7 +164,7 @@ void FMetalStateCache::SetRenderTargetsInfo(FRHISetRenderTargetsInfo const& InRe
 	ConditionalSwitchToRender();
 
 	// see if our new Info matches our previous Info
-	if (NeedsToSetRenderTarget(InRenderTargets))
+	if (NeedsToSetRenderTarget(InRenderTargets) || QueryBuffer != VisibilityResults)
 	{
 		// back this up for next frame
 		RenderTargetsInfo = InRenderTargets;
@@ -172,6 +173,7 @@ void FMetalStateCache::SetRenderTargetsInfo(FRHISetRenderTargetsInfo const& InRe
 		MTLRenderPassDescriptor* RenderPass = [MTLRenderPassDescriptor renderPassDescriptor];
 	
 		// if we need to do queries, write to the supplied query buffer
+		VisibilityResults = QueryBuffer;
 		RenderPass.visibilityResultBuffer = QueryBuffer;
 	
 		// default to non-msaa
diff --git a/Engine/Source/Runtime/Apple/MetalRHI/Private/MetalStateCache.h b/Engine/Source/Runtime/Apple/MetalRHI/Private/MetalStateCache.h
--- a/Engine/Source/Runtime/Apple/MetalRHI/Private/MetalStateCache.h
+++ b/Engine/Source/Runtime/Apple/MetalRHI/Private/MetalStateCache.h
@@ -74,6 +74,8 @@ private:
 	id<MTLBuffer> VertexBuffers[MaxMetalStreams];
 	uint32 VertexStrides[MaxMetalStreams];
 
+	id<MTLBuffer> VisibilityResults;
+	
 	TRefCountPtr<FMetalBlendState> BlendState;
 	TRefCountPtr<FMetalDepthStencilState> DepthStencilState;
 	TRefCountPtr<FMetalRasterizerState> RasterizerState;
