CIS Fixes: Hoping to fix non-unity compiler errors on mac.

diff --git a/Engine/Source/Runtime/Apple/MetalRHI/Private/MetalIndexBuffer.cpp b/Engine/Source/Runtime/Apple/MetalRHI/Private/MetalIndexBuffer.cpp
--- a/Engine/Source/Runtime/Apple/MetalRHI/Private/MetalIndexBuffer.cpp
+++ b/Engine/Source/Runtime/Apple/MetalRHI/Private/MetalIndexBuffer.cpp
@@ -5,7 +5,7 @@
 =============================================================================*/
 
 #include "MetalRHIPrivate.h"
-
+#include "MetalProfiler.h"
 
 /** Constructor */
 FMetalIndexBuffer::FMetalIndexBuffer(uint32 InStride, uint32 InSize, uint32 InUsage)
diff --git a/Engine/Source/Runtime/Apple/MetalRHI/Private/MetalProfiler.cpp b/Engine/Source/Runtime/Apple/MetalRHI/Private/MetalProfiler.cpp
--- a/Engine/Source/Runtime/Apple/MetalRHI/Private/MetalProfiler.cpp
+++ b/Engine/Source/Runtime/Apple/MetalRHI/Private/MetalProfiler.cpp
@@ -1,7 +1,5 @@
 // Copyright 1998-2015 Epic Games, Inc. All Rights Reserved.
 
-#pragma once
-
 #include "MetalRHIPrivate.h"
 #include "MetalProfiler.h"
 
@@ -236,7 +234,7 @@ void FMetalGPUProfiler::PopEvent()
 	FGPUProfiler::PopEvent();
 }
 
-TGlobalResource<FVector4VertexDeclaration> GMetalVector4VertexDeclaration;
+//TGlobalResource<FVector4VertexDeclaration> GMetalVector4VertexDeclaration;
 TGlobalResource<FTexture> GMetalLongTaskRT;
 
 static FGlobalBoundShaderState LongGPUTaskBoundShaderState;
diff --git a/Engine/Source/Runtime/Apple/MetalRHI/Private/MetalStructuredBuffer.cpp b/Engine/Source/Runtime/Apple/MetalRHI/Private/MetalStructuredBuffer.cpp
--- a/Engine/Source/Runtime/Apple/MetalRHI/Private/MetalStructuredBuffer.cpp
+++ b/Engine/Source/Runtime/Apple/MetalRHI/Private/MetalStructuredBuffer.cpp
@@ -2,7 +2,7 @@
 
 
 #include "MetalRHIPrivate.h"
-
+#include "MetalProfiler.h"
 
 FMetalStructuredBuffer::FMetalStructuredBuffer(uint32 Stride, uint32 Size, FResourceArrayInterface* ResourceArray, uint32 Usage)
 	: FRHIStructuredBuffer(Stride, Size, Usage)
diff --git a/Engine/Source/Runtime/Apple/MetalRHI/Private/MetalTexture.cpp b/Engine/Source/Runtime/Apple/MetalRHI/Private/MetalTexture.cpp
--- a/Engine/Source/Runtime/Apple/MetalRHI/Private/MetalTexture.cpp
+++ b/Engine/Source/Runtime/Apple/MetalRHI/Private/MetalTexture.cpp
@@ -5,6 +5,7 @@
 =============================================================================*/
 
 #include "MetalRHIPrivate.h"
+#include "MetalProfiler.h" // for STAT_MetalTexturePageOffTime
 
 uint8 FMetalSurface::NextKey = 1; // 0 is reserved for MTLPixelFormatInvalid
 
diff --git a/Engine/Source/Runtime/Apple/MetalRHI/Private/MetalUniformBuffer.cpp b/Engine/Source/Runtime/Apple/MetalRHI/Private/MetalUniformBuffer.cpp
--- a/Engine/Source/Runtime/Apple/MetalRHI/Private/MetalUniformBuffer.cpp
+++ b/Engine/Source/Runtime/Apple/MetalRHI/Private/MetalUniformBuffer.cpp
@@ -5,7 +5,7 @@
 =============================================================================*/
 
 #include "MetalRHIPrivate.h"
-
+#include "MetalProfiler.h"
 
 #define NUM_POOL_BUCKETS 29
 
