CIS Fixes: Hoping to fix further non-unity compiler errors on mac after CL 2730163 missed some.

diff --git a/Engine/Source/Runtime/Apple/MetalRHI/Private/MetalRenderPipelineDesc.cpp b/Engine/Source/Runtime/Apple/MetalRHI/Private/MetalRenderPipelineDesc.cpp
--- a/Engine/Source/Runtime/Apple/MetalRHI/Private/MetalRenderPipelineDesc.cpp
+++ b/Engine/Source/Runtime/Apple/MetalRHI/Private/MetalRenderPipelineDesc.cpp
@@ -2,6 +2,7 @@
 
 #include "MetalRHIPrivate.h"
 #include "MetalRenderPipelineDesc.h"
+#include "MetalProfiler.h"
 
 static_assert(Offset_End <= sizeof(FMetalRenderPipelineHash) * 8, "Too many bits used for the Hash");
 
diff --git a/Engine/Source/Runtime/Apple/MetalRHI/Private/MetalShaders.cpp b/Engine/Source/Runtime/Apple/MetalRHI/Private/MetalShaders.cpp
--- a/Engine/Source/Runtime/Apple/MetalRHI/Private/MetalShaders.cpp
+++ b/Engine/Source/Runtime/Apple/MetalRHI/Private/MetalShaders.cpp
@@ -8,6 +8,7 @@
 #include "MetalShaderResources.h"
 #include "MetalResources.h"
 #include "ShaderCache.h"
+#include "MetalProfiler.h"
 
 /** Set to 1 to enable shader debugging (makes the driver save the shader source) */
 #define DEBUG_METAL_SHADERS (UE_BUILD_DEBUG || UE_BUILD_DEVELOPMENT)
diff --git a/Engine/Source/Runtime/Apple/MetalRHI/Private/MetalVertexBuffer.cpp b/Engine/Source/Runtime/Apple/MetalRHI/Private/MetalVertexBuffer.cpp
--- a/Engine/Source/Runtime/Apple/MetalRHI/Private/MetalVertexBuffer.cpp
+++ b/Engine/Source/Runtime/Apple/MetalRHI/Private/MetalVertexBuffer.cpp
@@ -5,6 +5,7 @@
 =============================================================================*/
 
 #include "MetalRHIPrivate.h"
+#include "MetalProfiler.h"
 
 
 FMetalVertexBuffer::FMetalVertexBuffer(uint32 InSize, uint32 InUsage)
