Fix JIRA UE-22923: Windows TCHAR is 2-bytes but on iOS its 4-bytes so the header deserialisation was going wrong if the shaders were compiled on Windows. This would have affected Windows->Mac cooking too.

diff --git a/Engine/Source/Developer/Apple/MetalShaderFormat/Private/MetalShaderCompiler.cpp b/Engine/Source/Developer/Apple/MetalShaderFormat/Private/MetalShaderCompiler.cpp
--- a/Engine/Source/Developer/Apple/MetalShaderFormat/Private/MetalShaderCompiler.cpp
+++ b/Engine/Source/Developer/Apple/MetalShaderFormat/Private/MetalShaderCompiler.cpp
@@ -298,7 +298,7 @@ static void BuildMetalShaderOutput(
 	}
 
 	const int32 MaxSamplers = GetFeatureLevelMaxTextureSamplers(ERHIFeatureLevel::ES3_1);
-	Header.ShaderName = CCHeader.Name.GetCharArray();
+	Header.ShaderName = CCHeader.Name;
 
 	if (Header.Bindings.NumSamplers > MaxSamplers)
 	{
diff --git a/Engine/Source/Developer/Apple/MetalShaderFormat/Private/MetalShaderFormat.cpp b/Engine/Source/Developer/Apple/MetalShaderFormat/Private/MetalShaderFormat.cpp
--- a/Engine/Source/Developer/Apple/MetalShaderFormat/Private/MetalShaderFormat.cpp
+++ b/Engine/Source/Developer/Apple/MetalShaderFormat/Private/MetalShaderFormat.cpp
@@ -17,7 +17,7 @@ class FMetalShaderFormat : public IShaderFormat
 {
 	enum
 	{
-		HEADER_VERSION = 26,
+		HEADER_VERSION = 28,
 	};
 	struct FVersion
 	{
diff --git a/Engine/Source/Runtime/Apple/MetalRHI/Private/MetalShaders.cpp b/Engine/Source/Runtime/Apple/MetalRHI/Private/MetalShaders.cpp
--- a/Engine/Source/Runtime/Apple/MetalRHI/Private/MetalShaders.cpp
+++ b/Engine/Source/Runtime/Apple/MetalRHI/Private/MetalShaders.cpp
@@ -122,9 +122,9 @@ TMetalBaseShader<BaseResourceType, ShaderType>::TMetalBaseShader(const TArray<ui
 			NSError* Error;
 			NSString* ShaderSource = [NSString stringWithUTF8String:SourceCode];
 
-			if(Header.ShaderName.Num())
+			if(Header.ShaderName.Len())
 			{
-				ShaderSource = [NSString stringWithFormat:@"// %@\n%@", FString(Header.ShaderName.GetData()).GetNSString(), ShaderSource];
+				ShaderSource = [NSString stringWithFormat:@"// %@\n%@", Header.ShaderName.GetNSString(), ShaderSource];
 			}
 			
 			MTLCompileOptions *CompileOptions = [[MTLCompileOptions alloc] init];
diff --git a/Engine/Source/Runtime/Apple/MetalRHI/Public/MetalShaderResources.h b/Engine/Source/Runtime/Apple/MetalRHI/Public/MetalShaderResources.h
--- a/Engine/Source/Runtime/Apple/MetalRHI/Public/MetalShaderResources.h
+++ b/Engine/Source/Runtime/Apple/MetalRHI/Public/MetalShaderResources.h
@@ -138,7 +138,7 @@ struct FMetalCodeHeader
 	uint32 Frequency;
 	FMetalShaderBindings Bindings;
 	TArray<CrossCompiler::FUniformBufferCopyInfo> UniformBuffersCopyInfo;
-	TArray<TCHAR> ShaderName;
+	FString ShaderName;
 
 	uint8 NumThreadsX;
 	uint8 NumThreadsY;
