Merging //UE4/Release-4.11 to //UE4/Main (Up to 2874930)

==========================
MAJOR FEATURES + CHANGES
==========================

Change 2868448 on 2016/02/16 by [REDACTED_PERSON]

	Bring MetalRHI & MetalShaderFormat mostly up to Dev-Platform CL #2867146 to address JIRA UE-26181 and as the first part of addressing JIRA UE-23208. Only CL #2854142 has been omitted for compatibility with the 4.11 branch.

Change 2868454 on 2016/02/16 by [REDACTED_PERSON]

	Shader changes necessary to properly fix Mobile Preview under Mac Metal (JIRA UE-23208) - Mac ES 3.1 doesn't support framebuffer fetch & is in fact a Mobile Emulation shader platform.

Change 2868650 on 2016/02/16 by [REDACTED_PERSON]

	Allow Developer modules to be pre-compiled when the target is editor.

Change 2868859 on 2016/02/16 by [REDACTED_PERSON]

	Merging StereoPanorama fixes from Main to Release-4.11 (CL 2811839)

Change 2868927 on 2016/02/16 by [REDACTED_PERSON]

	Fixed some regular-use crashes in LogVisualizer #UE4

Change 2868994 on 2016/02/16 by [REDACTED_PERSON]

	Fix refresh UI issue with remove joint

Change 2868996 on 2016/02/16 by [REDACTED_PERSON]

	Fix node stop working when negative value of LODThreshold

Change 2868998 on 2016/02/16 by [REDACTED_PERSON]

	Fix with crash when invalid index has entered

Change 2869003 on 2016/02/16 by [REDACTED_PERSON]

	- Fix thread safety issue when cloth child collision or environment collision is used.
	- Fix cloth bounds growing when stale transform data is used.

Change 2869109 on 2016/02/16 by mason.seay

	Updated test assets for Restitution testing

Change 2869223 on 2016/02/16 by [REDACTED_PERSON]

Change 2869558 on 2016/02/16 by [REDACTED_PERSON]

	Conservative fix for crash that occurs when adding a weak object ptr to an array of object ptrs in a blueprint

Change 2869891 on 2016/02/17 by [REDACTED_PERSON]

	Fix crash when re-compiling anim BPs that are dependencies of 'parent' Blueprints

	Force a re-initialzation of nodes when initializing the whole anim instance. The bInitialized flag was intended as a runtime optimization - we assume our function will not change for the lifetime of the UAnimInstance. While parts of the graph will be re-initialized we dont need to re-acquire our UFunction ptrs (etc.) unless the whole instance is getting re-initialized.

Change 2869956 on 2016/02/17 by [REDACTED_PERSON]

Change 2869965 on 2016/02/17 by [REDACTED_PERSON]

	Fix check() being hit when loading KiteDemo x1_y1 (and other old foliage maps)

Change 2870007 on 2016/02/17 by [REDACTED_PERSON]

	Merging from //UE4/Dev-Editor CL 2867609

	Fixed auto-generation of unique object/package names when duplicating, so they generate the same names if they the source object was named the same as its package.

Change 2870072 on 2016/02/17 by [REDACTED_PERSON]

	Level Blueprints will no longer preload newly created Blueprints that they reference when they are reloaded.

Change 2870087 on 2016/02/17 by [REDACTED_PERSON]

	Removed Android and IOS from list of platforms to build only for UnrealMatch3, so that they can be cooked and packaged correctly.

Change 2870141 on 2016/02/17 by mason.seay

	Test content for Copy Pose From Mesh

Change 2870195 on 2016/02/17 by [REDACTED_PERSON]

	PR #2052: Added missing Super::BeginPlay() (Contributed by [REDACTED_PERSON])

Change 2870325 on 2016/02/17 by [REDACTED_PERSON]

diff --git a/Engine/Build/HTML5/GameX.html.template b/Engine/Build/HTML5/GameX.html.template
--- a/Engine/Build/HTML5/GameX.html.template
+++ b/Engine/Build/HTML5/GameX.html.template
@@ -349,6 +349,7 @@ function postRunEmscripten() {
 	$("#compilingmessage").remove();
 	// it still takes couple of ticks after run before anything can be drawn. show the canvas at the last possible moment. 
 	$("#canvas").show(); 
+	window.dispatchEvent(new Event('resize')); // set canvas to browser resolution
 }
 
 // create a deffered object for the data file. 
@@ -461,11 +462,14 @@ $(document).ready(function() {
 				addLog ("all downloads complete ");
 			},
 		function(){
-
+				// one or more downloads failed.
+				// force a reload...
 				$('#progressbar').remove(); 
-				$ ('#mainarea').append ( '<div class="alert alert-danger centered-axis-xy" style ="min-height: 20px" role="alert">One or more downloads failed! Please refresh this page</div>');
-
+				$ ('#mainarea').append ( '<div class="alert alert-danger centered-axis-xy" style ="min-height: 20px" role="alert">Reloading!  Please standby...</div>');
 				addLog("Download failed");
+				setTimeout(function() {
+					location.reload();
+				}, 1000);
 			}
 		);
 		
diff --git a/Engine/Plugins/Runtime/GearVR/Source/GearVR/Private/GearVRRender.cpp b/Engine/Plugins/Runtime/GearVR/Source/GearVR/Private/GearVRRender.cpp
--- a/Engine/Plugins/Runtime/GearVR/Source/GearVR/Private/GearVRRender.cpp
+++ b/Engine/Plugins/Runtime/GearVR/Source/GearVR/Private/GearVRRender.cpp
@@ -353,6 +353,11 @@ void FViewExtension::PreRenderViewFamily_RenderThread(FRHICommandListImmediate&
 
 	bFrameBegun = true;
 
+	FQuat OldOrientation;
+	FVector OldPosition;
+	CurrentFrame->PoseToOrientationAndPosition(CurrentFrame->CurSensorState.HeadPose.Pose, OldOrientation, OldPosition);
+	const FTransform OldRelativeTransform(OldOrientation, OldPosition);
+
 	if (ShowFlags.Rendering)
 	{
 		check(pPresentBridge->GetRenderThreadId() == gettid());
@@ -363,6 +368,13 @@ void FViewExtension::PreRenderViewFamily_RenderThread(FRHICommandListImmediate&
 			return;
 		}
 	}
+
+	FQuat NewOrientation;
+	FVector NewPosition;
+	CurrentFrame->PoseToOrientationAndPosition(NewTracking.HeadPose.Pose, NewOrientation, NewPosition);
+	const FTransform NewRelativeTransform(NewOrientation, NewPosition);
+
+	Delegate->ApplyLateUpdate(ViewFamily.Scene, OldRelativeTransform, NewRelativeTransform);
 }
 
 void FGearVR::CalculateRenderTargetSize(const FViewport& Viewport, uint32& InOutSizeX, uint32& InOutSizeY)
diff --git a/Engine/Shaders/AtmosphereCommon.usf b/Engine/Shaders/AtmosphereCommon.usf
--- a/Engine/Shaders/AtmosphereCommon.usf
+++ b/Engine/Shaders/AtmosphereCommon.usf
@@ -49,7 +49,7 @@ const static int InscatterNuNum = 8;
 
 #define TRANSMITTANCE_NON_LINEAR			1
 #define INSCATTER_NON_LINEAR				1
-#define ATMOSPHERIC_TEXTURE_SAMPLE_FIX		1
+#define ATMOSPHERIC_TEXTURE_SAMPLE_FIX		(!COMPILER_METAL || !MAC) // @todo AMD & Intel's Metal shader compilers apply overly aggressive optimisations to this path & thus generate NaN when they shouldn't. They'll need to fix this in a subsequent version of OS X.
 
 #ifndef ATMOSPHERIC_NO_SUN_DISK
 #define	ATMOSPHERIC_NO_SUN_DISK				0
diff --git a/Engine/Shaders/Common.usf b/Engine/Shaders/Common.usf
--- a/Engine/Shaders/Common.usf
+++ b/Engine/Shaders/Common.usf
@@ -185,7 +185,7 @@
 #define USE_DEVELOPMENT_SHADERS (COMPILE_SHADERS_FOR_DEVELOPMENT && PLATFORM_SUPPORTS_DEVELOPMENT_SHADERS)
 
 #ifndef MOBILE_EMULATION
-#define MOBILE_EMULATION ((FEATURE_LEVEL == FEATURE_LEVEL_ES2 || FEATURE_LEVEL == FEATURE_LEVEL_ES3_1) && (!COMPILER_GLSL_ES2 && USE_DEVELOPMENT_SHADERS && !METAL_PROFILE && !VULKAN_PROFILE))
+#define MOBILE_EMULATION ((FEATURE_LEVEL == FEATURE_LEVEL_ES2 || FEATURE_LEVEL == FEATURE_LEVEL_ES3_1) && (!COMPILER_GLSL_ES2 && USE_DEVELOPMENT_SHADERS && !(METAL_PROFILE && MAC) && !VULKAN_PROFILE))
 #endif
 
 // If we didn't request color component swizzling, just make it empty
@@ -983,7 +983,7 @@ float CalcSceneDepth( float2 ScreenUV )
 			float SceneW = ConvertFromDeviceZ(Texture2DSampleLevel(SceneDepthTexture, SceneDepthTextureSampler, ScreenUV, 0).r);		
 			return DepthbufferFetchES2(SceneW, View.InvDeviceZToWorldZTransform[2], View.InvDeviceZToWorldZTransform[3]);
 		#endif 
-	#elif METAL_PROFILE
+	#elif METAL_PROFILE && !MAC
 			return FramebufferFetchES2().w;
 	#else
 		return ConvertFromDeviceZ(Texture2DSampleLevel(SceneDepthTexture, SceneDepthTextureSampler, ScreenUV, 0).r);
diff --git a/Engine/Shaders/MaterialTemplate.usf b/Engine/Shaders/MaterialTemplate.usf
--- a/Engine/Shaders/MaterialTemplate.usf
+++ b/Engine/Shaders/MaterialTemplate.usf
@@ -742,7 +742,7 @@ MaterialFloat2 GetParticleMacroUV(FMaterialPixelParameters Parameters)
 }
 
 #ifndef MOBILE_EMULATION
-#define MOBILE_EMULATION ((FEATURE_LEVEL == FEATURE_LEVEL_ES2 || FEATURE_LEVEL == FEATURE_LEVEL_ES3_1) && (!COMPILER_GLSL_ES2 && USE_DEVELOPMENT_SHADERS && !METAL_PROFILE))
+#define MOBILE_EMULATION ((FEATURE_LEVEL == FEATURE_LEVEL_ES2 || FEATURE_LEVEL == FEATURE_LEVEL_ES3_1) && (!COMPILER_GLSL_ES2 && USE_DEVELOPMENT_SHADERS && !(METAL_PROFILE && MAC)))
 #endif
 
 MaterialFloat4 ProcessMaterialColorTextureLookup(MaterialFloat4 TextureValue)
diff --git a/Engine/Shaders/PostProcessMobile.usf b/Engine/Shaders/PostProcessMobile.usf
--- a/Engine/Shaders/PostProcessMobile.usf
+++ b/Engine/Shaders/PostProcessMobile.usf
@@ -79,7 +79,7 @@ void SunMaskPS_ES2(
 	out half4 OutColor : SV_Target0
 	)
 {
-	#if (COMPILER_GLSL_ES2 && ES2_USE_FETCH) || METAL_PROFILE
+	#if (COMPILER_GLSL_ES2 && ES2_USE_FETCH) || (METAL_PROFILE && !MAC)
 		OutColor = FramebufferFetchES2();
 	#else
 		OutColor = PostprocessInput0.Sample(PostprocessInput0Sampler, InUVPos.xy);
diff --git a/Engine/Source/Developer/Apple/MetalShaderFormat/Private/MetalBackend.h b/Engine/Source/Developer/Apple/MetalShaderFormat/Private/MetalBackend.h
--- a/Engine/Source/Developer/Apple/MetalShaderFormat/Private/MetalBackend.h
+++ b/Engine/Source/Developer/Apple/MetalShaderFormat/Private/MetalBackend.h
@@ -3,8 +3,10 @@
 #pragma once
 
 #include "hlslcc.h"
+PRAGMA_DISABLE_SHADOW_VARIABLE_WARNINGS
 #include "ir.h"
 #include "PackUniformBuffers.h"
+PRAGMA_ENABLE_SHADOW_VARIABLE_WARNINGS
 #include "LanguageSpec.h"
 
 
diff --git a/Engine/Source/Developer/Apple/MetalShaderFormat/Private/MetalShaderCompiler.cpp b/Engine/Source/Developer/Apple/MetalShaderFormat/Private/MetalShaderCompiler.cpp
--- a/Engine/Source/Developer/Apple/MetalShaderFormat/Private/MetalShaderCompiler.cpp
+++ b/Engine/Source/Developer/Apple/MetalShaderFormat/Private/MetalShaderCompiler.cpp
@@ -521,6 +521,7 @@ void CompileShader_Metal(const FShaderCompilerInput& Input,FShaderCompilerOutput
 	static FName NAME_SF_METAL_MRT(TEXT("SF_METAL_MRT"));
 	static FName NAME_SF_METAL_SM4(TEXT("SF_METAL_SM4"));
 	static FName NAME_SF_METAL_SM5(TEXT("SF_METAL_SM5"));
+	static FName NAME_SF_METAL_MACES3_1(TEXT("SF_METAL_MACES3_1"));
 	
 	TCHAR const* Standard = TEXT("-std=ios-metal1.0");
 
@@ -532,6 +533,12 @@ void CompileShader_Metal(const FShaderCompilerInput& Input,FShaderCompilerOutput
 	{
 		AdditionalDefines.SetDefine(TEXT("METAL_MRT_PROFILE"), 1);
 	}
+	else if (Input.ShaderFormat == NAME_SF_METAL_MACES3_1)
+	{
+		AdditionalDefines.SetDefine(TEXT("METAL_PROFILE"), 1);
+		Standard = TEXT("-std=osx-metal1.1");
+		MetalCompilerTarget = HCT_FeatureLevelES3_1;
+	}
 	else if (Input.ShaderFormat == NAME_SF_METAL_SM4)
 	{
 		AdditionalDefines.SetDefine(TEXT("METAL_SM4_PROFILE"), 1);
diff --git a/Engine/Source/Developer/Apple/MetalShaderFormat/Private/MetalShaderFormat.cpp b/Engine/Source/Developer/Apple/MetalShaderFormat/Private/MetalShaderFormat.cpp
--- a/Engine/Source/Developer/Apple/MetalShaderFormat/Private/MetalShaderFormat.cpp
+++ b/Engine/Source/Developer/Apple/MetalShaderFormat/Private/MetalShaderFormat.cpp
@@ -12,6 +12,7 @@ static FName NAME_SF_METAL(TEXT("SF_METAL"));
 static FName NAME_SF_METAL_MRT(TEXT("SF_METAL_MRT"));
 static FName NAME_SF_METAL_SM4(TEXT("SF_METAL_SM4"));
 static FName NAME_SF_METAL_SM5(TEXT("SF_METAL_SM5"));
+static FName NAME_SF_METAL_MACES3_1(TEXT("SF_METAL_MACES3_1"));
 
 class FMetalShaderFormat : public IShaderFormat
 {
@@ -53,10 +54,11 @@ public:
 		OutFormats.Add(NAME_SF_METAL_MRT);
 		OutFormats.Add(NAME_SF_METAL_SM4);
 		OutFormats.Add(NAME_SF_METAL_SM5);
+		OutFormats.Add(NAME_SF_METAL_MACES3_1);
 	}
 	virtual void CompileShader(FName Format, const struct FShaderCompilerInput& Input, struct FShaderCompilerOutput& Output,const FString& WorkingDirectory) const
 	{
-		check(Format == NAME_SF_METAL || Format == NAME_SF_METAL_MRT || Format == NAME_SF_METAL_SM4 || Format == NAME_SF_METAL_SM5);
+		check(Format == NAME_SF_METAL || Format == NAME_SF_METAL_MRT || Format == NAME_SF_METAL_SM4 || Format == NAME_SF_METAL_SM5 || Format == NAME_SF_METAL_MACES3_1);
 		CompileShader_Metal(Input, Output, WorkingDirectory);
 	}
 };
diff --git a/Engine/Source/Developer/LogVisualizer/Private/SVisualLoggerTimelinesContainer.cpp b/Engine/Source/Developer/LogVisualizer/Private/SVisualLoggerTimelinesContainer.cpp
--- a/Engine/Source/Developer/LogVisualizer/Private/SVisualLoggerTimelinesContainer.cpp
+++ b/Engine/Source/Developer/LogVisualizer/Private/SVisualLoggerTimelinesContainer.cpp
@@ -199,7 +199,11 @@ FReply SVisualLoggerTimelinesContainer::OnKeyDown(const FGeometry& MyGeometry, c
 			ContainingBorder->RemoveSlot(CurrentNode.ToSharedRef());
 		}
 
-		SetSelectionState(NotSelectedOne.Pin(), true, true);
+		if (NotSelectedOne.IsValid())
+		{
+			SetSelectionState(NotSelectedOne.Pin(), true, true);
+		}
+
 		return FReply::Handled();
 	}
 	else if (InKeyEvent.GetKey() == EKeys::Up || InKeyEvent.GetKey() == EKeys::Down)
diff --git a/Engine/Source/Developer/LogVisualizer/Private/VisualLoggerDatabase.cpp b/Engine/Source/Developer/LogVisualizer/Private/VisualLoggerDatabase.cpp
--- a/Engine/Source/Developer/LogVisualizer/Private/VisualLoggerDatabase.cpp
+++ b/Engine/Source/Developer/LogVisualizer/Private/VisualLoggerDatabase.cpp
@@ -151,13 +151,17 @@ void FVisualLoggerDatabase::AddItem(const FVisualLogDevice::FVisualLogEntryItem&
 	else
 	{
 		RowIndex = RowNameToIndex[NewItem.OwnerName];
-		CurrentRow = &Rows[RowIndex];
+		if (ensure(Rows.IsValidIndex(RowIndex)))
+		{
+			CurrentRow = &Rows[RowIndex];
+		}
+	}
+	
+	if (CurrentRow)
+	{
+		CurrentRow->AddItem(NewItem);
+		FVisualLoggerGraphsDatabase::Get().AddItem(NewItem);
 	}
-	check(CurrentRow);
-
-	CurrentRow->AddItem(NewItem);
-
-	FVisualLoggerGraphsDatabase::Get().AddItem(NewItem);
 }
 
 bool FVisualLoggerDatabase::ContainsRowByName(FName InName)
@@ -255,7 +259,7 @@ void FVisualLoggerDatabase::RemoveRow(FName RowName)
 	}
 	if (RowNameToIndex.Contains(RowName))
 	{
-		const int32 RemovedIndex = RowNameToIndex[RowName];
+		const int32 RemovedIndex = RowNameToIndex.FindAndRemoveChecked(RowName);
 		Rows.RemoveAtSwap(RemovedIndex, 1, false);
 		if (Rows.IsValidIndex(RemovedIndex))
 		{
diff --git a/Engine/Source/Developer/Mac/MacTargetPlatform/Private/GenericMacTargetPlatform.h b/Engine/Source/Developer/Mac/MacTargetPlatform/Private/GenericMacTargetPlatform.h
--- a/Engine/Source/Developer/Mac/MacTargetPlatform/Private/GenericMacTargetPlatform.h
+++ b/Engine/Source/Developer/Mac/MacTargetPlatform/Private/GenericMacTargetPlatform.h
@@ -136,6 +136,8 @@ return TSuper::SupportsFeature(Feature);
 				OutFormats.AddUnique(NAME_SF_METAL_SM4);
 				static FName NAME_SF_METAL_SM5(TEXT("SF_METAL_SM5"));
 				OutFormats.AddUnique(NAME_SF_METAL_SM5);
+				static FName NAME_SF_METAL_MACES3_1(TEXT("SF_METAL_MACES3_1"));
+				OutFormats.AddUnique(NAME_SF_METAL_MACES3_1);
 			}
 #endif
 		}
diff --git a/Engine/Source/Developer/MeshUtilities/Private/MeshUtilities.cpp b/Engine/Source/Developer/MeshUtilities/Private/MeshUtilities.cpp
--- a/Engine/Source/Developer/MeshUtilities/Private/MeshUtilities.cpp
+++ b/Engine/Source/Developer/MeshUtilities/Private/MeshUtilities.cpp
@@ -3452,6 +3452,7 @@ public:
 	virtual FVector GetVertexPosition(uint32 WedgeIndex) = 0;
 	virtual FVector GetVertexPosition(uint32 FaceIndex, uint32 TriIndex) = 0;
 	virtual FVector2D GetVertexUV(uint32 FaceIndex, uint32 TriIndex, uint32 UVIndex) = 0;
+	virtual uint32 GetFaceSmoothingGroups(uint32 FaceIndex) = 0;
 
 	virtual uint32 GetNumFaces() = 0;
 	virtual uint32 GetNumWedges() = 0;
@@ -3543,6 +3544,11 @@ public:
 		return Wedges[Faces[FaceIndex].iWedge[TriIndex]].UVs[UVIndex];
 	}
 
+	virtual uint32 GetFaceSmoothingGroups(uint32 FaceIndex)
+	{
+		return Faces[FaceIndex].SmoothingGroups;
+	}
+
 	virtual uint32 GetNumFaces() override
 	{
 		return Faces.Num();
@@ -4215,7 +4221,7 @@ public:
 									if (!NextFace.bFilled) // && !NextFace.bBlendTangents)
 									{
 										if ((NextFaceIndex != OtherFaceIdx)
-											)//&& (RawMesh.FaceSmoothingMasks[NextFace.FaceIndex] & RawMesh.FaceSmoothingMasks[OtherFace.FaceIndex]))
+											 && (BuildData->GetFaceSmoothingGroups(NextFace.FaceIndex) & BuildData->GetFaceSmoothingGroups(OtherFace.FaceIndex)))
 										{
 											int32 CommonVertices = 0;
 											int32 CommonNormalVertices = 0;
diff --git a/Engine/Source/Editor/BlueprintGraph/Private/EdGraphSchema_K2.cpp b/Engine/Source/Editor/BlueprintGraph/Private/EdGraphSchema_K2.cpp
--- a/Engine/Source/Editor/BlueprintGraph/Private/EdGraphSchema_K2.cpp
+++ b/Engine/Source/Editor/BlueprintGraph/Private/EdGraphSchema_K2.cpp
@@ -5893,6 +5893,7 @@ UK2Node* UEdGraphSchema_K2::CreateSplitPinNode(UEdGraphPin* Pin, FKismetCompiler
 		{
 			UK2Node_MakeStruct* MakeStructNode = (CompilerContext ? CompilerContext->SpawnIntermediateNode<UK2Node_MakeStruct>(GraphNode, SourceGraph) : NewObject<UK2Node_MakeStruct>(Graph));
 			MakeStructNode->StructType = StructType;
+			MakeStructNode->bMadeAfterOverridePinRemoval = true;
 			SplitPinNode = MakeStructNode;
 		}
 		else
diff --git a/Engine/Source/Editor/BlueprintGraph/Private/K2Node_FormatText.cpp b/Engine/Source/Editor/BlueprintGraph/Private/K2Node_FormatText.cpp
--- a/Engine/Source/Editor/BlueprintGraph/Private/K2Node_FormatText.cpp
+++ b/Engine/Source/Editor/BlueprintGraph/Private/K2Node_FormatText.cpp
@@ -212,6 +212,7 @@ void UK2Node_FormatText::ExpandNode(class FKismetCompilerContext& CompilerContex
 		UK2Node_MakeStruct* PinMakeStruct = CompilerContext.SpawnIntermediateNode<UK2Node_MakeStruct>(this, SourceGraph);
 		PinMakeStruct->StructType = FormatArgumentDataStruct;
 		PinMakeStruct->AllocateDefaultPins();
+		PinMakeStruct->bMadeAfterOverridePinRemoval = true;
 
 		// Set the struct's "ArgumentName" pin literal to be the argument pin's name.
 		PinMakeStruct->GetSchema()->TrySetDefaultValue(*PinMakeStruct->FindPin("ArgumentName"), ArgumentPin->PinName);
diff --git a/Engine/Source/Editor/DetailCustomizations/Private/MacTargetSettingsDetails.cpp b/Engine/Source/Editor/DetailCustomizations/Private/MacTargetSettingsDetails.cpp
--- a/Engine/Source/Editor/DetailCustomizations/Private/MacTargetSettingsDetails.cpp
+++ b/Engine/Source/Editor/DetailCustomizations/Private/MacTargetSettingsDetails.cpp
@@ -227,6 +227,10 @@ FText GetFriendlyNameFromRHINameMac(const FString& InRHIName)
 	{
 		FriendlyRHIName = LOCTEXT("OpenGL3", "OpenGL 3 (SM4)");
 	}
+	else if (InRHIName == TEXT("SF_METAL_MACES3_1"))
+	{
+		FriendlyRHIName = LOCTEXT("MetalES3.1", "Metal (ES3.1, Mobile Preview)");
+	}
 	else if (InRHIName == TEXT("SF_METAL_SM4"))
 	{
 		FriendlyRHIName = LOCTEXT("MetalSM4", "Metal (SM4)");
diff --git a/Engine/Source/Editor/FontEditor/Private/SCompositeFontEditor.cpp b/Engine/Source/Editor/FontEditor/Private/SCompositeFontEditor.cpp
--- a/Engine/Source/Editor/FontEditor/Private/SCompositeFontEditor.cpp
+++ b/Engine/Source/Editor/FontEditor/Private/SCompositeFontEditor.cpp
@@ -159,6 +159,7 @@ void SCompositeFontEditor::FlushCachedFont()
 	if(CompositeFont)
 	{
 		CompositeFont->MakeDirty();
+		FSlateApplication::Get().GetRenderer()->GetFontCache()->FlushCompositeFont(*CompositeFont);
 	}
 
 	TSharedPtr<IFontEditor> FontEditor = FontEditorPtr.Pin();
diff --git a/Engine/Source/Editor/GraphEditor/Private/SGraphPanel.cpp b/Engine/Source/Editor/GraphEditor/Private/SGraphPanel.cpp
--- a/Engine/Source/Editor/GraphEditor/Private/SGraphPanel.cpp
+++ b/Engine/Source/Editor/GraphEditor/Private/SGraphPanel.cpp
@@ -1392,7 +1392,26 @@ void SGraphPanel::Update()
 			UEdGraphNode* Node = GraphObj->Nodes[NodeIndex];
 			if (Node)
 			{
-				AddNode(Node, CheckUserAddedNodesList);
+				// Helps detect cases of UE-26998 without causing a crash. Prevents the node from being rendered altogether and provides info on the state of the graph vs the node.
+				// Because the editor won't crash, a GLEO can be expected if the node's outer is in the transient package.
+				if (ensureMsgf(Node->GetOuter() == GraphObj, TEXT("Found %s ('%s') that does not belong to %s. Node Outer: %s, Node Outer Type: %s, Graph Outer: %s, Graph Outer Type: %s"),
+					*Node->GetName(), *Node->GetClass()->GetName(),
+					*GraphObj->GetName(),
+					*Node->GetOuter()->GetName(), *Node->GetOuter()->GetClass()->GetName(),
+					*GraphObj->GetOuter()->GetName(), *GraphObj->GetOuter()->GetClass()->GetName()
+					))
+ 				{
+					AddNode(Node, CheckUserAddedNodesList);
+				}
+				else
+				{
+					UE_LOG(LogGraphPanel, Error, TEXT("Found %s ('%s') that does not belong to %s. Node Outer: %s, Node Outer Type: %s, Graph Outer: %s, Graph Outer Type: %s"),
+						*Node->GetName(), *Node->GetClass()->GetName(),
+						*GraphObj->GetName(),
+						*Node->GetOuter()->GetName(), *Node->GetOuter()->GetClass()->GetName(),
+						*GraphObj->GetOuter()->GetName(), *GraphObj->GetOuter()->GetClass()->GetName()
+					);
+				}
 			}
 			else
 			{
diff --git a/Engine/Source/Editor/HierarchicalLODOutliner/Private/HLODOutliner.cpp b/Engine/Source/Editor/HierarchicalLODOutliner/Private/HLODOutliner.cpp
--- a/Engine/Source/Editor/HierarchicalLODOutliner/Private/HLODOutliner.cpp
+++ b/Engine/Source/Editor/HierarchicalLODOutliner/Private/HLODOutliner.cpp
@@ -124,7 +124,6 @@ namespace HLODOutliner
 			+ SVerticalBox::Slot()
 			.Padding(FMargin(0.0f, 5.0f))
 			[
-
 				SNew(SHorizontalBox)
 				+ SHorizontalBox::Slot()
 				.AutoWidth()
@@ -135,7 +134,7 @@ namespace HLODOutliner
 					.HAlign(HAlign_Center)
 					.Text(LOCTEXT("GenerateClusters", "Generate Clusters"))
 					.OnClicked(this, &SHLODOutliner::HandlePreviewHLODs)
-					.ToolTipText(LOCTEXT("GenerateClusterToolTip", "Generates Clusters for Meshes in the Level"))
+					.ToolTipText(LOCTEXT("GenerateClusterToolTip", "Generates Clusters (but not proxy meshes) for Meshes in the Level"))
 				]
 
 				+ SHorizontalBox::Slot()
diff --git a/Engine/Source/Editor/KismetCompiler/Private/KismetCompilerModule.cpp b/Engine/Source/Editor/KismetCompiler/Private/KismetCompilerModule.cpp
--- a/Engine/Source/Editor/KismetCompiler/Private/KismetCompilerModule.cpp
+++ b/Engine/Source/Editor/KismetCompiler/Private/KismetCompilerModule.cpp
@@ -177,6 +177,12 @@ void FKismet2CompilerModule::CompileBlueprint(class UBlueprint* Blueprint, const
 		FKismetCompilerOptions SkeletonCompileOptions;
 		SkeletonCompileOptions.CompileType = EKismetCompileType::SkeletonOnly;
 		CompileBlueprintInner(Blueprint, SkeletonCompileOptions, SkeletonResults, ParentReinstancer, ObjLoaded);
+
+		// Only when doing full compiles do we want to compile all skeletons before continuing to compile 
+		if (CompileOptions.CompileType == EKismetCompileType::Full)
+		{
+			SkeletonReinstancer->ReinstanceObjects();
+		}
 	}
 
 	// If this was a full compile, take appropriate actions depending on the success of failure of the compile
diff --git a/Engine/Source/Editor/KismetWidgets/Private/SPinTypeSelector.cpp b/Engine/Source/Editor/KismetWidgets/Private/SPinTypeSelector.cpp
--- a/Engine/Source/Editor/KismetWidgets/Private/SPinTypeSelector.cpp
+++ b/Engine/Source/Editor/KismetWidgets/Private/SPinTypeSelector.cpp
@@ -456,7 +456,7 @@ TSharedRef< SWidget > SPinTypeSelector::GetAllowedObjectTypes(FPinTypeTreeItem I
 
 void SPinTypeSelector::OnSelectPinType(FPinTypeTreeItem InItem, FString InPinCategory)
 {
-	const FScopedTransaction Transaction( LOCTEXT("ChangeParam", "Change Paramater Type") );
+	const FScopedTransaction Transaction( LOCTEXT("ChangeParam", "Change Parameter Type") );
 
 	FEdGraphPinType NewTargetPinType = TargetPinType.Get();
 	//Call delegate in order to notify pin type change is about to happen
diff --git a/Engine/Source/Editor/LandscapeEditor/Private/LandscapeEdModeMirrorTool.cpp b/Engine/Source/Editor/LandscapeEditor/Private/LandscapeEdModeMirrorTool.cpp
--- a/Engine/Source/Editor/LandscapeEditor/Private/LandscapeEdModeMirrorTool.cpp
+++ b/Engine/Source/Editor/LandscapeEditor/Private/LandscapeEdModeMirrorTool.cpp
@@ -103,55 +103,59 @@ public:
 
 	virtual void Render(const FSceneView* View, FViewport* Viewport, FPrimitiveDrawInterface* PDI) override
 	{
+		// The editor can try to render the tool before the UpdateLandscapeEditorData command runs and the landscape editor realizes that the landscape has been hidden/deleted
 		const ULandscapeInfo* const LandscapeInfo = EdMode->CurrentToolTarget.LandscapeInfo.Get();
 		const ALandscapeProxy* const LandscapeProxy = LandscapeInfo->GetLandscapeProxy();
-		const FTransform LandscapeToWorld = LandscapeProxy->LandscapeActorToWorld();
-
-		int32 MinX, MinY, MaxX, MaxY;
-		if (LandscapeInfo->GetLandscapeExtent(MinX, MinY, MaxX, MaxY))
+		if (LandscapeProxy)
 		{
-			FVector MirrorPoint3D = FVector((MaxX + MinX) / 2.0f, (MaxY + MinY) / 2.0f, 0);
-			FVector MirrorPlaneScale = FVector(0, 1, 100);
+			const FTransform LandscapeToWorld = LandscapeProxy->LandscapeActorToWorld();
 
-			if (EdMode->UISettings->MirrorOp == ELandscapeMirrorOperation::MinusXToPlusX ||
-				EdMode->UISettings->MirrorOp == ELandscapeMirrorOperation::PlusXToMinusX)
-			{
-				MirrorPoint3D.X = EdMode->UISettings->MirrorPoint.X;
-				MirrorPlaneScale.Y = (MaxY - MinY) / 2.0f;
-			}
-			else
+			int32 MinX, MinY, MaxX, MaxY;
+			if (LandscapeInfo->GetLandscapeExtent(MinX, MinY, MaxX, MaxY))
 			{
-				MirrorPoint3D.Y = EdMode->UISettings->MirrorPoint.Y;
-				MirrorPlaneScale.Y = (MaxX - MinX) / 2.0f;
-			}
+				FVector MirrorPoint3D = FVector((MaxX + MinX) / 2.0f, (MaxY + MinY) / 2.0f, 0);
+				FVector MirrorPlaneScale = FVector(0, 1, 100);
 
-			MirrorPoint3D.Z = GetLocalZAtPoint(LandscapeInfo, FMath::RoundToInt(MirrorPoint3D.X), FMath::RoundToInt(MirrorPoint3D.Y));
-			MirrorPoint3D = LandscapeToWorld.TransformPosition(MirrorPoint3D);
+				if (EdMode->UISettings->MirrorOp == ELandscapeMirrorOperation::MinusXToPlusX ||
+					EdMode->UISettings->MirrorOp == ELandscapeMirrorOperation::PlusXToMinusX)
+				{
+					MirrorPoint3D.X = EdMode->UISettings->MirrorPoint.X;
+					MirrorPlaneScale.Y = (MaxY - MinY) / 2.0f;
+				}
+				else
+				{
+					MirrorPoint3D.Y = EdMode->UISettings->MirrorPoint.Y;
+					MirrorPlaneScale.Y = (MaxX - MinX) / 2.0f;
+				}
 
-			FMatrix Matrix;
-			if (EdMode->UISettings->MirrorOp == ELandscapeMirrorOperation::MinusYToPlusY ||
-				EdMode->UISettings->MirrorOp == ELandscapeMirrorOperation::PlusYToMinusY)
-			{
-				Matrix = FScaleRotationTranslationMatrix(MirrorPlaneScale, FRotator(0, 90, 0), FVector::ZeroVector);
-			}
-			else
-			{
-				Matrix = FScaleMatrix(MirrorPlaneScale);
-			}
+				MirrorPoint3D.Z = GetLocalZAtPoint(LandscapeInfo, FMath::RoundToInt(MirrorPoint3D.X), FMath::RoundToInt(MirrorPoint3D.Y));
+				MirrorPoint3D = LandscapeToWorld.TransformPosition(MirrorPoint3D);
+
+				FMatrix Matrix;
+				if (EdMode->UISettings->MirrorOp == ELandscapeMirrorOperation::MinusYToPlusY ||
+					EdMode->UISettings->MirrorOp == ELandscapeMirrorOperation::PlusYToMinusY)
+				{
+					Matrix = FScaleRotationTranslationMatrix(MirrorPlaneScale, FRotator(0, 90, 0), FVector::ZeroVector);
+				}
+				else
+				{
+					Matrix = FScaleMatrix(MirrorPlaneScale);
+				}
 
-			Matrix *= LandscapeToWorld.ToMatrixWithScale();
-			Matrix.SetOrigin(MirrorPoint3D);
+				Matrix *= LandscapeToWorld.ToMatrixWithScale();
+				Matrix.SetOrigin(MirrorPoint3D);
 
-			// Convert plane from horizontal to vertical
-			Matrix = FMatrix(FVector(0, 1, 0), FVector(0, 0, 1), FVector(1, 0, 0), FVector(0, 0, 0)) * Matrix;
+				// Convert plane from horizontal to vertical
+				Matrix = FMatrix(FVector(0, 1, 0), FVector(0, 0, 1), FVector(1, 0, 0), FVector(0, 0, 0)) * Matrix;
 
-			const FBox Box = FBox(FVector(-1, -1, 0), FVector(+1, +1, 0));
-			DrawWireBox(PDI, Matrix, Box, FLinearColor::Green, SDPG_World);
+				const FBox Box = FBox(FVector(-1, -1, 0), FVector(+1, +1, 0));
+				DrawWireBox(PDI, Matrix, Box, FLinearColor::Green, SDPG_World);
 
-			const float LandscapeScaleRatio = LandscapeToWorld.GetScale3D().Z / LandscapeToWorld.GetScale3D().X;
-			FVector2D UVScale = FVector2D(FMath::RoundToFloat(MirrorPlaneScale.Y / 10), FMath::RoundToFloat(MirrorPlaneScale.Z * LandscapeScaleRatio / 10 / 2) * 2);
-			MirrorPlaneMaterial->SetVectorParameterValue(FName("GridSize"), FVector(UVScale, 0));
-			DrawPlane10x10(PDI, Matrix, 1, FVector2D(0, 0), FVector2D(1, 1), MirrorPlaneMaterial->GetRenderProxy(false), SDPG_World);
+				const float LandscapeScaleRatio = LandscapeToWorld.GetScale3D().Z / LandscapeToWorld.GetScale3D().X;
+				FVector2D UVScale = FVector2D(FMath::RoundToFloat(MirrorPlaneScale.Y / 10), FMath::RoundToFloat(MirrorPlaneScale.Z * LandscapeScaleRatio / 10 / 2) * 2);
+				MirrorPlaneMaterial->SetVectorParameterValue(FName("GridSize"), FVector(UVScale, 0));
+				DrawPlane10x10(PDI, Matrix, 1, FVector2D(0, 0), FVector2D(1, 1), MirrorPlaneMaterial->GetRenderProxy(false), SDPG_World);
+			}
 		}
 	}
 
diff --git a/Engine/Source/Editor/LandscapeEditor/Private/LandscapeEdModeRampTool.cpp b/Engine/Source/Editor/LandscapeEditor/Private/LandscapeEdModeRampTool.cpp
--- a/Engine/Source/Editor/LandscapeEditor/Private/LandscapeEdModeRampTool.cpp
+++ b/Engine/Source/Editor/LandscapeEditor/Private/LandscapeEdModeRampTool.cpp
@@ -310,9 +310,10 @@ public:
 
 	virtual void Render(const FSceneView* View, FViewport* Viewport, FPrimitiveDrawInterface* PDI) override
 	{
-		if (NumPoints > 0)
+		// The editor can try to render the tool before the UpdateLandscapeEditorData command runs and the landscape editor realizes that the landscape has been hidden/deleted
+		const ALandscapeProxy* LandscapeProxy = EdMode->CurrentToolTarget.LandscapeInfo->GetLandscapeProxy();
+		if (LandscapeProxy && NumPoints > 0)
 		{
-			const ALandscapeProxy* LandscapeProxy = EdMode->CurrentToolTarget.LandscapeInfo->GetLandscapeProxy();
 			const FTransform LandscapeToWorld = LandscapeProxy->LandscapeActorToWorld();
 
 			const FLinearColor SelectedSpriteColor = FLinearColor::White + (GEngine->GetSelectedMaterialColor() * GEngine->SelectionHighlightIntensity * 10);
@@ -403,7 +404,7 @@ public:
 	{
 		if (SelectedPoint != INDEX_NONE)
 		{
-			// The editor can try to render the transform widget before the landscape editor ticks and realises that the landscape has been hidden/deleted
+			// The editor can try to render the transform widget before the landscape editor ticks and realizes that the landscape has been hidden/deleted
 			const ALandscapeProxy* LandscapeProxy = EdMode->CurrentToolTarget.LandscapeInfo->GetLandscapeProxy();
 			if (LandscapeProxy)
 			{
diff --git a/Engine/Source/Editor/LandscapeEditor/Private/LandscapeEdModeSplineTools.cpp b/Engine/Source/Editor/LandscapeEditor/Private/LandscapeEdModeSplineTools.cpp
--- a/Engine/Source/Editor/LandscapeEditor/Private/LandscapeEdModeSplineTools.cpp
+++ b/Engine/Source/Editor/LandscapeEditor/Private/LandscapeEdModeSplineTools.cpp
@@ -1484,51 +1484,56 @@ public:
 
 	virtual void Render(const FSceneView* View, FViewport* Viewport, FPrimitiveDrawInterface* PDI) override
 	{
-		for (ULandscapeSplineControlPoint* ControlPoint : SelectedSplineControlPoints)
+		// The editor can try to render the tool before the UpdateLandscapeEditorData command runs and the landscape editor realizes that the landscape has been hidden/deleted
+		const ALandscapeProxy* LandscapeProxy = EdMode->CurrentToolTarget.LandscapeInfo->GetLandscapeProxy();
+		if (LandscapeProxy)
 		{
-			const ULandscapeSplinesComponent* SplinesComponent = ControlPoint->GetOuterULandscapeSplinesComponent();
+			for (ULandscapeSplineControlPoint* ControlPoint : SelectedSplineControlPoints)
+			{
+				const ULandscapeSplinesComponent* SplinesComponent = ControlPoint->GetOuterULandscapeSplinesComponent();
 
-			FVector HandlePos0 = SplinesComponent->ComponentToWorld.TransformPosition(ControlPoint->Location + ControlPoint->Rotation.Vector() * -20);
-			FVector HandlePos1 = SplinesComponent->ComponentToWorld.TransformPosition(ControlPoint->Location + ControlPoint->Rotation.Vector() * 20);
-			DrawDashedLine(PDI, HandlePos0, HandlePos1, FColor::White, 20, SDPG_Foreground);
+				FVector HandlePos0 = SplinesComponent->ComponentToWorld.TransformPosition(ControlPoint->Location + ControlPoint->Rotation.Vector() * -20);
+				FVector HandlePos1 = SplinesComponent->ComponentToWorld.TransformPosition(ControlPoint->Location + ControlPoint->Rotation.Vector() * 20);
+				DrawDashedLine(PDI, HandlePos0, HandlePos1, FColor::White, 20, SDPG_Foreground);
 
-			if (GLevelEditorModeTools().GetWidgetMode() == FWidget::WM_Scale)
-			{
-				for (const FLandscapeSplineConnection& Connection : ControlPoint->ConnectedSegments)
+				if (GLevelEditorModeTools().GetWidgetMode() == FWidget::WM_Scale)
 				{
-					FVector StartLocation; FRotator StartRotation;
-					Connection.GetNearConnection().ControlPoint->GetConnectionLocationAndRotation(Connection.GetNearConnection().SocketName, StartLocation, StartRotation);
+					for (const FLandscapeSplineConnection& Connection : ControlPoint->ConnectedSegments)
+					{
+						FVector StartLocation; FRotator StartRotation;
+						Connection.GetNearConnection().ControlPoint->GetConnectionLocationAndRotation(Connection.GetNearConnection().SocketName, StartLocation, StartRotation);
 
-					FVector StartPos = SplinesComponent->ComponentToWorld.TransformPosition(StartLocation);
-					FVector HandlePos = SplinesComponent->ComponentToWorld.TransformPosition(StartLocation + StartRotation.Vector() * Connection.GetNearConnection().TangentLen / 2);
-					PDI->DrawLine(StartPos, HandlePos, FColor::White, SDPG_Foreground);
+						FVector StartPos = SplinesComponent->ComponentToWorld.TransformPosition(StartLocation);
+						FVector HandlePos = SplinesComponent->ComponentToWorld.TransformPosition(StartLocation + StartRotation.Vector() * Connection.GetNearConnection().TangentLen / 2);
+						PDI->DrawLine(StartPos, HandlePos, FColor::White, SDPG_Foreground);
 
-					if (PDI->IsHitTesting()) PDI->SetHitProxy(new HLandscapeSplineProxy_Tangent(Connection.Segment, Connection.End));
-					PDI->DrawPoint(HandlePos, FColor::White, 10.0f, SDPG_Foreground);
-					if (PDI->IsHitTesting()) PDI->SetHitProxy(NULL);
+						if (PDI->IsHitTesting()) PDI->SetHitProxy(new HLandscapeSplineProxy_Tangent(Connection.Segment, Connection.End));
+						PDI->DrawPoint(HandlePos, FColor::White, 10.0f, SDPG_Foreground);
+						if (PDI->IsHitTesting()) PDI->SetHitProxy(NULL);
+					}
 				}
 			}
-		}
 
-		if (GLevelEditorModeTools().GetWidgetMode() == FWidget::WM_Scale)
-		{
-			for (ULandscapeSplineSegment* Segment : SelectedSplineSegments)
+			if (GLevelEditorModeTools().GetWidgetMode() == FWidget::WM_Scale)
 			{
-				const ULandscapeSplinesComponent* SplinesComponent = Segment->GetOuterULandscapeSplinesComponent();
-				for (int32 End = 0; End <= 1; End++)
+				for (ULandscapeSplineSegment* Segment : SelectedSplineSegments)
 				{
-					const FLandscapeSplineSegmentConnection& Connection = Segment->Connections[End];
+					const ULandscapeSplinesComponent* SplinesComponent = Segment->GetOuterULandscapeSplinesComponent();
+					for (int32 End = 0; End <= 1; End++)
+					{
+						const FLandscapeSplineSegmentConnection& Connection = Segment->Connections[End];
 
-					FVector StartLocation; FRotator StartRotation;
-					Connection.ControlPoint->GetConnectionLocationAndRotation(Connection.SocketName, StartLocation, StartRotation);
+						FVector StartLocation; FRotator StartRotation;
+						Connection.ControlPoint->GetConnectionLocationAndRotation(Connection.SocketName, StartLocation, StartRotation);
 
-					FVector EndPos = SplinesComponent->ComponentToWorld.TransformPosition(StartLocation);
-					FVector EndHandlePos = SplinesComponent->ComponentToWorld.TransformPosition(StartLocation + StartRotation.Vector() * Connection.TangentLen / 2);
+						FVector EndPos = SplinesComponent->ComponentToWorld.TransformPosition(StartLocation);
+						FVector EndHandlePos = SplinesComponent->ComponentToWorld.TransformPosition(StartLocation + StartRotation.Vector() * Connection.TangentLen / 2);
 
-					PDI->DrawLine(EndPos, EndHandlePos, FColor::White, SDPG_Foreground);
-					if (PDI->IsHitTesting()) PDI->SetHitProxy(new HLandscapeSplineProxy_Tangent(Segment, !!End));
-					PDI->DrawPoint(EndHandlePos, FColor::White, 10.0f, SDPG_Foreground);
-					if (PDI->IsHitTesting()) PDI->SetHitProxy(NULL);
+						PDI->DrawLine(EndPos, EndHandlePos, FColor::White, SDPG_Foreground);
+						if (PDI->IsHitTesting()) PDI->SetHitProxy(new HLandscapeSplineProxy_Tangent(Segment, !!End));
+						PDI->DrawPoint(EndHandlePos, FColor::White, 10.0f, SDPG_Foreground);
+						if (PDI->IsHitTesting()) PDI->SetHitProxy(NULL);
+					}
 				}
 			}
 		}
@@ -1554,7 +1559,12 @@ public:
 	{
 		if (SelectedSplineControlPoints.Num() > 0)
 		{
-			return true;
+			// The editor can try to render the transform widget before the landscape editor ticks and realizes that the landscape has been hidden/deleted
+			const ALandscapeProxy* LandscapeProxy = EdMode->CurrentToolTarget.LandscapeInfo->GetLandscapeProxy();
+			if (LandscapeProxy)
+			{
+				return true;
+			}
 		}
 
 		return false;
@@ -1587,9 +1597,13 @@ public:
 	{
 		if (SelectedSplineControlPoints.Num() > 0)
 		{
-			ULandscapeSplineControlPoint* FirstPoint = *SelectedSplineControlPoints.CreateConstIterator();
-			ULandscapeSplinesComponent* SplinesComponent = FirstPoint->GetOuterULandscapeSplinesComponent();
-			return SplinesComponent->ComponentToWorld.TransformPosition(FirstPoint->Location);
+			const ALandscapeProxy* LandscapeProxy = EdMode->CurrentToolTarget.LandscapeInfo->GetLandscapeProxy();
+			if (LandscapeProxy)
+			{
+				ULandscapeSplineControlPoint* FirstPoint = *SelectedSplineControlPoints.CreateConstIterator();
+				ULandscapeSplinesComponent* SplinesComponent = FirstPoint->GetOuterULandscapeSplinesComponent();
+				return SplinesComponent->ComponentToWorld.TransformPosition(FirstPoint->Location);
+			}
 		}
 
 		return FVector::ZeroVector;
@@ -1599,9 +1613,13 @@ public:
 	{
 		if (SelectedSplineControlPoints.Num() > 0)
 		{
-			ULandscapeSplineControlPoint* FirstPoint = *SelectedSplineControlPoints.CreateConstIterator();
-			ULandscapeSplinesComponent* SplinesComponent = FirstPoint->GetOuterULandscapeSplinesComponent();
-			return FQuatRotationTranslationMatrix(FirstPoint->Rotation.Quaternion() * SplinesComponent->ComponentToWorld.GetRotation(), FVector::ZeroVector);
+			const ALandscapeProxy* LandscapeProxy = EdMode->CurrentToolTarget.LandscapeInfo->GetLandscapeProxy();
+			if (LandscapeProxy)
+			{
+				ULandscapeSplineControlPoint* FirstPoint = *SelectedSplineControlPoints.CreateConstIterator();
+				ULandscapeSplinesComponent* SplinesComponent = FirstPoint->GetOuterULandscapeSplinesComponent();
+				return FQuatRotationTranslationMatrix(FirstPoint->Rotation.Quaternion() * SplinesComponent->ComponentToWorld.GetRotation(), FVector::ZeroVector);
+			}
 		}
 
 		return FMatrix::Identity;
diff --git a/Engine/Source/Editor/Persona/Private/SSkeletonTree.cpp b/Engine/Source/Editor/Persona/Private/SSkeletonTree.cpp
--- a/Engine/Source/Editor/Persona/Private/SSkeletonTree.cpp
+++ b/Engine/Source/Editor/Persona/Private/SSkeletonTree.cpp
@@ -1790,6 +1790,8 @@ void SSkeletonTree::RemoveFromLOD(int32 LODIndex, bool bIncludeSelected, bool bI
 
 		// remove from current LOD
 		MeshUtilities.RemoveBonesFromMesh(PreviewComponent->SkeletalMesh, LODIndex, &BonesToRemove);
+		// update UI to reflect the change
+		OnLODSwitched();
 	}
 }
 
diff --git a/Engine/Source/Editor/UMGEditor/Private/WidgetBlueprintEditorUtils.cpp b/Engine/Source/Editor/UMGEditor/Private/WidgetBlueprintEditorUtils.cpp
--- a/Engine/Source/Editor/UMGEditor/Private/WidgetBlueprintEditorUtils.cpp
+++ b/Engine/Source/Editor/UMGEditor/Private/WidgetBlueprintEditorUtils.cpp
@@ -85,8 +85,16 @@ bool FWidgetBlueprintEditorUtils::VerifyWidgetRename(TSharedRef<class FWidgetBlu
 		return false;
 	}
 
+	UWidget* RenamedTemplateWidget = Widget.GetTemplate();
+	if ( !RenamedTemplateWidget )
+	{
+		// In certain situations, the template might be lost due to mid recompile with focus lost on the rename box at
+		// during a strange moment.
+		return false;
+	}
+
 	// Slug the new name down to a valid object name
-	const FName NewNameSlug = MakeObjectNameFromDisplayLabel(NewNameString, Widget.GetTemplate()->GetFName());
+	const FName NewNameSlug = MakeObjectNameFromDisplayLabel(NewNameString, RenamedTemplateWidget->GetFName());
 
 	UWidgetBlueprint* Blueprint = BlueprintEditor->GetWidgetBlueprintObj();
 	UWidget* ExistingTemplate = Blueprint->WidgetTree->FindWidget(NewNameSlug);
@@ -94,7 +102,7 @@ bool FWidgetBlueprintEditorUtils::VerifyWidgetRename(TSharedRef<class FWidgetBlu
 	bool bIsSameWidget = false;
 	if (ExistingTemplate != nullptr)
 	{
-		if (Widget.GetTemplate() != ExistingTemplate)
+		if ( RenamedTemplateWidget != ExistingTemplate )
 		{
 			OutErrorMessage = LOCTEXT("ExistingWidgetName", "Existing Widget Name");
 			return false;
@@ -117,7 +125,7 @@ bool FWidgetBlueprintEditorUtils::VerifyWidgetRename(TSharedRef<class FWidgetBlu
 				return false;
 			}
 		}
-		UWidget* WidgetTemplate = Widget.GetTemplate();
+		UWidget* WidgetTemplate = RenamedTemplateWidget;
 		if (WidgetTemplate)
 		{
 			// Dummy rename with flag REN_Test returns if rename is possible
diff --git a/Engine/Source/Editor/UnrealEd/Private/CookOnTheFlyServer.cpp b/Engine/Source/Editor/UnrealEd/Private/CookOnTheFlyServer.cpp
--- a/Engine/Source/Editor/UnrealEd/Private/CookOnTheFlyServer.cpp
+++ b/Engine/Source/Editor/UnrealEd/Private/CookOnTheFlyServer.cpp
@@ -1619,8 +1619,11 @@ uint32 UCookOnTheFlyServer::TickCookOnTheSide( const float TimeSlice, uint32 &Co
 			// did not cook this package 
 #if DO_CHECK
 			// make sure this package doesn't exist
-			const FString SandboxFilename = ConvertToFullSandboxPath(ToBuild.GetFilename().ToString(), false);
-			check(IFileManager::Get().FileExists(*SandboxFilename) == false);
+			for (const auto& TargetPlatformName : ToBuild.GetPlatformnames())
+			{
+				const FString SandboxFilename = ConvertToFullSandboxPath(ToBuild.GetFilename().ToString(), true, TargetPlatformName.ToString());
+				check(IFileManager::Get().FileExists(*SandboxFilename) == false);
+			}
 #endif
 			CookedPackages.Add( FFilePlatformCookedPackage( ToBuild.GetFilename(), TargetPlatformNames, false) );
 			continue;
@@ -3206,6 +3209,9 @@ void UCookOnTheFlyServer::CleanSandbox( const bool bIterative )
 				SandboxToPackage.Add(PackageSandbox.Replace(*SandboxFile->GetSandboxDirectory(), TEXT("")), PackageFile);
 			}
 
+			// Daniel: optimization
+			// should I loop through the cooked packages which we think we have instead of going through what's on disk?
+
 			// See what files are out of date in the sandbox folder
 			for (int32 Index = 0; Index < Platforms.Num(); Index++)
 			{
@@ -3263,7 +3269,7 @@ void UCookOnTheFlyServer::CleanSandbox( const bool bIterative )
 #if DEBUG_COOKONTHEFLY
 							UE_LOG(LogCook, Display, TEXT("Deleting out of date cooked file: %s"), *CookedFilename);
 #endif
-							IFileManager::Get().Delete(*CookedFilename);
+							IFileManager::Get().Delete(*CookedFilename); // do I need to delete this? the cooked packages list is what the cooker uses to recook this file should be deleted later
 
 							CookedPackages.RemoveFileForPlatform(StandardUnCookedFileFName, PlatformFName);
 						}
diff --git a/Engine/Source/Editor/UnrealEd/Private/Factories/SkeletalMeshImport.cpp b/Engine/Source/Editor/UnrealEd/Private/Factories/SkeletalMeshImport.cpp
--- a/Engine/Source/Editor/UnrealEd/Private/Factories/SkeletalMeshImport.cpp
+++ b/Engine/Source/Editor/UnrealEd/Private/Factories/SkeletalMeshImport.cpp
@@ -129,6 +129,8 @@ void FSkeletalMeshImportData::CopyLODImportData(
 		Face.TangentZ[1]		= Faces[f].TangentZ[1];
 		Face.TangentZ[2]		= Faces[f].TangentZ[2];
 
+		Face.SmoothingGroups    = Faces[f].SmoothingGroups;
+
 		LODFaces[f] = Face;
 	}			
 
diff --git a/Engine/Source/Editor/UnrealEd/Private/Fbx/FbxSkeletalMeshImport.cpp b/Engine/Source/Editor/UnrealEd/Private/Fbx/FbxSkeletalMeshImport.cpp
--- a/Engine/Source/Editor/UnrealEd/Private/Fbx/FbxSkeletalMeshImport.cpp
+++ b/Engine/Source/Editor/UnrealEd/Private/Fbx/FbxSkeletalMeshImport.cpp
@@ -1300,7 +1300,7 @@ USkeletalMesh* UnFbx::FFbxImporter::ImportSkeletalMesh(UObject* InParent, TArray
 		BuildOptions.bKeepOverlappingVertices = ImportOptions->bKeepOverlappingVertices;
 		BuildOptions.bComputeNormals = !ImportOptions->ShouldImportNormals() || !SkelMeshImportDataPtr->bHasNormals;
 		BuildOptions.bComputeTangents = !ImportOptions->ShouldImportTangents() || !SkelMeshImportDataPtr->bHasTangents;
-		BuildOptions.bUseMikkTSpace = (ImportOptions->NormalGenerationMethod == EFBXNormalGenerationMethod::MikkTSpace);
+		BuildOptions.bUseMikkTSpace = (ImportOptions->NormalGenerationMethod == EFBXNormalGenerationMethod::MikkTSpace) && (!ImportOptions->ShouldImportNormals() || !ImportOptions->ShouldImportTangents());
 		BuildOptions.bRemoveDegenerateTriangles = false;
 
 		IMeshUtilities& MeshUtilities = FModuleManager::Get().LoadModuleChecked<IMeshUtilities>("MeshUtilities");
diff --git a/Engine/Source/Editor/UnrealEd/Private/Fbx/FbxStaticMeshImport.cpp b/Engine/Source/Editor/UnrealEd/Private/Fbx/FbxStaticMeshImport.cpp
--- a/Engine/Source/Editor/UnrealEd/Private/Fbx/FbxStaticMeshImport.cpp
+++ b/Engine/Source/Editor/UnrealEd/Private/Fbx/FbxStaticMeshImport.cpp
@@ -1317,7 +1317,7 @@ UStaticMesh* UnFbx::FFbxImporter::ImportStaticMeshAsSingle(UObject* InParent, TA
 		SrcModel.BuildSettings.bBuildReversedIndexBuffer = ImportOptions->bBuildReversedIndexBuffer;
 		SrcModel.BuildSettings.bRecomputeNormals = ImportOptions->NormalImportMethod == FBXNIM_ComputeNormals;
 		SrcModel.BuildSettings.bRecomputeTangents = ImportOptions->NormalImportMethod != FBXNIM_ImportNormalsAndTangents;
-		SrcModel.BuildSettings.bUseMikkTSpace = ImportOptions->NormalGenerationMethod == EFBXNormalGenerationMethod::MikkTSpace;
+		SrcModel.BuildSettings.bUseMikkTSpace = (ImportOptions->NormalGenerationMethod == EFBXNormalGenerationMethod::MikkTSpace) && (!ImportOptions->ShouldImportNormals() || !ImportOptions->ShouldImportTangents());
 		if( ImportOptions->bGenerateLightmapUVs )
 		{
 			SrcModel.BuildSettings.bGenerateLightmapUVs = true;
diff --git a/Engine/Source/Editor/UnrealEd/Private/Kismet2/ComponentEditorUtils.cpp b/Engine/Source/Editor/UnrealEd/Private/Kismet2/ComponentEditorUtils.cpp
--- a/Engine/Source/Editor/UnrealEd/Private/Kismet2/ComponentEditorUtils.cpp
+++ b/Engine/Source/Editor/UnrealEd/Private/Kismet2/ComponentEditorUtils.cpp
@@ -631,7 +631,10 @@ void FComponentEditorUtils::AdjustComponentDelta(USceneComponent* Component, FVe
 
 		if (!Component->bAbsoluteLocation)
 		{
-			Drag = ParentToWorldSpace.Inverse().TransformVector(Drag);
+			//transform the drag vector in relative to the parent transform
+			Drag = ParentToWorldSpace.InverseTransformVectorNoScale(Drag);
+			//Now that we have a global drag we can apply the parent scale
+			Drag = Drag * ParentToWorldSpace.Inverse().GetScale3D();
 		}
 
 		if (!Component->bAbsoluteRotation)
diff --git a/Engine/Source/Editor/UnrealEd/Private/Kismet2/Kismet2.cpp b/Engine/Source/Editor/UnrealEd/Private/Kismet2/Kismet2.cpp
--- a/Engine/Source/Editor/UnrealEd/Private/Kismet2/Kismet2.cpp
+++ b/Engine/Source/Editor/UnrealEd/Private/Kismet2/Kismet2.cpp
@@ -474,6 +474,9 @@ UBlueprint* FKismetEditorUtilities::CreateBlueprint(UClass* ParentClass, UObject
 	FKismetCompilerOptions CompileOptions;
 	Compiler.CompileBlueprint(NewBP, CompileOptions, Results);
 
+	// Mark the BP as being regenerated, so it will not be confused as needing to be loaded and regenerated when a referenced BP loads.
+	NewBP->bHasBeenRegenerated = true;
+
 	UBlueprintEditorSettings* Settings = GetMutableDefault<UBlueprintEditorSettings>();
 	if(Settings && Settings->bSpawnDefaultBlueprintNodes)
 	{
diff --git a/Engine/Source/Editor/UnrealEd/Private/Kismet2/KismetDebugUtilities.cpp b/Engine/Source/Editor/UnrealEd/Private/Kismet2/KismetDebugUtilities.cpp
--- a/Engine/Source/Editor/UnrealEd/Private/Kismet2/KismetDebugUtilities.cpp
+++ b/Engine/Source/Editor/UnrealEd/Private/Kismet2/KismetDebugUtilities.cpp
@@ -43,7 +43,7 @@ public:
 	int32 TargetGraphStackDepth;
 
 	// The last message that an exception delivered
-	FString LastExceptionMessage;
+	FText LastExceptionMessage;
 
 	// Only valid inside intraframe debugging
 	const FFrame* StackFrameAtIntraframeDebugging;
@@ -142,16 +142,13 @@ void FKismetDebugUtilities::OnScriptException(const UObject* ActiveObject, const
 		case EBlueprintExceptionType::AccessViolation:
 			if ( GIsEditor && GIsPlayInEditorWorld )
 			{
-				// a line of text classifying where the exception is coming from (is it a node?, function?, etc.)
-				TSharedRef<FTextToken> TraceClassificationMessageLog = FTextToken::Create(FText::FromString(TEXT("from code generated-function")));
-				
+				TSharedRef<FTokenizedMessage> ErrorMessage = FMessageLog("PIE").Error(LOCTEXT("RuntimeErrorMessage", "Blueprint Runtime Error from function:"));
+
 				// NOTE: StackFrame.Node is not a blueprint node like you may think ("Node" has some legacy meaning)
 				FString GeneratedFuncName = FString::Printf(TEXT("'%s'"), *StackFrame.Node->GetName());
 				// a log token, telling us specifically where the exception is coming from (here
 				// it's not helpful to link to a generated-function, so we just provide the plain name)
-				TSharedRef<IMessageToken> TraceTargetMessageLog = FTextToken::Create(FText::FromString(GeneratedFuncName));
-
-				FString MsgInBlueprintStr(TEXT("in blueprint"));
+				ErrorMessage->AddToken(FTextToken::Create(FText::FromString(GeneratedFuncName)));
 
 #if WITH_EDITORONLY_DATA // to protect access to GeneratedClass->DebugData
 				UBlueprintGeneratedClass* GeneratedClass = Cast<UBlueprintGeneratedClass>(ClassContainingCode);
@@ -161,22 +158,33 @@ void FKismetDebugUtilities::OnScriptException(const UObject* ActiveObject, const
 					// if instead, there is a node we can point to...
 					if (BlueprintNode != NULL)
 					{
-						TraceClassificationMessageLog = FTextToken::Create(FText::FromString(TEXT("from node")));
+						ErrorMessage->AddToken(FTextToken::Create(LOCTEXT("RuntimeErrorBlueprintNode", "from node:")));
 
 						FText NodeTitle = BlueprintNode->GetNodeTitle(ENodeTitleType::ListView); // a more user friendly name
 						// link to the last executed node (the one throwing the exception, presumably)
-						TraceTargetMessageLog = FUObjectToken::Create(BlueprintNode, NodeTitle)->OnMessageTokenActivated(FOnMessageTokenActivated::CreateStatic(&Local::OnMessageLogLinkActivated));
+						ErrorMessage->AddToken(
+							FUObjectToken::Create(
+								BlueprintNode, 
+								NodeTitle
+							)->OnMessageTokenActivated(FOnMessageTokenActivated::CreateStatic(&Local::OnMessageLogLinkActivated))
+						);
 
-						MsgInBlueprintStr = FString::Printf(TEXT("in graph '%s' in blueprint"), *GetNameSafe(BlueprintNode->GetGraph()));
+						ErrorMessage->AddToken( FTextToken::Create( LOCTEXT("RuntimeErrorBlueprintGraph", "in graph:") ) );
+						ErrorMessage->AddToken(
+							FUObjectToken::Create(
+								BlueprintNode->GetGraph(), 
+								FText::FromString(GetNameSafe(BlueprintNode->GetGraph()))
+							)->OnMessageTokenActivated(FOnMessageTokenActivated::CreateStatic(&Local::OnMessageLogLinkActivated)));
 					}
 				}
 #endif // WITH_EDITORONLY_DATA
 
-				FMessageLog("PIE").Error(FText::FromString(Info.GetDescription()))
-					->AddToken(TraceClassificationMessageLog)
-					->AddToken(TraceTargetMessageLog)
-					->AddToken(FTextToken::Create(FText::FromString(MsgInBlueprintStr)))
+				ErrorMessage->AddToken(FTextToken::Create(LOCTEXT("RuntimeErrorBlueprintObject", "in object:")));
+				ErrorMessage
 					->AddToken(FUObjectToken::Create(BlueprintObj, FText::FromString(BlueprintObj->GetName()))->OnMessageTokenActivated(FOnMessageTokenActivated::CreateStatic(&Local::OnMessageLogLinkActivated)));
+
+				ErrorMessage->AddToken(FTextToken::Create(LOCTEXT("RuntimeErrorBlueprintDescription", "with description:")));
+				ErrorMessage->AddToken(FTextToken::Create(Info.GetDescription()));
 			}
 			bForceToCurrentObject = true;
 			bShouldBreakExecution = GetDefault<UEditorExperimentalSettings>()->bBreakOnExceptions;
@@ -416,7 +424,7 @@ void FKismetDebugUtilities::OnScriptException(const UObject* ActiveObject, const
 		switch (Info.GetType())
 		{
 		case EBlueprintExceptionType::FatalError:
-			DisplayErrorLambda(LOCTEXT("FatalErrorType", "Fatal Error"), *Info.GetDescription());
+			DisplayErrorLambda(LOCTEXT("FatalErrorType", "Fatal Error"), *(Info.GetDescription().ToString()));
 			break;
 		case EBlueprintExceptionType::InfiniteLoop:
 			DisplayErrorLambda(LOCTEXT("InfiniteLoopErrorType", "Infinite Loop"), nullptr);
@@ -430,21 +438,15 @@ void FKismetDebugUtilities::OnScriptException(const UObject* ActiveObject, const
 
 UClass* FKismetDebugUtilities::FindClassForNode(const UObject* Object, UFunction* Function)
 {
+	if (NULL != Function)
+	{
+		UClass* FunctionOwner = Function->GetOwnerClass();
+		return FunctionOwner;
+	}
 	if(NULL != Object)
 	{
 		UClass* ObjClass = Object->GetClass();
-		if(NULL != Function)
-		{
-			UClass* FunctionOwner = Function->GetOwnerClass();
-			if(ensure(ObjClass->IsChildOf(FunctionOwner)))
-			{
-				return FunctionOwner;
-			}
-		}
-		else
-		{
-			return ObjClass;
-		}
+		return ObjClass;
 	}
 	return NULL;
 }	
@@ -1120,9 +1122,9 @@ FKismetDebugUtilities::EWatchTextResult FKismetDebugUtilities::GetWatchText(FStr
 FText FKismetDebugUtilities::GetAndClearLastExceptionMessage()
 {
 	FKismetDebugUtilitiesData& Data = FKismetDebugUtilitiesData::Get();
-	const FString Result = Data.LastExceptionMessage;
-	Data.LastExceptionMessage.Empty();
-	return FText::FromString(Result);
+	const FText Result = Data.LastExceptionMessage;
+	Data.LastExceptionMessage = FText();
+	return Result;
 }
 
 #undef LOCTEXT_NAMESPACE
diff --git a/Engine/Source/Editor/UnrealEd/Private/Kismet2/KismetReinstanceUtilities.cpp b/Engine/Source/Editor/UnrealEd/Private/Kismet2/KismetReinstanceUtilities.cpp
--- a/Engine/Source/Editor/UnrealEd/Private/Kismet2/KismetReinstanceUtilities.cpp
+++ b/Engine/Source/Editor/UnrealEd/Private/Kismet2/KismetReinstanceUtilities.cpp
@@ -514,7 +514,14 @@ void FBlueprintCompileReinstancer::CompileChildren()
 		{
 			ReparentChild(BP);
 
-			FKismetEditorUtilities::CompileBlueprint(BP, false, bSkipGarbageCollection);
+			// Full compiles first recompile all skeleton classes, so they are up-to-date
+			const bool bSkeletonUpToDate = true;
+			FKismetEditorUtilities::CompileBlueprint(BP, false, bSkipGarbageCollection, false, nullptr, bSkeletonUpToDate, false);
+		}
+		else if (bIsReinstancingSkeleton)
+		{
+			const bool bForceRegeneration = true;
+			FKismetEditorUtilities::GenerateBlueprintSkeleton(BP, bForceRegeneration);
 		}
 	}
 }
@@ -635,11 +642,21 @@ void FBlueprintCompileReinstancer::ReinstanceObjects(bool bForceAlwaysReinstance
 				Iter.RemoveCurrent();
 				if (auto BP = BPPtr.Get())
 				{
-					// it's unsafe to GC in the middle of reinstancing because there may be other reinstancers still alive with references to 
-					// otherwise unreferenced classes:
-					const bool bSkipGC = true;
-					FKismetEditorUtilities::CompileBlueprint(BP, false, bSkipGC, false, nullptr, false, true);
-					CompiledBlueprints.Add(BP);
+					if (bIsReinstancingSkeleton)
+					{
+						const bool bForceRegeneration = true;
+						FKismetEditorUtilities::GenerateBlueprintSkeleton(BP, bForceRegeneration);
+					}
+					else
+					{
+						// it's unsafe to GC in the middle of reinstancing because there may be other reinstancers still alive with references to 
+											// otherwise unreferenced classes:
+						const bool bSkipGC = true;
+						// Full compiles first recompile all skeleton classes, so they are up-to-date
+						const bool bSkeletonUpToDate = true;
+						FKismetEditorUtilities::CompileBlueprint(BP, false, bSkipGC, false, nullptr, bSkeletonUpToDate, true);
+						CompiledBlueprints.Add(BP);
+					}
 				}
 			}
 
@@ -657,50 +674,58 @@ void FBlueprintCompileReinstancer::ReinstanceObjects(bool bForceAlwaysReinstance
 
 			ensure(0 == DependentBlueprintsToRecompile.Num());
 
-			TGuardValue<bool> ReinstancingGuard(GIsReinstancing, true);
+			if (!bIsReinstancingSkeleton)
+			{
+				TGuardValue<bool> ReinstancingGuard(GIsReinstancing, true);
 
-			TArray<TSharedPtr<FReinstanceFinalizer>> Finalizers;
+				TArray<TSharedPtr<FReinstanceFinalizer>> Finalizers;
 
-			// All children were recompiled. It's safe to reinstance.
-			for (int32 Idx = 0; Idx < QueueToReinstance.Num(); ++Idx)
-			{
-				auto Finalizer = QueueToReinstance[Idx]->ReinstanceInner(bForceAlwaysReinstance);
-				if (Finalizer.IsValid())
+				// All children were recompiled. It's safe to reinstance.
+				for (int32 Idx = 0; Idx < QueueToReinstance.Num(); ++Idx)
 				{
-					Finalizers.Push(Finalizer);
+					auto Finalizer = QueueToReinstance[Idx]->ReinstanceInner(bForceAlwaysReinstance);
+					if (Finalizer.IsValid())
+					{
+						Finalizers.Push(Finalizer);
+					}
+					QueueToReinstance[Idx]->bHasReinstanced = true;
 				}
-				QueueToReinstance[Idx]->bHasReinstanced = true;
-			}
-			QueueToReinstance.Empty();
+				QueueToReinstance.Empty();
 
-			for (auto Finalizer : Finalizers)
-			{
-				if (Finalizer.IsValid())
+				for (auto Finalizer : Finalizers)
 				{
-					Finalizer->Finalize();
+					if (Finalizer.IsValid())
+					{
+						Finalizer->Finalize();
+					}
 				}
-			}
 
-			for (auto CompiledBP : CompiledBlueprints)
-			{
-				CompiledBP->BroadcastCompiled();
-			}
+				for (auto CompiledBP : CompiledBlueprints)
+				{
+					CompiledBP->BroadcastCompiled();
+				}
 
-			{
-				BP_SCOPED_COMPILER_EVENT_STAT(EKismetCompilerStats_RefreshDependentBlueprintsInReinstancer);
-				for (auto BPPtr : DependentBlueprintsToRefresh)
 				{
-					if (BPPtr.IsValid())
+					BP_SCOPED_COMPILER_EVENT_STAT(EKismetCompilerStats_RefreshDependentBlueprintsInReinstancer);
+					for (auto BPPtr : DependentBlueprintsToRefresh)
 					{
-						BPPtr->BroadcastChanged();
+						if (BPPtr.IsValid())
+						{
+							BPPtr->BroadcastChanged();
+						}
 					}
+					DependentBlueprintsToRefresh.Empty();
 				}
-				DependentBlueprintsToRefresh.Empty();
-			}
 
-			if (GEditor)
+				if (GEditor)
+				{
+					GEditor->BroadcastBlueprintCompiled();
+				}
+			}
+			else
 			{
-				GEditor->BroadcastBlueprintCompiled();
+				QueueToReinstance.Empty();
+				DependentBlueprintsToRefresh.Empty();
 			}
 		}
 	}
diff --git a/Engine/Source/Editor/UnrealEd/Private/ObjectTools.cpp b/Engine/Source/Editor/UnrealEd/Private/ObjectTools.cpp
--- a/Engine/Source/Editor/UnrealEd/Private/ObjectTools.cpp
+++ b/Engine/Source/Editor/UnrealEd/Private/ObjectTools.cpp
@@ -2347,7 +2347,6 @@ namespace ObjectTools
 
 			if (NewPackage)
 			{
-				FString PackagePrefix = PackageName;
 				FString ObjectPrefix = ObjectName;
 				int32 Suffix = 2;
 
@@ -2372,6 +2371,11 @@ namespace ObjectTools
 					}
 				}
 				
+				// If the package and object names were equal before, ensure that the generated names are also equal
+				const FString PackageShortName = FPackageName::GetLongPackageAssetName(*PackageName);
+				const FString PackagePath = FPackageName::GetLongPackagePath(*PackageName);
+				FString PackagePrefix = (ObjectName == PackageShortName) ? (PackagePath / ObjectPrefix) : PackageName;
+
 				for (; NewPackage && StaticFindObjectFast(NULL, NewPackage, FName(*ObjectName)); Suffix++)
 				{
 					// DlgName exists in DlgPackage - generate a new one with a numbered suffix
diff --git a/Engine/Source/Programs/AutomationTool/AutomationUtils/Automation.cs b/Engine/Source/Programs/AutomationTool/AutomationUtils/Automation.cs
--- a/Engine/Source/Programs/AutomationTool/AutomationUtils/Automation.cs
+++ b/Engine/Source/Programs/AutomationTool/AutomationUtils/Automation.cs
@@ -619,10 +619,6 @@ AutomationTool.exe [-verbose] [-compileonly] [-p4] Command0 [-Arg0 -Arg1 -Arg2 
 		{
 			if (!bRunningRocket.HasValue)
 			{
-				if (GlobalCommandLine.Rocket)
-				{
-					Log.TraceWarning("Use of -rocket argument on command-line to test Rocket behavior is deprecated, please ensure that you've built a true Rocket build.");
-				}
 				bRunningRocket = GlobalCommandLine.Rocket;
 				string RocketFile = Path.Combine(CommandUtils.CmdEnv.LocalRoot, "Engine", "Build", "Rocket.txt");
 				bRunningRocket |= File.Exists(RocketFile);
diff --git a/Engine/Source/Programs/AutomationTool/AutomationUtils/BuildEnvironment.cs b/Engine/Source/Programs/AutomationTool/AutomationUtils/BuildEnvironment.cs
--- a/Engine/Source/Programs/AutomationTool/AutomationUtils/BuildEnvironment.cs
+++ b/Engine/Source/Programs/AutomationTool/AutomationUtils/BuildEnvironment.cs
@@ -22,6 +22,7 @@ namespace AutomationTool
 		static public readonly string LogFolder = "uebp_LogFolder";
         static public readonly string CSVFile = "uebp_CSVFile";
 		static public readonly string EngineSavedFolder = "uebp_EngineSavedFolder";
+		static public readonly string MacMallocNanoZone = "MallocNanoZone";
 
 		// Perforce Environment
 		static public readonly string P4Port = "uebp_PORT";		
@@ -60,6 +61,7 @@ namespace AutomationTool
 		public bool HasCapabilityToCompile { get; protected set; }
 		public string MsBuildExe { get; protected set; }
 		public string MsDevExe { get; protected set; }
+		public string MallocNanoZone { get; protected set; }
 
 		#endregion
 
@@ -117,7 +119,8 @@ namespace AutomationTool
 			RobocopyExe = GetSystemExePath("robocopy.exe");
 			MountExe = GetSystemExePath("mount.exe");
 			CmdExe = Utils.IsRunningOnMono ? "/bin/sh" : GetSystemExePath("cmd.exe");
-
+			MallocNanoZone = "0";
+			CommandUtils.SetEnvVar(EnvVarNames.MacMallocNanoZone, MallocNanoZone);
 			if (String.IsNullOrEmpty(LogFolder))
 			{
 				throw new AutomationException("Environment is not set up correctly: LogFolder is not set!");
diff --git a/Engine/Source/Programs/AutomationTool/Scripts/CopyBuildToStagingDirectory.Automation.cs b/Engine/Source/Programs/AutomationTool/Scripts/CopyBuildToStagingDirectory.Automation.cs
--- a/Engine/Source/Programs/AutomationTool/Scripts/CopyBuildToStagingDirectory.Automation.cs
+++ b/Engine/Source/Programs/AutomationTool/Scripts/CopyBuildToStagingDirectory.Automation.cs
@@ -257,7 +257,7 @@ public partial class Project : CommandUtils
             string DLCName = Params.DLCName;
 
             // Making a plugin, grab the binaries too
-            SC.StageFiles(StagedFileType.NonUFS, CombinePaths(SC.ProjectRoot, "Plugins", DLCName), "*.uplugin", true, null, null, true);
+            SC.StageFiles(StagedFileType.NonUFS, CombinePaths(SC.ProjectRoot, "Plugins", DLCName), "*.uplugin", false, null, null, true);
             SC.StageFiles(StagedFileType.NonUFS, CombinePaths(SC.ProjectRoot, "Plugins", DLCName, "Binaries"), "libUE4-*.so", true, null, null, true);
             SC.StageFiles(StagedFileType.NonUFS, CombinePaths(SC.ProjectRoot, "Plugins", DLCName, "Binaries"), "UE4-*.dll", true, null, null, true);
             SC.StageFiles(StagedFileType.NonUFS, CombinePaths(SC.ProjectRoot, "Plugins", DLCName, "Binaries"), "libUE4Server-*.so", true, null, null, true);
diff --git a/Engine/Source/Programs/ShaderCompileWorker/Private/ShaderCompileWorker.cpp b/Engine/Source/Programs/ShaderCompileWorker/Private/ShaderCompileWorker.cpp
--- a/Engine/Source/Programs/ShaderCompileWorker/Private/ShaderCompileWorker.cpp
+++ b/Engine/Source/Programs/ShaderCompileWorker/Private/ShaderCompileWorker.cpp
@@ -550,6 +550,7 @@ static FName NAME_GLSL_310_ES_EXT(TEXT("GLSL_310_ES_EXT"));
 static FName NAME_SF_METAL_SM5(TEXT("SF_METAL_SM5"));
 static FName NAME_PC_VULKAN_ES2(TEXT("PC_VULKAN_ES2"));
 static FName NAME_SF_METAL_SM4(TEXT("SF_METAL_SM4"));
+static FName NAME_SF_METAL_MACES3_1(TEXT("SF_METAL_MACES3_1"));
 
 static EShaderPlatform FormatNameToEnum(FName ShaderFormat)
 {
@@ -574,6 +575,7 @@ static EShaderPlatform FormatNameToEnum(FName ShaderFormat)
 	if (ShaderFormat == NAME_SF_METAL_SM5)		return SP_METAL_SM5;
 	if (ShaderFormat == NAME_PC_VULKAN_ES2)		return SP_VULKAN_ES2;
 	if (ShaderFormat == NAME_SF_METAL_SM4)		return SP_METAL_SM4;
+	if (ShaderFormat == NAME_SF_METAL_MACES3_1)	return SP_METAL_MACES3_1;
 	return SP_NumPlatforms;
 }
 
diff --git a/Engine/Source/Programs/UnrealBuildTool/Android/AndroidPluginLanguage.cs b/Engine/Source/Programs/UnrealBuildTool/Android/AndroidPluginLanguage.cs
--- a/Engine/Source/Programs/UnrealBuildTool/Android/AndroidPluginLanguage.cs
+++ b/Engine/Source/Programs/UnrealBuildTool/Android/AndroidPluginLanguage.cs
@@ -1304,7 +1304,7 @@ namespace UnrealBuildTool
 								{
 									if (!GlobalContext.ElementVariables.TryGetValue(Name, out Element))
 									{
-										Log.TraceWarning("\nMissing element variable '{0}' in '{1}' (skipping instruction)", Tag, TraceNodeString(Node));
+										Log.TraceWarning("\nMissing element variable '{0}' in '{1}' (skipping instruction)", Name, TraceNodeString(Node));
 										continue;
 									}
 								}
diff --git a/Engine/Source/Programs/UnrealBuildTool/Configuration/UEBuildGame.cs b/Engine/Source/Programs/UnrealBuildTool/Configuration/UEBuildGame.cs
--- a/Engine/Source/Programs/UnrealBuildTool/Configuration/UEBuildGame.cs
+++ b/Engine/Source/Programs/UnrealBuildTool/Configuration/UEBuildGame.cs
@@ -53,10 +53,13 @@ namespace UnrealBuildTool
 			base.SetupBinaries();
 
 			{
+				// Don't write to the Engine intermediate directory for Monolithic builds
+				DirectoryReference IntermediateDirectory = ShouldCompileMonolithic() ? ProjectIntermediateDirectory : EngineIntermediateDirectory;
+
 				// Make the game executable.
 				UEBuildBinaryConfiguration Config = new UEBuildBinaryConfiguration(InType: UEBuildBinaryType.Executable,
 																					InOutputFilePaths: OutputPaths,
-																					InIntermediateDirectory: EngineIntermediateDirectory,
+																					InIntermediateDirectory: IntermediateDirectory,
 																					bInCreateImportLibrarySeparately: (ShouldCompileMonolithic() ? false : true),
 																					bInAllowExports: !ShouldCompileMonolithic(),
 																					InModuleNames: new List<string>() { "Launch" });
diff --git a/Engine/Source/Programs/UnrealBuildTool/Configuration/UEBuildTarget.cs b/Engine/Source/Programs/UnrealBuildTool/Configuration/UEBuildTarget.cs
--- a/Engine/Source/Programs/UnrealBuildTool/Configuration/UEBuildTarget.cs
+++ b/Engine/Source/Programs/UnrealBuildTool/Configuration/UEBuildTarget.cs
@@ -3335,7 +3335,7 @@ namespace UnrealBuildTool
 									bCanPrecompile = false;
 									break;
 								case ModuleRules.PrecompileTargetsType.Default:
-									bCanPrecompile = !ModuleFileName.IsUnderDirectory(DeveloperDirectory);
+									bCanPrecompile = !ModuleFileName.IsUnderDirectory(DeveloperDirectory) || TargetType == TargetRules.TargetType.Editor;
 									break;
 								case ModuleRules.PrecompileTargetsType.Game:
 									bCanPrecompile = (TargetType == TargetRules.TargetType.Client || TargetType == TargetRules.TargetType.Server || TargetType == TargetRules.TargetType.Game);
diff --git a/Engine/Source/Programs/UnrealBuildTool/IOS/IOSToolChain.cs b/Engine/Source/Programs/UnrealBuildTool/IOS/IOSToolChain.cs
--- a/Engine/Source/Programs/UnrealBuildTool/IOS/IOSToolChain.cs
+++ b/Engine/Source/Programs/UnrealBuildTool/IOS/IOSToolChain.cs
@@ -1005,7 +1005,7 @@ namespace UnrealBuildTool
 					if (Directory.Exists(Project))
 					{
 						// ensure the plist, entitlements, and provision files are properly copied
-						var DeployHandler = new UEDeployIOS();
+						var DeployHandler = (CppPlatform == CPPTargetPlatform.IOS ? new UEDeployIOS() : new UEDeployTVOS());
 						DeployHandler.PrepTargetForDeployment(Target);
 
 						var ConfigName = Target.Configuration.ToString();
diff --git a/Engine/Source/Runtime/AIModule/Private/Perception/AIPerceptionSystem.cpp b/Engine/Source/Runtime/AIModule/Private/Perception/AIPerceptionSystem.cpp
--- a/Engine/Source/Runtime/AIModule/Private/Perception/AIPerceptionSystem.cpp
+++ b/Engine/Source/Runtime/AIModule/Private/Perception/AIPerceptionSystem.cpp
@@ -109,10 +109,10 @@ void UAIPerceptionSystem::PerformSourceRegistration()
 {
 	SCOPE_CYCLE_COUNTER(STAT_AI_PerceptionSys);
 
-	for (const auto& PercSource : SourcesToRegister)
+	for (const FPerceptionSourceRegistration& PercSource : SourcesToRegister)
 	{
 		AActor* SourceActor = PercSource.Source.Get();
-		if (SourceActor != nullptr && SourceActor->IsPendingKillPending() == false)
+		if (SourceActor != nullptr && SourceActor->IsPendingKillPending() == false && Senses[PercSource.SenseID] != nullptr)
 		{
 			Senses[PercSource.SenseID]->RegisterSource(*SourceActor);
 
@@ -358,7 +358,7 @@ void UAIPerceptionSystem::UnregisterSource(AActor& SourceActor, TSubclassOf<UAIS
 		if (Sense)
 		{
 			const FAISenseID SenseID = UAISense::GetSenseID(Sense);
-			if (StimuliSource->RelevantSenses.ShouldRespondToChannel(Senses[SenseID]->GetSenseID()))
+			if (Senses[SenseID] != nullptr && StimuliSource->RelevantSenses.ShouldRespondToChannel(Senses[SenseID]->GetSenseID()))
 			{
 				Senses[SenseID]->UnregisterSource(SourceActor);
 				StimuliSource->RelevantSenses.FilterOutChannel(SenseID);
@@ -368,7 +368,7 @@ void UAIPerceptionSystem::UnregisterSource(AActor& SourceActor, TSubclassOf<UAIS
 		{
 			for (int32 SenseID = 0; SenseID < Senses.Num(); ++SenseID)
 			{
-				if (StimuliSource->RelevantSenses.ShouldRespondToChannel(Senses[SenseID]->GetSenseID()))
+				if (Senses[SenseID] != nullptr && StimuliSource->RelevantSenses.ShouldRespondToChannel(Senses[SenseID]->GetSenseID()))
 				{
 					Senses[SenseID]->UnregisterSource(SourceActor);
 				}
diff --git a/Engine/Source/Runtime/AnimGraphRuntime/Private/BoneControllers/AnimNode_Trail.cpp b/Engine/Source/Runtime/AnimGraphRuntime/Private/BoneControllers/AnimNode_Trail.cpp
--- a/Engine/Source/Runtime/AnimGraphRuntime/Private/BoneControllers/AnimNode_Trail.cpp
+++ b/Engine/Source/Runtime/AnimGraphRuntime/Private/BoneControllers/AnimNode_Trail.cpp
@@ -212,7 +212,13 @@ bool FAnimNode_Trail::IsValidToEvaluate(const USkeleton* Skeleton, const FBoneCo
 	{
 		for (auto& ChainIndex : ChainBoneIndices)
 		{
-			if (RequiredBones.Contains(ChainIndex) == false)
+			if (ChainIndex == INDEX_NONE)
+			{
+				// we don't want this to spam. We want this to happen once. 
+				ensureMsgf(false, TEXT("AnimNode_Trail [%s]: contains invalid chain index. Verify if the node contains correct length of chain. "), *TrailBone.BoneName.ToString());
+				return false;
+			}
+			else if (RequiredBones.Contains(ChainIndex) == false)
 			{
 				return false;
 			}
diff --git a/Engine/Source/Runtime/AnimGraphRuntime/Public/BoneControllers/AnimNode_Trail.h b/Engine/Source/Runtime/AnimGraphRuntime/Public/BoneControllers/AnimNode_Trail.h
--- a/Engine/Source/Runtime/AnimGraphRuntime/Public/BoneControllers/AnimNode_Trail.h
+++ b/Engine/Source/Runtime/AnimGraphRuntime/Public/BoneControllers/AnimNode_Trail.h
@@ -104,5 +104,5 @@ private:
 	FVector GetAlignVector(EAxis::Type AxisOption, bool bInvert);
 
 	// skeleton index
-	TArray<FBoneIndexType> ChainBoneIndices;
+	TArray<int32> ChainBoneIndices;
 };
\ No newline at end of file
diff --git a/Engine/Source/Runtime/Apple/MetalRHI/Private/MetalBufferPools.cpp b/Engine/Source/Runtime/Apple/MetalRHI/Private/MetalBufferPools.cpp
--- a/Engine/Source/Runtime/Apple/MetalRHI/Private/MetalBufferPools.cpp
+++ b/Engine/Source/Runtime/Apple/MetalRHI/Private/MetalBufferPools.cpp
@@ -218,7 +218,6 @@ uint32 FRingBuffer::Allocate(uint32 Size, uint32 Alignment)
 		}
 		else // wrap
 		{
-			// 		NSLog(@"Wrapping at frame %lld [size = %d]", GFrameCounter, (uint32)Buffer.length);
 			Offset = 0;
 		}
 	}
@@ -234,7 +233,12 @@ uint32 FRingBuffer::Allocate(uint32 Size, uint32 Alignment)
 	}
 	else
 	{
-		UE_LOG(LogMetal, Warning, TEXT("Wrapping offset %d into outstanding buffer region %d at frame %lld [size = %d]"), Offset, LastRead, GFrameCounter, (uint32)Buffer.length);
+		uint32 BufferSize = Buffer.length;
+		UE_LOG(LogMetal, Warning, TEXT("Reallocating ring-buffer from %d to %d to avoid wrapping write at offset %d into outstanding buffer region %d at frame %lld]"), BufferSize, BufferSize * 2, Offset, LastRead, GFrameCounter);
+		SafeReleaseMetalResource(Buffer);
+		Buffer = [GetMetalDeviceContext().GetDevice() newBufferWithLength:BufferSize * 2 options:BUFFER_CACHE_MODE];
+		Offset = 0;
+		LastRead = Size;
 		
 		// get current location
 		uint32 ReturnOffset = Offset;
diff --git a/Engine/Source/Runtime/Apple/MetalRHI/Private/MetalCommandEncoder.cpp b/Engine/Source/Runtime/Apple/MetalRHI/Private/MetalCommandEncoder.cpp
--- a/Engine/Source/Runtime/Apple/MetalRHI/Private/MetalCommandEncoder.cpp
+++ b/Engine/Source/Runtime/Apple/MetalRHI/Private/MetalCommandEncoder.cpp
@@ -10,8 +10,9 @@
 
 #pragma mark - Public C++ Boilerplate -
 
-FMetalCommandEncoder::FMetalCommandEncoder(void)
-: FrontFacingWinding(MTLWindingClockwise)
+FMetalCommandEncoder::FMetalCommandEncoder(FMetalCommandList& CmdList)
+: CommandList(CmdList)
+, FrontFacingWinding(MTLWindingClockwise)
 , CullMode(MTLCullModeNone)
 #if METAL_API_1_1
 , DepthClipMode(MTLDepthClipModeClip)
@@ -72,8 +73,81 @@ FMetalCommandEncoder::~FMetalCommandEncoder(void)
 		[DepthStencilState release];
 	}
 	[RenderPassDesc release];
+
+	if(DebugGroups)
+	{
+		[DebugGroups release];
+	}
+	
+	for(int32 i = 0; i < SF_NumFrequencies; i++)
+	{
+		for(int32 j = 0; j < ML_MaxSamplers; j++)
+		{
+			if(ShaderSamplers[i].Samplers[j])
+			{
+				[ShaderSamplers[i].Samplers[j] release];
+			}
+		}
+	}
 }
+
+void FMetalCommandEncoder::Reset(void)
+{
+	check(!CommandBuffer);
 	
+	if(ComputePipelineState)
+	{
+		[ComputePipelineState release];
+		ComputePipelineState = nil;
+	}
+	if(RenderPipelineState)
+	{
+		[RenderPipelineState release];
+		RenderPipelineState = nil;
+	}
+	if(DepthStencilState)
+	{
+		[DepthStencilState release];
+		DepthStencilState = nil;
+	}
+	[RenderPassDesc release];
+	RenderPassDesc = nil;
+	
+	[DebugGroups removeAllObjects];
+	
+	FrontFacingWinding = MTLWindingClockwise;
+	CullMode = MTLCullModeNone;
+#if METAL_API_1_1
+	DepthClipMode = MTLDepthClipModeClip;
+#endif
+	FillMode = MTLTriangleFillModeFill;
+	VisibilityMode = MTLVisibilityResultModeDisabled;
+	
+	Viewport = {0.0, 0.0, 0.0, 0.0, FLT_MIN, FLT_MAX};
+	DepthBias[0] = DepthBias[1] = 0.0f;
+	DepthBias[2] = FLT_MAX;
+	ScissorRect = {0,0, 1,1};
+	FMemory::Memzero(BlendColor);
+	StencilRef[0] = StencilRef[1] = 0xff;
+	
+	FMemory::Memzero(ShaderBuffers);
+	FMemory::Memzero(ShaderTextures);
+	for(int32 i = 0; i < SF_NumFrequencies; i++)
+	{
+		for(int32 j = 0; j < ML_MaxSamplers; j++)
+		{
+			if(ShaderSamplers[i].Samplers[j])
+			{
+				[ShaderSamplers[i].Samplers[j] release];
+			}
+			ShaderSamplers[i].Samplers[j] = nil;
+			ShaderSamplers[i].MaxLods[j] = FLT_MAX;
+			ShaderSamplers[i].MinLods[j] = 0.0;
+		}
+		ShaderSamplers[i].Bound = 0;
+	}
+}
+
 #pragma mark - Public Command Buffer Mutators -
 
 void FMetalCommandEncoder::StartCommandBuffer(id<MTLCommandBuffer> const InCommandBuffer)
@@ -86,6 +160,11 @@ void FMetalCommandEncoder::StartCommandBuffer(id<MTLCommandBuffer> const InComma
 	id<MTLCommandBuffer> OldCommandBuffer = CommandBuffer;
 	CommandBuffer = InCommandBuffer;
 	[OldCommandBuffer release];
+	
+	if ([DebugGroups count] > 0)
+	{
+		[InCommandBuffer setLabel:[DebugGroups lastObject]];
+	}
 }
 	
 void FMetalCommandEncoder::CommitCommandBuffer(bool const bWait)
@@ -93,11 +172,13 @@ void FMetalCommandEncoder::CommitCommandBuffer(bool const bWait)
 	check(CommandBuffer);
 	check(IsRenderCommandEncoderActive() == false && IsComputeCommandEncoderActive() == false && IsBlitCommandEncoderActive() == false);
 	
-	[CommandBuffer commit];
-	if(bWait)
+	if (CommandBuffer.label == nil && [DebugGroups count] > 0)
 	{
-		[CommandBuffer waitUntilCompleted];
+		[CommandBuffer setLabel:[DebugGroups lastObject]];
 	}
+	
+	CommandList.Commit(CommandBuffer, bWait);
+
 	[CommandBuffer release];
 	CommandBuffer = nil;
 }
@@ -136,35 +217,26 @@ id<MTLBlitCommandEncoder> FMetalCommandEncoder::GetBlitCommandEncoder(void) cons
 	
 #pragma mark - Public Command Encoder Mutators -
 
-void FMetalCommandEncoder::BeginRenderCommandEncoding(MTLRenderPassDescriptor* const RenderPass)
+void FMetalCommandEncoder::BeginRenderCommandEncoding(void)
 {
-	check(RenderPass);
+	check(RenderPassDesc);
 	check(CommandBuffer);
 	check(IsRenderCommandEncoderActive() == false && IsComputeCommandEncoderActive() == false && IsBlitCommandEncoderActive() == false);
 	
-	if(RenderPass != RenderPassDesc)
-	{
-		[RenderPass retain];
-		if(RenderPassDesc)
-		{
-			GetMetalDeviceContext().ReleaseObject(RenderPassDesc);
-		}
-		RenderPassDesc = RenderPass;
-	}
-	
-	RenderCommandEncoder = [[CommandBuffer renderCommandEncoderWithDescriptor:RenderPass] retain];
+	RenderCommandEncoder = [[CommandBuffer renderCommandEncoderWithDescriptor:RenderPassDesc] retain];
 	
 	if(GEmitDrawEvents)
 	{
 		NSString* Label = [DebugGroups count] > 0 ? [DebugGroups lastObject] : (NSString*)CFSTR("InitialPass");
 		[RenderCommandEncoder setLabel:Label];
-    }
-    
-    for(uint32 i = 0; i < SF_NumFrequencies; i++)
-    {
-        ShaderBuffers[i].Bound = 0;
-		ShaderTextures[i].Bound = 0;
-		ShaderSamplers[i].Bound = 0;
+		
+		if([DebugGroups count])
+		{
+			for (NSString* Group in DebugGroups)
+			{
+				[RenderCommandEncoder pushDebugGroup:Group];
+			}
+		}
     }
 }
 
@@ -174,7 +246,7 @@ void FMetalCommandEncoder::RestoreRenderCommandEncoding(void)
 	check(CommandBuffer);
 	check(IsRenderCommandEncoderActive() == false && IsComputeCommandEncoderActive() == false && IsBlitCommandEncoderActive() == false);
 	
-	BeginRenderCommandEncoding(RenderPassDesc);
+	BeginRenderCommandEncoding();
 	RestoreRenderCommandEncodingState();
 }
 
@@ -219,43 +291,37 @@ void FMetalCommandEncoder::RestoreRenderCommandEncodingState(void)
 	
 	for(uint32 i = 0; i < ML_MaxBuffers; i++)
 	{
-		if(ShaderBuffers[SF_Vertex].Buffers[i] != nil)
+		if(ShaderBuffers[SF_Vertex].Buffers[i] != nil && (ShaderBuffers[SF_Vertex].Bound & (1 << i)))
 		{
-            ShaderBuffers[SF_Vertex].Bound |= (1 << i);
 			[RenderCommandEncoder setVertexBuffer:ShaderBuffers[SF_Vertex].Buffers[i] offset:ShaderBuffers[SF_Vertex].Offsets[i] atIndex:i];
 		}
-		if(ShaderBuffers[SF_Pixel].Buffers[i] != nil)
+		if(ShaderBuffers[SF_Pixel].Buffers[i] != nil && (ShaderBuffers[SF_Pixel].Bound & (1 << i)))
         {
-            ShaderBuffers[SF_Pixel].Bound |= (1 << i);
 			[RenderCommandEncoder setFragmentBuffer:ShaderBuffers[SF_Pixel].Buffers[i] offset:ShaderBuffers[SF_Pixel].Offsets[i] atIndex:i];
 		}
 	}
 	
 	for(uint32 i = 0; i < ML_MaxTextures; i++)
 	{
-		if(ShaderTextures[SF_Vertex].Textures[i])
+		if(ShaderTextures[SF_Vertex].Textures[i] && (ShaderTextures[SF_Vertex].Bound & (1 << i)))
 		{
-			ShaderTextures[SF_Vertex].Bound |= (1 << i);
 			[RenderCommandEncoder setVertexTexture:ShaderTextures[SF_Vertex].Textures[i] atIndex:i];
 		}
 		
-		if(ShaderTextures[SF_Pixel].Textures[i])
+		if(ShaderTextures[SF_Pixel].Textures[i] && (ShaderTextures[SF_Pixel].Bound & (1 << i)))
 		{
-			ShaderTextures[SF_Pixel].Bound |= (1 << i);
 			[RenderCommandEncoder setFragmentTexture:ShaderTextures[SF_Pixel].Textures[i] atIndex:i];
 		}
 	}
 	
 	for(uint32 i = 0; i < ML_MaxSamplers; i++)
 	{
-		if(ShaderSamplers[SF_Vertex].Samplers[i] != nil)
+		if(ShaderSamplers[SF_Vertex].Samplers[i] != nil && (ShaderSamplers[SF_Vertex].Bound & (1 << i)))
 		{
-			ShaderSamplers[SF_Vertex].Bound |= (1 << i);
 			[RenderCommandEncoder setVertexSamplerState:ShaderSamplers[SF_Vertex].Samplers[i] lodMinClamp:ShaderSamplers[SF_Vertex].MinLods[i] lodMaxClamp:ShaderSamplers[SF_Vertex].MaxLods[i] atIndex:i];
 		}
-		if(ShaderSamplers[SF_Pixel].Samplers[i] != nil)
+		if(ShaderSamplers[SF_Pixel].Samplers[i] != nil && (ShaderSamplers[SF_Pixel].Bound & (1 << i)))
 		{
-			ShaderSamplers[SF_Pixel].Bound |= (1 << i);
 			[RenderCommandEncoder setFragmentSamplerState:ShaderSamplers[SF_Pixel].Samplers[i] lodMinClamp:ShaderSamplers[SF_Pixel].MinLods[i] lodMaxClamp:ShaderSamplers[SF_Pixel].MaxLods[i] atIndex:i];
 		}
 	}
@@ -274,6 +340,14 @@ void FMetalCommandEncoder::BeginComputeCommandEncoding(void)
 	{
 		NSString* Label = [DebugGroups count] > 0 ? [DebugGroups lastObject] : (NSString*)CFSTR("InitialPass");
 		[ComputeCommandEncoder setLabel:Label];
+		
+		if([DebugGroups count])
+		{
+			for (NSString* Group in DebugGroups)
+			{
+				[ComputeCommandEncoder pushDebugGroup:Group];
+			}
+		}
 	}
 	
     for(uint32 i = 0; i < SF_NumFrequencies; i++)
@@ -284,12 +358,64 @@ void FMetalCommandEncoder::BeginComputeCommandEncoding(void)
     }
 }
 
+void FMetalCommandEncoder::RestoreComputeCommandEncodingState(void)
+{
+	check(ComputeCommandEncoder);
+	
+	PushDebugGroup(@"RestoreComputeCommandEncodingState");
+	
+	[ComputeCommandEncoder setComputePipelineState:ComputePipelineState];
+	
+	for(uint32 i = 0; i < ML_MaxBuffers; i++)
+	{
+		if(ShaderBuffers[SF_Compute].Buffers[i] != nil)
+		{
+			ShaderBuffers[SF_Compute].Bound |= (1 << i);
+			[ComputeCommandEncoder setBuffer:ShaderBuffers[SF_Compute].Buffers[i] offset:ShaderBuffers[SF_Compute].Offsets[i] atIndex:i];
+		}
+	}
+	
+	for(uint32 i = 0; i < ML_MaxTextures; i++)
+	{
+		if(ShaderTextures[SF_Compute].Textures[i])
+		{
+			ShaderTextures[SF_Compute].Bound |= (1 << i);
+			[ComputeCommandEncoder setTexture:ShaderTextures[SF_Compute].Textures[i] atIndex:i];
+		}
+	}
+	
+	for(uint32 i = 0; i < ML_MaxSamplers; i++)
+	{
+		if(ShaderSamplers[SF_Compute].Samplers[i] != nil)
+		{
+			ShaderSamplers[SF_Compute].Bound |= (1 << i);
+			[ComputeCommandEncoder setSamplerState:ShaderSamplers[SF_Compute].Samplers[i] lodMinClamp:ShaderSamplers[SF_Compute].MinLods[i] lodMaxClamp:ShaderSamplers[SF_Compute].MaxLods[i] atIndex:i];
+		}
+	}
+	
+	PopDebugGroup();
+}
+
 void FMetalCommandEncoder::BeginBlitCommandEncoding(void)
 {
 	check(CommandBuffer);
 	check(IsRenderCommandEncoderActive() == false && IsComputeCommandEncoderActive() == false && IsBlitCommandEncoderActive() == false);
 	
 	BlitCommandEncoder = [[CommandBuffer blitCommandEncoder] retain];
+	
+	if(GEmitDrawEvents)
+	{
+		NSString* Label = [DebugGroups count] > 0 ? [DebugGroups lastObject] : (NSString*)CFSTR("InitialPass");
+		[BlitCommandEncoder setLabel:Label];
+		
+		if([DebugGroups count])
+		{
+			for (NSString* Group in DebugGroups)
+			{
+				[BlitCommandEncoder pushDebugGroup:Group];
+			}
+		}
+	}
 }
 
 void FMetalCommandEncoder::EndEncoding(void)
@@ -318,47 +444,114 @@ void FMetalCommandEncoder::EndEncoding(void)
 
 void FMetalCommandEncoder::InsertDebugSignpost(NSString* const String)
 {
-	[RenderCommandEncoder insertDebugSignpost:String];
+	if (RenderCommandEncoder)
+	{
+		[RenderCommandEncoder insertDebugSignpost:String];
+	}
+	if (ComputeCommandEncoder)
+	{
+		[ComputeCommandEncoder insertDebugSignpost:String];
+	}
+	if (BlitCommandEncoder)
+	{
+		[BlitCommandEncoder insertDebugSignpost:String];
+	}
 }
 
 void FMetalCommandEncoder::PushDebugGroup(NSString* const String)
 {
 	[DebugGroups addObject:String];
-	[RenderCommandEncoder pushDebugGroup:String];
+	if (RenderCommandEncoder)
+	{
+		[RenderCommandEncoder pushDebugGroup:String];
+	}
+	else if (ComputeCommandEncoder)
+	{
+		[ComputeCommandEncoder pushDebugGroup:String];
+	}
+	else if (BlitCommandEncoder)
+	{
+		[BlitCommandEncoder pushDebugGroup:String];
+	}
 }
 
 void FMetalCommandEncoder::PopDebugGroup(void)
 {
 	[DebugGroups removeLastObject];
-	[RenderCommandEncoder popDebugGroup];
+	if (RenderCommandEncoder)
+	{
+		[RenderCommandEncoder popDebugGroup];
+	}
+	else if (ComputeCommandEncoder)
+	{
+		[ComputeCommandEncoder popDebugGroup];
+	}
+	else if (BlitCommandEncoder)
+	{
+		[BlitCommandEncoder popDebugGroup];
+	}
 }
 	
 #pragma mark - Public Render State Mutators -
 
+void FMetalCommandEncoder::SetRenderPassDescriptor(MTLRenderPassDescriptor* const RenderPass)
+{
+	check(IsRenderCommandEncoderActive() == false && IsComputeCommandEncoderActive() == false && IsBlitCommandEncoderActive() == false);
+	
+	if(RenderPass != RenderPassDesc)
+	{
+		[RenderPass retain];
+		if(RenderPassDesc)
+		{
+			GetMetalDeviceContext().ReleaseObject(RenderPassDesc);
+		}
+		RenderPassDesc = RenderPass;
+	}
+	
+	for(uint32 i = 0; i < SF_NumFrequencies; i++)
+	{
+		ShaderBuffers[i].Bound = 0;
+		ShaderTextures[i].Bound = 0;
+		ShaderSamplers[i].Bound = 0;
+	}
+}
+
 void FMetalCommandEncoder::SetRenderPipelineState(id<MTLRenderPipelineState> PipelineState)
 {
 	[PipelineState retain];
 	[RenderPipelineState release];
 	RenderPipelineState = PipelineState;
-	[RenderCommandEncoder setRenderPipelineState:RenderPipelineState];
+	if (RenderCommandEncoder)
+	{
+		[RenderCommandEncoder setRenderPipelineState:RenderPipelineState];
+	}
 }
 
 void FMetalCommandEncoder::SetViewport(MTLViewport const& InViewport)
 {
 	Viewport = InViewport;
-	[RenderCommandEncoder setViewport:Viewport];
+	if (RenderCommandEncoder)
+	{
+		[RenderCommandEncoder setViewport:Viewport];
+	}
 }
 
 void FMetalCommandEncoder::SetFrontFacingWinding(MTLWinding const InFrontFacingWinding)
 {
 	FrontFacingWinding = InFrontFacingWinding;
-	[RenderCommandEncoder setFrontFacingWinding:FrontFacingWinding];
+	if (RenderCommandEncoder)
+	{
+		[RenderCommandEncoder setFrontFacingWinding:FrontFacingWinding];
+	}
 }
 
 void FMetalCommandEncoder::SetCullMode(MTLCullMode const InCullMode)
 {
 	CullMode = InCullMode;
-	[RenderCommandEncoder setCullMode:InCullMode];
+	if (RenderCommandEncoder)
+	{
+		[RenderCommandEncoder setCullMode:InCullMode];
+	}
 }
 
 #if METAL_API_1_1
@@ -367,7 +560,10 @@ void FMetalCommandEncoder::SetDepthClipMode(MTLDepthClipMode const InDepthClipMo
 	DepthClipMode = InDepthClipMode;
 	if(GetMetalDeviceContext().SupportsFeature(EMetalFeaturesDepthClipMode))
 	{
-		[RenderCommandEncoder setDepthClipMode:InDepthClipMode];
+		if (RenderCommandEncoder)
+		{
+			[RenderCommandEncoder setDepthClipMode:InDepthClipMode];
+		}
 	}
 	else
 	{
@@ -381,19 +577,28 @@ void FMetalCommandEncoder::SetDepthBias(float const InDepthBias, float const InS
 	DepthBias[0] = InDepthBias;
 	DepthBias[1] = InSlopeScale;
 	DepthBias[2] = InClamp;
-	[RenderCommandEncoder setDepthBias:InDepthBias slopeScale:InSlopeScale clamp:InClamp];
+	if (RenderCommandEncoder)
+	{
+		[RenderCommandEncoder setDepthBias:InDepthBias slopeScale:InSlopeScale clamp:InClamp];
+	}
 }
 
 void FMetalCommandEncoder::SetScissorRect(MTLScissorRect const& Rect)
 {
 	ScissorRect = Rect;
-	[RenderCommandEncoder setScissorRect:ScissorRect];
+	if (RenderCommandEncoder)
+	{
+		[RenderCommandEncoder setScissorRect:ScissorRect];
+	}
 }
 
 void FMetalCommandEncoder::SetTriangleFillMode(MTLTriangleFillMode const InFillMode)
 {
 	FillMode = InFillMode;
-	[RenderCommandEncoder setTriangleFillMode:FillMode];
+	if (RenderCommandEncoder)
+	{
+		[RenderCommandEncoder setTriangleFillMode:FillMode];
+	}
 }
 
 void FMetalCommandEncoder::SetBlendColor(float const Red, float const Green, float const Blue, float const Alpha)
@@ -402,7 +607,10 @@ void FMetalCommandEncoder::SetBlendColor(float const Red, float const Green, flo
 	BlendColor[1] = Green;
 	BlendColor[2] = Blue;
 	BlendColor[3] = Alpha;
-	[RenderCommandEncoder setBlendColorRed:Red green:Green blue:Blue alpha:Alpha];
+	if (RenderCommandEncoder)
+	{
+		[RenderCommandEncoder setBlendColorRed:Red green:Green blue:Blue alpha:Alpha];
+	}
 }
 
 void FMetalCommandEncoder::SetDepthStencilState(id<MTLDepthStencilState> const InDepthStencilState)
@@ -410,30 +618,39 @@ void FMetalCommandEncoder::SetDepthStencilState(id<MTLDepthStencilState> const I
 	[InDepthStencilState retain];
 	[DepthStencilState release];
 	DepthStencilState = InDepthStencilState;
-	[RenderCommandEncoder setDepthStencilState:InDepthStencilState];
+	if (RenderCommandEncoder)
+	{
+		[RenderCommandEncoder setDepthStencilState:InDepthStencilState];
+	}
 }
 
 void FMetalCommandEncoder::SetStencilReferenceValue(uint32 const ReferenceValue)
 {
 	StencilRef[0] = ReferenceValue;
 	StencilRef[1] = ReferenceValue;
-	[RenderCommandEncoder setStencilReferenceValue:ReferenceValue];
+	if (RenderCommandEncoder)
+	{
+		[RenderCommandEncoder setStencilReferenceValue:ReferenceValue];
+	}
 }
 
 void FMetalCommandEncoder::SetStencilReferenceValue(uint32 const FrontReferenceValue, uint32 const BackReferenceValue)
 {
 	StencilRef[0] = FrontReferenceValue;
 	StencilRef[1] = BackReferenceValue;
-#if METAL_API_1_1
-	if(GetMetalDeviceContext().SupportsFeature(EMetalFeaturesSeparateStencil))
+	if (RenderCommandEncoder)
 	{
-		[RenderCommandEncoder setStencilFrontReferenceValue: FrontReferenceValue backReferenceValue:BackReferenceValue];
-	}
-	else
+#if METAL_API_1_1
+		if(GetMetalDeviceContext().SupportsFeature(EMetalFeaturesSeparateStencil))
+		{
+			[RenderCommandEncoder setStencilFrontReferenceValue: FrontReferenceValue backReferenceValue:BackReferenceValue];
+		}
+		else
 #endif
-	{
-		UE_LOG(LogMetal, Warning, TEXT("Attempting to set separate stencil ref on device which only supports unified stencil ref."));
-		[RenderCommandEncoder setStencilReferenceValue: FrontReferenceValue];
+		{
+			UE_LOG(LogMetal, Warning, TEXT("Attempting to set separate stencil ref on device which only supports unified stencil ref."));
+			[RenderCommandEncoder setStencilReferenceValue: FrontReferenceValue];
+		}
 	}
 }
 
@@ -441,7 +658,10 @@ void FMetalCommandEncoder::SetVisibilityResultMode(MTLVisibilityResultMode const
 {
 	VisibilityMode = Mode;
 	VisibilityOffset = Offset;
-	[RenderCommandEncoder setVisibilityResultMode: Mode offset:Offset];
+	if (RenderCommandEncoder)
+	{
+		[RenderCommandEncoder setVisibilityResultMode: Mode offset:Offset];
+	}
 }
 	
 #pragma mark - Public Shader Resource Mutators -
@@ -468,14 +688,22 @@ void FMetalCommandEncoder::SetShaderBuffer(EShaderFrequency Frequency, id<MTLBuf
         switch (Frequency)
         {
             case SF_Vertex:
-                [RenderCommandEncoder setVertexBuffer:Buffer offset:Offset atIndex:index];
+				if (RenderCommandEncoder)
+				{
+					[RenderCommandEncoder setVertexBuffer:Buffer offset:Offset atIndex:index];
+				}
                 break;
             case SF_Pixel:
-                [RenderCommandEncoder setFragmentBuffer:Buffer offset:Offset atIndex:index];
+				if (RenderCommandEncoder)
+				{
+					[RenderCommandEncoder setFragmentBuffer:Buffer offset:Offset atIndex:index];
+				}
                 break;
             case SF_Compute:
-                checkf(ComputeCommandEncoder != nil, TEXT("Attempted to use ComputeCommandEncoder before calling FMetalContext::ConditionalSwitchToCompute(void) to create the encoder"));
-                [ComputeCommandEncoder setBuffer:Buffer offset:Offset atIndex:index];
+				if (ComputeCommandEncoder)
+				{
+					[ComputeCommandEncoder setBuffer:Buffer offset:Offset atIndex:index];
+				}
                 break;
             default:
                 check(false);
@@ -495,14 +723,22 @@ void FMetalCommandEncoder::SetShaderBufferOffset(EShaderFrequency Frequency, NSU
 		switch (Frequency)
 		{
 			case SF_Vertex:
-				[RenderCommandEncoder setVertexBufferOffset:Offset atIndex:index];
+				if (RenderCommandEncoder)
+				{
+					[RenderCommandEncoder setVertexBufferOffset:Offset atIndex:index];
+				}
 				break;
 			case SF_Pixel:
-				[RenderCommandEncoder setFragmentBufferOffset:Offset atIndex:index];
+				if (RenderCommandEncoder)
+				{
+					[RenderCommandEncoder setFragmentBufferOffset:Offset atIndex:index];
+				}
 				break;
 			case SF_Compute:
-				checkf(ComputeCommandEncoder != nil, TEXT("Attempted to use ComputeCommandEncoder before calling FMetalContext::ConditionalSwitchToCompute(void) to create the encoder"));
-				[ComputeCommandEncoder setBufferOffset:Offset atIndex:index];
+				if (ComputeCommandEncoder)
+				{
+					[ComputeCommandEncoder setBufferOffset:Offset atIndex:index];
+				}
 				break;
 			default:
 				check(false);
@@ -535,14 +771,22 @@ void FMetalCommandEncoder::SetShaderBuffers(EShaderFrequency Frequency, id<MTLBu
 	switch (Frequency)
 	{
 		case SF_Vertex:
-			[RenderCommandEncoder setVertexBuffers:Buffers offsets:Offsets withRange:Range];
+			if (RenderCommandEncoder)
+			{
+				[RenderCommandEncoder setVertexBuffers:Buffers offsets:Offsets withRange:Range];
+			}
 			break;
 		case SF_Pixel:
-			[RenderCommandEncoder setFragmentBuffers:Buffers offsets:Offsets withRange:Range];
+			if (RenderCommandEncoder)
+			{
+				[RenderCommandEncoder setFragmentBuffers:Buffers offsets:Offsets withRange:Range];
+			}
 			break;
 		case SF_Compute:
-			checkf(ComputeCommandEncoder != nil, TEXT("Attempted to use ComputeCommandEncoder before calling FMetalContext::ConditionalSwitchToCompute(void) to create the encoder"));
-			[ComputeCommandEncoder setBuffers:Buffers offsets:Offsets withRange:Range];
+			if (ComputeCommandEncoder)
+			{
+				[ComputeCommandEncoder setBuffers:Buffers offsets:Offsets withRange:Range];
+			}
 			break;
 		default:
 			check(false);
@@ -568,14 +812,22 @@ void FMetalCommandEncoder::SetShaderTexture(EShaderFrequency Frequency, id<MTLTe
 	switch (Frequency)
 	{
 		case SF_Vertex:
-			[RenderCommandEncoder setVertexTexture:Texture atIndex:index];
+			if (RenderCommandEncoder)
+			{
+				[RenderCommandEncoder setVertexTexture:Texture atIndex:index];
+			}
 			break;
 		case SF_Pixel:
-			[RenderCommandEncoder setFragmentTexture:Texture atIndex:index];
+			if (RenderCommandEncoder)
+			{
+				[RenderCommandEncoder setFragmentTexture:Texture atIndex:index];
+			}
 			break;
 		case SF_Compute:
-			checkf(ComputeCommandEncoder != nil, TEXT("Attempted to use ComputeCommandEncoder before calling FMetalContext::ConditionalSwitchToCompute(void) to create the encoder"));
-			[ComputeCommandEncoder setTexture:Texture atIndex:index];
+			if (ComputeCommandEncoder)
+			{
+				[ComputeCommandEncoder setTexture:Texture atIndex:index];
+			}
 			break;
 		default:
 			check(false);
@@ -603,14 +855,22 @@ void FMetalCommandEncoder::SetShaderTextures(EShaderFrequency Frequency, const i
 	switch (Frequency)
 	{
 		case SF_Vertex:
-			[RenderCommandEncoder setVertexTextures:Textures withRange:Range];
+			if (RenderCommandEncoder)
+			{
+				[RenderCommandEncoder setVertexTextures:Textures withRange:Range];
+			}
 			break;
 		case SF_Pixel:
-			[RenderCommandEncoder setFragmentTextures:Textures withRange:Range];
+			if (RenderCommandEncoder)
+			{
+				[RenderCommandEncoder setFragmentTextures:Textures withRange:Range];
+			}
 			break;
 		case SF_Compute:
-			checkf(ComputeCommandEncoder != nil, TEXT("Attempted to use ComputeCommandEncoder before calling FMetalContext::ConditionalSwitchToCompute(void) to create the encoder"));
-			[ComputeCommandEncoder setTextures:Textures withRange:Range];
+			if (ComputeCommandEncoder)
+			{
+				[ComputeCommandEncoder setTextures:Textures withRange:Range];
+			}
 			break;
 		default:
 			check(false);
@@ -640,14 +900,22 @@ void FMetalCommandEncoder::SetShaderSamplerState(EShaderFrequency Frequency, id<
 	switch (Frequency)
 	{
 		case SF_Vertex:
-			[RenderCommandEncoder setVertexSamplerState:Sampler atIndex:index];
+			if (RenderCommandEncoder)
+			{
+				[RenderCommandEncoder setVertexSamplerState:Sampler atIndex:index];
+			}
 			break;
 		case SF_Pixel:
-			[RenderCommandEncoder setFragmentSamplerState:Sampler atIndex:index];
+			if (RenderCommandEncoder)
+			{
+				[RenderCommandEncoder setFragmentSamplerState:Sampler atIndex:index];
+			}
 			break;
 		case SF_Compute:
-			checkf(ComputeCommandEncoder != nil, TEXT("Attempted to use ComputeCommandEncoder before calling FMetalContext::ConditionalSwitchToCompute(void) to create the encoder"));
-			[ComputeCommandEncoder setSamplerState:Sampler atIndex:index];
+			if (ComputeCommandEncoder)
+			{
+				[ComputeCommandEncoder setSamplerState:Sampler atIndex:index];
+			}
 			break;
 		default:
 			check(false);
@@ -680,14 +948,22 @@ void FMetalCommandEncoder::SetShaderSamplerStates(EShaderFrequency Frequency, co
 	switch (Frequency)
 	{
 		case SF_Vertex:
-			[RenderCommandEncoder setVertexSamplerStates:Samplers withRange:Range];
+			if (RenderCommandEncoder)
+			{
+				[RenderCommandEncoder setVertexSamplerStates:Samplers withRange:Range];
+			}
 			break;
 		case SF_Pixel:
-			[RenderCommandEncoder setFragmentSamplerStates:Samplers withRange:Range];
+			if (RenderCommandEncoder)
+			{
+				[RenderCommandEncoder setFragmentSamplerStates:Samplers withRange:Range];
+			}
 			break;
 		case SF_Compute:
-			checkf(ComputeCommandEncoder != nil, TEXT("Attempted to use ComputeCommandEncoder before calling FMetalContext::ConditionalSwitchToCompute(void) to create the encoder"));
-			[ComputeCommandEncoder setSamplerStates:Samplers withRange:Range];
+			if (ComputeCommandEncoder)
+			{
+				[ComputeCommandEncoder setSamplerStates:Samplers withRange:Range];
+			}
 			break;
 		default:
 			check(false);
@@ -717,14 +993,22 @@ void FMetalCommandEncoder::SetShaderSamplerState(EShaderFrequency Frequency, id<
 	switch (Frequency)
 	{
 		case SF_Vertex:
-			[RenderCommandEncoder setVertexSamplerState:Sampler lodMinClamp:lodMinClamp lodMaxClamp:lodMaxClamp atIndex:index];
+			if (RenderCommandEncoder)
+			{
+				[RenderCommandEncoder setVertexSamplerState:Sampler lodMinClamp:lodMinClamp lodMaxClamp:lodMaxClamp atIndex:index];
+			}
 			break;
 		case SF_Pixel:
-			[RenderCommandEncoder setFragmentSamplerState:Sampler lodMinClamp:lodMinClamp lodMaxClamp:lodMaxClamp atIndex:index];
+			if (RenderCommandEncoder)
+			{
+				[RenderCommandEncoder setFragmentSamplerState:Sampler lodMinClamp:lodMinClamp lodMaxClamp:lodMaxClamp atIndex:index];
+			}
 			break;
 		case SF_Compute:
-			checkf(ComputeCommandEncoder != nil, TEXT("Attempted to use ComputeCommandEncoder before calling FMetalContext::ConditionalSwitchToCompute(void) to create the encoder"));
-			[ComputeCommandEncoder setSamplerState:Sampler lodMinClamp:lodMinClamp lodMaxClamp:lodMaxClamp atIndex:index];
+			if (ComputeCommandEncoder)
+			{
+				[ComputeCommandEncoder setSamplerState:Sampler lodMinClamp:lodMinClamp lodMaxClamp:lodMaxClamp atIndex:index];
+			}
 			break;
 		default:
 			check(false);
@@ -756,14 +1040,22 @@ void FMetalCommandEncoder::SetShaderSamplerStates(EShaderFrequency Frequency, co
 	switch (Frequency)
 	{
 		case SF_Vertex:
-			[RenderCommandEncoder setVertexSamplerStates:Samplers lodMinClamps:LodMinClamps lodMaxClamps:LodMaxClamps withRange:Range];
+			if (RenderCommandEncoder)
+			{
+				[RenderCommandEncoder setVertexSamplerStates:Samplers lodMinClamps:LodMinClamps lodMaxClamps:LodMaxClamps withRange:Range];
+			}
 			break;
 		case SF_Pixel:
-			[RenderCommandEncoder setFragmentSamplerStates:Samplers lodMinClamps:LodMinClamps lodMaxClamps:LodMaxClamps withRange:Range];
+			if (RenderCommandEncoder)
+			{
+				[RenderCommandEncoder setFragmentSamplerStates:Samplers lodMinClamps:LodMinClamps lodMaxClamps:LodMaxClamps withRange:Range];
+			}
 			break;
 		case SF_Compute:
-			checkf(ComputeCommandEncoder != nil, TEXT("Attempted to use ComputeCommandEncoder before calling FMetalContext::ConditionalSwitchToCompute(void) to create the encoder"));
-			[ComputeCommandEncoder setSamplerStates:Samplers lodMinClamps:LodMinClamps lodMaxClamps:LodMaxClamps withRange:Range];
+			if (ComputeCommandEncoder)
+			{
+				[ComputeCommandEncoder setSamplerStates:Samplers lodMinClamps:LodMinClamps lodMaxClamps:LodMaxClamps withRange:Range];
+			}
 			break;
 		default:
 			check(false);
@@ -775,11 +1067,13 @@ void FMetalCommandEncoder::SetShaderSamplerStates(EShaderFrequency Frequency, co
 
 void FMetalCommandEncoder::SetComputePipelineState(id<MTLComputePipelineState> const State)
 {
-	checkf(ComputeCommandEncoder != nil, TEXT("Attempted to use ComputeCommandEncoder before calling FMetalContext::ConditionalSwitchToCompute(void) to create the encoder"));
 	[State retain];
 	[ComputePipelineState release];
 	ComputePipelineState = State;
-	[ComputeCommandEncoder setComputePipelineState:State];
+	if (ComputeCommandEncoder)
+	{
+		[ComputeCommandEncoder setComputePipelineState:State];
+	}
 }
 
 #pragma mark - Public Extension Accessors -
diff --git a/Engine/Source/Runtime/Apple/MetalRHI/Private/MetalCommandEncoder.h b/Engine/Source/Runtime/Apple/MetalRHI/Private/MetalCommandEncoder.h
--- a/Engine/Source/Runtime/Apple/MetalRHI/Private/MetalCommandEncoder.h
+++ b/Engine/Source/Runtime/Apple/MetalRHI/Private/MetalCommandEncoder.h
@@ -4,6 +4,8 @@
 
 #include <Metal/Metal.h>
 
+class FMetalCommandList;
+
 /**
  * FMetalCommandEncoder:
  *	Wraps the details of switching between different command encoders on the command-buffer, allowing for restoration of the render encoder if needed.
@@ -17,11 +19,14 @@ public:
 #pragma mark - Public C++ Boilerplate -
 
 	/** Default constructor */
-	FMetalCommandEncoder(void);
+	FMetalCommandEncoder(FMetalCommandList& CmdList);
 	
 	/** Destructor */
 	~FMetalCommandEncoder(void);
 	
+	/** Reset cached state for reuse */
+	void Reset(void);
+	
 #pragma mark - Public Command Buffer Mutators -
 
 	/**
@@ -60,10 +65,9 @@ public:
 #pragma mark - Public Command Encoder Mutators -
 
 	/**
- 	 * Begins encoding rendering commands into the current command buffer. No other encoder may be active.
- 	 * @param RenderPass The descriptor for the new render pass.
+ 	 * Begins encoding rendering commands into the current command buffer. No other encoder may be active & the MTLRenderPassDescriptor must previously have been set.
 	 */
-	void BeginRenderCommandEncoding(MTLRenderPassDescriptor* const RenderPass);
+	void BeginRenderCommandEncoding(void);
 	
 	/** Restore the render command encoder from whatever state we were in. Any previous encoder must first be terminated with EndEncoding. */
 	void RestoreRenderCommandEncoding(void);
@@ -74,6 +78,9 @@ public:
 	/** Begins encoding compute commands into the current command buffer. No other encoder may be active. */
 	void BeginComputeCommandEncoding(void);
 	
+	/** Restores the compute command state into a new compute command encoder. */
+	void RestoreComputeCommandEncodingState(void);
+	
 	/** Begins encoding blit commands into the current command buffer. No other encoder may be active. */
 	void BeginBlitCommandEncoding(void);
 	
@@ -99,6 +106,12 @@ public:
 	
 #pragma mark - Public Render State Mutators -
 
+	/**
+	 * Set the render pass descriptor - no encoder may be active when this function is called.
+	 * @param RenderPass The render pass descriptor to set. May be nil.
+	 */
+	void SetRenderPassDescriptor(MTLRenderPassDescriptor* const RenderPass);
+	
 	/*
 	 * Sets the current render pipeline state object.
 	 * @param PipelineState The pipeline state to set. Must not be nil.
@@ -351,6 +364,8 @@ private:
 	FMetalTextureBindings ShaderTextures[SF_NumFrequencies];
 	FMetalSamplerBindings ShaderSamplers[SF_NumFrequencies];
 	
+	FMetalCommandList& CommandList;
+	
 	MTLViewport Viewport;
 	MTLWinding FrontFacingWinding;
 	MTLCullMode CullMode;
diff --git a/Engine/Source/Runtime/Apple/MetalRHI/Private/MetalCommandList.cpp b/Engine/Source/Runtime/Apple/MetalRHI/Private/MetalCommandList.cpp
new file mode 100644
--- /dev/null
+++ b/Engine/Source/Runtime/Apple/MetalRHI/Private/MetalCommandList.cpp
@@ -0,0 +1,91 @@
+// Copyright 1998-2016 Epic Games, Inc. All Rights Reserved.
+
+/*=============================================================================
+	MetalCommandList.cpp: Metal command buffer list wrapper.
+=============================================================================*/
+
+#include "MetalRHIPrivate.h"
+
+#include "MetalCommandList.h"
+#include "MetalCommandQueue.h"
+
+#pragma mark - Public C++ Boilerplate -
+
+FMetalCommandList::FMetalCommandList(FMetalCommandQueue& InCommandQueue, bool const bInImmediate)
+: CommandQueue(InCommandQueue)
+, SubmittedBuffers(!bInImmediate ? [NSMutableArray<id<MTLCommandBuffer>> new] : nil)
+, bImmediate(bInImmediate)
+{
+}
+
+FMetalCommandList::~FMetalCommandList(void)
+{
+	[SubmittedBuffers release];
+	SubmittedBuffers = nil;
+}
+	
+#pragma mark - Public Command List Mutators -
+
+void FMetalCommandList::Commit(id<MTLCommandBuffer> Buffer, bool const bWait)
+{
+	check(Buffer);
+	
+#if UE_BUILD_DEBUG || UE_BUILD_DEVELOPMENT
+	[Buffer addCompletedHandler : ^ (id <MTLCommandBuffer> CompletedBuffer)
+	{
+		if (CompletedBuffer.status == MTLCommandBufferStatusError)
+		{
+			NSString* Label = CompletedBuffer.label;
+			int32 Code = CompletedBuffer.error.code;
+			NSString* Domain = CompletedBuffer.error.domain;
+			NSString* ErrorDesc = CompletedBuffer.error.localizedDescription;
+			NSString* FailureDesc = CompletedBuffer.error.localizedFailureReason;
+			NSString* RecoveryDesc = CompletedBuffer.error.localizedRecoverySuggestion;
+			
+			FString LabelString = Label ? FString(Label) : FString(TEXT("Unknown"));
+			FString DomainString = Domain ? FString(Domain) : FString(TEXT("Unknown"));
+			FString ErrorString = ErrorDesc ? FString(ErrorDesc) : FString(TEXT("Unknown"));
+			FString FailureString = FailureDesc ? FString(FailureDesc) : FString(TEXT("Unknown"));
+			FString RecoveryString = RecoveryDesc ? FString(RecoveryDesc) : FString(TEXT("Unknown"));
+			
+			UE_LOG(LogMetal, Fatal, TEXT("Command Buffer %s Failure! Error Domain: %s Code: %d Description %s %s %s"), *LabelString, *DomainString, Code, *ErrorString, *FailureString, *RecoveryString);
+		}
+	}];
+#endif
+	
+	if (bImmediate)
+	{
+		CommandQueue.CommitCommandBuffer(Buffer);
+		if (bWait)
+		{
+			[Buffer waitUntilCompleted];
+		}
+	}
+	else
+	{
+		check(SubmittedBuffers);
+		check(!bWait);
+		[SubmittedBuffers addObject: Buffer];
+	}
+}
+
+void FMetalCommandList::Submit(uint32 Index, uint32 Count)
+{
+	// Only deferred contexts should call Submit, the immediate context commits directly to the command-queue.
+	check(!bImmediate);
+
+	// Command queue takes ownership of the array
+	CommandQueue.SubmitCommandBuffers(this, Index, Count);
+}
+
+void FMetalCommandList::OnScheduled(void)
+{
+	[SubmittedBuffers removeAllObjects];
+}
+	
+#pragma mark - Public Command List Accessors -
+
+NSArray<id<MTLCommandBuffer>>* FMetalCommandList::GetCommandBuffers() const
+{
+	return SubmittedBuffers;
+}
diff --git a/Engine/Source/Runtime/Apple/MetalRHI/Private/MetalCommandList.h b/Engine/Source/Runtime/Apple/MetalRHI/Private/MetalCommandList.h
new file mode 100644
--- /dev/null
+++ b/Engine/Source/Runtime/Apple/MetalRHI/Private/MetalCommandList.h
@@ -0,0 +1,63 @@
+// Copyright 1998-2016 Epic Games, Inc. All Rights Reserved.
+
+#pragma once
+
+#include <Metal/Metal.h>
+
+class FMetalCommandQueue;
+
+/**
+ * FMetalCommandList:
+ * Encapsulates multiple command-buffers into an ordered list for submission. 
+ * For the immediate context this is irrelevant and is merely a pass-through into the CommandQueue, but
+ * for deferred/parallel contexts it is required as they must queue their command buffers until they can 
+ * be committed to the command-queue in the proper order which is only known at the end of parallel encoding.
+ */
+class FMetalCommandList
+{
+public:
+#pragma mark - Public C++ Boilerplate -
+
+	/**
+	 * Constructor
+	 * @param InCommandQueue The command-queue to which the command-list's buffers will be submitted.
+	 */
+	FMetalCommandList(FMetalCommandQueue& InCommandQueue, bool const bInImmediate);
+	
+	/** Destructor */
+	~FMetalCommandList(void);
+	
+#pragma mark - Public Command List Mutators -
+	
+	/** 
+	 * Commits the provided buffer to the command-list for execution. When parallel encoding this will be submitted later.
+	 * @param Buffer The buffer to submit to the command-list.
+	 * @param bWait Whether to wait for the command buffer to complete - it is an error to set this to true on a deferred command-list.
+	 */
+	void Commit(id<MTLCommandBuffer> Buffer, bool const bWait);
+	
+	/**
+	 * Submits all outstanding command-buffers in the proper commit order to the command-queue.
+	 * When more than one command-list is active the command-queue will buffer the command-lists until all are committed to guarantee order of submission to the GPU.
+	 * @param Index The command-list's intended index in the command-queue.
+	 * @param Count The number of command-lists that will be committed to the command-queue.
+	 */
+	void Submit(uint32 Index, uint32 Count);
+	
+	/** Called by the command-queue once all the buffers have been scheduled for execution. */
+	void OnScheduled(void);
+	
+#pragma mark - Public Command List Accessors -
+
+	/**
+	* Accessor to the internal array of submitted command buffers waiting for execution - called by the Command Queue.
+	* @returns The internal array of submitted command buffers waiting for scheduling.
+	*/
+	NSArray<id<MTLCommandBuffer>>* GetCommandBuffers() const;
+	
+private:
+#pragma mark - Private Member Variables -
+	FMetalCommandQueue& CommandQueue;
+	NSMutableArray<id<MTLCommandBuffer>>* SubmittedBuffers;
+	bool bImmediate;
+};
\ No newline at end of file
diff --git a/Engine/Source/Runtime/Apple/MetalRHI/Private/MetalCommandQueue.cpp b/Engine/Source/Runtime/Apple/MetalRHI/Private/MetalCommandQueue.cpp
--- a/Engine/Source/Runtime/Apple/MetalRHI/Private/MetalCommandQueue.cpp
+++ b/Engine/Source/Runtime/Apple/MetalRHI/Private/MetalCommandQueue.cpp
@@ -7,6 +7,7 @@
 #include "MetalRHIPrivate.h"
 
 #include "MetalCommandQueue.h"
+#include "MetalCommandList.h"
 #if METAL_STATISTICS
 #include "MetalStatistics.h"
 #include "ModuleManager.h"
@@ -59,20 +60,60 @@ FMetalCommandQueue::~FMetalCommandQueue(void)
 
 id<MTLCommandBuffer> FMetalCommandQueue::CreateRetainedCommandBuffer(void)
 {
-	return [CommandQueue commandBuffer];
+	@autoreleasepool
+	{
+		return [[CommandQueue commandBuffer] retain];
+	}
 }
 
 id<MTLCommandBuffer> FMetalCommandQueue::CreateUnretainedCommandBuffer(void)
 {
-    static bool bUnretainedRefs = !FParse::Param(FCommandLine::Get(),TEXT("metalretainrefs"));
-    return bUnretainedRefs ? [CommandQueue commandBufferWithUnretainedReferences] : [CommandQueue commandBuffer];
+	@autoreleasepool
+	{
+		static bool bUnretainedRefs = !FParse::Param(FCommandLine::Get(),TEXT("metalretainrefs"));
+		return bUnretainedRefs ? [[CommandQueue commandBufferWithUnretainedReferences] retain] : [[CommandQueue commandBuffer] retain];
+	}
+}
+
+void FMetalCommandQueue::CommitCommandBuffer(id<MTLCommandBuffer> const CommandBuffer)
+{
+	check(CommandBuffer);
+	[CommandBuffer commit];
+	[CommandBuffer release];
+}
+
+void FMetalCommandQueue::SubmitCommandBuffers(FMetalCommandList* BufferList, uint32 Index, uint32 Count)
+{
+	check(BufferList);
+	CommandBuffers.SetNumZeroed(Count);
+	CommandBuffers[Index] = BufferList;
+	bool bComplete = true;
+	for (uint32 i = 0; i < Count; i++)
+	{
+		bComplete &= (CommandBuffers[i] != nullptr);
+	}
+	if (bComplete)
+	{
+		GetMetalDeviceContext().SubmitCommandsHint(true);
+		
+		for (uint32 i = 0; i < Count; i++)
+		{
+			FMetalCommandList* List = CommandBuffers[i];
+			for (id<MTLCommandBuffer> Buffer in List->GetCommandBuffers())
+			{
+				CommitCommandBuffer(Buffer);
+			}
+			List->OnScheduled();
+			CommandBuffers[i] = nullptr;
+		}
+	}
 }
 
 #pragma mark - Public Command Queue Accessors -
 	
-id<MTLCommandQueue> FMetalCommandQueue::GetCommandQueue(void) const
+id<MTLDevice> FMetalCommandQueue::GetDevice(void) const
 {
-	return CommandQueue;
+	return CommandQueue.device;
 }
 
 #pragma mark - Public Debug Support -
diff --git a/Engine/Source/Runtime/Apple/MetalRHI/Private/MetalCommandQueue.h b/Engine/Source/Runtime/Apple/MetalRHI/Private/MetalCommandQueue.h
--- a/Engine/Source/Runtime/Apple/MetalRHI/Private/MetalCommandQueue.h
+++ b/Engine/Source/Runtime/Apple/MetalRHI/Private/MetalCommandQueue.h
@@ -4,6 +4,8 @@
 
 #include <Metal/Metal.h>
 
+class FMetalCommandList;
+
 /**
  * FMetalCommandQueue:
  */
@@ -39,15 +41,23 @@ public:
 	id<MTLCommandBuffer> CreateUnretainedCommandBuffer(void);
 	
 	/**
-	 * Commit the supplied command buffer.
+	 * Commit the supplied command buffer immediately.
 	 * @param CommandBuffer The command buffer to commit, must be non-nil.
  	 */
 	void CommitCommandBuffer(id<MTLCommandBuffer> const CommandBuffer);
+	
+	/**
+	 * Deferred contexts submit their internal lists of command-buffers out of order, the command-queue takes ownership and handles reordering them & lazily commits them once all command-buffer lists are submitted.
+	 * @param BufferList The list of buffers to enqueue into the command-queue at the given index.
+	 * @param Index The 0-based index to commit BufferList's contents into relative to other active deferred contexts.
+	 * @param Count The total number of deferred contexts that will submit - only once all are submitted can any command-buffer be committed.
+	 */
+	void SubmitCommandBuffers(FMetalCommandList* BufferList, uint32 Index, uint32 Count);
 
 #pragma mark - Public Command Queue Accessors -
 	
-	/** @returns The native command queue. */
-	id<MTLCommandQueue> GetCommandQueue(void) const;
+	/** @returns The command queue's native device. */
+	id<MTLDevice> GetDevice(void) const;
 
 #pragma mark - Public Debug Support -
 
@@ -67,4 +77,5 @@ private:
 #if METAL_STATISTICS
 	class IMetalStatistics* Statistics;
 #endif
+	TArray<FMetalCommandList*> CommandBuffers;
 };
\ No newline at end of file
diff --git a/Engine/Source/Runtime/Apple/MetalRHI/Private/MetalCommands.cpp b/Engine/Source/Runtime/Apple/MetalRHI/Private/MetalCommands.cpp
--- a/Engine/Source/Runtime/Apple/MetalRHI/Private/MetalCommands.cpp
+++ b/Engine/Source/Runtime/Apple/MetalRHI/Private/MetalCommands.cpp
@@ -73,9 +73,6 @@ void FMetalRHICommandContext::RHISetComputeShader(FComputeShaderRHIParamRef Comp
 {
 	FMetalComputeShader* ComputeShader = ResourceCast(ComputeShaderRHI);
 
-	// active compute mode
-	Context->GetCurrentState().ConditionalSwitchToCompute();
-	
 	// cache this for Dispatch
 	// sets this compute shader pipeline as the current (this resets all state, so we need to set all resources after calling this)
 	Context->GetCurrentState().SetComputeShader(ComputeShader);
@@ -111,8 +108,6 @@ void FMetalRHICommandContext::RHISetViewport(uint32 MinX,uint32 MinY,float MinZ,
 	Viewport.znear = MinZ;
 	Viewport.zfar = MaxZ;
 	
-	Context->GetCommandEncoder().SetViewport(Viewport);
-
 	Context->GetCurrentState().SetViewport(Viewport);
 	
 	FShaderCache::SetViewport(MinX,MinY, MinZ, MaxX, MaxY, MaxZ);
@@ -238,35 +233,7 @@ void FMetalRHICommandContext::RHISetShaderTexture(FComputeShaderRHIParamRef Comp
 void FMetalRHICommandContext::RHISetShaderResourceViewParameter(FVertexShaderRHIParamRef VertexShaderRHI, uint32 TextureIndex, FShaderResourceViewRHIParamRef SRVRHI)
 {
 	FMetalShaderResourceView* SRV = ResourceCast(SRVRHI);
-	if (SRV)
-	{
-		FRHITexture* Texture = SRV->SourceTexture.GetReference();
-		if (Texture)
-		{
-			FMetalSurface* Surface = SRV->TextureView;
-			if (Surface != nullptr)
-			{
-				Context->GetCommandEncoder().SetShaderTexture(SF_Vertex, Surface->Texture, TextureIndex);
-			}
-			else
-			{
-				Context->GetCommandEncoder().SetShaderTexture(SF_Vertex, nil, TextureIndex);
-			}
-		}
-		else
-		{
-			FMetalVertexBuffer* VB = SRV->SourceVertexBuffer.GetReference();
-			if (VB)
-			{
-				Context->GetCommandEncoder().SetShaderBuffer(SF_Vertex, VB->Buffer, 0, TextureIndex);
-			}
-		}
-	}
-	else
-	{
-		Context->GetCommandEncoder().SetShaderTexture(SF_Vertex, nil, TextureIndex);
-	}
-	FShaderCache::SetSRV(SF_Vertex, TextureIndex, SRVRHI);
+	Context->SetShaderResourceView(SF_Vertex, TextureIndex, SRV);
 }
 
 void FMetalRHICommandContext::RHISetShaderResourceViewParameter(FHullShaderRHIParamRef HullShaderRHI,uint32 TextureIndex,FShaderResourceViewRHIParamRef SRVRHI)
@@ -287,69 +254,13 @@ void FMetalRHICommandContext::RHISetShaderResourceViewParameter(FGeometryShaderR
 void FMetalRHICommandContext::RHISetShaderResourceViewParameter(FPixelShaderRHIParamRef PixelShaderRHI,uint32 TextureIndex,FShaderResourceViewRHIParamRef SRVRHI)
 {
 	FMetalShaderResourceView* SRV = ResourceCast(SRVRHI);
-	if (SRV)
-	{
-		FRHITexture* Texture = SRV->SourceTexture.GetReference();
-		if (Texture)
-		{
-			FMetalSurface* Surface = SRV->TextureView;
-			if (Surface != nullptr)
-			{
-				Context->GetCommandEncoder().SetShaderTexture(SF_Pixel, Surface->Texture, TextureIndex);
-			}
-			else
-			{
-				Context->GetCommandEncoder().SetShaderTexture(SF_Pixel, nil, TextureIndex);
-			}
-		}
-		else
-		{
-			FMetalVertexBuffer* VB = SRV->SourceVertexBuffer.GetReference();
-			if (VB)
-			{
-				Context->GetCommandEncoder().SetShaderBuffer(SF_Pixel, VB->Buffer, 0, TextureIndex);
-			}
-		}
-	}
-	else
-	{
-		Context->GetCommandEncoder().SetShaderTexture(SF_Pixel, nil, TextureIndex);
-	}
-	FShaderCache::SetSRV(SF_Pixel, TextureIndex, SRVRHI);
+	Context->SetShaderResourceView(SF_Pixel, TextureIndex, SRV);
 }
 
 void FMetalRHICommandContext::RHISetShaderResourceViewParameter(FComputeShaderRHIParamRef ComputeShaderRHI,uint32 TextureIndex,FShaderResourceViewRHIParamRef SRVRHI)
 {
 	FMetalShaderResourceView* SRV = ResourceCast(SRVRHI);
-	if (SRV)
-	{
-		FRHITexture* Texture = SRV->SourceTexture.GetReference();
-		if (Texture)
-		{
-			FMetalSurface* Surface = SRV->TextureView;
-			if (Surface != nullptr)
-			{
-				Context->GetCommandEncoder().SetShaderTexture(SF_Compute, Surface->Texture, TextureIndex);
-			}
-			else
-			{
-				Context->GetCommandEncoder().SetShaderTexture(SF_Compute, nil, TextureIndex);
-			}
-		}
-		else
-		{
-			FMetalVertexBuffer* VB = SRV->SourceVertexBuffer.GetReference();
-			if (VB)
-			{
-				Context->GetCommandEncoder().SetShaderBuffer(SF_Compute, VB->Buffer, 0, TextureIndex);
-			}
-		}
-	}
-	else
-	{
-		Context->GetCommandEncoder().SetShaderTexture(SF_Compute, nil, TextureIndex);
-	}
-	FShaderCache::SetSRV(SF_Compute, TextureIndex, SRVRHI);
+	Context->SetShaderResourceView(SF_Compute, TextureIndex, SRV);
 }
 
 
@@ -532,9 +443,7 @@ void FMetalRHICommandContext::RHISetRenderTargets(uint32 NumSimultaneousRenderTa
 
 void FMetalDynamicRHI::RHIDiscardRenderTargets(bool Depth, bool Stencil, uint32 ColorBitMask)
 {
-#if PLATFORM_MAC // @todo zebra
-	NOT_SUPPORTED("RHIDiscardRenderTargets");
-#endif
+	// Deliberate do nothing - Metal doesn't care about this.
 }
 
 void FMetalRHICommandContext::RHISetRenderTargetsAndClear(const FRHISetRenderTargetsInfo& RenderTargetsInfo)
@@ -1142,7 +1051,7 @@ void FMetalRHICommandContext::RHIClearMRT(bool bClearColor,int32 NumClearColors,
 
 void FMetalDynamicRHI::RHIBlockUntilGPUIdle()
 {
-	NOT_SUPPORTED("RHIBlockUntilGPUIdle");
+	Context->SubmitCommandBufferAndWait();
 }
 
 uint32 FMetalDynamicRHI::RHIGetGPUFrameCycles()
diff --git a/Engine/Source/Runtime/Apple/MetalRHI/Private/MetalContext.cpp b/Engine/Source/Runtime/Apple/MetalRHI/Private/MetalContext.cpp
--- a/Engine/Source/Runtime/Apple/MetalRHI/Private/MetalContext.cpp
+++ b/Engine/Source/Runtime/Apple/MetalRHI/Private/MetalContext.cpp
@@ -8,7 +8,7 @@
 #include "MetalContext.h"
 #include "MetalProfiler.h"
 
-const uint32 RingBufferSize = 64 * 1024 * 1024;
+const uint32 RingBufferSize = 8 * 1024 * 1024;
 
 #if SHOULD_TRACK_OBJECTS
 TMap<id, int32> ClassCounts;
@@ -109,11 +109,12 @@ FMetalDeviceContext* FMetalDeviceContext::CreateDeviceContext()
 }
 
 FMetalDeviceContext::FMetalDeviceContext(id<MTLDevice> MetalDevice, FMetalCommandQueue* Queue)
-: FMetalContext(*Queue)
+: FMetalContext(*Queue, true)
 , Device(MetalDevice)
 , SceneFrameCounter(0)
 , FrameCounter(0)
 , Features(0)
+, ActiveContexts(1)
 {
 #if PLATFORM_IOS
 	NSOperatingSystemVersion Vers = [[NSProcessInfo processInfo] operatingSystemVersion];
@@ -144,6 +145,8 @@ FMetalDeviceContext::FMetalDeviceContext(id<MTLDevice> MetalDevice, FMetalComman
 		FIOSPlatformRHIFramePacer::InitWithEvent( FrameReadyEvent );
 	}
 #endif
+	
+	InitFrame(true);
 }
 
 FMetalDeviceContext::~FMetalDeviceContext()
@@ -155,6 +158,7 @@ FMetalDeviceContext::~FMetalDeviceContext()
 
 void FMetalDeviceContext::BeginFrame()
 {
+	FPlatformTLS::SetTlsValue(CurrentContextTLSSlot, this);
 }
 
 void FMetalDeviceContext::EndFrame()
@@ -191,6 +195,8 @@ void FMetalDeviceContext::EndFrame()
 
 void FMetalDeviceContext::BeginScene()
 {
+	FPlatformTLS::SetTlsValue(CurrentContextTLSSlot, this);
+	
 	// Increment the frame counter. INDEX_NONE is a special value meaning "uninitialized", so if
 	// we hit it just wrap around to zero.
 	SceneFrameCounter++;
@@ -213,6 +219,7 @@ void FMetalDeviceContext::EndScene()
 
 void FMetalDeviceContext::BeginDrawingViewport(FMetalViewport* Viewport)
 {
+	FPlatformTLS::SetTlsValue(CurrentContextTLSSlot, this);
 }
 
 void FMetalDeviceContext::EndDrawingViewport(FMetalViewport* Viewport, bool bPresent)
@@ -224,8 +231,6 @@ void FMetalDeviceContext::EndDrawingViewport(FMetalViewport* Viewport, bool bPre
 	}
 	
 	// kick the whole buffer
-	uint32 RingBufferOffset = RingBuffer.GetOffset();
-	
 	FMetalDelayedFreeList* NewList = new FMetalDelayedFreeList;
 	NewList->Signal = FPlatformProcess::GetSynchEventFromPool(true);
 	if(GUseRHIThread)
@@ -240,12 +245,21 @@ void FMetalDeviceContext::EndDrawingViewport(FMetalViewport* Viewport, bool bPre
 	}
 	DelayedFreeLists.Add(NewList);
 	
+	uint32 RingBufferOffset = RingBuffer->GetOffset();
+	id<MTLBuffer> CurrentRingBuffer = RingBuffer->Buffer;
+	TWeakPtr<FRingBuffer, ESPMode::ThreadSafe>* WeakRingBufferRef = new TWeakPtr<FRingBuffer, ESPMode::ThreadSafe>(RingBuffer.ToSharedRef());
 	[CurrentCommandBuffer addCompletedHandler : ^ (id <MTLCommandBuffer> Buffer)
 	 {
-		RingBuffer.SetLastRead(RingBufferOffset);
+		TSharedPtr<FRingBuffer, ESPMode::ThreadSafe> CmdBufferRingBuffer = WeakRingBufferRef->Pin();
+		if(CmdBufferRingBuffer.IsValid() && CmdBufferRingBuffer->Buffer == CurrentRingBuffer)
+		{
+			CmdBufferRingBuffer->SetLastRead(RingBufferOffset);
+		}
+		delete WeakRingBufferRef;
+	 
 		dispatch_semaphore_signal(CommandBufferSemaphore);
 		NewList->Signal->Trigger();
-	 }];
+	}];
 	
 	// We may be limiting our framerate to the display link
 	if( FrameReadyEvent != nullptr )
@@ -267,6 +281,7 @@ void FMetalDeviceContext::EndDrawingViewport(FMetalViewport* Viewport, bool bPre
 	
 	UNTRACK_OBJECT(CurrentCommandBuffer);
 	[CurrentCommandBuffer release];
+	CurrentCommandBuffer = nil;
 	
 #if SHOULD_TRACK_OBJECTS
 	// print out outstanding objects
@@ -345,22 +360,53 @@ void FMetalDeviceContext::ReleasePooledBuffer(FMetalPooledBuffer Buf)
 	}
 }
 
+
+FMetalRHICommandContext* FMetalDeviceContext::AcquireContext()
+{
+	FMetalRHICommandContext* Context = ParallelContexts.Pop();
+	if (!Context)
+	{
+		FMetalContext* MetalContext = new FMetalContext(GetCommandQueue(), false);
+		check(MetalContext);
+		
+		FMetalRHICommandContext* CmdContext = static_cast<FMetalRHICommandContext*>(RHIGetDefaultContext());
+		check(CmdContext);
+		
+		Context = new FMetalRHICommandContext(CmdContext->GetProfiler(), MetalContext);
+	}
+	check(Context);
+	++ActiveContexts;
+	return Context;
+}
+
+void FMetalDeviceContext::ReleaseContext(FMetalRHICommandContext* Context)
+{
+	//delete Context;
+	ParallelContexts.Push(Context);
+	--ActiveContexts;
+}
+
+uint32 FMetalDeviceContext::GetNumActiveContexts(void) const
+{
+	return ActiveContexts;
+}
+
 uint32 FMetalContext::AutoReleasePoolTLSSlot = FPlatformTLS::AllocTlsSlot();
 uint32 FMetalContext::CurrentContextTLSSlot = FPlatformTLS::AllocTlsSlot();
 
-FMetalContext::FMetalContext(FMetalCommandQueue& Queue)
-: Device(Queue.GetCommandQueue().device)
+FMetalContext::FMetalContext(FMetalCommandQueue& Queue, bool const bIsImmediate)
+: Device(Queue.GetDevice())
 , CommandQueue(Queue)
+, CommandList(Queue, bIsImmediate)
+, CommandEncoder(CommandList)
 , StateCache(CommandEncoder)
 , CurrentCommandBuffer(nil)
-, RingBuffer(Device, RingBufferSize, BufferOffsetAlignment)
+, RingBuffer(new FRingBuffer(Device, RingBufferSize, BufferOffsetAlignment))
 , QueryBuffer(new FMetalQueryBufferPool(this))
 , ResourceTableFrameCounter(INDEX_NONE)
 {
 	// create a semaphore for multi-buffering the command buffer
 	CommandBufferSemaphore = dispatch_semaphore_create(FParse::Param(FCommandLine::Get(),TEXT("gpulockstep")) ? 1 : 3);
-	
-	InitFrame(false);
 }
 
 FMetalContext::~FMetalContext()
@@ -385,6 +431,11 @@ FMetalCommandQueue& FMetalContext::GetCommandQueue()
 	return CommandQueue;
 }
 
+FMetalCommandList& FMetalContext::GetCommandList()
+{
+	return CommandList;
+}
+
 FMetalCommandEncoder& FMetalContext::GetCommandEncoder()
 {
 	return CommandEncoder;
@@ -394,7 +445,7 @@ id<MTLRenderCommandEncoder> FMetalContext::GetRenderContext()
 {
 	if(!CommandEncoder.IsRenderCommandEncoderActive())
 	{
-		UE_LOG(LogMetal, Warning, TEXT("Attempted to use GetRenderContext before calling SetRenderTarget to create the encoder"));
+		UE_LOG(LogMetal, Fatal, TEXT("Attempted to use GetRenderContext before calling SetRenderTarget to create the encoder"));
 		ConditionalSwitchToGraphics();
 		CommandEncoder.RestoreRenderCommandEncoding();
 	}
@@ -430,6 +481,9 @@ void FMetalContext::InitFrame(bool const bImmediateContext)
 {
 	FPlatformTLS::SetTlsValue(CurrentContextTLSSlot, this);
 	
+	// Reset cached state in the encoder
+	StateCache.Reset();
+	
 	// start an auto release pool (EndFrame will drain and remake)
 	CreateAutoreleasePool();
 
@@ -470,10 +524,17 @@ void FMetalContext::CreateCurrentCommandBuffer(bool bWait)
 
 void FMetalContext::SubmitCommandsHint(bool const bCreateNew)
 {
-    uint32 RingBufferOffset = RingBuffer.GetOffset();
+	uint32 RingBufferOffset = RingBuffer->GetOffset();
+	id<MTLBuffer> CurrentRingBuffer = RingBuffer->Buffer;
+	TWeakPtr<FRingBuffer, ESPMode::ThreadSafe>* WeakRingBufferRef = new TWeakPtr<FRingBuffer, ESPMode::ThreadSafe>(RingBuffer.ToSharedRef());
     [CurrentCommandBuffer addCompletedHandler : ^ (id <MTLCommandBuffer> Buffer)
     {
-        RingBuffer.SetLastRead(RingBufferOffset);
+		TSharedPtr<FRingBuffer, ESPMode::ThreadSafe> CmdBufferRingBuffer = WeakRingBufferRef->Pin();
+		if(CmdBufferRingBuffer.IsValid() && CmdBufferRingBuffer->Buffer == CurrentRingBuffer)
+		{
+			CmdBufferRingBuffer->SetLastRead(RingBufferOffset);
+		}
+		delete WeakRingBufferRef;
     }];
     
     // commit the render context to the commandBuffer
@@ -495,15 +556,26 @@ void FMetalContext::SubmitCommandsHint(bool const bCreateNew)
 	{
 		CreateCurrentCommandBuffer(false);
 	}
+	else
+	{
+		CurrentCommandBuffer = nil;
+	}
 }
 
 void FMetalContext::SubmitCommandBufferAndWait()
 {
-	uint32 RingBufferOffset = RingBuffer.GetOffset();
+	uint32 RingBufferOffset = RingBuffer->GetOffset();
+	id<MTLBuffer> CurrentRingBuffer = RingBuffer->Buffer;
+	TWeakPtr<FRingBuffer, ESPMode::ThreadSafe>* WeakRingBufferRef = new TWeakPtr<FRingBuffer, ESPMode::ThreadSafe>(RingBuffer.ToSharedRef());
 	[CurrentCommandBuffer addCompletedHandler : ^ (id <MTLCommandBuffer> Buffer)
 	 {
-		RingBuffer.SetLastRead(RingBufferOffset);
-	 }];
+		TSharedPtr<FRingBuffer, ESPMode::ThreadSafe> CmdBufferRingBuffer = WeakRingBufferRef->Pin();
+		if(CmdBufferRingBuffer.IsValid() && CmdBufferRingBuffer->Buffer == CurrentRingBuffer)
+		{
+			CmdBufferRingBuffer->SetLastRead(RingBufferOffset);
+		}
+		delete WeakRingBufferRef;
+	}];
 	
 	// commit the render context to the commandBuffer
 	if (CommandEncoder.IsRenderCommandEncoderActive() || CommandEncoder.IsComputeCommandEncoderActive() || CommandEncoder.IsBlitCommandEncoderActive())
@@ -535,6 +607,7 @@ void FMetalContext::ResetRenderCommandEncoder()
 	SubmitCommandsHint();
 	
 	ConditionalSwitchToGraphics();
+	
     if (IsFeatureLevelSupported( GMaxRHIShaderPlatform, ERHIFeatureLevel::SM4 ))
     {
         StateCache.SetRenderTargetsInfo(StateCache.GetRenderTargetsInfo(), QueryBuffer->GetCurrentQueryBuffer()->Buffer);
@@ -543,7 +616,15 @@ void FMetalContext::ResetRenderCommandEncoder()
     {
         StateCache.SetRenderTargetsInfo(StateCache.GetRenderTargetsInfo(), NULL);
     }
-	CommandEncoder.RestoreRenderCommandEncodingState();
+	
+	if (CommandEncoder.IsRenderCommandEncoderActive())
+	{
+		CommandEncoder.RestoreRenderCommandEncodingState();
+	}
+	else
+	{
+		CommandEncoder.RestoreRenderCommandEncoding();
+	}
 }
 
 void FMetalContext::PrepareToDraw(uint32 PrimitiveType)
@@ -647,10 +728,10 @@ void FMetalContext::SetRenderTargetsInfo(const FRHISetRenderTargetsInfo& RenderT
 {
 	// Force submit if there's an RHI thread to prevent the GPU going idle.
 	// There may be other locations where this is necessary.
-	if (GUseRHIThread)
-	{
-		SubmitCommandsHint();
-	}
+//	if (GUseRHIThread)
+//	{
+//		SubmitCommandsHint();
+//	}
     if (IsFeatureLevelSupported( GMaxRHIShaderPlatform, ERHIFeatureLevel::SM4 ))
     {
         StateCache.SetRenderTargetsInfo(RenderTargetsInfo, QueryBuffer->GetCurrentQueryBuffer()->Buffer);
@@ -663,20 +744,18 @@ void FMetalContext::SetRenderTargetsInfo(const FRHISetRenderTargetsInfo& RenderT
 
 uint32 FMetalContext::AllocateFromRingBuffer(uint32 Size, uint32 Alignment)
 {
-	return RingBuffer.Allocate(Size, Alignment);
+	return RingBuffer->Allocate(Size, Alignment);
 }
 
 
-void FMetalContext::SetResource(uint32 ShaderStage, uint32 BindIndex, FRHITexture* RESTRICT TextureRHI)
+void FMetalContext::SetResource(uint32 ShaderStage, uint32 BindIndex, FRHITexture* RESTRICT TextureRHI, float CurrentTime)
 {
-	//	check(Surface->Texture != nil);
-	
-	// todo: this does multiple virtual function calls to find the right type to cast to
-	// this is due to multiple inheritance nastiness, NEEDS CLEANUP
 	FMetalSurface* Surface = GetMetalSurfaceFromRHITexture(TextureRHI);
 	
 	if (Surface != nullptr)
 	{
+		TextureRHI->SetLastRenderTime(CurrentTime);
+		
 		switch (ShaderStage)
 		{
 			case CrossCompiler::SHADER_STAGE_PIXEL:
@@ -725,7 +804,63 @@ void FMetalContext::SetResource(uint32 ShaderStage, uint32 BindIndex, FRHITextur
 	}
 }
 
-void FMetalContext::SetResource(uint32 ShaderStage, uint32 BindIndex, FMetalSamplerState* RESTRICT SamplerState)
+void FMetalContext::SetShaderResourceView(EShaderFrequency ShaderStage, uint32 BindIndex, FMetalShaderResourceView* RESTRICT SRV)
+{
+	if (SRV)
+	{
+		FRHITexture* Texture = SRV->SourceTexture.GetReference();
+		if (Texture)
+		{
+			FMetalSurface* Surface = SRV->TextureView;
+			if (Surface != nullptr)
+			{
+				Surface->UpdateSRV();
+				GetCommandEncoder().SetShaderTexture(ShaderStage, Surface->Texture, BindIndex);
+			}
+			else
+			{
+				GetCommandEncoder().SetShaderTexture(ShaderStage, nil, BindIndex);
+			}
+		}
+		else
+		{
+			FMetalVertexBuffer* VB = SRV->SourceVertexBuffer.GetReference();
+			if (VB)
+			{
+				GetCommandEncoder().SetShaderBuffer(ShaderStage, VB->Buffer, 0, BindIndex);
+			}
+		}
+	}
+	else
+	{
+		GetCommandEncoder().SetShaderTexture(ShaderStage, nil, BindIndex);
+	}
+	FShaderCache::SetSRV(ShaderStage, BindIndex, SRV);
+}
+
+void FMetalContext::SetResource(uint32 ShaderStage, uint32 BindIndex, FMetalShaderResourceView* RESTRICT SRV, float CurrentTime)
+{
+	switch (ShaderStage)
+	{
+		case CrossCompiler::SHADER_STAGE_PIXEL:
+			SetShaderResourceView(SF_Pixel, BindIndex, SRV);
+			break;
+			
+		case CrossCompiler::SHADER_STAGE_VERTEX:
+			SetShaderResourceView(SF_Vertex, BindIndex, SRV);
+			break;
+			
+		case CrossCompiler::SHADER_STAGE_COMPUTE:
+			SetShaderResourceView(SF_Compute, BindIndex, SRV);
+			break;
+			
+		default:
+			check(0);
+			break;
+	}
+}
+
+void FMetalContext::SetResource(uint32 ShaderStage, uint32 BindIndex, FMetalSamplerState* RESTRICT SamplerState, float CurrentTime)
 {
 	check(SamplerState->State != nil);
 	switch (ShaderStage)
@@ -755,6 +890,8 @@ void FMetalContext::SetResource(uint32 ShaderStage, uint32 BindIndex, FMetalSamp
 template <typename MetalResourceType>
 inline int32 FMetalContext::SetShaderResourcesFromBuffer(uint32 ShaderStage, FMetalUniformBuffer* RESTRICT Buffer, const uint32* RESTRICT ResourceMap, int32 BufferIndex)
 {
+	const TRefCountPtr<FRHIResource>* RESTRICT Resources = Buffer->ResourceTable.GetData();
+	float CurrentTime = FApp::GetCurrentTime();
 	int32 NumSetCalls = 0;
 	uint32 BufferOffset = ResourceMap[BufferIndex];
 	if (BufferOffset > 0)
@@ -766,16 +903,16 @@ inline int32 FMetalContext::SetShaderResourcesFromBuffer(uint32 ShaderStage, FMe
 			checkSlow(FRHIResourceTableEntry::GetUniformBufferIndex(ResourceInfo) == BufferIndex);
 			const uint16 ResourceIndex = FRHIResourceTableEntry::GetResourceIndex(ResourceInfo);
 			const uint8 BindIndex = FRHIResourceTableEntry::GetBindIndex(ResourceInfo);
-
+			
+			MetalResourceType* ResourcePtr = (MetalResourceType*)Resources[ResourceIndex].GetReference();
+			
 			// todo: could coalesce adjacent bound resources.
-			MetalResourceType* ResourcePtr = (MetalResourceType*)Buffer->RawResourceTable[ResourceIndex];
-			SetResource(ShaderStage, BindIndex, ResourcePtr);
-
+			SetResource(ShaderStage, BindIndex, ResourcePtr, CurrentTime);
+			
 			NumSetCalls++;
 			ResourceInfo = *ResourceInfos++;
 		} while (FRHIResourceTableEntry::GetUniformBufferIndex(ResourceInfo) == BufferIndex);
 	}
-
 	return NumSetCalls;
 }
 
@@ -811,7 +948,6 @@ void FMetalContext::SetResourcesFromTables(ShaderType Shader, uint32 ShaderStage
 	
 	// Mask the dirty bits by those buffers from which the shader has bound resources.
 	uint32 DirtyBits = Shader->Bindings.ShaderResourceTable.ResourceTableBits & StateCache.GetDirtyUniformBuffers(Frequency);
-	uint32 NumSetCalls = 0;
 	while (DirtyBits)
 	{
 		// Scan for the lowest set bit, compute its index, clear it in the set of dirty bits.
@@ -822,18 +958,13 @@ void FMetalContext::SetResourcesFromTables(ShaderType Shader, uint32 ShaderStage
 		check(Buffer);
 		check(BufferIndex < Shader->Bindings.ShaderResourceTable.ResourceTableLayoutHashes.Num());
 		check(Buffer->GetLayout().GetHash() == Shader->Bindings.ShaderResourceTable.ResourceTableLayoutHashes[BufferIndex]);
-		Buffer->CacheResources(ResourceTableFrameCounter);
 
 		// todo: could make this two pass: gather then set
-//		NumSetCalls += SetShaderResourcesFromBuffer<FMetalSurface>(ShaderStage, Buffer, Shader->Bindings.ShaderResourceTable.ShaderResourceViewMap.GetData(), BufferIndex);
-		if(Buffer->RawResourceTable.Num())
-		{
-			NumSetCalls += SetShaderResourcesFromBuffer<FRHITexture>(ShaderStage, Buffer, Shader->Bindings.ShaderResourceTable.TextureMap.GetData(), BufferIndex);
-			NumSetCalls += SetShaderResourcesFromBuffer<FMetalSamplerState>(ShaderStage, Buffer, Shader->Bindings.ShaderResourceTable.SamplerMap.GetData(), BufferIndex);
-		}
+		SetShaderResourcesFromBuffer<FRHITexture>(ShaderStage, Buffer, Shader->Bindings.ShaderResourceTable.TextureMap.GetData(), BufferIndex);
+		SetShaderResourcesFromBuffer<FMetalShaderResourceView>(ShaderStage, Buffer, Shader->Bindings.ShaderResourceTable.ShaderResourceViewMap.GetData(), BufferIndex);
+		SetShaderResourcesFromBuffer<FMetalSamplerState>(ShaderStage, Buffer, Shader->Bindings.ShaderResourceTable.SamplerMap.GetData(), BufferIndex);
 	}
 	StateCache.SetDirtyUniformBuffers(Frequency, 0);
-//	SetTextureInTableCalls += NumSetCalls;
 }
 
 void FMetalContext::CommitGraphicsResourceTables()
@@ -870,11 +1001,6 @@ void FMetalContext::ConditionalSwitchToGraphics()
 	StateCache.ConditionalSwitchToRender();
 }
 
-void FMetalContext::ConditionalSwitchToCompute()
-{
-	StateCache.ConditionalSwitchToCompute();
-}
-
 void FMetalContext::ConditionalSwitchToBlit()
 {
 	StateCache.ConditionalSwitchToBlit();
@@ -882,6 +1008,13 @@ void FMetalContext::ConditionalSwitchToBlit()
 
 void FMetalContext::Dispatch(uint32 ThreadGroupCountX, uint32 ThreadGroupCountY, uint32 ThreadGroupCountZ)
 {
+	if (!CommandEncoder.IsComputeCommandEncoderActive())
+	{
+		StateCache.ConditionalSwitchToCompute();
+		CommandEncoder.RestoreComputeCommandEncodingState();
+	}
+	check(CommandEncoder.IsComputeCommandEncoderActive());
+	
 	TRefCountPtr<FMetalComputeShader> CurrentComputeShader = StateCache.GetComputeShader();
 	check(CurrentComputeShader);
 
@@ -904,6 +1037,13 @@ void FMetalContext::Dispatch(uint32 ThreadGroupCountX, uint32 ThreadGroupCountY,
 #if PLATFORM_MAC
 void FMetalContext::DispatchIndirect(FMetalVertexBuffer* ArgumentBuffer, uint32 ArgumentOffset)
 {
+	if (!CommandEncoder.IsComputeCommandEncoderActive())
+	{
+		StateCache.ConditionalSwitchToCompute();
+		CommandEncoder.RestoreComputeCommandEncodingState();
+	}
+	check(CommandEncoder.IsComputeCommandEncoderActive());
+	
 	TRefCountPtr<FMetalComputeShader> CurrentComputeShader = StateCache.GetComputeShader();
 	check(CurrentComputeShader);
 	
@@ -925,7 +1065,6 @@ void FMetalContext::DispatchIndirect(FMetalVertexBuffer* ArgumentBuffer, uint32
 
 void FMetalContext::StartTiming(class FMetalEventNode* EventNode)
 {
-	bool bRecreateRender = CommandEncoder.IsRenderCommandEncoderActive();
 	if (CommandEncoder.IsRenderCommandEncoderActive() || CommandEncoder.IsComputeCommandEncoderActive() || CommandEncoder.IsBlitCommandEncoderActive())
 	{
 		CommandEncoder.EndEncoding();
@@ -945,16 +1084,10 @@ void FMetalContext::StartTiming(class FMetalEventNode* EventNode)
 	[CurrentCommandBuffer release];
 	
 	CreateCurrentCommandBuffer(false);
-	
-	if(bRecreateRender)
-	{
-		CommandEncoder.RestoreRenderCommandEncoding();
-	}
 }
 
 void FMetalContext::EndTiming(class FMetalEventNode* EventNode)
 {
-	bool bRecreateRender = CommandEncoder.IsRenderCommandEncoderActive();
 	if (CommandEncoder.IsRenderCommandEncoderActive() || CommandEncoder.IsComputeCommandEncoderActive() || CommandEncoder.IsBlitCommandEncoderActive())
 	{
 		CommandEncoder.EndEncoding();
@@ -972,11 +1105,6 @@ void FMetalContext::EndTiming(class FMetalEventNode* EventNode)
 	[CurrentCommandBuffer release];
 
     CreateCurrentCommandBuffer(false);
-	
-	if(bRecreateRender)
-	{
-		CommandEncoder.RestoreRenderCommandEncoding();
-	}
 }
 
 #if METAL_SUPPORTS_PARALLEL_RHI_EXECUTE
@@ -992,8 +1120,9 @@ public:
 	void operator delete(void *RawMemory);
 	
 	FMetalCommandContextContainer()
-	: CmdContext(nullptr)
+	: CmdContext(GetMetalDeviceContext().AcquireContext())
 	{
+		check(CmdContext);
 	}
 	
 	virtual ~FMetalCommandContextContainer() override
@@ -1002,35 +1131,29 @@ public:
 	
 	virtual IRHICommandContext* GetContext() override
 	{
-		check(!CmdContext);
-		FMetalContext* Ctx = new FMetalContext(GetMetalDeviceContext().GetCommandQueue());
-		
-		FMetalRHICommandContext* Context = static_cast<FMetalRHICommandContext*>(RHIGetDefaultContext());
-		check(Context);
-		
-		CmdContext = new FMetalRHICommandContext(Context->GetProfiler(), Ctx);
+		CmdContext->GetInternalContext().InitFrame(false);
 		return CmdContext;
 	}
 	virtual void FinishContext() override
 	{
 		check(CmdContext);
 		CmdContext->GetInternalContext().FinishFrame();
-		
-		delete CmdContext;
-		CmdContext = nullptr;
-		check(!CmdContext);
 	}
 	virtual void SubmitAndFreeContextContainer(int32 Index, int32 Num) override
 	{
 		if (CmdContext)
 		{
-			FinishContext();
+			CmdContext->GetInternalContext().GetCommandList().Submit(Index, Num);
+			
+			GetMetalDeviceContext().ReleaseContext(CmdContext);
+			CmdContext = nullptr;
+			check(!CmdContext);
 		}
 		delete this;
 	}
 };
 
-static TLockFreeFixedSizeAllocator<sizeof(FMetalCommandContextContainer), FThreadSafeCounter> FMetalCommandContextContainerAllocator;
+static TLockFreeFixedSizeAllocator<sizeof(FMetalCommandContextContainer), PLATFORM_CACHE_LINE_SIZE, FThreadSafeCounter> FMetalCommandContextContainerAllocator;
 
 void* FMetalCommandContextContainer::operator new(size_t Size)
 {
@@ -1059,4 +1182,4 @@ IRHICommandContextContainer* FMetalDynamicRHI::RHIGetCommandContextContainer()
 	return nullptr;
 }
 
-#endif
\ No newline at end of file
+#endif
diff --git a/Engine/Source/Runtime/Apple/MetalRHI/Private/MetalContext.h b/Engine/Source/Runtime/Apple/MetalRHI/Private/MetalContext.h
--- a/Engine/Source/Runtime/Apple/MetalRHI/Private/MetalContext.h
+++ b/Engine/Source/Runtime/Apple/MetalRHI/Private/MetalContext.h
@@ -9,9 +9,11 @@
 #include "MetalBufferPools.h"
 #include "MetalCommandEncoder.h"
 #include "MetalCommandQueue.h"
+#include "MetalCommandList.h"
 #if PLATFORM_IOS
 #include "IOSView.h"
 #endif
+#include "LockFreeList.h"
 
 #define NUM_SAFE_FRAMES 4
 
@@ -33,18 +35,20 @@ enum EMetalFeatures
 	EMetalFeaturesDepthStencilBlitOptions = 1 << 4
 };
 
+class FMetalRHICommandContext;
 
 class FMetalContext
 {
 	friend class FMetalCommandContextContainer;
 public:
-	FMetalContext(FMetalCommandQueue& Queue);
+	FMetalContext(FMetalCommandQueue& Queue, bool const bIsImmediate);
 	virtual ~FMetalContext();
 	
 	static FMetalContext* GetCurrentContext();
 	
 	id<MTLDevice> GetDevice();
 	FMetalCommandQueue& GetCommandQueue();
+	FMetalCommandList& GetCommandList();
 	FMetalCommandEncoder& GetCommandEncoder();
 	id<MTLRenderCommandEncoder> GetRenderContext();
 	id<MTLBlitCommandEncoder> GetBlitContext();
@@ -80,7 +84,7 @@ public:
 	uint32 AllocateFromRingBuffer(uint32 Size, uint32 Alignment=0);
 	id<MTLBuffer> GetRingBuffer()
 	{
-		return RingBuffer.Buffer;
+		return RingBuffer->Buffer;
 	}
 
 	TSharedRef<FMetalQueryBufferPool, ESPMode::ThreadSafe> GetQueryBufferPool()
@@ -92,6 +96,8 @@ public:
 	void SubmitCommandBufferAndWait();
 	void SubmitComputeCommandBufferAndWait();
 	void ResetRenderCommandEncoder();
+	
+	void SetShaderResourceView(EShaderFrequency ShaderStage, uint32 BindIndex, FMetalShaderResourceView* RESTRICT SRV);
 
 	void Dispatch(uint32 ThreadGroupCountX, uint32 ThreadGroupCountY, uint32 ThreadGroupCountZ);
 #if PLATFORM_MAC
@@ -113,11 +119,6 @@ protected:
 	 */
 	void ConditionalSwitchToGraphics();
 	
-	/**
-	 * Possibly switch from graphics to compute
-	 */
-	void ConditionalSwitchToCompute();
-	
 	/**
 	 * Switch to blitting
 	 */
@@ -129,9 +130,11 @@ protected:
 	void CommitNonComputeShaderConstants();
 	
 private:
-	FORCEINLINE void SetResource(uint32 ShaderStage, uint32 BindIndex, FRHITexture* RESTRICT TextureRHI);
+	FORCEINLINE void SetResource(uint32 ShaderStage, uint32 BindIndex, FRHITexture* RESTRICT TextureRHI, float CurrentTime);
+	
+	FORCEINLINE void SetResource(uint32 ShaderStage, uint32 BindIndex, FMetalShaderResourceView* RESTRICT SRV, float CurrentTime);
 	
-	FORCEINLINE void SetResource(uint32 ShaderStage, uint32 BindIndex, FMetalSamplerState* RESTRICT SamplerState);
+	FORCEINLINE void SetResource(uint32 ShaderStage, uint32 BindIndex, FMetalSamplerState* RESTRICT SamplerState, float CurrentTime);
 	
 	template <typename MetalResourceType>
 	inline int32 SetShaderResourcesFromBuffer(uint32 ShaderStage, FMetalUniformBuffer* RESTRICT Buffer, const uint32* RESTRICT ResourceMap, int32 BufferIndex);
@@ -146,6 +149,9 @@ protected:
 	/** The wrapper around the device command-queue for creating & committing command buffers to */
 	FMetalCommandQueue& CommandQueue;
 	
+	/** The wrapper around commabd buffers for ensuring correct parallel execution order */
+	FMetalCommandList CommandList;
+	
 	/** The wrapper for encoding commands into the current command buffer. */
 	FMetalCommandEncoder CommandEncoder;
 	
@@ -159,7 +165,7 @@ protected:
 	dispatch_semaphore_t CommandBufferSemaphore;
 	
 	/** A simple fixed-size ring buffer for dynamic data */
-	FRingBuffer RingBuffer;
+	TSharedPtr<FRingBuffer, ESPMode::ThreadSafe> RingBuffer;
 	
 	/** A pool of buffers for writing visibility query results. */
 	TSharedPtr<FMetalQueryBufferPool, ESPMode::ThreadSafe> QueryBuffer;
@@ -201,6 +207,15 @@ public:
 	void BeginDrawingViewport(FMetalViewport* Viewport);
 	void EndDrawingViewport(FMetalViewport* Viewport, bool bPresent);
 	
+	/** Take a parallel FMetalContext from the free-list or allocate a new one if required */
+	FMetalRHICommandContext* AcquireContext();
+	
+	/** Release a parallel FMetalContext back into the free-list */
+	void ReleaseContext(FMetalRHICommandContext* Context);
+	
+	/** Returns the number of concurrent contexts encoding commands, including the device context. */
+	uint32 GetNumActiveContexts(void) const;
+	
 private:
 	FMetalDeviceContext(id<MTLDevice> MetalDevice, FMetalCommandQueue* Queue);
 	
@@ -223,6 +238,9 @@ private:
 	};
 	TArray<FMetalDelayedFreeList*> DelayedFreeLists;
 	
+	/** Free-list of contexts for parallel encoding */
+	TLockFreePointerListLIFO<FMetalRHICommandContext> ParallelContexts;
+	
 	/** Critical section for FreeList */
 	FCriticalSection FreeListMutex;
 	
@@ -237,4 +255,7 @@ private:
 	
 	/** Bitfield of supported Metal features with varying availability depending on OS/device */
 	uint32 Features;
+	
+	/** Count of concurrent contexts encoding commands. */
+	uint32 ActiveContexts;
 };
diff --git a/Engine/Source/Runtime/Apple/MetalRHI/Private/MetalRHI.cpp b/Engine/Source/Runtime/Apple/MetalRHI/Private/MetalRHI.cpp
--- a/Engine/Source/Runtime/Apple/MetalRHI/Private/MetalRHI.cpp
+++ b/Engine/Source/Runtime/Apple/MetalRHI/Private/MetalRHI.cpp
@@ -69,7 +69,12 @@ FMetalDynamicRHI::FMetalDynamicRHI()
     {
         GMaxRHIFeatureLevel = ERHIFeatureLevel::ES3_1;
         GMaxRHIShaderPlatform = SP_METAL;
-    }
+	}
+	
+	GShaderPlatformForFeatureLevel[ERHIFeatureLevel::ES2] = SP_METAL;
+	GShaderPlatformForFeatureLevel[ERHIFeatureLevel::ES3_1] = SP_METAL;
+	GShaderPlatformForFeatureLevel[ERHIFeatureLevel::SM4] = (GMaxRHIFeatureLevel >= ERHIFeatureLevel::SM4) ? (EShaderPlatform)GMaxRHIFeatureLevel : SP_NumPlatforms;
+	GShaderPlatformForFeatureLevel[ERHIFeatureLevel::SM5] = SP_NumPlatforms;
 
 #else // @todo zebra
     // get the device to ask about capabilities?
@@ -90,6 +95,11 @@ FMetalDynamicRHI::FMetalDynamicRHI()
 		GMaxRHIShaderPlatform = SP_METAL_SM4;
 	}
 	
+	GShaderPlatformForFeatureLevel[ERHIFeatureLevel::ES2] = SP_METAL_MACES3_1;
+	GShaderPlatformForFeatureLevel[ERHIFeatureLevel::ES3_1] = SP_METAL_MACES3_1;
+	GShaderPlatformForFeatureLevel[ERHIFeatureLevel::SM4] = SP_METAL_SM4;
+	GShaderPlatformForFeatureLevel[ERHIFeatureLevel::SM5] = (GMaxRHIFeatureLevel >= ERHIFeatureLevel::SM5) ? (EShaderPlatform)GMaxRHIFeatureLevel : SP_NumPlatforms;
+	
 	GRHIAdapterName = FString(Device.name);
 	
 	bool bSupportsPointLights = false;
@@ -137,8 +147,6 @@ FMetalDynamicRHI::FMetalDynamicRHI()
 	
 #endif
 
-    
-
 	if (FPlatformMisc::IsDebuggerPresent() && UE_BUILD_DEBUG)
 	{
 #if PLATFORM_IOS // @todo zebra : needs a RENDER_API or whatever
@@ -178,16 +186,6 @@ FMetalDynamicRHI::FMetalDynamicRHI()
 	GMaxCubeTextureDimensions = 4096;
 	GMaxTextureArrayLayers = 2048;
 
-	GShaderPlatformForFeatureLevel[ERHIFeatureLevel::ES2] = SP_NumPlatforms;
-	GShaderPlatformForFeatureLevel[ERHIFeatureLevel::ES3_1] = SP_METAL;
-#if PLATFORM_IOS
-	GShaderPlatformForFeatureLevel[ERHIFeatureLevel::SM4] = (GMaxRHIShaderPlatform == SP_METAL_MRT) ? SP_METAL_MRT : SP_NumPlatforms;
-	GShaderPlatformForFeatureLevel[ERHIFeatureLevel::SM5] = SP_NumPlatforms;
-#elif PLATFORM_MAC
-	GShaderPlatformForFeatureLevel[ERHIFeatureLevel::SM4] = SP_METAL_SM4;
-	GShaderPlatformForFeatureLevel[ERHIFeatureLevel::SM5] = (GMaxRHIShaderPlatform == SP_METAL_SM5) ? SP_METAL_SM5 : SP_NumPlatforms;
-#endif
-
 	// Initialize the platform pixel format map.
 	GPixelFormats[PF_Unknown			].PlatformFormat	= MTLPixelFormatInvalid;
 	GPixelFormats[PF_A32B32G32R32F		].PlatformFormat	= MTLPixelFormatRGBA32Float;
diff --git a/Engine/Source/Runtime/Apple/MetalRHI/Private/MetalRenderPipelineDesc.cpp b/Engine/Source/Runtime/Apple/MetalRHI/Private/MetalRenderPipelineDesc.cpp
--- a/Engine/Source/Runtime/Apple/MetalRHI/Private/MetalRenderPipelineDesc.cpp
+++ b/Engine/Source/Runtime/Apple/MetalRHI/Private/MetalRenderPipelineDesc.cpp
@@ -7,6 +7,8 @@
 static_assert(Offset_End <= sizeof(FMetalRenderPipelineHash) * 8, "Too many bits used for the Hash");
 
 NSMutableDictionary* FMetalRenderPipelineDesc::MetalPipelineCache = [NSMutableDictionary new];
+FCriticalSection FMetalRenderPipelineDesc::MetalPipelineMutex;
+
 uint32 FMetalRenderPipelineDesc::BlendBitOffsets[] = { Offset_BlendState0, Offset_BlendState1, Offset_BlendState2, Offset_BlendState3, Offset_BlendState4, Offset_BlendState5, };
 uint32 FMetalRenderPipelineDesc::RTBitOffsets[] = { Offset_RenderTargetFormat0, Offset_RenderTargetFormat1, Offset_RenderTargetFormat2, Offset_RenderTargetFormat3, Offset_RenderTargetFormat4, Offset_RenderTargetFormat5, };
 
@@ -54,6 +56,12 @@ id<MTLRenderPipelineState> FMetalRenderPipelineDesc::CreatePipelineStateForBound
 	check(SampleCount > 0);
 
 	NSError* Error = nil;
+	
+	if(GUseRHIThread)
+	{
+		MetalPipelineMutex.Lock();
+	}
+	
 	id<MTLRenderPipelineState> PipelineState = (id<MTLRenderPipelineState>)[MetalPipelineCache objectForKey:PipelineDescriptor];
 	if(PipelineState == nil)
 	{
@@ -66,6 +74,11 @@ id<MTLRenderPipelineState> FMetalRenderPipelineDesc::CreatePipelineStateForBound
 		}
 	}
 	
+	if(GUseRHIThread)
+	{
+		MetalPipelineMutex.Unlock();
+	}
+	
 	PipelineDescriptor.depthAttachmentPixelFormat = DepthFormat;
 	PipelineDescriptor.stencilAttachmentPixelFormat = StencilFormat;
 	
diff --git a/Engine/Source/Runtime/Apple/MetalRHI/Private/MetalRenderPipelineDesc.h b/Engine/Source/Runtime/Apple/MetalRHI/Private/MetalRenderPipelineDesc.h
--- a/Engine/Source/Runtime/Apple/MetalRHI/Private/MetalRenderPipelineDesc.h
+++ b/Engine/Source/Runtime/Apple/MetalRHI/Private/MetalRenderPipelineDesc.h
@@ -72,6 +72,7 @@ struct FMetalRenderPipelineDesc
 	// running hash of the pipeline state
 	FMetalRenderPipelineHash Hash;
 
+	static FCriticalSection MetalPipelineMutex;
 	static NSMutableDictionary* MetalPipelineCache;
 	static uint32 BlendBitOffsets[6];
 	static uint32 RTBitOffsets[6];
diff --git a/Engine/Source/Runtime/Apple/MetalRHI/Private/MetalStateCache.cpp b/Engine/Source/Runtime/Apple/MetalRHI/Private/MetalStateCache.cpp
--- a/Engine/Source/Runtime/Apple/MetalRHI/Private/MetalStateCache.cpp
+++ b/Engine/Source/Runtime/Apple/MetalRHI/Private/MetalStateCache.cpp
@@ -36,8 +36,10 @@ FMetalStateCache::FMetalStateCache(FMetalCommandEncoder& InCommandEncoder)
 , FrameBufferSize(CGSizeMake(0.0, 0.0))
 , RenderTargetArraySize(1)
 , bHasValidRenderTarget(false)
+, bScissorRectEnabled(false)
 {
 	Viewport.originX = Viewport.originY = Viewport.width = Viewport.height = Viewport.znear = Viewport.zfar = 0.0;
+	Scissor.x = Scissor.y = Scissor.width = Scissor.height;
 	
 	FMemory::Memzero(VertexBuffers, sizeof(VertexBuffers));
 	FMemory::Memzero(VertexStrides, sizeof(VertexStrides));
@@ -58,6 +60,70 @@ FMetalStateCache::~FMetalStateCache()
 	
 }
 
+void FMetalStateCache::Reset(void)
+{
+	CommandEncoder.Reset();
+	
+	for (uint32 i = 0; i < CrossCompiler::NUM_SHADER_STAGES; i++)
+	{
+		ShaderParameters[i].MarkAllDirty();
+	}
+	for (uint32 i = 0; i < SF_NumFrequencies; i++)
+	{
+		BoundUniformBuffers[i].Empty();
+	}
+	
+	PipelineDesc.Hash = 0;
+	
+	Viewport.originX = Viewport.originY = Viewport.width = Viewport.height = Viewport.znear = Viewport.zfar = 0.0;
+	
+	FMemory::Memzero(RenderTargetsInfo);
+	bHasValidRenderTarget = false;
+	
+	FMemory::Memzero(DirtyUniformBuffers);
+	
+	FMemory::Memzero(VertexBuffers, sizeof(VertexBuffers));
+	FMemory::Memzero(VertexStrides, sizeof(VertexStrides));
+	
+	VisibilityResults = nil;
+	
+	BlendState.SafeRelease();
+	DepthStencilState.SafeRelease();
+	RasterizerState.SafeRelease();
+	BoundShaderState.SafeRelease();
+	ComputeShader.SafeRelease();
+	StencilRef = 0;
+	
+	BlendFactor = FLinearColor::Transparent;
+	FrameBufferSize = CGSizeMake(0.0, 0.0);
+	RenderTargetArraySize = 0;
+	
+}
+
+void FMetalStateCache::SetScissorRect(bool const bEnable, MTLScissorRect const& Rect)
+{
+	bScissorRectEnabled = bEnable;
+	if (bEnable)
+	{
+		Scissor = Rect;
+	}
+	else
+	{
+		Scissor.x = Viewport.originX;
+		Scissor.y = Viewport.originY;
+		Scissor.width = Viewport.width;
+		Scissor.height = Viewport.height;
+	}
+	
+	// Clamp to framebuffer size - Metal doesn't allow scissor to be larger.
+	Scissor.x = Scissor.x;
+	Scissor.y = Scissor.y;
+	Scissor.width = (Scissor.x + Scissor.width <= FrameBufferSize.width) ? Scissor.width : FrameBufferSize.width - Scissor.x;
+	Scissor.height = (Scissor.y + Scissor.height <= FrameBufferSize.height) ? Scissor.height : FrameBufferSize.height - Scissor.y;
+	
+	CommandEncoder.SetScissorRect(Scissor);
+}
+
 void FMetalStateCache::SetBlendFactor(FLinearColor const& InBlendFactor)
 {
 	//if(BlendFactor != InBlendFactor) // @todo zebra
@@ -169,68 +235,11 @@ void FMetalStateCache::SetComputeShader(FMetalComputeShader* InComputeShader)
 
 void FMetalStateCache::SetRenderTargetsInfo(FRHISetRenderTargetsInfo const& InRenderTargets, id<MTLBuffer> const QueryBuffer)
 {
-#if METAL_API_1_1
-	if(IsValidRef(DepthStencilTexture) && DepthStencilTexture->GetFormat() == PF_DepthStencil)
-	{
-		MTLStoreAction StoreAction = MTLStoreActionDontCare;
-#if PLATFORM_MAC
-		if(RenderTargetsInfo.DepthStencilRenderTarget.GetDepthStencilAccess().IsStencilWrite())
-		{
-			StoreAction = GetMetalRTStoreAction(RenderTargetsInfo.DepthStencilRenderTarget.GetStencilStoreAction());
-		}
-		else
-		{
-			StoreAction = MTLStoreActionDontCare;
-		}
-#else
-		StoreAction = GetMetalRTStoreAction(RenderTargetsInfo.DepthStencilRenderTarget.GetStencilStoreAction());
-#endif
-		
-		if(StoreAction == MTLStoreActionStore)
-		{
-			FMetalSurface& Surface = *GetMetalSurfaceFromRHITexture(DepthStencilTexture);
-			
-			if(Surface.StencilTexture != Surface.Texture && (Surface.Texture.pixelFormat == MTLPixelFormatDepth32Float_Stencil8
-#if PLATFORM_MAC
-				|| Surface.Texture.pixelFormat == MTLPixelFormatDepth24Unorm_Stencil8
-#endif
-				))
-			{
-				switch(Surface.Type)
-				{
-					case RRT_Texture2D:
-					{
-						id<MTLBlitCommandEncoder> Blitter = GetMetalDeviceContext().GetBlitContext();
-						
-						uint32 Size = Surface.Texture.width * Surface.Texture.height;
-						FMetalPooledBuffer Buf = GetMetalDeviceContext().CreatePooledBuffer(FMetalPooledBufferArgs(GetMetalDeviceContext().GetDevice(), Size, BUFFER_STORAGE_MODE));
-						
-						[Blitter copyFromTexture:Surface.Texture sourceSlice:0 sourceLevel:0 sourceOrigin:MTLOriginMake(0,0,0) sourceSize:MTLSizeMake(Surface.Texture.width, Surface.Texture.height, 1) toBuffer:Buf.Buffer destinationOffset:0 destinationBytesPerRow:Surface.Texture.width destinationBytesPerImage:Size options:MTLBlitOptionStencilFromDepthStencil];
-						
-						[Blitter copyFromBuffer:Buf.Buffer sourceOffset:0 sourceBytesPerRow:Surface.Texture.width sourceBytesPerImage:Size sourceSize:MTLSizeMake(Surface.Texture.width, Surface.Texture.height, 1) toTexture:Surface.StencilTexture destinationSlice:0 destinationLevel:0 destinationOrigin:MTLOriginMake(0,0,0)];
-						
-						SafeReleasePooledBuffer(Buf.Buffer);
-						break;
-					}
-					case RRT_Texture2DArray:
-					case RRT_Texture3D:
-					case RRT_TextureCube:
-					default:
-						NOT_SUPPORTED("Stencil sampling from non-2D texture!");
-						break;
-				}
-			}
-		}
-	}
-#endif
-	
-	ConditionalSwitchToRender();
-	
-	DepthStencilTexture = InRenderTargets.DepthStencilRenderTarget.Texture;
-
 	// see if our new Info matches our previous Info
 	if (NeedsToSetRenderTarget(InRenderTargets) || QueryBuffer != VisibilityResults)
 	{
+		bool bNeedsClear = false;
+		
 		// back this up for next frame
 		RenderTargetsInfo = InRenderTargets;
 		
@@ -357,6 +366,9 @@ void FMetalStateCache::SetRenderTargetsInfo(FRHISetRenderTargetsInfo const& InRe
 				}
 				
 				ColorAttachment.loadAction = GetMetalRTLoadAction(RenderTargetView.LoadAction);
+				
+				bNeedsClear |= (ColorAttachment.loadAction == MTLLoadActionClear);
+				
 				const FClearValueBinding& ClearValue = RenderTargetsInfo.ColorRenderTarget[RenderTargetIndex].Texture->GetClearBinding();
 				if (ClearValue.ColorBinding == EClearBinding::EColorBound)
 				{
@@ -514,6 +526,9 @@ void FMetalStateCache::SetRenderTargetsInfo(FRHISetRenderTargetsInfo const& InRe
 				// set up the depth attachment
 				DepthAttachment.texture = DepthTexture;
 				DepthAttachment.loadAction = GetMetalRTLoadAction(RenderTargetsInfo.DepthStencilRenderTarget.DepthLoadAction);
+				
+				bNeedsClear |= (DepthAttachment.loadAction == MTLLoadActionClear);
+				
 #if PLATFORM_MAC
 				if(RenderTargetsInfo.DepthStencilRenderTarget.GetDepthStencilAccess().IsDepthWrite())
 				{
@@ -548,6 +563,9 @@ void FMetalStateCache::SetRenderTargetsInfo(FRHISetRenderTargetsInfo const& InRe
 				// set up the stencil attachment
 				StencilAttachment.texture = StencilTexture;
 				StencilAttachment.loadAction = GetMetalRTLoadAction(RenderTargetsInfo.DepthStencilRenderTarget.StencilLoadAction);
+				
+				bNeedsClear |= (StencilAttachment.loadAction == MTLLoadActionClear);
+				
 #if PLATFORM_MAC
 				if(RenderTargetsInfo.DepthStencilRenderTarget.GetDepthStencilAccess().IsStencilWrite())
 				{
@@ -580,13 +598,18 @@ void FMetalStateCache::SetRenderTargetsInfo(FRHISetRenderTargetsInfo const& InRe
 		PipelineDesc.SetHashValue(Offset_SampleCount, NumBits_SampleCount, PipelineDesc.SampleCount);
 	
 		// commit pending commands on the old render target
-		if(CommandEncoder.IsRenderCommandEncoderActive())
+		if(CommandEncoder.IsRenderCommandEncoderActive() || CommandEncoder.IsBlitCommandEncoderActive() || CommandEncoder.IsComputeCommandEncoderActive())
 		{
 			CommandEncoder.EndEncoding();
 		}
 	
-		// make a new render context to use to render to the framebuffer
-		CommandEncoder.BeginRenderCommandEncoding(RenderPass);
+		// Set render to the framebuffer
+		CommandEncoder.SetRenderPassDescriptor(RenderPass);
+		
+		// if (bNeedsClear)
+		{
+			CommandEncoder.BeginRenderCommandEncoding();
+		}
 		
 		// Reset any existing state as that must be fully reinitialised by the caller.
 		DepthStencilState.SafeRelease();
@@ -604,6 +627,18 @@ void FMetalStateCache::SetHasValidRenderTarget(bool InHasValidRenderTarget)
 void FMetalStateCache::SetViewport(const MTLViewport& InViewport)
 {
 	Viewport = InViewport;
+	
+	CommandEncoder.SetViewport(Viewport);
+	
+	if (!bScissorRectEnabled)
+	{
+		MTLScissorRect Rect;
+		Rect.x = InViewport.originX;
+		Rect.y = InViewport.originY;
+		Rect.width = InViewport.width;
+		Rect.height = InViewport.height;
+		SetScissorRect(false, Rect);
+	}
 }
 
 void FMetalStateCache::SetVertexBuffer(uint32 const Index, id<MTLBuffer> Buffer, uint32 const Stride, uint32 const Offset)
@@ -670,9 +705,6 @@ void FMetalStateCache::ConditionalSwitchToCompute(void)
 	}
 	if(!CommandEncoder.IsComputeCommandEncoderActive())
 	{
-		// Clear any previous compute shader
-		ComputeShader.SafeRelease();
-		
 		// start encoding for compute
 		CommandEncoder.BeginComputeCommandEncoding();
 	}
@@ -754,6 +786,18 @@ bool FMetalStateCache::NeedsToSetRenderTarget(const FRHISetRenderTargetsInfo& In
 //				break;
 //			}
 		}
+		
+		if (InRenderTargetsInfo.DepthStencilRenderTarget.Texture && (InRenderTargetsInfo.DepthStencilRenderTarget.DepthLoadAction == ERenderTargetLoadAction::EClear || InRenderTargetsInfo.DepthStencilRenderTarget.StencilLoadAction == ERenderTargetLoadAction::EClear))
+		{
+			bAllChecksPassed = false;
+		}
+		
+#if PLATFORM_MAC
+		if (!(InRenderTargetsInfo.DepthStencilRenderTarget.GetDepthStencilAccess() == RenderTargetsInfo.DepthStencilRenderTarget.GetDepthStencilAccess()))
+		{
+			bAllChecksPassed = false;
+		}
+#endif
 	}
 
 	// if we are setting them to nothing, then this is probably end of frame, and we can't make a framebuffer
diff --git a/Engine/Source/Runtime/Apple/MetalRHI/Private/MetalStateCache.h b/Engine/Source/Runtime/Apple/MetalRHI/Private/MetalStateCache.h
--- a/Engine/Source/Runtime/Apple/MetalRHI/Private/MetalStateCache.h
+++ b/Engine/Source/Runtime/Apple/MetalRHI/Private/MetalStateCache.h
@@ -11,7 +11,11 @@ class FMetalStateCache
 public:
 	FMetalStateCache(FMetalCommandEncoder& CommandEncoder);
 	~FMetalStateCache();
+	
+	/** Reset cached state for reuse */
+	void Reset(void);
 
+	void SetScissorRect(bool const bEnable, MTLScissorRect const& Rect);
 	void SetBlendFactor(FLinearColor const& InBlendFactor);
 	void SetStencilRef(uint32 const InStencilRef);
 	void SetBlendState(FMetalBlendState* InBlendState);
@@ -59,6 +63,7 @@ public:
 	TArray<TRefCountPtr<FRHIUniformBuffer>> GetBoundUniformBuffers(EShaderFrequency const Freq) const { return BoundUniformBuffers[Freq]; }
 	uint32 GetDirtyUniformBuffers(EShaderFrequency const Freq) const { return DirtyUniformBuffers[Freq]; }
 	id<MTLBuffer> GetVisibilityResultsBuffer() const { return VisibilityResults; }
+	bool GetScissorRectEnabled() const { return bScissorRectEnabled; }
 
 private:
 	void ConditionalUpdateBackBuffer(FMetalSurface& Surface);
@@ -93,8 +98,10 @@ private:
 	uint32 RenderTargetArraySize;
 
 	MTLViewport Viewport;
+	MTLScissorRect Scissor;
 	
 	FRHISetRenderTargetsInfo RenderTargetsInfo;
-	FTextureRHIRef DepthStencilTexture;
 	bool bHasValidRenderTarget;
+	
+	bool bScissorRectEnabled;
 };
diff --git a/Engine/Source/Runtime/Apple/MetalRHI/Private/MetalTexture.cpp b/Engine/Source/Runtime/Apple/MetalRHI/Private/MetalTexture.cpp
--- a/Engine/Source/Runtime/Apple/MetalRHI/Private/MetalTexture.cpp
+++ b/Engine/Source/Runtime/Apple/MetalRHI/Private/MetalTexture.cpp
@@ -43,7 +43,7 @@ FMetalSurface* GetMetalSurfaceFromRHITexture(FRHITexture* Texture)
     if (!Texture)
     {
         return NULL;
-    }
+	}
 	if(Texture->GetTexture2D())
 	{
 		return &((FMetalTexture2D*)Texture)->Surface;
@@ -261,6 +261,15 @@ FMetalSurface::FMetalSurface(FMetalSurface& Source, NSRange const MipRange, EPix
 	}
 #endif
 	
+	if(Source.MSAATexture)
+	{
+		MSAATexture = Texture;
+	}
+	if(Source.StencilTexture)
+	{
+		StencilTexture = Texture;
+	}
+	
 	NSRange Slices = NSMakeRange(0, Source.Texture.arrayLength * (bIsCubemap ? 6 : 1));
     // @todo Zebra Temporary workaround for absence of X24_G8 or equivalent to GL_STENCIL_INDEX so that the stencil part of a texture may be sampled
 	// For now, if we find ourselves *requiring* this we lazily blit the stencil data out to a separate texture. radr://21813831
@@ -326,6 +335,7 @@ FMetalSurface::FMetalSurface(FMetalSurface& Source, NSRange const MipRange, EPix
 					SafeReleasePooledBuffer(Buf.Buffer);
 				}
 				//fallthrough to assign the stencil texture
+				StencilTexture = [Source.Texture retain];
 			}
 #endif
 			case MTLPixelFormatDepth32Float:
@@ -345,17 +355,6 @@ FMetalSurface::FMetalSurface(FMetalSurface& Source, NSRange const MipRange, EPix
     {
         Texture = [Source.Texture retain];
     }
-	if(Source.MSAATexture)
-	{
-		MSAATexture = Texture;
-	}
-	if(Source.StencilTexture)
-	{
-#if PLATFORM_IOS
-		//check(false); // @todo Zebra Must handle separate stencil texture SRV!
-#endif
-		StencilTexture = Texture;
-	}
 	
 	const uint32 BlockSizeX = GPixelFormats[PixelFormat].BlockSizeX;
 	const uint32 BlockSizeY = GPixelFormats[PixelFormat].BlockSizeY;
@@ -562,7 +561,7 @@ FMetalSurface::FMetalSurface(ERHIResourceType ResourceType, EPixelFormat Format,
 			
 			if (Texture == nil)
 			{
-				NSLog(@"Failed to create texture, desc  %@", Desc);
+				UE_LOG(LogMetal, Fatal, TEXT("Failed to create texture, desc %s"), *FString([Desc description]));
 			}
 		}
 		TRACK_OBJECT(Texture);
@@ -575,7 +574,7 @@ FMetalSurface::FMetalSurface(ERHIResourceType ResourceType, EPixelFormat Format,
 		Texture = [GetMetalDeviceContext().GetDevice() newTextureWithDescriptor:Desc];
 		if (Texture == nil)
 		{
-			NSLog(@"Failed to create texture, desc  %@", Desc);
+			UE_LOG(LogMetal, Fatal, TEXT("Failed to create texture, desc %s"), *FString([Desc description]));
 		}
 		TRACK_OBJECT(Texture);
 		
@@ -714,8 +713,13 @@ FMetalSurface::~FMetalSurface()
 	if (MSAATexture != nil)
     {
 		SafeReleaseMetalResource(MSAATexture);
-    }
-    
+	}
+	
+	if (StencilTexture != nil && StencilTexture != Texture)
+	{
+		SafeReleaseMetalResource(StencilTexture);
+	}
+	
 	if (!(Flags & TexCreate_Presentable) && MSAATexture != Texture)
 	{
 		SafeReleaseMetalResource(Texture);
@@ -821,6 +825,8 @@ void FMetalSurface::Unlock(uint32 MipIndex, uint32 ArrayIndex)
 #if PLATFORM_MAC
 		if(Texture.storageMode == MTLStorageModePrivate)
 		{
+			SCOPED_AUTORELEASE_POOL;
+			
 			// @todo: gather these all up over a frame
 			id<MTLCommandBuffer> CommandBuffer = GetMetalDeviceContext().CreateCommandBuffer(false/*bRetainReferences*/);
 			
@@ -830,7 +836,7 @@ void FMetalSurface::Unlock(uint32 MipIndex, uint32 ArrayIndex)
 			[Blitter copyFromBuffer:LockedMemory[MipIndex] sourceOffset:0 sourceBytesPerRow:Stride sourceBytesPerImage:BytesPerImage sourceSize:Region.size toTexture:Texture destinationSlice:ArrayIndex destinationLevel:MipIndex destinationOrigin:Region.origin];
 			
 			[Blitter endEncoding];
-			[CommandBuffer commit];
+			GetMetalDeviceContext().GetCommandList().Commit(CommandBuffer, false);
 		}
 		else
 #else
@@ -925,6 +931,28 @@ id<MTLTexture> FMetalSurface::GetDrawableTexture()
 	return Texture;
 }
 
+void FMetalSurface::UpdateSRV()
+{
+#if METAL_API_1_1
+	if (PixelFormat == PF_X24_G8 && Texture != StencilTexture && GetMetalDeviceContext().SupportsFeature(EMetalFeaturesDepthStencilBlitOptions))
+	{
+		// In this case StencilTexture is the source depth/stencil texture and Texture is our target Stencil-only copy
+		
+		id<MTLBlitCommandEncoder> BlitEncoder = FMetalContext::GetCurrentContext()->GetBlitContext();
+		
+		uint32 SizePerImage = Texture.width * Texture.height;
+		FMetalPooledBuffer Buf = GetMetalDeviceContext().CreatePooledBuffer(FMetalPooledBufferArgs(GetMetalDeviceContext().GetDevice(), SizePerImage, MTLStorageModeShared));
+		[Buf.Buffer retain];
+		
+		[BlitEncoder copyFromTexture:StencilTexture sourceSlice:0 sourceLevel:0 sourceOrigin:MTLOriginMake(0,0,0) sourceSize:MTLSizeMake(Texture.width, Texture.height, 1) toBuffer:Buf.Buffer destinationOffset:0 destinationBytesPerRow:Texture.width destinationBytesPerImage:SizePerImage options:MTLBlitOptionStencilFromDepthStencil];
+		
+		[BlitEncoder copyFromBuffer:Buf.Buffer sourceOffset:0 sourceBytesPerRow:Texture.width sourceBytesPerImage:SizePerImage sourceSize:MTLSizeMake(Texture.width, Texture.height, 1) toTexture:Texture destinationSlice:0 destinationLevel:0 destinationOrigin:MTLOriginMake(0,0,0)];
+		
+		SafeReleasePooledBuffer(Buf.Buffer);
+	}
+#endif
+}
+
 
 /*-----------------------------------------------------------------------------
 	Texture allocator support.
@@ -1010,7 +1038,7 @@ void FMetalDynamicRHI::RHIGenerateMips(FTextureRHIParamRef SourceSurfaceRHI)
 		
 		// kick it off!
 		[Blitter endEncoding];
-		[CommandBuffer commit];
+		GetMetalDeviceContext().GetCommandList().Commit(CommandBuffer, false);
 		
 		if (GIsRHIInitialized)
 		{
@@ -1083,8 +1111,8 @@ FTexture2DRHIRef FMetalDynamicRHI::RHIAsyncReallocateTexture2D(FTexture2DRHIPara
 					sourceOrigin:Origin
 					  sourceSize:MTLSizeMake(MipSizeX, MipSizeY, 1)
 					   toTexture:NewTexture->Surface.Texture
-			    destinationSlice:SliceIndex
-			    destinationLevel:MipIndex + DestMipOffset
+				destinationSlice:SliceIndex
+				destinationLevel:MipIndex + DestMipOffset
 			   destinationOrigin:Origin];
 	}
 
@@ -1097,7 +1125,7 @@ FTexture2DRHIRef FMetalDynamicRHI::RHIAsyncReallocateTexture2D(FTexture2DRHIPara
 
 	// kick it off!
 	[Blitter endEncoding];
-	[CommandBuffer commit];
+	GetMetalDeviceContext().GetCommandList().Commit(CommandBuffer, false);
 
 	return NewTexture;
 }
@@ -1192,6 +1220,8 @@ void FMetalDynamicRHI::RHIUpdateTexture2D(FTexture2DRHIParamRef TextureRHI, uint
 #if PLATFORM_MAC
 	if(Tex.storageMode == MTLStorageModePrivate)
 	{
+		SCOPED_AUTORELEASE_POOL;
+		
 		// @todo: gather these all up over a frame
 		id<MTLCommandBuffer> CommandBuffer = GetMetalDeviceContext().CreateCommandBuffer(false/*bRetainReferences*/);
 		
@@ -1214,7 +1244,7 @@ void FMetalDynamicRHI::RHIUpdateTexture2D(FTexture2DRHIParamRef TextureRHI, uint
 		[Blitter copyFromBuffer:LockedMemory sourceOffset:0 sourceBytesPerRow:SourcePitch sourceBytesPerImage:BytesPerImage sourceSize:Region.size toTexture:Tex destinationSlice:0 destinationLevel:MipIndex destinationOrigin:Region.origin];
 		
 		[Blitter endEncoding];
-		[CommandBuffer commit];
+		GetMetalDeviceContext().GetCommandList().Commit(CommandBuffer, false);
 		
 		GetMetalDeviceContext().ReleaseObject(LockedMemory);
 	}
@@ -1236,6 +1266,8 @@ void FMetalDynamicRHI::RHIUpdateTexture3D(FTexture3DRHIParamRef TextureRHI,uint3
 #if PLATFORM_MAC
 	if(Tex.storageMode == MTLStorageModePrivate)
 	{
+		SCOPED_AUTORELEASE_POOL;
+		
 		// @todo: gather these all up over a frame
 		id<MTLCommandBuffer> CommandBuffer = GetMetalDeviceContext().CreateCommandBuffer(false/*bRetainReferences*/);
 		
@@ -1258,7 +1290,7 @@ void FMetalDynamicRHI::RHIUpdateTexture3D(FTexture3DRHIParamRef TextureRHI,uint3
 		[Blitter copyFromBuffer:LockedMemory sourceOffset:0 sourceBytesPerRow:SourceRowPitch sourceBytesPerImage:BytesPerImage sourceSize:Region.size toTexture:Tex destinationSlice:0 destinationLevel:MipIndex destinationOrigin:Region.origin];
 		
 		[Blitter endEncoding];
-		[CommandBuffer commit];
+		GetMetalDeviceContext().GetCommandList().Commit(CommandBuffer, false);
 		
 		GetMetalDeviceContext().ReleaseObject(LockedMemory);
 	}
diff --git a/Engine/Source/Runtime/Apple/MetalRHI/Private/MetalUAV.cpp b/Engine/Source/Runtime/Apple/MetalRHI/Private/MetalUAV.cpp
--- a/Engine/Source/Runtime/Apple/MetalRHI/Private/MetalUAV.cpp
+++ b/Engine/Source/Runtime/Apple/MetalRHI/Private/MetalUAV.cpp
@@ -12,6 +12,12 @@ FMetalShaderResourceView::~FMetalShaderResourceView()
 		{
 			[TextureView->Texture release];
 			TextureView->Texture = nil;
+			
+			[TextureView->StencilTexture release];
+			TextureView->StencilTexture = nil;
+			
+			[TextureView->MSAATexture release];
+			TextureView->MSAATexture = nil;
 		}
 		delete TextureView;
 		TextureView = nullptr;
diff --git a/Engine/Source/Runtime/Apple/MetalRHI/Private/MetalUniformBuffer.cpp b/Engine/Source/Runtime/Apple/MetalRHI/Private/MetalUniformBuffer.cpp
--- a/Engine/Source/Runtime/Apple/MetalRHI/Private/MetalUniformBuffer.cpp
+++ b/Engine/Source/Runtime/Apple/MetalRHI/Private/MetalUniformBuffer.cpp
@@ -196,12 +196,12 @@ id<MTLBuffer> SuballocateUB(uint32 Size, uint32& OutOffset)
 }
 
 
-FMetalUniformBuffer::FMetalUniformBuffer(const void* Contents, const FRHIUniformBufferLayout& Layout, EUniformBufferUsage Usage)
+FMetalUniformBuffer::FMetalUniformBuffer(const void* Contents, const FRHIUniformBufferLayout& Layout, EUniformBufferUsage InUsage)
 	: FRHIUniformBuffer(Layout)
 	, Buffer(nil)
 	, Offset(0)
 	, Size(Layout.ConstantBufferSize)
-	, LastCachedFrame(INDEX_NONE)
+	, Usage(InUsage)
 {
 	if (Layout.ConstantBufferSize > 0)
 	{
@@ -269,79 +269,19 @@ FMetalUniformBuffer::FMetalUniformBuffer(const void* Contents, const FRHIUniform
 			check(InResources[i]);
 			ResourceTable[i] = InResources[i];
 		}
-		RawResourceTable.Empty(NumResources);
-		RawResourceTable.AddZeroed(NumResources);
 	}
 }
 
 FMetalUniformBuffer::~FMetalUniformBuffer()
 {
 	// don't need to free the ring buffer!
-	if (GIsRHIInitialized && Buffer != nil && Buffer != GetMetalDeviceContext().GetRingBuffer())
+	if (GIsRHIInitialized && Buffer != nil && !(Usage == UniformBuffer_SingleDraw && !GUseRHIThread))
 	{
 		check(Size <= 65536);
 		AddNewlyFreedBufferToUniformBufferPool(Buffer, Offset, Size);
 	}
 }
 
-void FMetalUniformBuffer::CacheResourcesInternal()
-{
-	const FRHIUniformBufferLayout& Layout = GetLayout();
-	int32 NumResources = Layout.Resources.Num();
-	const uint8* RESTRICT ResourceTypes = Layout.Resources.GetData();
-	const TRefCountPtr<FRHIResource>* RESTRICT Resources = ResourceTable.GetData();
-	void** RESTRICT RawResources = RawResourceTable.GetData();
-	float CurrentTime = FApp::GetCurrentTime();
-
-	// todo: Immutable resources, i.e. not textures, can be safely cached across frames.
-	// Texture streaming makes textures complicated :)
-	for (int32 i = 0; i < NumResources; ++i)
-	{
-		switch (ResourceTypes[i])
-		{
-			case UBMT_SRV:
-				{
-					NOT_SUPPORTED("FMetalUniformBuffer::CacheResourcesInternal UBMT_SRV");
-					
-					FMetalShaderResourceView* SRV = (FMetalShaderResourceView*)Resources[i].GetReference();
-					if (IsValidRef(SRV->SourceTexture))
-					{
-						FMetalSurface* Surface = SRV->TextureView;
-						RawResources[i] = Surface;
-					}
-					else
-					{
-						RawResources[i] = &SRV->SourceVertexBuffer->Buffer;
-					}
-				}
-				break;
-
-			case UBMT_TEXTURE:
-				{
-					FRHITexture* TextureRHI = (FRHITexture*)Resources[i].GetReference();
-					TextureRHI->SetLastRenderTime(CurrentTime);
-					RawResources[i] = TextureRHI;
-				}
-				break;
-
-			case UBMT_UAV:
-				NOT_SUPPORTED("FMetalUniformBuffer::CacheResourcesInternal UBMT_UAV");
-				RawResources[i] = 0;
-				break;
-
-			case UBMT_SAMPLER:
-				RawResources[i] = (FMetalSamplerState*)Resources[i].GetReference();
-				break;
-
-			default:
-				check(0);
-				break;
-		}
-	}
-}
-
-
-
 FUniformBufferRHIRef FMetalDynamicRHI::RHICreateUniformBuffer(const void* Contents, const FRHIUniformBufferLayout& Layout, EUniformBufferUsage Usage)
 {
 	check(IsInRenderingThread() || IsInRHIThread());
diff --git a/Engine/Source/Runtime/Apple/MetalRHI/Private/MetalViewport.cpp b/Engine/Source/Runtime/Apple/MetalRHI/Private/MetalViewport.cpp
--- a/Engine/Source/Runtime/Apple/MetalRHI/Private/MetalViewport.cpp
+++ b/Engine/Source/Runtime/Apple/MetalRHI/Private/MetalViewport.cpp
@@ -110,17 +110,20 @@ id<MTLDrawable> FMetalViewport::GetDrawable()
 {
 	if (!Drawable)
 	{
-		uint32 IdleStart = FPlatformTime::Cycles();
-		
-#if PLATFORM_MAC
-		CAMetalLayer* CurrentLayer = (CAMetalLayer*)[View layer];
-		Drawable = CurrentLayer ? [[CurrentLayer nextDrawable] retain] : nil;
-#else
-		Drawable = [[[IOSAppDelegate GetDelegate].IOSView MakeDrawable] retain];
-#endif
-		
-		GRenderThreadIdle[ERenderThreadIdleTypes::WaitingForGPUPresent] += FPlatformTime::Cycles() - IdleStart;
-		GRenderThreadNumIdle[ERenderThreadIdleTypes::WaitingForGPUPresent]++;
+		@autoreleasepool
+		{
+			uint32 IdleStart = FPlatformTime::Cycles();
+			
+	#if PLATFORM_MAC
+			CAMetalLayer* CurrentLayer = (CAMetalLayer*)[View layer];
+			Drawable = CurrentLayer ? [[CurrentLayer nextDrawable] retain] : nil;
+	#else
+			Drawable = [[[IOSAppDelegate GetDelegate].IOSView MakeDrawable] retain];
+	#endif
+			
+			GRenderThreadIdle[ERenderThreadIdleTypes::WaitingForGPUPresent] += FPlatformTime::Cycles() - IdleStart;
+			GRenderThreadNumIdle[ERenderThreadIdleTypes::WaitingForGPUPresent]++;
+		}
 	}
 	return Drawable;
 }
diff --git a/Engine/Source/Runtime/Apple/MetalRHI/Public/MetalResources.h b/Engine/Source/Runtime/Apple/MetalRHI/Public/MetalResources.h
--- a/Engine/Source/Runtime/Apple/MetalRHI/Public/MetalResources.h
+++ b/Engine/Source/Runtime/Apple/MetalRHI/Public/MetalResources.h
@@ -10,7 +10,8 @@
 #include "MetalShaderResources.h"
 #include "ShaderCache.h"
 
-#define METAL_SUPPORTS_PARALLEL_RHI_EXECUTE 0
+/** Parallel execution is available on Mac but not iOS for the moment - it needs to be tested because it isn't cost-free */
+#define METAL_SUPPORTS_PARALLEL_RHI_EXECUTE PLATFORM_MAC
 
 class FMetalContext;
 
@@ -201,6 +202,9 @@ public:
 	
 	/** Gets the drawable texture if this is a back-buffer surface. */
 	id<MTLTexture> GetDrawableTexture();
+	
+	/** Updates an SRV surface's internal data if required. */
+	void UpdateSRV();
 
 	ERHIResourceType Type;
 	EPixelFormat PixelFormat;
@@ -500,16 +504,6 @@ public:
 	~FMetalUniformBuffer();
 	
 
-	/** Cache resources if needed. */
-	inline void CacheResources(uint32 InFrameCounter)
-	{
-		if (InFrameCounter == INDEX_NONE || LastCachedFrame != InFrameCounter)
-		{
-			CacheResourcesInternal();
-			LastCachedFrame = InFrameCounter;
-		}
-	}
-
 	bool IsConstantBuffer() const
 	{
 		return Buffer.length > 0;
@@ -524,19 +518,12 @@ public:
 
 	uint32 Size; // @todo zebra: HACK! This should be removed and the code that uses it should be changed to use GetSize() instead once we fix the problem with FRHIUniformBufferLayout being released too early
 
+	/** The intended usage of the uniform buffer. */
+	EUniformBufferUsage Usage;
+	
 	/** Resource table containing RHI references. */
 	TArray<TRefCountPtr<FRHIResource> > ResourceTable;
 
-	/** Raw resource table, cached once per frame. */
-	TArray<void*> RawResourceTable;
-
-	/** The frame in which RawResourceTable was last cached. */
-	uint32 LastCachedFrame;
-
-private:
-	/** Actually cache resources. */
-	void CacheResourcesInternal();
-
 };
 
 
diff --git a/Engine/Source/Runtime/Core/Private/Mac/CocoaThread.cpp b/Engine/Source/Runtime/Core/Private/Mac/CocoaThread.cpp
--- a/Engine/Source/Runtime/Core/Private/Mac/CocoaThread.cpp
+++ b/Engine/Source/Runtime/Core/Private/Mac/CocoaThread.cpp
@@ -457,7 +457,11 @@ void RunGameThread(id Target, SEL Selector)
 	
 	// Create a separate game thread and set it to the stack size to be the same as the main thread default of 8MB ( http://developer.apple.com/library/mac/#qa/qa1419/_index.html )
 	FCocoaGameThread* GameThread = [[FCocoaGameThread alloc] initWithTarget:Target selector:Selector object:nil];
-	[GameThread setStackSize:16*1024*1024];
+#if UE_BUILD_DEBUG
+	[GameThread setStackSize:64*1024*1024];
+#else
+	[GameThread setStackSize:32*1024*1024];
+#endif
 	[GameThread start];
 #else
 	[Target performSelector:Selector withObject:nil];
diff --git a/Engine/Source/Runtime/Core/Private/Mac/MacPlatformSplash.cpp b/Engine/Source/Runtime/Core/Private/Mac/MacPlatformSplash.cpp
--- a/Engine/Source/Runtime/Core/Private/Mac/MacPlatformSplash.cpp
+++ b/Engine/Source/Runtime/Core/Private/Mac/MacPlatformSplash.cpp
@@ -311,13 +311,15 @@ void FMacPlatformSplash::Hide()
 {
 	if (GSplashWindow)
 	{
-		SCOPED_AUTORELEASE_POOL;
+		MainThreadCall(^{
+			SCOPED_AUTORELEASE_POOL;
 
-		[GSplashWindow close];
-		GSplashWindow = NULL;
+			[GSplashWindow close];
+			GSplashWindow = NULL;
 
-		[GSplashScreenImage release];
-		GSplashScreenImage = NULL;
+			[GSplashScreenImage release];
+			GSplashScreenImage = NULL;
+		}, NSDefaultRunLoopMode, true);
 
 		FPlatformMisc::PumpMessages(true);
 	}
diff --git a/Engine/Source/Runtime/Core/Private/Mac/MacWindow.cpp b/Engine/Source/Runtime/Core/Private/Mac/MacWindow.cpp
--- a/Engine/Source/Runtime/Core/Private/Mac/MacWindow.cpp
+++ b/Engine/Source/Runtime/Core/Private/Mac/MacWindow.cpp
@@ -115,11 +115,11 @@ void FMacWindow::Initialize( FMacApplication* const Application, const TSharedRe
 						break;
 
 					case EWindowType::Menu:
-						WindowLevel = NSMainMenuWindowLevel;
+						WindowLevel = NSStatusWindowLevel;
 						break;
 
 					case EWindowType::ToolTip:
-						WindowLevel = NSStatusWindowLevel;
+						WindowLevel = NSPopUpMenuWindowLevel;
 						break;
 
 					case EWindowType::Notification:
@@ -127,7 +127,7 @@ void FMacWindow::Initialize( FMacApplication* const Application, const TSharedRe
 						break;
 
 					case EWindowType::CursorDecorator:
-						WindowLevel = NSPopUpMenuWindowLevel;
+						WindowLevel = NSMainMenuWindowLevel;
 						break;
 				}
 			}
diff --git a/Engine/Source/Runtime/Core/Public/IOS/IOSPlatformMisc.h b/Engine/Source/Runtime/Core/Public/IOS/IOSPlatformMisc.h
--- a/Engine/Source/Runtime/Core/Public/IOS/IOSPlatformMisc.h
+++ b/Engine/Source/Runtime/Core/Public/IOS/IOSPlatformMisc.h
@@ -8,6 +8,31 @@
 #include "GenericPlatform/GenericPlatformMisc.h"
 #include "IOS/IOSSystemIncludes.h"
 
+#ifdef __OBJC__
+
+class FScopeAutoreleasePool
+{
+public:
+
+	FScopeAutoreleasePool()
+	{
+		Pool = [[NSAutoreleasePool alloc] init];
+	}
+
+	~FScopeAutoreleasePool()
+	{
+		[Pool release];
+	}
+
+private:
+
+	NSAutoreleasePool*	Pool;
+};
+
+#define SCOPED_AUTORELEASE_POOL const FScopeAutoreleasePool PREPROCESSOR_JOIN(Pool,__LINE__);
+
+#endif // __OBJC__
+
 /**
 * iOS implementation of the misc OS functions
 **/
diff --git a/Engine/Source/Runtime/Core/Public/Mac/MacPlatformCompilerPreSetup.h b/Engine/Source/Runtime/Core/Public/Mac/MacPlatformCompilerPreSetup.h
--- a/Engine/Source/Runtime/Core/Public/Mac/MacPlatformCompilerPreSetup.h
+++ b/Engine/Source/Runtime/Core/Public/Mac/MacPlatformCompilerPreSetup.h
@@ -105,3 +105,9 @@
 #pragma clang diagnostic ignored "-Winconsistent-missing-override"
 #pragma clang diagnostic ignored "-Wundefined-bool-conversion"
 #pragma clang diagnostic ignored "-Wunused-local-typedef"
+
+// We can pragma optimisation's on and off as of Apple LLVM 7.3.0 but not before.
+#if __clang_major__ >= 7 && __clang_minor__ >= 3
+#define PRAGMA_DISABLE_OPTIMIZATION_ACTUAL _Pragma("clang optimize off")
+#define PRAGMA_ENABLE_OPTIMIZATION_ACTUAL  _Pragma("clang optimize on")
+#endif
diff --git a/Engine/Source/Runtime/CoreUObject/Private/UObject/ScriptCore.cpp b/Engine/Source/Runtime/CoreUObject/Private/UObject/ScriptCore.cpp
--- a/Engine/Source/Runtime/CoreUObject/Private/UObject/ScriptCore.cpp
+++ b/Engine/Source/Runtime/CoreUObject/Private/UObject/ScriptCore.cpp
@@ -14,6 +14,8 @@ DEFINE_LOG_CATEGORY_STATIC(LogScriptCore, Log, All);
 
 DECLARE_CYCLE_STAT(TEXT("Blueprint Time"),STAT_BlueprintTime,STATGROUP_Game);
 
+#define LOCTEXT_NAMESPACE "ScriptCore"
+
 #if TOTAL_OVERHEAD_SCRIPT_STATS
 COREUOBJECT_API FBlueprintEventTimer::FScopedVMTimer* FBlueprintEventTimer::ActiveVMTimer = nullptr;
 COREUOBJECT_API FBlueprintEventTimer::FPausableScopeTimer* FBlueprintEventTimer::ActiveTimer = nullptr;
@@ -121,7 +123,7 @@ void FBlueprintCoreDelegates::ThrowScriptException(const UObject* ActiveObject,
 
 	if (bShouldLogWarning)
 	{
-		UE_SUPPRESS(LogScript, Warning, const_cast<FFrame*>(&StackFrame)->Logf(TEXT("%s"), *(Info.GetDescription())));
+		UE_SUPPRESS(LogScript, Warning, const_cast<FFrame*>(&StackFrame)->Logf(TEXT("%s"), *(Info.GetDescription().ToString())));
 	}
 
 	// cant fire arbitrary delegates here off the game thread
@@ -841,8 +843,13 @@ void UObject::ProcessInternal( FFrame& Stack, RESULT_DECL )
 			ClearReturnValue(ReturnProp, RESULT_PARAM);
 
 			// Notify anyone who cares that we've had a fatal error, so we can shut down PIE, etc
-			const FString Desc = FString::Printf(TEXT("Infinite script recursion (%i calls) detected"), RECURSE_LIMIT);
-			FBlueprintExceptionInfo InfiniteRecursionExceptionInfo(EBlueprintExceptionType::InfiniteLoop, Desc);
+			FBlueprintExceptionInfo InfiniteRecursionExceptionInfo(
+				EBlueprintExceptionType::InfiniteLoop, 
+				FText::Format(
+					LOCTEXT("InfiniteLoop", "Infinite script recursion ({0} calls) detected - see log for stack trace"), 
+					FText::AsNumber(RECURSE_LIMIT)
+				)
+			);
 			FBlueprintCoreDelegates::ThrowScriptException(this, Stack, InfiniteRecursionExceptionInfo);
 
 			// This flag prevents repeated warnings of infinite loop, script exception handler 
@@ -866,8 +873,13 @@ void UObject::ProcessInternal( FFrame& Stack, RESULT_DECL )
 				ClearReturnValue(ReturnProp, RESULT_PARAM);
 
 				// Notify anyone who cares that we've had a fatal error, so we can shut down PIE, etc
-				const FString Desc = FString::Printf(TEXT("Runaway loop detected (over %i iterations)"), GMaximumScriptLoopIterations );
-				FBlueprintExceptionInfo RunawayLoopExceptionInfo(EBlueprintExceptionType::InfiniteLoop, Desc);
+				FBlueprintExceptionInfo RunawayLoopExceptionInfo(
+					EBlueprintExceptionType::InfiniteLoop, 
+					FText::Format(
+						LOCTEXT("RunawayLoop", "Runaway loop detected (over {0} iterations) - see log for stack trace"),
+						FText::AsNumber(GMaximumScriptLoopIterations)
+					)
+				);
 
 				// Need to reset Runaway counter BEFORE throwing script exception, because the exception causes a modal dialog,
 				// and other scripts running will then erroneously think they are also "runaway".
@@ -1292,9 +1304,9 @@ void UObject::execInstanceVariable(FFrame& Stack, RESULT_DECL)
 	UProperty* VarProperty = (UProperty*)Stack.ReadObject();
 	Stack.MostRecentProperty = VarProperty;
 
-	if (VarProperty == nullptr)
+	if (VarProperty == nullptr || !IsA((UClass*)VarProperty->GetOuter()))
 	{
-		static FBlueprintExceptionInfo ExceptionInfo(EBlueprintExceptionType::AccessViolation, TEXT("Attempt to access missing property. If this is a packaged/cooked build, are you attempting to use an editor-only property?"));
+		FBlueprintExceptionInfo ExceptionInfo(EBlueprintExceptionType::AccessViolation, LOCTEXT("MissingProperty", "Attempted to access missing property. If this is a packaged/cooked build, are you attempting to use an editor-only property?"));
 		FBlueprintCoreDelegates::ThrowScriptException(this, Stack, ExceptionInfo);
 
 		Stack.MostRecentPropertyAddress = nullptr;
@@ -1319,23 +1331,23 @@ void UObject::execDefaultVariable(FFrame& Stack, RESULT_DECL)
 	Stack.MostRecentProperty = VarProperty;
 	Stack.MostRecentPropertyAddress = nullptr;
 
-	if(VarProperty == nullptr)
+	UObject* DefaultObject = nullptr;
+	if (HasAnyFlags(RF_ClassDefaultObject))
 	{
-		static FBlueprintExceptionInfo ExceptionInfo(EBlueprintExceptionType::AccessViolation, TEXT("Attempt to access a missing property. If this is a packaged/cooked build, are you attempting to use an editor-only property?"));
-		FBlueprintCoreDelegates::ThrowScriptException(this, Stack, ExceptionInfo);
+		DefaultObject = this;
 	}
 	else
 	{
-		UObject* DefaultObject = nullptr;
-		if(HasAnyFlags(RF_ClassDefaultObject))
-		{
-			DefaultObject = this;
-		}
-		else
-		{
-			// @todo - allow access to archetype properties through object references?
-		}
+		// @todo - allow access to archetype properties through object references?
+	}
 
+	if (VarProperty == nullptr || (DefaultObject && !DefaultObject->IsA((UClass*)VarProperty->GetOuter())))
+	{
+		FBlueprintExceptionInfo ExceptionInfo(EBlueprintExceptionType::AccessViolation, LOCTEXT("MissingPropertyDefaultObject", "Attempted to access a missing property on a CDO. If this is a packaged/cooked build, are you attempting to use an editor-only property?"));
+		FBlueprintCoreDelegates::ThrowScriptException(this, Stack, ExceptionInfo);
+	}
+	else
+	{
 		if(DefaultObject != nullptr)
 		{
 			Stack.MostRecentPropertyAddress = VarProperty->ContainerPtrToValuePtr<uint8>(DefaultObject);
@@ -1346,7 +1358,7 @@ void UObject::execDefaultVariable(FFrame& Stack, RESULT_DECL)
 		}
 		else
 		{
-			static FBlueprintExceptionInfo ExceptionInfo(EBlueprintExceptionType::AccessViolation, TEXT("Attempt to access a default property through an invalid or otherwise unsupported context."));
+			FBlueprintExceptionInfo ExceptionInfo(EBlueprintExceptionType::AccessViolation, LOCTEXT("AccessNoneDefaultObject", "Accessed None attempting to read a default property"));
 			FBlueprintCoreDelegates::ThrowScriptException(this, Stack, ExceptionInfo);
 		}
 	}
@@ -1412,13 +1424,18 @@ void UObject::execClassContext(FFrame& Stack, RESULT_DECL)
 	{
 		if (Stack.MostRecentProperty != NULL)
 		{
-			const FString Desc = FString::Printf(TEXT("Accessed null class context '%s'"), *Stack.MostRecentProperty->GetName());
-			FBlueprintExceptionInfo ExceptionInfo(EBlueprintExceptionType::AccessViolation, Desc);
+			FBlueprintExceptionInfo ExceptionInfo(
+				EBlueprintExceptionType::AccessViolation, 
+				FText::Format(
+					LOCTEXT("AccessedNoneClass", "Accessed None trying to read Class from property {0}"), 
+					FText::FromString(Stack.MostRecentProperty->GetName())
+				)
+			);
 			FBlueprintCoreDelegates::ThrowScriptException(this, Stack, ExceptionInfo);
 		}
 		else
 		{
-			static FBlueprintExceptionInfo ExceptionInfo(EBlueprintExceptionType::AccessViolation, TEXT("Accessed null class context"));
+			FBlueprintExceptionInfo ExceptionInfo(EBlueprintExceptionType::AccessViolation, LOCTEXT("AccessedNoneClassUnknownProperty", "Accessed None reading a Class"));
 			FBlueprintCoreDelegates::ThrowScriptException(this, Stack, ExceptionInfo);
 		}
 
@@ -1472,7 +1489,7 @@ void UObject::execBreakpoint( FFrame& Stack, RESULT_DECL )
 #if WITH_EDITORONLY_DATA
 	if (GIsEditor)
 	{
-		static FBlueprintExceptionInfo BreakpointExceptionInfo(EBlueprintExceptionType::Breakpoint);
+		FBlueprintExceptionInfo BreakpointExceptionInfo(EBlueprintExceptionType::Breakpoint);
 		FBlueprintCoreDelegates::ThrowScriptException(this, Stack, BreakpointExceptionInfo);
 	}
 #endif
@@ -1484,7 +1501,7 @@ void UObject::execTracepoint( FFrame& Stack, RESULT_DECL )
 #if WITH_EDITORONLY_DATA
 	if (GIsEditor)
 	{
-		static FBlueprintExceptionInfo TracepointExceptionInfo(EBlueprintExceptionType::Tracepoint);
+		FBlueprintExceptionInfo TracepointExceptionInfo(EBlueprintExceptionType::Tracepoint);
 		FBlueprintCoreDelegates::ThrowScriptException(this, Stack, TracepointExceptionInfo);
 	}
 #endif
@@ -1496,7 +1513,7 @@ void UObject::execWireTracepoint( FFrame& Stack, RESULT_DECL )
 #if WITH_EDITORONLY_DATA
 	if (GIsEditor)
 	{
-		static FBlueprintExceptionInfo TracepointExceptionInfo(EBlueprintExceptionType::WireTracepoint);
+		FBlueprintExceptionInfo TracepointExceptionInfo(EBlueprintExceptionType::WireTracepoint);
 		FBlueprintCoreDelegates::ThrowScriptException(this, Stack, TracepointExceptionInfo);
 	}
 #endif
@@ -1512,12 +1529,12 @@ void UObject::execInstrumentation( FFrame& Stack, RESULT_DECL )
 	{
 		if (EventType == EScriptInstrumentation::NodeEntry)
 		{
-			static FBlueprintExceptionInfo TracepointExceptionInfo(EBlueprintExceptionType::Tracepoint);
+			FBlueprintExceptionInfo TracepointExceptionInfo(EBlueprintExceptionType::Tracepoint);
 			FBlueprintCoreDelegates::ThrowScriptException(this, Stack, TracepointExceptionInfo);
 		}
 		else if (EventType == EScriptInstrumentation::NodeExit)
 		{
-			static FBlueprintExceptionInfo WiretraceExceptionInfo(EBlueprintExceptionType::WireTracepoint);
+			FBlueprintExceptionInfo WiretraceExceptionInfo(EBlueprintExceptionType::WireTracepoint);
 			FBlueprintCoreDelegates::ThrowScriptException(this, Stack, WiretraceExceptionInfo);
 		}
 	}
@@ -1695,7 +1712,13 @@ void UObject::execSwitchValue(FFrame& Stack, RESULT_DECL)
 	uint8* IndexAdress = Stack.MostRecentPropertyAddress;
 	if (!ensure(IndexAdress))
 	{
-		FBlueprintExceptionInfo ExceptionInfo(EBlueprintExceptionType::NonFatalError, TEXT("Switch - Unknown index"));
+		FBlueprintExceptionInfo ExceptionInfo(
+			EBlueprintExceptionType::NonFatalError, 
+			FText::Format(
+				LOCTEXT("SwitchValueIndex", "Switch statement failed to read property for index value for index property {0}"),
+				FText::FromString(IndexProperty->GetName())
+			)
+		);
 		FBlueprintCoreDelegates::ThrowScriptException(this, Stack, ExceptionInfo);
 	}
 
@@ -1727,8 +1750,13 @@ void UObject::execSwitchValue(FFrame& Stack, RESULT_DECL)
 	}
 	else
 	{
-		FBlueprintExceptionInfo ExceptionInfo(EBlueprintExceptionType::NonFatalError, FString::Printf(TEXT("Switch - Out of bounds index %d (cases %d-%d)"),
-			IndexAdress ? *IndexAdress : -1, 0, NumCases-1));
+		FBlueprintExceptionInfo ExceptionInfo(
+			EBlueprintExceptionType::NonFatalError, 
+			FText::Format(
+				LOCTEXT("SwitchValueOutOfBounds", "Switch statement failed to match case for index property {0}"),
+				FText::FromString(IndexProperty->GetName())
+			)
+		);
 		FBlueprintCoreDelegates::ThrowScriptException(this, Stack, ExceptionInfo);
 
 		// get default value
@@ -1739,8 +1767,6 @@ IMPLEMENT_VM_FUNCTION(EX_SwitchValue, execSwitchValue);
 
 void UObject::execLet(FFrame& Stack, RESULT_DECL)
 {
-	checkSlow(!dynamic_cast<UBoolProperty*>(this));
-
 	Stack.MostRecentProperty = nullptr;
 	UProperty* LocallyKnownProperty = Stack.ReadPropertyUnchecked();
 
@@ -1752,7 +1778,9 @@ void UObject::execLet(FFrame& Stack, RESULT_DECL)
 	uint8* LocalTempResult = nullptr;
 	if (Stack.MostRecentPropertyAddress == nullptr)
 	{
-		static FBlueprintExceptionInfo ExceptionInfo(EBlueprintExceptionType::AccessViolation, TEXT("Attempt to assign variable through None"));
+		FBlueprintExceptionInfo ExceptionInfo(
+			EBlueprintExceptionType::AccessViolation, 
+			LOCTEXT("LetAccessNone", "Attempted to assign to None"));
 		FBlueprintCoreDelegates::ThrowScriptException(this, Stack, ExceptionInfo);
 
 		if (LocallyKnownProperty)
@@ -1786,7 +1814,9 @@ void UObject::execLetObj( FFrame& Stack, RESULT_DECL )
 
 	if (Stack.MostRecentPropertyAddress == NULL)
 	{
-		static FBlueprintExceptionInfo ExceptionInfo(EBlueprintExceptionType::AccessViolation, TEXT("Attempt to assign variable through None"));
+		FBlueprintExceptionInfo ExceptionInfo(
+			EBlueprintExceptionType::AccessViolation, 
+			LOCTEXT("LetObjAccessNone", "Accessed None attempting to assign variable on an object"));
 		FBlueprintCoreDelegates::ThrowScriptException(this, Stack, ExceptionInfo);
 	}
 
@@ -1821,7 +1851,9 @@ void UObject::execLetWeakObjPtr( FFrame& Stack, RESULT_DECL )
 
 	if (Stack.MostRecentPropertyAddress == NULL)
 	{
-		static FBlueprintExceptionInfo ExceptionInfo(EBlueprintExceptionType::AccessViolation, TEXT("Attempt to assign variable through None"));
+		FBlueprintExceptionInfo ExceptionInfo(
+			EBlueprintExceptionType::AccessViolation,
+			LOCTEXT("LetWeakObjAccessNone", "Accessed None attempting to assign variable on a weakly referenced object"));
 		FBlueprintCoreDelegates::ThrowScriptException(this, Stack, ExceptionInfo);
 	}
 
@@ -1980,15 +2012,25 @@ void UObject::ProcessContextOpcode( FFrame& Stack, RESULT_DECL, bool bCanFailSil
 		{
 			if (NewContext && NewContext->IsPendingKill())
 			{
-				const FString Desc = FString::Printf(TEXT("Cannot access '%s'. It is pending kill. Property: '%s'")
-					, *GetNameSafe(NewContext), *GetNameSafe(Stack.MostRecentProperty));
-				FBlueprintExceptionInfo ExceptionInfo(EBlueprintExceptionType::AccessViolation, Desc);
+				FBlueprintExceptionInfo ExceptionInfo(
+					EBlueprintExceptionType::AccessViolation, 
+					FText::Format(
+						LOCTEXT("AccessPendingKill", "Attempted to access {0} via property {1}, but {0} is pending kill"),
+						FText::FromString( GetNameSafe(NewContext) ), 
+						FText::FromString( GetNameSafe(Stack.MostRecentProperty) )
+					)
+				);
 				FBlueprintCoreDelegates::ThrowScriptException(this, Stack, ExceptionInfo);
 			}
 			else if (Stack.MostRecentProperty != NULL)
 			{
-				const FString Desc = FString::Printf(TEXT("Accessed None '%s'"), *Stack.MostRecentProperty->GetName());
-				FBlueprintExceptionInfo ExceptionInfo(EBlueprintExceptionType::AccessViolation, Desc);
+				FBlueprintExceptionInfo ExceptionInfo(
+					EBlueprintExceptionType::AccessViolation, 
+					FText::Format( 
+						LOCTEXT("AccessNoneContext", "Accessed None trying to read property {0}"), 
+						FText::FromString( Stack.MostRecentProperty->GetName() )
+					)
+				);
 				FBlueprintCoreDelegates::ThrowScriptException(this, Stack, ExceptionInfo);
 			}
 			else
@@ -1997,7 +2039,10 @@ void UObject::ProcessContextOpcode( FFrame& Stack, RESULT_DECL, bool bCanFailSil
 				//   1. the context expression was a function call which returned an object
 				//   2. the context expression was a literal object reference
 				//   3. the context expression was an instance variable that no longer exists (it was editor-only, etc.)
-				static FBlueprintExceptionInfo ExceptionInfo(EBlueprintExceptionType::AccessViolation, TEXT("Accessed None"));
+				FBlueprintExceptionInfo ExceptionInfo(
+					EBlueprintExceptionType::AccessViolation, 
+					LOCTEXT("AccessNoneNoContext", "Accessed None")
+				);
 				FBlueprintCoreDelegates::ThrowScriptException(this, Stack, ExceptionInfo);
 			}
 		}
@@ -2042,7 +2087,13 @@ void UObject::execStructMemberContext(FFrame& Stack, RESULT_DECL)
 	else
 	{
 		// Access none
-		static FBlueprintExceptionInfo ExceptionInfo(EBlueprintExceptionType::AccessViolation, TEXT("Struct expression was None"));
+		FBlueprintExceptionInfo ExceptionInfo(
+			EBlueprintExceptionType::AccessViolation, 
+			FText::Format(
+				LOCTEXT("AccessNoneStructure", "Accessed None reading structure {0}"),
+				FText::FromString(StructProperty->GetName())
+			)
+		);
 		FBlueprintCoreDelegates::ThrowScriptException(this, Stack, ExceptionInfo);
 
 		Stack.MostRecentPropertyAddress = NULL;
@@ -2593,3 +2644,4 @@ void UObject::execInterfaceToObject(FFrame& Stack, RESULT_DECL)
 }
 IMPLEMENT_VM_FUNCTION( EX_InterfaceToObjCast, execInterfaceToObject );
 
+#undef LOCTEXT_NAMESPACE
diff --git a/Engine/Source/Runtime/CoreUObject/Public/UObject/Script.h b/Engine/Source/Runtime/CoreUObject/Public/UObject/Script.h
--- a/Engine/Source/Runtime/CoreUObject/Public/UObject/Script.h
+++ b/Engine/Source/Runtime/CoreUObject/Public/UObject/Script.h
@@ -310,7 +310,7 @@ public:
 	{
 	}
 
-	FBlueprintExceptionInfo(EBlueprintExceptionType::Type InEventType, const FString& InDescription)
+	FBlueprintExceptionInfo(EBlueprintExceptionType::Type InEventType, const FText& InDescription)
 		: EventType(InEventType)
 		, Description(InDescription)
 	{
@@ -321,13 +321,13 @@ public:
 		return EventType;
 	}
 
-	const FString& GetDescription() const
+	const FText& GetDescription() const
 	{
 		return Description;
 	}
 protected:
 	EBlueprintExceptionType::Type EventType;
-	FString Description;
+	FText Description;
 };
 
 // Information about a blueprint instrumentation event
diff --git a/Engine/Source/Runtime/Engine/Classes/Components/SkeletalMeshComponent.h b/Engine/Source/Runtime/Engine/Classes/Components/SkeletalMeshComponent.h
--- a/Engine/Source/Runtime/Engine/Classes/Components/SkeletalMeshComponent.h
+++ b/Engine/Source/Runtime/Engine/Classes/Components/SkeletalMeshComponent.h
@@ -90,6 +90,10 @@ private:
 	FVector WindDirection;
 	float WindAdaption;
 
+	/** Whether we have deferred cloth update transform. If true we will update any cloth collision when appropriate */
+	bool bPendingClothUpdateTransform;
+	ETeleportType PendingTeleportType;
+
 	friend class USkeletalMeshComponent;
 };
 
@@ -1259,10 +1263,10 @@ public:
 	 */
 	void UpdateClothStateAndSimulate(float DeltaTime, FTickFunction& ThisTickFunction);
 	/** 
-	 * Updates clothing actor's global pose.
-	 * So should be called when ComponentToWorld is changed.
+	 * Updates cloth collision outside the cloth asset (environment collision, child collision, etc...)
+	 * Should be called when scene changes or world position changes
 	 */
-	void UpdateClothTransform();
+	void UpdateClothTransform(ETeleportType TeleportType);
 
 	/** only check whether there are valid clothing actors or not */
 	bool HasValidClothingActors() const;
@@ -1409,6 +1413,11 @@ private:
 	void CompleteParallelBlendPhysics();
 	void PostBlendPhysics();
 
+#if WITH_APEX_CLOTHING
+	/** See UpdateClothTransform for documentation. */
+	void UpdateClothTransformImp();
+#endif
+
 	friend class FTickClothingTask;
 
 	// these are deprecated variables from removing SingleAnimSkeletalComponent
diff --git a/Engine/Source/Runtime/Engine/Classes/Kismet/KismetArrayLibrary.h b/Engine/Source/Runtime/Engine/Classes/Kismet/KismetArrayLibrary.h
--- a/Engine/Source/Runtime/Engine/Classes/Kismet/KismetArrayLibrary.h
+++ b/Engine/Source/Runtime/Engine/Classes/Kismet/KismetArrayLibrary.h
@@ -215,7 +215,7 @@ public:
  
  		Stack.MostRecentPropertyAddress = NULL;
  		Stack.StepCompiledIn<UProperty>(StorageSpace);
-		void* NewItemPtr = (Stack.MostRecentPropertyAddress != NULL) ? Stack.MostRecentPropertyAddress : StorageSpace;
+		void* NewItemPtr = (Stack.MostRecentPropertyAddress != NULL && Stack.MostRecentProperty->GetClass() == InnerProp->GetClass()) ? Stack.MostRecentPropertyAddress : StorageSpace;
  
  		P_FINISH;
 		P_NATIVE_BEGIN;
@@ -244,7 +244,7 @@ public:
 
 		Stack.MostRecentPropertyAddress = NULL;
 		Stack.StepCompiledIn<UProperty>(StorageSpace);
-		void* NewItemPtr = (Stack.MostRecentPropertyAddress != NULL) ? Stack.MostRecentPropertyAddress : StorageSpace;
+		void* NewItemPtr = (Stack.MostRecentPropertyAddress != NULL && Stack.MostRecentProperty->GetClass() == InnerProp->GetClass()) ? Stack.MostRecentPropertyAddress : StorageSpace;
 
 		P_FINISH;
 		P_NATIVE_BEGIN;
@@ -320,7 +320,7 @@ public:
 
 		Stack.MostRecentPropertyAddress = NULL;
 		Stack.StepCompiledIn<UProperty>(StorageSpace);
-		void* NewItemPtr = (Stack.MostRecentPropertyAddress != NULL) ? Stack.MostRecentPropertyAddress : StorageSpace;
+		void* NewItemPtr = (Stack.MostRecentPropertyAddress != NULL && Stack.MostRecentProperty->GetClass() == InnerProp->GetClass()) ? Stack.MostRecentPropertyAddress : StorageSpace;
 
 		P_GET_PROPERTY(UIntProperty, Index);
 		P_FINISH;
@@ -471,7 +471,7 @@ public:
 
 		Stack.MostRecentPropertyAddress = NULL;
 		Stack.StepCompiledIn<UProperty>(StorageSpace);
-		void* ItemPtr = (Stack.MostRecentPropertyAddress != NULL) ? Stack.MostRecentPropertyAddress : StorageSpace;
+		void* ItemPtr = (Stack.MostRecentPropertyAddress != NULL && Stack.MostRecentProperty->GetClass() == InnerProp->GetClass()) ? Stack.MostRecentPropertyAddress : StorageSpace;
 
 		P_FINISH;
 		P_NATIVE_BEGIN;
@@ -501,7 +501,7 @@ public:
 
 		Stack.MostRecentPropertyAddress = NULL;
 		Stack.StepCompiledIn<UProperty>(StorageSpace);
-		void* NewItemPtr = (Stack.MostRecentPropertyAddress != NULL) ? Stack.MostRecentPropertyAddress : StorageSpace;
+		void* NewItemPtr = (Stack.MostRecentPropertyAddress != NULL && Stack.MostRecentProperty->GetClass() == InnerProp->GetClass()) ? Stack.MostRecentPropertyAddress : StorageSpace;
 
 		P_GET_UBOOL(bSizeToFit);
 
diff --git a/Engine/Source/Runtime/Engine/Private/AI/Navigation/NavLinkProxy.cpp b/Engine/Source/Runtime/Engine/Private/AI/Navigation/NavLinkProxy.cpp
--- a/Engine/Source/Runtime/Engine/Private/AI/Navigation/NavLinkProxy.cpp
+++ b/Engine/Source/Runtime/Engine/Private/AI/Navigation/NavLinkProxy.cpp
@@ -104,6 +104,8 @@ void ANavLinkProxy::BeginPlay()
 	{
 		REDIRECT_OBJECT_TO_VLOG(this, NavSys);
 	}
+
+	Super::BeginPlay();
 }
 #endif // ENABLE_VISUAL_LOG
 
diff --git a/Engine/Source/Runtime/Engine/Private/AI/Navigation/RecastNavMesh.cpp b/Engine/Source/Runtime/Engine/Private/AI/Navigation/RecastNavMesh.cpp
--- a/Engine/Source/Runtime/Engine/Private/AI/Navigation/RecastNavMesh.cpp
+++ b/Engine/Source/Runtime/Engine/Private/AI/Navigation/RecastNavMesh.cpp
@@ -2117,7 +2117,7 @@ void ARecastNavMesh::PostEditChangeProperty(FPropertyChangedEvent& PropertyChang
 			if (RuntimeGeneration == ERuntimeGenerationType::Static)
 			{
 				const FString EngineIniFilename = FPaths::ConvertRelativePathToFull(GetDefault<UEngine>()->GetDefaultConfigFilename());
-				GConfig->SetString(TEXT("/Script/Engine.RecastNavMesh"), *NAME_RuntimeGeneration.ToString(), TEXT(""), *EngineIniFilename);
+				GConfig->SetString(TEXT("/Script/Engine.RecastNavMesh"), *NAME_RuntimeGeneration.ToString(), TEXT("Static"), *EngineIniFilename);
 				GConfig->Flush(false);
 			}
 		}
diff --git a/Engine/Source/Runtime/Engine/Private/Actor.cpp b/Engine/Source/Runtime/Engine/Private/Actor.cpp
--- a/Engine/Source/Runtime/Engine/Private/Actor.cpp
+++ b/Engine/Source/Runtime/Engine/Private/Actor.cpp
@@ -1440,33 +1440,19 @@ void AActor::DetachSceneComponentsFromParent(USceneComponent* InParentComponent,
 
 AActor* AActor::GetAttachParentActor() const
 {
-	if(RootComponent != NULL)
+	if (GetRootComponent() && GetRootComponent()->GetAttachParent())
 	{
-		for( const USceneComponent* Test=GetRootComponent()->GetAttachParent(); Test!=NULL; Test=Test->GetAttachParent())
-		{
-			AActor* TestOwner = Test->GetOwner();
-			if( TestOwner != this )
-			{
-				return TestOwner;
-			}
-		}
+		return GetRootComponent()->GetAttachParent()->GetOwner();
 	}
 
-	return NULL;
+	return nullptr;
 }
 
 FName AActor::GetAttachParentSocketName() const
 {
-	if(RootComponent != NULL)
+	if (GetRootComponent() && GetRootComponent()->GetAttachParent())
 	{
-		for( const USceneComponent* Test=GetRootComponent(); Test!=NULL; Test=Test->GetAttachParent() )
-		{
-			AActor* TestOwner = Test->GetOwner();
-			if( TestOwner != this )
-			{
-				return Test->GetAttachSocketName();
-			}
-		}
+		return GetRootComponent()->GetAttachSocketName();
 	}
 
 	return NAME_None;
diff --git a/Engine/Source/Runtime/Engine/Private/Animation/AnimInstanceProxy.cpp b/Engine/Source/Runtime/Engine/Private/Animation/AnimInstanceProxy.cpp
--- a/Engine/Source/Runtime/Engine/Private/Animation/AnimInstanceProxy.cpp
+++ b/Engine/Source/Runtime/Engine/Private/Animation/AnimInstanceProxy.cpp
@@ -125,17 +125,22 @@ void FAnimInstanceProxy::InitializeRootNode()
 				FAnimNode_Base* AnimNode = Property->ContainerPtrToValuePtr<FAnimNode_Base>(AnimInstanceObject);
 				if (AnimNode)
 				{
+					// Force our functions to be re-evaluated - this reinitialization may have been a 
+					// consequence of our class being recompiled and functions will be invalid in that
+					// case.
+					AnimNode->EvaluateGraphExposedInputs.bInitialized = false;
+
 					if (AnimNode->HasPreUpdate())
 					{
 						GameThreadPreUpdateNodes.Add(AnimNode);
 					}
 
-			if(Property->Struct->IsChildOf(FAnimNode_StateMachine::StaticStruct()))
-			{
+					if(Property->Struct->IsChildOf(FAnimNode_StateMachine::StaticStruct()))
+					{
 						FAnimNode_StateMachine* StateMachine = static_cast<FAnimNode_StateMachine*>(AnimNode);
-				StateMachine->CacheMachineDescription(AnimClassInterface);
-			}
-		}
+						StateMachine->CacheMachineDescription(AnimClassInterface);
+					}
+				}
 			}
 		}
 
diff --git a/Engine/Source/Runtime/Engine/Private/Animation/AnimNodeBase.cpp b/Engine/Source/Runtime/Engine/Private/Animation/AnimNodeBase.cpp
--- a/Engine/Source/Runtime/Engine/Private/Animation/AnimNodeBase.cpp
+++ b/Engine/Source/Runtime/Engine/Private/Animation/AnimNodeBase.cpp
@@ -83,7 +83,7 @@ void FAnimNode_Base::Initialize(const FAnimationInitializeContext& Context)
 
 bool FAnimNode_Base::IsLODEnabled(FAnimInstanceProxy* AnimInstanceProxy, int32 InLODThreshold)
 {
-	return (InLODThreshold == INDEX_NONE || AnimInstanceProxy->GetSkelMeshComponent()->PredictedLODLevel <= InLODThreshold);
+	return (InLODThreshold < 0 || AnimInstanceProxy->GetSkelMeshComponent()->PredictedLODLevel <= InLODThreshold);
 }
 
 /////////////////////////////////////////////////////
diff --git a/Engine/Source/Runtime/Engine/Private/Components/SkeletalMeshComponent.cpp b/Engine/Source/Runtime/Engine/Private/Components/SkeletalMeshComponent.cpp
--- a/Engine/Source/Runtime/Engine/Private/Components/SkeletalMeshComponent.cpp
+++ b/Engine/Source/Runtime/Engine/Private/Components/SkeletalMeshComponent.cpp
@@ -1142,6 +1142,18 @@ void USkeletalMeshComponent::UpdateClothSimulationContext()
 	{
 		CheckClothTeleport();
 		InternalClothSimulationContext.ClothTeleportMode = ClothTeleportMode;
+
+		if(InternalClothSimulationContext.bPendingClothUpdateTransform)	//it's possible we want to update cloth collision based on a pending transform
+		{
+			InternalClothSimulationContext.bPendingClothUpdateTransform = false;
+			if(InternalClothSimulationContext.PendingTeleportType == ETeleportType::TeleportPhysics)	//If the pending transform came from a teleport, make sure to teleport the cloth in this upcoming simulation
+			{
+				InternalClothSimulationContext.ClothTeleportMode = FClothingActor::TeleportMode::Teleport;
+			}
+
+			UpdateClothTransformImp();
+		}
+
 		ClothTeleportMode = FClothingActor::TeleportMode::Continuous;
 	}
 
@@ -1297,6 +1309,8 @@ FClothSimulationContext::FClothSimulationContext()
 	InMasterBoneMapCacheCount = -1;
 	bUseMasterPose = false;
 	WindAdaption = 2.f;	//This is the const that the previous code was using. Not sure where it comes from
+	bPendingClothUpdateTransform = false;
+	PendingTeleportType = ETeleportType::None;
 }
 
 void USkeletalMeshComponent::PostAnimEvaluation(FAnimationEvaluationContext& EvaluationContext)
diff --git a/Engine/Source/Runtime/Engine/Private/Font.cpp b/Engine/Source/Runtime/Engine/Private/Font.cpp
--- a/Engine/Source/Runtime/Engine/Private/Font.cpp
+++ b/Engine/Source/Runtime/Engine/Private/Font.cpp
@@ -394,8 +394,11 @@ SIZE_T UFont::GetResourceSize(EResourceSizeMode::Type Mode)
 				int32 TypefaceResourceSize = 0;
 				for (const FTypefaceEntry& TypefaceEntry : Typeface.Fonts)
 				{
-					// We use GetBulkDataSizeOnDisk since that will be the resident size once the bulk data has been decompressed
-					TypefaceResourceSize += TypefaceEntry.Font.BulkDataPtr->GetBulkDataSizeOnDisk();
+					if (TypefaceEntry.Font.BulkDataPtr)
+					{
+						// We use GetBulkDataSizeOnDisk since that will be the resident size once the bulk data has been decompressed
+						TypefaceResourceSize += TypefaceEntry.Font.BulkDataPtr->GetBulkDataSizeOnDisk();
+					}
 				}
 				return TypefaceResourceSize;
 			};
diff --git a/Engine/Source/Runtime/Engine/Private/HierarchicalInstancedStaticMesh.cpp b/Engine/Source/Runtime/Engine/Private/HierarchicalInstancedStaticMesh.cpp
--- a/Engine/Source/Runtime/Engine/Private/HierarchicalInstancedStaticMesh.cpp
+++ b/Engine/Source/Runtime/Engine/Private/HierarchicalInstancedStaticMesh.cpp
@@ -2488,7 +2488,7 @@ void UHierarchicalInstancedStaticMeshComponent::PostLoad()
 		}
 	}
 	// For some reason we don't have a tree, or it is out of date. Build one now!
-	if (StaticMesh && PerInstanceSMData.Num() > 0 && (!ClusterTreePtr.IsValid() || ClusterTreePtr->Num() == 0 || (NumBuiltInstances != PerInstanceSMData.Num()) || GetLinkerUE4Version() < VER_UE4_REBUILD_HIERARCHICAL_INSTANCE_TREES))
+	if (StaticMesh && PerInstanceSMData.Num() > 0 && (!ClusterTreePtr.IsValid() || ClusterTreePtr->Num() == 0 || (NumBuiltInstances != PerInstanceSMData.Num()) || UnbuiltInstanceBoundsList.Num() > 0 || GetLinkerUE4Version() < VER_UE4_REBUILD_HIERARCHICAL_INSTANCE_TREES))
 	{
 		UE_LOG(LogStaticMesh, Warning, TEXT("Rebuilding foliage, please resave map %s."), *GetFullName());
 		check(!IsAsyncBuilding());
diff --git a/Engine/Source/Runtime/Engine/Private/KismetFunctionLibrary.cpp b/Engine/Source/Runtime/Engine/Private/KismetFunctionLibrary.cpp
--- a/Engine/Source/Runtime/Engine/Private/KismetFunctionLibrary.cpp
+++ b/Engine/Source/Runtime/Engine/Private/KismetFunctionLibrary.cpp
@@ -31,7 +31,7 @@ FStringAssetReference UBlueprintFunctionLibrary::Generic_MakeStringAssetReferenc
 	FStringAssetReference Ref(AssetLongPathname);
 	if (!AssetLongPathname.IsEmpty() && !Ref.IsValid())
 	{
-		FBlueprintExceptionInfo Info(EBlueprintExceptionType::FatalError, TEXT("Asset path not valid. Only long path name is allowed."));
+		FBlueprintExceptionInfo Info(EBlueprintExceptionType::FatalError, NSLOCTEXT("BlueprintFunctionLibrary", "AssetStringInvalid", "Asset path not valid. Only long path name is allowed."));
 		FBlueprintCoreDelegates::ThrowScriptException(Stack.Object, Stack, Info);
 		return FStringAssetReference();
 	}
diff --git a/Engine/Source/Runtime/Engine/Private/SkeletalMeshComponentPhysics.cpp b/Engine/Source/Runtime/Engine/Private/SkeletalMeshComponentPhysics.cpp
--- a/Engine/Source/Runtime/Engine/Private/SkeletalMeshComponentPhysics.cpp
+++ b/Engine/Source/Runtime/Engine/Private/SkeletalMeshComponentPhysics.cpp
@@ -1429,9 +1429,10 @@ void USkeletalMeshComponent::OnUpdateTransform(bool bSkipPhysicsMove, ETeleportT
 #if WITH_APEX_CLOTHING
 	if(ClothingActors.Num() > 0)
 	{
+		
 		//@todo: Should cloth know whether we're teleporting?
 		// Updates cloth animation states because transform is updated
-		UpdateClothTransform();
+		UpdateClothTransform(Teleport);
 	}
 #endif //#if WITH_APEX_CLOTHING
 }
@@ -2143,7 +2144,7 @@ bool USkeletalMeshComponent::ComponentOverlapMultiImpl(TArray<struct FOverlapRes
 
 void USkeletalMeshComponent::AddClothingBounds(FBoxSphereBounds& InOutBounds, const FTransform& LocalToWorld) const
 {
-	for(const FClothingActor& ClothingActor : ClothingActors)
+	for (const FClothingActor& ClothingActor : ClothingActors)
 	{
 		if(NxClothingActor* Actor = ClothingActor.ApexClothingActor)
 		{
@@ -3296,7 +3297,7 @@ void USkeletalMeshComponent::PostPhysicsTickComponent(FSkeletalMeshComponentPost
 
 #if WITH_APEX_CLOTHING
 
-void USkeletalMeshComponent::UpdateClothTransform()
+void USkeletalMeshComponent::UpdateClothTransformImp()
 {
 	int32 NumActors = ClothingActors.Num();
 
@@ -3328,6 +3329,13 @@ void USkeletalMeshComponent::UpdateClothTransform()
 //#endif
 }
 
+void USkeletalMeshComponent::UpdateClothTransform(ETeleportType TeleportType)
+{
+	//Note that it's not safe to run the update here. This is because cloth sim could still be running on another thread. We defer it
+	InternalClothSimulationContext.bPendingClothUpdateTransform = true;
+	InternalClothSimulationContext.PendingTeleportType = TeleportType;
+}
+
 void USkeletalMeshComponent::CheckClothTeleport()
 {
 	// Get the root bone transform
diff --git a/Engine/Source/Runtime/Engine/Public/SkeletalMeshTypes.h b/Engine/Source/Runtime/Engine/Public/SkeletalMeshTypes.h
--- a/Engine/Source/Runtime/Engine/Public/SkeletalMeshTypes.h
+++ b/Engine/Source/Runtime/Engine/Public/SkeletalMeshTypes.h
@@ -90,7 +90,8 @@ struct FMeshFace
 	FVector	TangentX[3];
 	FVector	TangentY[3];
 	FVector	TangentZ[3];
-	
+
+	uint32   SmoothingGroups; // 32-bit flag for smoothing groups.
 };
 template <> struct TIsPODType<FMeshFace> { enum { Value = true }; };
 
diff --git a/Engine/Source/Runtime/OpenGLDrv/Private/Android/AndroidOpenGL.cpp b/Engine/Source/Runtime/OpenGLDrv/Private/Android/AndroidOpenGL.cpp
--- a/Engine/Source/Runtime/OpenGLDrv/Private/Android/AndroidOpenGL.cpp
+++ b/Engine/Source/Runtime/OpenGLDrv/Private/Android/AndroidOpenGL.cpp
@@ -379,6 +379,10 @@ void FAndroidOpenGL::ProcessExtensions(const FString& ExtensionsString)
 	const bool bIsAdrenoBased = RendererString.Contains(TEXT("Adreno"));
 	if (bIsAdrenoBased)
 	{
+		// This is to avoid a bug in Adreno drivers that define GL_EXT_shader_framebuffer_fetch even when device does not support this extension
+		// OpenGL ES 3.1 [REDACTED_EMAIL] (GIT@I1af360237c)
+		bRequiresShaderFramebufferFetchUndef = !bSupportsShaderFramebufferFetch;
+
 		// Adreno 2xx doesn't work with packed depth stencil enabled
 		if (RendererString.Contains(TEXT("Adreno (TM) 2")))
 		{
diff --git a/Engine/Source/Runtime/OpenGLDrv/Private/OpenGLES2.cpp b/Engine/Source/Runtime/OpenGLDrv/Private/OpenGLES2.cpp
--- a/Engine/Source/Runtime/OpenGLDrv/Private/OpenGLES2.cpp
+++ b/Engine/Source/Runtime/OpenGLDrv/Private/OpenGLES2.cpp
@@ -133,6 +133,9 @@ bool FOpenGLES2::bNeedsVertexAttribRemap = false;
 /* This hack fixes an issue with SGX540 compiler which can get upset with some operations that mix highp and mediump */
 bool FOpenGLES2::bRequiresTexture2DPrecisionHack = false;
 
+/* This is to avoid a bug in Adreno drivers that define GL_EXT_shader_framebuffer_fetch even when device does not support this extension  */
+bool FOpenGLES2::bRequiresShaderFramebufferFetchUndef = false;
+
 /* Indicates shader compiler hack checks are being tested */
 bool FOpenGLES2::bIsCheckingShaderCompilerHacks = false;
 
diff --git a/Engine/Source/Runtime/OpenGLDrv/Private/OpenGLES31.cpp b/Engine/Source/Runtime/OpenGLDrv/Private/OpenGLES31.cpp
--- a/Engine/Source/Runtime/OpenGLDrv/Private/OpenGLES31.cpp
+++ b/Engine/Source/Runtime/OpenGLDrv/Private/OpenGLES31.cpp
@@ -136,6 +136,9 @@ bool FOpenGLES31::bRequiresGLFragCoordVaryingLimitHack = false;
 /* This hack fixes an issue with SGX540 compiler which can get upset with some operations that mix highp and mediump */
 bool FOpenGLES31::bRequiresTexture2DPrecisionHack = false;
 
+/* This is to avoid a bug in Adreno drivers that define GL_EXT_shader_framebuffer_fetch even when device does not support this extension  */
+bool FOpenGLES31::bRequiresShaderFramebufferFetchUndef = false;
+
 /* Indicates shader compiler hack checks are being tested */
 bool FOpenGLES31::bIsCheckingShaderCompilerHacks = false;
 
diff --git a/Engine/Source/Runtime/OpenGLDrv/Private/OpenGLShaders.cpp b/Engine/Source/Runtime/OpenGLDrv/Private/OpenGLShaders.cpp
--- a/Engine/Source/Runtime/OpenGLDrv/Private/OpenGLShaders.cpp
+++ b/Engine/Source/Runtime/OpenGLDrv/Private/OpenGLShaders.cpp
@@ -620,6 +620,7 @@ void OPENGLDRV_API GetCurrentOpenGLShaderDeviceCapabilities(FOpenGLShaderDeviceC
 	Capabilities.bSupportsStandardDerivativesExtension = FOpenGL::SupportsStandardDerivativesExtension();
 	Capabilities.bSupportsRenderTargetFormat_PF_FloatRGBA = GSupportsRenderTargetFormat_PF_FloatRGBA;
 	Capabilities.bSupportsShaderFramebufferFetch = FOpenGL::SupportsShaderFramebufferFetch();
+	Capabilities.bRequiresShaderFramebufferFetchUndef = FOpenGL::RequiresShaderFramebufferFetchUndef();
 	Capabilities.bRequiresDontEmitPrecisionForTextureSamplers = FOpenGL::RequiresDontEmitPrecisionForTextureSamplers();
 	Capabilities.bSupportsShaderTextureLod = FOpenGL::SupportsShaderTextureLod();
 	Capabilities.bSupportsShaderTextureCubeLod = FOpenGL::SupportsShaderTextureCubeLod();
@@ -734,7 +735,7 @@ void OPENGLDRV_API GLSLToDeviceCompatibleGLSL(FAnsiCharArray& GlslCodeOriginal,
 				}
 			}
 
-			if (!Capabilities.bSupportsShaderFramebufferFetch && TypeEnum == GL_FRAGMENT_SHADER)
+			if (Capabilities.bRequiresShaderFramebufferFetchUndef && TypeEnum == GL_FRAGMENT_SHADER)
 			{
 				// This is to avoid a bug in Adreno drivers that define GL_EXT_shader_framebuffer_fetch even when device does not support this extension
 				// OpenGL ES 3.1 [REDACTED_EMAIL] (GIT@I1af360237c)
diff --git a/Engine/Source/Runtime/OpenGLDrv/Public/OpenGLES2.h b/Engine/Source/Runtime/OpenGLDrv/Public/OpenGLES2.h
--- a/Engine/Source/Runtime/OpenGLDrv/Public/OpenGLES2.h
+++ b/Engine/Source/Runtime/OpenGLDrv/Public/OpenGLES2.h
@@ -138,6 +138,7 @@ struct FOpenGLES2 : public FOpenGLBase
 	static FORCEINLINE bool SupportsStandardDerivativesExtension()		{ return bSupportsStandardDerivativesExtension; }
 	static FORCEINLINE bool RequiresGLFragCoordVaryingLimitHack()		{ return bRequiresGLFragCoordVaryingLimitHack; }
 	static FORCEINLINE bool RequiresTexture2DPrecisionHack()			{ return bRequiresTexture2DPrecisionHack; }
+	static FORCEINLINE bool RequiresShaderFramebufferFetchUndef()		{ return bRequiresShaderFramebufferFetchUndef; }
 	static FORCEINLINE bool IsCheckingShaderCompilerHacks()				{ return bIsCheckingShaderCompilerHacks; }
 
 	static FORCEINLINE int32 GetReadHalfFloatPixelsEnum()				{ return GL_HALF_FLOAT_OES; }
@@ -496,6 +497,9 @@ public:
 	/* This hack fixes an issue with SGX540 compiler which can get upset with some operations that mix highp and mediump */
 	static bool bRequiresTexture2DPrecisionHack;
 
+	/* This is to avoid a bug in Adreno drivers that define GL_EXT_shader_framebuffer_fetch even when device does not support this extension  */
+	static bool bRequiresShaderFramebufferFetchUndef;
+
 	/* Indicates shader compiler hack checks are being tested */
 	static bool bIsCheckingShaderCompilerHacks;
 };
diff --git a/Engine/Source/Runtime/OpenGLDrv/Public/OpenGLES31.h b/Engine/Source/Runtime/OpenGLDrv/Public/OpenGLES31.h
--- a/Engine/Source/Runtime/OpenGLDrv/Public/OpenGLES31.h
+++ b/Engine/Source/Runtime/OpenGLDrv/Public/OpenGLES31.h
@@ -96,6 +96,7 @@ struct FOpenGLES31 : public FOpenGLBase
 	static FORCEINLINE bool RequiresGLFragCoordVaryingLimitHack()		{ return bRequiresGLFragCoordVaryingLimitHack; }
 	static FORCEINLINE GLenum GetVertexHalfFloatFormat()				{ return bES2Fallback ? GL_HALF_FLOAT_OES : GL_HALF_FLOAT; }
 	static FORCEINLINE bool RequiresTexture2DPrecisionHack()			{ return bRequiresTexture2DPrecisionHack; }
+	static FORCEINLINE bool RequiresShaderFramebufferFetchUndef()		{ return bRequiresShaderFramebufferFetchUndef; }
 	static FORCEINLINE bool IsCheckingShaderCompilerHacks()				{ return bIsCheckingShaderCompilerHacks; }
 
 	// On iOS both glMapBufferOES() and glBufferSubData() for immediate vertex and index data
@@ -954,6 +955,9 @@ public:
 	/* This hack fixes an issue with SGX540 compiler which can get upset with some operations that mix highp and mediump */
 	static bool bRequiresTexture2DPrecisionHack;
 
+	/* This is to avoid a bug in Adreno drivers that define GL_EXT_shader_framebuffer_fetch even when device does not support this extension  */
+	static bool bRequiresShaderFramebufferFetchUndef;
+
 	/* Indicates shader compiler hack checks are being tested */
 	static bool bIsCheckingShaderCompilerHacks;
 };
diff --git a/Engine/Source/Runtime/OpenGLDrv/Public/OpenGLShaders.h b/Engine/Source/Runtime/OpenGLDrv/Public/OpenGLShaders.h
--- a/Engine/Source/Runtime/OpenGLDrv/Public/OpenGLShaders.h
+++ b/Engine/Source/Runtime/OpenGLDrv/Public/OpenGLShaders.h
@@ -32,6 +32,7 @@ struct FOpenGLShaderDeviceCapabilities
 	bool bSupportsStandardDerivativesExtension;
 	bool bSupportsRenderTargetFormat_PF_FloatRGBA;
 	bool bSupportsShaderFramebufferFetch;
+	bool bRequiresShaderFramebufferFetchUndef;
 	bool bRequiresDontEmitPrecisionForTextureSamplers;
 	bool bSupportsShaderTextureLod;
 	bool bSupportsShaderTextureCubeLod;
diff --git a/Engine/Source/Runtime/RHI/Private/RHI.cpp b/Engine/Source/Runtime/RHI/Private/RHI.cpp
--- a/Engine/Source/Runtime/RHI/Private/RHI.cpp
+++ b/Engine/Source/Runtime/RHI/Private/RHI.cpp
@@ -361,6 +361,7 @@ static FName NAME_GLSL_310_ES_EXT(TEXT("GLSL_310_ES_EXT"));
 static FName NAME_SF_METAL_SM5(TEXT("SF_METAL_SM5"));
 static FName NAME_PC_VULKAN_ES2(TEXT("PC_VULKAN_ES2"));
 static FName NAME_SF_METAL_SM4(TEXT("SF_METAL_SM4"));
+static FName NAME_SF_METAL_MACES3_1(TEXT("SF_METAL_MACES3_1"));
 
 FName LegacyShaderPlatformToShaderFormat(EShaderPlatform Platform)
 {
@@ -405,6 +406,8 @@ FName LegacyShaderPlatformToShaderFormat(EShaderPlatform Platform)
 		return NAME_SF_METAL_SM4;
 	case SP_METAL_SM5:
 		return NAME_SF_METAL_SM5;
+	case SP_METAL_MACES3_1:
+		return NAME_SF_METAL_MACES3_1;
 	case SP_OPENGL_ES31_EXT:
 		return NAME_GLSL_310_ES_EXT;
 	case SP_VULKAN_ES2:
@@ -439,6 +442,7 @@ EShaderPlatform ShaderFormatToLegacyShaderPlatform(FName ShaderFormat)
 	if (ShaderFormat == NAME_SF_METAL_SM5)		return SP_METAL_SM5;
 	if (ShaderFormat == NAME_PC_VULKAN_ES2)		return SP_VULKAN_ES2;
 	if (ShaderFormat == NAME_SF_METAL_SM4)		return SP_METAL_SM4;
+	if (ShaderFormat == NAME_SF_METAL_MACES3_1)	return SP_METAL_MACES3_1;
 	return SP_NumPlatforms;
 }
 
diff --git a/Engine/Source/Runtime/RHI/Public/RHIDefinitions.h b/Engine/Source/Runtime/RHI/Public/RHIDefinitions.h
--- a/Engine/Source/Runtime/RHI/Public/RHIDefinitions.h
+++ b/Engine/Source/Runtime/RHI/Public/RHIDefinitions.h
@@ -59,8 +59,9 @@ enum EShaderPlatform
 	SP_METAL_SM5		= 17,
 	SP_VULKAN_ES2		= 18,
 	SP_METAL_SM4		= 19,
+	SP_METAL_MACES3_1	= 20,
 
-	SP_NumPlatforms		= 20,
+	SP_NumPlatforms		= 21,
 	SP_NumBits			= 5,
 };
 static_assert(SP_NumPlatforms <= (1 << SP_NumBits), "SP_NumPlatforms will not fit on SP_NumBits");
@@ -610,19 +611,19 @@ enum class EAsyncComputeBudget
 
 inline bool IsPCPlatform(const EShaderPlatform Platform)
 {
-	return Platform == SP_PCD3D_SM5 || Platform == SP_PCD3D_SM4 || Platform == SP_PCD3D_ES2 || Platform == SP_PCD3D_ES3_1 || Platform ==  SP_OPENGL_SM4 || Platform == SP_OPENGL_SM4_MAC || Platform == SP_OPENGL_SM5 || Platform == SP_OPENGL_PCES2 || Platform == SP_OPENGL_PCES3_1 || Platform == SP_METAL_SM4 || Platform == SP_METAL_SM5;
+	return Platform == SP_PCD3D_SM5 || Platform == SP_PCD3D_SM4 || Platform == SP_PCD3D_ES2 || Platform == SP_PCD3D_ES3_1 || Platform ==  SP_OPENGL_SM4 || Platform == SP_OPENGL_SM4_MAC || Platform == SP_OPENGL_SM5 || Platform == SP_OPENGL_PCES2 || Platform == SP_OPENGL_PCES3_1 || Platform == SP_METAL_SM4 || Platform == SP_METAL_SM5 || Platform == SP_METAL_MACES3_1;
 }
 
 /** Whether the shader platform corresponds to the ES2 feature level. */
 inline bool IsES2Platform(const EShaderPlatform Platform)
 {
-	return Platform == SP_PCD3D_ES2 || Platform == SP_OPENGL_PCES2 || Platform == SP_OPENGL_ES2_ANDROID || Platform == SP_OPENGL_ES2_WEBGL || Platform == SP_OPENGL_ES2_IOS || Platform == SP_METAL || Platform == SP_VULKAN_ES2;
+	return Platform == SP_PCD3D_ES2 || Platform == SP_OPENGL_PCES2 || Platform == SP_OPENGL_ES2_ANDROID || Platform == SP_OPENGL_ES2_WEBGL || Platform == SP_OPENGL_ES2_IOS || Platform == SP_METAL || Platform == SP_VULKAN_ES2 || Platform == SP_METAL_MACES3_1;
 }
 
 /** Whether the shader platform corresponds to the ES2 feature level. */
 inline bool IsMobilePlatform(const EShaderPlatform Platform)
 {
-	return IsES2Platform(Platform) || Platform == SP_METAL || Platform == SP_PCD3D_ES3_1 || Platform == SP_OPENGL_PCES3_1;
+	return IsES2Platform(Platform) || Platform == SP_METAL || Platform == SP_PCD3D_ES3_1 || Platform == SP_OPENGL_PCES3_1 || Platform == SP_METAL_MACES3_1;
 }
 
 inline bool IsOpenGLPlatform(const EShaderPlatform Platform)
@@ -633,7 +634,7 @@ inline bool IsOpenGLPlatform(const EShaderPlatform Platform)
 
 inline bool IsMetalPlatform(const EShaderPlatform Platform)
 {
-	return Platform == SP_METAL || Platform == SP_METAL_MRT || Platform == SP_METAL_SM4 || Platform == SP_METAL_SM5;
+	return Platform == SP_METAL || Platform == SP_METAL_MRT || Platform == SP_METAL_SM4 || Platform == SP_METAL_SM5 || Platform == SP_METAL_MACES3_1;
 }
 
 inline bool IsConsolePlatform(const EShaderPlatform Platform)
@@ -668,6 +669,7 @@ inline ERHIFeatureLevel::Type GetMaxSupportedFeatureLevel(EShaderPlatform InShad
 	case SP_METAL:
 	case SP_PCD3D_ES3_1:
 	case SP_OPENGL_PCES3_1:
+	case SP_METAL_MACES3_1:
 		return ERHIFeatureLevel::ES3_1;
 	default:
 		check(0);
@@ -719,6 +721,8 @@ inline bool IsFeatureLevelSupported(EShaderPlatform InShaderPlatform, ERHIFeatur
 		return InFeatureLevel <= ERHIFeatureLevel::SM5;
 	case SP_METAL_SM5:
 		return InFeatureLevel <= ERHIFeatureLevel::SM5;
+	case SP_METAL_MACES3_1:
+		return InFeatureLevel <= ERHIFeatureLevel::ES3_1;
 	default:
 		return false;
 	}	
diff --git a/Engine/Source/Runtime/ShaderCore/Private/ShaderCache.cpp b/Engine/Source/Runtime/ShaderCore/Private/ShaderCache.cpp
--- a/Engine/Source/Runtime/ShaderCore/Private/ShaderCache.cpp
+++ b/Engine/Source/Runtime/ShaderCore/Private/ShaderCache.cpp
@@ -147,6 +147,7 @@ static bool ShaderPlatformCanPrebindBoundShaderState(EShaderPlatform Platform)
 		case SP_OPENGL_SM4_MAC:
 		case SP_METAL_MRT:
 		case SP_METAL_SM5:
+		case SP_METAL_MACES3_1:
 		{
 			return true;
 		}
diff --git a/Engine/Source/Runtime/SlateCore/Private/Fonts/FontCache.cpp b/Engine/Source/Runtime/SlateCore/Private/Fonts/FontCache.cpp
--- a/Engine/Source/Runtime/SlateCore/Private/Fonts/FontCache.cpp
+++ b/Engine/Source/Runtime/SlateCore/Private/Fonts/FontCache.cpp
@@ -1221,6 +1221,11 @@ void FSlateFontCache::FlushObject( const UObject* const InObject )
 	}
 }
 
+void FSlateFontCache::FlushCompositeFont(const FCompositeFont& InCompositeFont)
+{
+	CompositeFontCache->FlushCompositeFont(InCompositeFont);
+}
+
 bool FSlateFontCache::ConditionalFlushCache()
 {
 	bool bFlushed = false;
diff --git a/Engine/Source/Runtime/SlateCore/Private/Fonts/FontCacheCompositeFont.cpp b/Engine/Source/Runtime/SlateCore/Private/Fonts/FontCacheCompositeFont.cpp
--- a/Engine/Source/Runtime/SlateCore/Private/Fonts/FontCacheCompositeFont.cpp
+++ b/Engine/Source/Runtime/SlateCore/Private/Fonts/FontCacheCompositeFont.cpp
@@ -54,6 +54,21 @@ const FFontData* FCachedTypefaceData::GetFontData(const FName& InName) const
 	return SingularFontData;
 }
 
+void FCachedTypefaceData::GetCachedFontData(TArray<const FFontData*>& OutFontData) const
+{
+	if (NameToFontDataMap.Num() > 0)
+	{
+		for (const auto& NameToFontDataPair : NameToFontDataMap)
+		{
+			OutFontData.Add(NameToFontDataPair.Value);
+		}
+	}
+	else
+	{
+		OutFontData.Add(SingularFontData);
+	}
+}
+
 
 FCachedCompositeFontData::FCachedCompositeFontData()
 	: CompositeFont(nullptr)
@@ -121,6 +136,14 @@ const FCachedTypefaceData* FCachedCompositeFontData::GetTypefaceForCharacter(con
 	return CachedTypefaces[0].Get();
 }
 
+void FCachedCompositeFontData::GetCachedFontData(TArray<const FFontData*>& OutFontData) const
+{
+	for (const auto& CachedTypeface : CachedTypefaces)
+	{
+		CachedTypeface->GetCachedFontData(OutFontData);
+	}
+}
+
 
 FCompositeFontCache::FCompositeFontCache(const FFreeTypeLibrary* InFTLibrary)
 	: FTLibrary(InFTLibrary)
@@ -236,6 +259,23 @@ const TSet<FName>& FCompositeFontCache::GetFontAttributes(const FFontData& InFon
 	return (FaceAndMemory.IsValid()) ? FaceAndMemory->GetAttributes() : DummyAttributes;
 }
 
+void FCompositeFontCache::FlushCompositeFont(const FCompositeFont& InCompositeFont)
+{
+	TSharedPtr<FCachedCompositeFontData>* const FoundCompositeFontData = CompositeFontToCachedDataMap.Find(&InCompositeFont);
+	if (FoundCompositeFontData)
+	{
+		// Also clear out any font face map entries for the cached data
+		TArray<const FFontData*> AllCachedFontData;
+		(*FoundCompositeFontData)->GetCachedFontData(AllCachedFontData);
+		for (const FFontData* CachedFontData : AllCachedFontData)
+		{
+			FontFaceMap.Remove(CachedFontData);
+		}
+
+		CompositeFontToCachedDataMap.Remove(&InCompositeFont);
+	}
+}
+
 void FCompositeFontCache::FlushCache()
 {
 	CompositeFontToCachedDataMap.Empty();
diff --git a/Engine/Source/Runtime/SlateCore/Private/Fonts/FontCacheCompositeFont.h b/Engine/Source/Runtime/SlateCore/Private/Fonts/FontCacheCompositeFont.h
--- a/Engine/Source/Runtime/SlateCore/Private/Fonts/FontCacheCompositeFont.h
+++ b/Engine/Source/Runtime/SlateCore/Private/Fonts/FontCacheCompositeFont.h
@@ -33,6 +33,9 @@ public:
 		return ScalingFactor;
 	}
 
+	/** Get all the font data cached by this entry */
+	void GetCachedFontData(TArray<const FFontData*>& OutFontData) const;
+
 private:
 	/** Typeface we cached data from */
 	const FTypeface* Typeface;
@@ -76,6 +79,9 @@ public:
 	/** Get the typeface that should be used for the given character */
 	const FCachedTypefaceData* GetTypefaceForCharacter(const TCHAR InChar) const;
 
+	/** Get all the font data cached by this entry */
+	void GetCachedFontData(TArray<const FFontData*>& OutFontData) const;
+
 private:
 	/** Entry containing a range and the typeface associated with that range */
 	struct FCachedFontRange
@@ -133,6 +139,9 @@ public:
 	/** Get the attributes associated with the given font data */
 	const TSet<FName>& GetFontAttributes(const FFontData& InFontData);
 
+	/** Flush a single composite font entry from this cache */
+	void FlushCompositeFont(const FCompositeFont& InCompositeFont);
+
 	/** Flush this cache */
 	void FlushCache();
 
diff --git a/Engine/Source/Runtime/SlateCore/Public/Fonts/FontCache.h b/Engine/Source/Runtime/SlateCore/Public/Fonts/FontCache.h
--- a/Engine/Source/Runtime/SlateCore/Public/Fonts/FontCache.h
+++ b/Engine/Source/Runtime/SlateCore/Public/Fonts/FontCache.h
@@ -742,6 +742,11 @@ public:
 	 */
 	void FlushObject( const UObject* const InObject );
 
+	/**
+	 * Flush the given composite font out of the cache
+	 */
+	void FlushCompositeFont(const FCompositeFont& InCompositeFont);
+
 	/** 
 	 * Flush the cache if needed
 	 */
diff --git a/Engine/Source/Runtime/SlateRHIRenderer/Private/SlateRHIRenderingPolicy.cpp b/Engine/Source/Runtime/SlateRHIRenderer/Private/SlateRHIRenderingPolicy.cpp
--- a/Engine/Source/Runtime/SlateRHIRenderer/Private/SlateRHIRenderingPolicy.cpp
+++ b/Engine/Source/Runtime/SlateRHIRenderer/Private/SlateRHIRenderingPolicy.cpp
@@ -552,35 +552,60 @@ void FSlateRHIRenderingPolicy::DrawElements(FRHICommandListImmediate& RHICmdList
 
 						uint32 PrimitiveCount = RenderBatch.DrawPrimitiveType == ESlateDrawPrimitive::LineList ? RenderBatch.NumIndices / 2 : RenderBatch.NumIndices / 3;
 
-						// for RHIs that can't handle VertexOffset, we need to offset the stream source each time
-						if (!GRHISupportsBaseVertexIndex)
+						if ( bUseInstancing )
 						{
-							if (bUseInstancing)
+							uint32 InstanceCount = RenderBatch.InstanceCount;
+
+							if ( GRHISupportsInstancing )
 							{
 								FSlateUpdatableInstanceBuffer* InstanceBuffer = (FSlateUpdatableInstanceBuffer*)RenderBatch.InstanceData;
 								InstanceBuffer->BindStreamSource(RHICmdList, 1, RenderBatch.InstanceOffset);
-							}
-							else
-							{
-								RHICmdList.SetStreamSource(1, nullptr, 0, 0);
-							}
 
-							RHICmdList.SetStreamSource(0, VertexBuffer->VertexBufferRHI, sizeof(FSlateVertex), RenderBatch.VertexOffset * sizeof(FSlateVertex));
-							RHICmdList.DrawIndexedPrimitive(IndexBuffer->IndexBufferRHI, GetRHIPrimitiveType(RenderBatch.DrawPrimitiveType), 0, 0, RenderBatch.NumVertices, RenderBatch.IndexOffset, PrimitiveCount, RenderBatch.InstanceCount);
+								// for RHIs that can't handle VertexOffset, we need to offset the stream source each time
+								if ( !GRHISupportsBaseVertexIndex )
+								{
+									RHICmdList.SetStreamSource(0, VertexBuffer->VertexBufferRHI, sizeof(FSlateVertex), RenderBatch.VertexOffset * sizeof(FSlateVertex));
+									RHICmdList.DrawIndexedPrimitive(IndexBuffer->IndexBufferRHI, GetRHIPrimitiveType(RenderBatch.DrawPrimitiveType), 0, 0, RenderBatch.NumVertices, RenderBatch.IndexOffset, PrimitiveCount, InstanceCount);
+								}
+								else
+								{
+									RHICmdList.DrawIndexedPrimitive(IndexBuffer->IndexBufferRHI, GetRHIPrimitiveType(RenderBatch.DrawPrimitiveType), RenderBatch.VertexOffset, 0, RenderBatch.NumVertices, RenderBatch.IndexOffset, PrimitiveCount, InstanceCount);
+								}
+							}
+							//else
+							//{
+							//	for ( uint32 InstanceIndex = 0; InstanceIndex < InstanceCount; InstanceIndex++ )
+							//	{
+							//		FSlateUpdatableInstanceBuffer* InstanceBuffer = (FSlateUpdatableInstanceBuffer*)RenderBatch.InstanceData;
+							//		InstanceBuffer->BindStreamSource(RHICmdList, 1, RenderBatch.InstanceOffset + InstanceIndex);
+
+							//		// for RHIs that can't handle VertexOffset, we need to offset the stream source each time
+							//		if ( !GRHISupportsBaseVertexIndex )
+							//		{
+							//			RHICmdList.SetStreamSource(0, VertexBuffer->VertexBufferRHI, sizeof(FSlateVertex), RenderBatch.VertexOffset * sizeof(FSlateVertex));
+							//			RHICmdList.DrawIndexedPrimitive(IndexBuffer->IndexBufferRHI, GetRHIPrimitiveType(RenderBatch.DrawPrimitiveType), 0, 0, RenderBatch.NumVertices, RenderBatch.IndexOffset, PrimitiveCount, 1);
+							//		}
+							//		else
+							//		{
+							//			RHICmdList.DrawIndexedPrimitive(IndexBuffer->IndexBufferRHI, GetRHIPrimitiveType(RenderBatch.DrawPrimitiveType), RenderBatch.VertexOffset, 0, RenderBatch.NumVertices, RenderBatch.IndexOffset, PrimitiveCount, 1);
+							//		}
+							//	}
+							//}
 						}
 						else
 						{
-							if (bUseInstancing)
+							RHICmdList.SetStreamSource(1, nullptr, 0, 0);
+
+							// for RHIs that can't handle VertexOffset, we need to offset the stream source each time
+							if ( !GRHISupportsBaseVertexIndex )
 							{
-								FSlateUpdatableInstanceBuffer* InstanceBuffer = (FSlateUpdatableInstanceBuffer*)RenderBatch.InstanceData;
-								InstanceBuffer->BindStreamSource(RHICmdList, 1, RenderBatch.InstanceOffset);
+								RHICmdList.SetStreamSource(0, VertexBuffer->VertexBufferRHI, sizeof(FSlateVertex), RenderBatch.VertexOffset * sizeof(FSlateVertex));
+								RHICmdList.DrawIndexedPrimitive(IndexBuffer->IndexBufferRHI, GetRHIPrimitiveType(RenderBatch.DrawPrimitiveType), 0, 0, RenderBatch.NumVertices, RenderBatch.IndexOffset, PrimitiveCount, 1);
 							}
 							else
 							{
-								RHICmdList.SetStreamSource(1, nullptr, 0, 0);
+								RHICmdList.DrawIndexedPrimitive(IndexBuffer->IndexBufferRHI, GetRHIPrimitiveType(RenderBatch.DrawPrimitiveType), RenderBatch.VertexOffset, 0, RenderBatch.NumVertices, RenderBatch.IndexOffset, PrimitiveCount, 1);
 							}
-
-							RHICmdList.DrawIndexedPrimitive(IndexBuffer->IndexBufferRHI, GetRHIPrimitiveType(RenderBatch.DrawPrimitiveType), RenderBatch.VertexOffset, 0, RenderBatch.NumVertices, RenderBatch.IndexOffset, PrimitiveCount, RenderBatch.InstanceCount);
 						}
 					}
 				}
diff --git a/Engine/Source/Runtime/UMG/Private/Components/GridSlot.cpp b/Engine/Source/Runtime/UMG/Private/Components/GridSlot.cpp
--- a/Engine/Source/Runtime/UMG/Private/Components/GridSlot.cpp
+++ b/Engine/Source/Runtime/UMG/Private/Components/GridSlot.cpp
@@ -23,7 +23,6 @@ void UGridSlot::ReleaseSlateResources(bool bReleaseChildren)
 	Slot = nullptr;
 }
 
-// COMPILE PLEASE
 void UGridSlot::BuildSlot(TSharedRef<SGridPanel> GridPanel)
 {
 	Slot = &GridPanel->AddSlot(Column, Row, SGridPanel::Layer(Layer))
@@ -108,4 +107,5 @@ void UGridSlot::SynchronizeProperties()
 	SetColumnSpan(ColumnSpan);
 	SetHorizontalAlignment(HorizontalAlignment);
 	SetVerticalAlignment(VerticalAlignment);
+	SetLayer(Layer);
 }
