Copying //UE4/Dev-Platform to //UE4/Dev-Main (Source: //UE4/Dev-Platform @ 3046626)

diff --git a/Engine/Source/Programs/UnrealBuildTool/System/XcodeProject.cs b/Engine/Source/Programs/UnrealBuildTool/System/XcodeProject.cs
--- a/Engine/Source/Programs/UnrealBuildTool/System/XcodeProject.cs
+++ b/Engine/Source/Programs/UnrealBuildTool/System/XcodeProject.cs
@@ -608,11 +608,24 @@ namespace UnrealBuildTool
 			Content.Append("\t\t\tisa = XCBuildConfiguration;" + ProjectFileGenerator.NewLine);
 			Content.Append("\t\t\tbuildSettings = {" + ProjectFileGenerator.NewLine);
 
-			MacPlatform MacBuildPlat = UEBuildPlatform.GetBuildPlatform(UnrealTargetPlatform.Mac) as MacPlatform;
-			MacPlatformContext PlatformContextMac = (MacPlatformContext)MacBuildPlat.CreateContext(ProjectFile);
-			PlatformContextMac.SetUpProjectEnvironment();
-
-			string UE4Dir = ConvertPath(Path.GetFullPath(Directory.GetCurrentDirectory() + "../../.."));
+            if (UnrealBuildTool.IsValidPlatform(UnrealTargetPlatform.Mac))
+            {
+                MacPlatform MacBuildPlat = UEBuildPlatform.GetBuildPlatform(UnrealTargetPlatform.Mac) as MacPlatform;
+                MacPlatformContext PlatformContextMac = (MacPlatformContext)MacBuildPlat.CreateContext(ProjectFile);
+                PlatformContextMac.SetUpProjectEnvironment();
+            }
+            else if (UnrealBuildTool.IsValidPlatform(UnrealTargetPlatform.IOS))
+            {
+                IOSPlatform IOSBuildPlat = UEBuildPlatform.GetBuildPlatform(UnrealTargetPlatform.IOS) as IOSPlatform;
+                IOSPlatformContext PlatformContextIOS = (IOSPlatformContext)IOSBuildPlat.CreateContext(ProjectFile);
+                PlatformContextIOS.SetUpProjectEnvironment();
+            }
+            else
+            {
+                throw new BuildException("No valid build platform found when trying to generate an Xcode project file.");
+            }
+
+            string UE4Dir = ConvertPath(Path.GetFullPath(Directory.GetCurrentDirectory() + "../../.."));
 			string MacExecutableDir = ConvertPath(Config.MacExecutablePath.Directory.FullName);
 			string MacExecutableFileName = Config.MacExecutablePath.GetFileName();
 
@@ -1038,14 +1051,27 @@ namespace UnrealBuildTool
 												}
 											}
 
-											MacPlatform MacBuildPlat = UEBuildPlatform.GetBuildPlatform(UnrealTargetPlatform.Mac) as MacPlatform;
-											MacPlatformContext PlatformContextMac = (MacPlatformContext)MacBuildPlat.CreateContext(OnlyGameProject);
-
-											string MacExecutableName = UEBuildTarget.MakeBinaryFileName(ExeName, UnrealTargetPlatform.Mac, (ExeName == "UE4Editor" && Configuration == UnrealTargetConfiguration.DebugGame) ? UnrealTargetConfiguration.Development : Configuration, PlatformContextMac.GetActiveArchitecture(), ProjectTarget.TargetRules.UndecoratedConfiguration, UEBuildBinaryType.Executable);
-											string IOSExecutableName = MacExecutableName.Replace("-Mac-", "-IOS-");
-											string TVOSExecutableName = MacExecutableName.Replace("-Mac-", "-TVOS-");
-											BuildConfigs.Add(new XcodeBuildConfig(ConfigName, TargetName, FileReference.Combine(OutputDirectory, "Mac", MacExecutableName), FileReference.Combine(OutputDirectory, "IOS", IOSExecutableName), FileReference.Combine(OutputDirectory, "TVOS", TVOSExecutableName), ProjectTarget, Configuration));
-										}
+                                            if (BuildPlatform.Platform == UnrealTargetPlatform.Mac)
+                                            {
+                                                MacPlatform MacBuildPlat = UEBuildPlatform.GetBuildPlatform(UnrealTargetPlatform.Mac) as MacPlatform;
+                                                MacPlatformContext PlatformContextMac = (MacPlatformContext)MacBuildPlat.CreateContext(OnlyGameProject);
+
+                                                string MacExecutableName = UEBuildTarget.MakeBinaryFileName(ExeName, UnrealTargetPlatform.Mac, (ExeName == "UE4Editor" && Configuration == UnrealTargetConfiguration.DebugGame) ? UnrealTargetConfiguration.Development : Configuration, PlatformContextMac.GetActiveArchitecture(), ProjectTarget.TargetRules.UndecoratedConfiguration, UEBuildBinaryType.Executable);
+                                                string IOSExecutableName = MacExecutableName.Replace("-Mac-", "-IOS-");
+                                                string TVOSExecutableName = MacExecutableName.Replace("-Mac-", "-TVOS-");
+                                                BuildConfigs.Add(new XcodeBuildConfig(ConfigName, TargetName, FileReference.Combine(OutputDirectory, "Mac", MacExecutableName), FileReference.Combine(OutputDirectory, "IOS", IOSExecutableName), FileReference.Combine(OutputDirectory, "TVOS", TVOSExecutableName), ProjectTarget, Configuration));
+                                            }
+                                            else if (BuildPlatform.Platform == UnrealTargetPlatform.IOS)
+                                            {
+                                                IOSPlatform IOSBuildPlat = UEBuildPlatform.GetBuildPlatform(UnrealTargetPlatform.IOS) as IOSPlatform;
+                                                IOSPlatformContext PlatformContextIOS = (IOSPlatformContext)IOSBuildPlat.CreateContext(OnlyGameProject);
+
+                                                string IOSExecutableName = UEBuildTarget.MakeBinaryFileName(ExeName, UnrealTargetPlatform.IOS, (ExeName == "UE4Editor" && Configuration == UnrealTargetConfiguration.DebugGame) ? UnrealTargetConfiguration.Development : Configuration, PlatformContextIOS.GetActiveArchitecture(), ProjectTarget.TargetRules.UndecoratedConfiguration, UEBuildBinaryType.Executable);
+                                                string TVOSExecutableName = IOSExecutableName.Replace("-IOS-", "-TVOS-");
+                                                string MacExecutableName = IOSExecutableName.Replace("-IOS-", "-Mac-");
+                                                BuildConfigs.Add(new XcodeBuildConfig(ConfigName, TargetName, FileReference.Combine(OutputDirectory, "Mac", IOSExecutableName), FileReference.Combine(OutputDirectory, "IOS", IOSExecutableName), FileReference.Combine(OutputDirectory, "TVOS", TVOSExecutableName), ProjectTarget, Configuration));
+                                            }
+                                        }
 									}
 								}
 							}
diff --git a/Engine/Source/Programs/UnrealBuildTool/TVOS/UEBuildTVOS.cs b/Engine/Source/Programs/UnrealBuildTool/TVOS/UEBuildTVOS.cs
--- a/Engine/Source/Programs/UnrealBuildTool/TVOS/UEBuildTVOS.cs
+++ b/Engine/Source/Programs/UnrealBuildTool/TVOS/UEBuildTVOS.cs
@@ -113,7 +113,7 @@ namespace UnrealBuildTool
 		/// <returns>True if the platform requires a deployment handler, false otherwise</returns>
 		public override UEBuildDeploy CreateDeploymentHandler()
 		{
-			return new UEDeployTVOS();
+			return new UEDeployTVOS(ProjectFile);
 		}
 
 	}
diff --git a/Engine/Source/Programs/UnrealBuildTool/TVOS/UEDeployTVOS.cs b/Engine/Source/Programs/UnrealBuildTool/TVOS/UEDeployTVOS.cs
--- a/Engine/Source/Programs/UnrealBuildTool/TVOS/UEDeployTVOS.cs
+++ b/Engine/Source/Programs/UnrealBuildTool/TVOS/UEDeployTVOS.cs
@@ -7,17 +7,23 @@ using System.Text;
 using System.Xml;
 using System.IO;
 using System.Diagnostics;
+using System.Xml.Linq;
 
 namespace UnrealBuildTool
 {
 	public class UEDeployTVOS : UEDeployIOS
 	{
+
+		public UEDeployTVOS(FileReference InProjectFile) : base(InProjectFile)
+		{
+		}
+
 		protected override string GetTargetPlatformName()
 		{
 			return "TVOS";
 		}
 
-		public static bool GenerateTVOSPList(string ProjectDirectory, bool bIsUE4Game, string GameName, string ProjectName, string InEngineDir, string AppDirectory)
+		public static bool GenerateTVOSPList(string ProjectDirectory, bool bIsUE4Game, string GameName, string ProjectName, string InEngineDir, string AppDirectory, UEDeployTVOS InThis = null)
 		{
 			// @todo tvos: THIS!
 
@@ -280,12 +286,35 @@ namespace UnrealBuildTool
             Text.AppendLine("</dict>");
 			Text.AppendLine("</plist>");
 
-			// write the file out
+			// Create the intermediate directory if needed
 			if (!Directory.Exists(IntermediateDirectory))
 			{
 				Directory.CreateDirectory(IntermediateDirectory);
 			}
-			File.WriteAllText(PListFile, Text.ToString());
+
+			if(InThis != null && InThis.UPL != null)
+			{
+				// Allow UPL to modify the plist here
+				XDocument XDoc;
+				try
+				{
+					XDoc = XDocument.Parse(Text.ToString());
+				}
+				catch (Exception e)
+				{
+					throw new BuildException("plist is invalid {0}\n{1}", e, Text.ToString());
+				}
+
+				XDoc.DocumentType.InternalSubset = "";
+				InThis.UPL.ProcessPluginNode("None", "iosPListUpdates", "", ref XDoc);
+				string result = XDoc.Declaration.ToString() + "\n" + XDoc.ToString();
+				File.WriteAllText(PListFile, result);
+			}
+			else
+			{
+				File.WriteAllText(PListFile, Text.ToString());
+			}
+
 			if (BuildHostPlatform.Current.Platform == UnrealTargetPlatform.Mac)
 			{
 				if (!Directory.Exists(AppDirectory))
@@ -300,7 +329,7 @@ namespace UnrealBuildTool
 
 		public override bool GeneratePList(string ProjectDirectory, bool bIsUE4Game, string GameName, string ProjectName, string InEngineDir, string AppDirectory)
 		{
-			return GenerateTVOSPList(ProjectDirectory, bIsUE4Game, GameName, ProjectName, InEngineDir, AppDirectory);
+			return GenerateTVOSPList(ProjectDirectory, bIsUE4Game, GameName, ProjectName, InEngineDir, AppDirectory, this);
 		}
 
         protected override void CopyGraphicsResources(bool bSkipDefaultPNGs, string InEngineDir, string AppDirectory, string BuildDirectory, string IntermediateDir)
diff --git a/Engine/Source/Programs/UnrealBuildTool/UnrealBuildTool.csproj b/Engine/Source/Programs/UnrealBuildTool/UnrealBuildTool.csproj
--- a/Engine/Source/Programs/UnrealBuildTool/UnrealBuildTool.csproj
+++ b/Engine/Source/Programs/UnrealBuildTool/UnrealBuildTool.csproj
@@ -1,4 +1,4 @@
-﻿<?xml version="1.0" encoding="utf-8"?>
+<?xml version="1.0" encoding="utf-8"?>
 <Project DefaultTargets="Build" xmlns="http://schemas.microsoft.com/developer/msbuild/2003" ToolsVersion="4.0">
   <PropertyGroup>
     <Configuration Condition=" '$(Configuration)' == '' ">Debug</Configuration>
@@ -131,8 +131,8 @@
     <Compile Include="Android\UEBuildAndroid.cs" />
     <Compile Include="Android\UEDeployAndroid.cs" />
     <Compile Include="Android\AndroidToolChain.cs" />
-    <Compile Include="Android\AndroidPluginLanguage.cs" />
     <Compile Include="Android\AndroidProjectGenerator.cs" />
+    <Compile Include="System\UnrealPluginLanguage.cs" />
   </ItemGroup>
   <ItemGroup>
     <Compile Include="Utilities\HashSetExtensions.cs" />
diff --git a/Engine/Source/Runtime/Apple/MetalRHI/Private/MetalBufferPools.h b/Engine/Source/Runtime/Apple/MetalRHI/Private/MetalBufferPools.h
--- a/Engine/Source/Runtime/Apple/MetalRHI/Private/MetalBufferPools.h
+++ b/Engine/Source/Runtime/Apple/MetalRHI/Private/MetalBufferPools.h
@@ -52,7 +52,7 @@ public:
 		NumPoolBucketSizes = 22, /** Number of pool bucket sizes */
 		NumPoolBuckets = NumPoolBucketSizes * 3, /** Number of pool bucket sizes * 2 for Shared/Managed/Private storage */
 		NumToDrainPerFrame = 65536, /** Max. number of resources to cull in a single frame */
-		CullAfterFramesNum = 30 * 60 /** Resources are culled if unused for more frames than this */
+		CullAfterFramesNum = 30 /** Resources are culled if unused for more frames than this */
 	};
 	
 	/** Get the pool bucket index from the size
diff --git a/Engine/Source/Runtime/Apple/MetalRHI/Private/MetalCommandBuffer.cpp b/Engine/Source/Runtime/Apple/MetalRHI/Private/MetalCommandBuffer.cpp
new file mode 100644
--- /dev/null
+++ b/Engine/Source/Runtime/Apple/MetalRHI/Private/MetalCommandBuffer.cpp
@@ -0,0 +1,280 @@
+// Copyright 1998-2016 Epic Games, Inc. All Rights Reserved.
+
+/*=============================================================================
+	MetalCommandBuffere.cpp: Metal command buffer wrapper.
+=============================================================================*/
+
+#include "MetalRHIPrivate.h"
+
+#include "MetalCommandBuffer.h"
+
+NSString* GMetalDebugCommandTypeNames[EMetalDebugCommandTypeInvalid] = {
+	@"RenderEncoder",
+	@"ComputeEncoder",
+	@"BlitEncoder",
+	@"EndEncoder",
+	@"Draw",
+	@"Dispatch",
+	@"Blit",
+	@"Signpost",
+	@"PushGroup",
+	@"PopGroup"
+};
+
+@implementation FMetalDebugCommandBuffer
+
+@synthesize InnerBuffer;
+
+-(id)initWithCommandBuffer:(id<MTLCommandBuffer>)Buffer
+{
+	id Self = [super init];
+	if (Self)
+	{
+		InnerBuffer = Buffer;
+		DebugGroup = [NSMutableArray new];
+		ActiveEncoder = nil;
+	}
+	return Self;
+}
+
+-(void)dealloc
+{
+	for (FMetalDebugCommand* Command : DebugCommands)
+	{
+		[Command->Label release];
+		[Command->PassDesc release];
+		delete Command;
+	}
+	
+	[InnerBuffer release];
+	[DebugGroup release];
+	
+	[super dealloc];
+}
+
+-(id <MTLDevice>) device
+{
+	return InnerBuffer.device;
+}
+
+-(id <MTLCommandQueue>) commandQueue
+{
+	return InnerBuffer.commandQueue;
+}
+
+-(BOOL) retainedReferences
+{
+	return InnerBuffer.retainedReferences;
+}
+
+-(NSString *)label
+{
+	return InnerBuffer.label;
+}
+
+-(void)setLabel:(NSString *)Text
+{
+	InnerBuffer.label = Text;
+}
+
+-(MTLCommandBufferStatus) status
+{
+	return InnerBuffer.status;
+}
+
+-(NSError *)error
+{
+	return InnerBuffer.error;
+}
+
+- (void)enqueue
+{
+	[InnerBuffer enqueue];
+}
+
+- (void)commit
+{
+	[InnerBuffer commit];
+}
+
+- (void)addScheduledHandler:(MTLCommandBufferHandler)block
+{
+	[InnerBuffer addScheduledHandler:^(id <MTLCommandBuffer>){
+		block(self);
+	 }];
+}
+
+- (void)presentDrawable:(id <MTLDrawable>)drawable
+{
+	[InnerBuffer presentDrawable:drawable];
+}
+
+- (void)presentDrawable:(id <MTLDrawable>)drawable atTime:(CFTimeInterval)presentationTime
+{
+	[InnerBuffer presentDrawable:drawable atTime:presentationTime];
+}
+
+- (void)waitUntilScheduled
+{
+	[InnerBuffer waitUntilScheduled];
+}
+
+- (void)addCompletedHandler:(MTLCommandBufferHandler)block
+{
+	[InnerBuffer addCompletedHandler:^(id <MTLCommandBuffer>){
+		block(self);
+	}];
+}
+
+- (void)waitUntilCompleted
+{
+	[InnerBuffer waitUntilCompleted];
+}
+
+- (id <MTLBlitCommandEncoder>)blitCommandEncoder
+{
+	return [InnerBuffer blitCommandEncoder];
+}
+
+- (id <MTLRenderCommandEncoder>)renderCommandEncoderWithDescriptor:(MTLRenderPassDescriptor *)renderPassDescriptor
+{
+	return [InnerBuffer renderCommandEncoderWithDescriptor:renderPassDescriptor];
+}
+
+- (id <MTLComputeCommandEncoder>)computeCommandEncoder
+{
+	return [InnerBuffer computeCommandEncoder];
+}
+
+- (id <MTLParallelRenderCommandEncoder>)parallelRenderCommandEncoderWithDescriptor:(MTLRenderPassDescriptor *)renderPassDescriptor
+{
+	return [InnerBuffer parallelRenderCommandEncoderWithDescriptor:renderPassDescriptor];
+}
+
+-(NSString*) description
+{
+	NSMutableString* String = [NSMutableString new];
+	NSString* Label = self.label ? self.label : @"Unknown";
+	[String appendFormat:@"Command Buffer %p %@:", self, Label];
+	return String;
+}
+
+-(NSString*) debugDescription
+{
+	NSMutableString* String = [NSMutableString new];
+	NSString* Label = self.label ? self.label : @"Unknown";
+	[String appendFormat:@"Command Buffer %p %@:", self, Label];
+
+	for (FMetalDebugCommand* Command : DebugCommands)
+	{
+		[String appendFormat:@"\n\t%@: %@", GMetalDebugCommandTypeNames[Command->Type], Command->Label];
+	}
+	return String;
+}
+
+-(void) beginRenderCommandEncoder:(NSString*)Label withDescriptor:(MTLRenderPassDescriptor*)Desc
+{
+	check(!ActiveEncoder);
+	ActiveEncoder = [Label retain];
+	FMetalDebugCommand* Command = new FMetalDebugCommand;
+	Command->Type = EMetalDebugCommandTypeRenderEncoder;
+	Command->Label = [ActiveEncoder retain];
+	Command->PassDesc = [Desc retain];
+	DebugCommands.Add(Command);
+}
+
+-(void) beginComputeCommandEncoder:(NSString*)Label
+{
+	check(!ActiveEncoder);
+	ActiveEncoder = [Label retain];
+	FMetalDebugCommand* Command = new FMetalDebugCommand;
+	Command->Type = EMetalDebugCommandTypeComputeEncoder;
+	Command->Label = [ActiveEncoder retain];
+	Command->PassDesc = nil;
+	DebugCommands.Add(Command);
+}
+
+-(void) beginBlitCommandEncoder:(NSString*)Label
+{
+	check(!ActiveEncoder);
+	ActiveEncoder = [Label retain];
+	FMetalDebugCommand* Command = new FMetalDebugCommand;
+	Command->Type = EMetalDebugCommandTypeBlitEncoder;
+	Command->Label = [ActiveEncoder retain];
+	Command->PassDesc = nil;
+	DebugCommands.Add(Command);
+}
+
+-(void) endCommandEncoder
+{
+	check(ActiveEncoder);
+	FMetalDebugCommand* Command = new FMetalDebugCommand;
+	Command->Type = EMetalDebugCommandTypeEndEncoder;
+	Command->Label = ActiveEncoder;
+	Command->PassDesc = nil;
+	DebugCommands.Add(Command);
+	ActiveEncoder = nil;
+}
+
+-(void) draw:(NSString*)Desc withPipeline:(FMetalBoundShaderState*)BSS
+{
+	FMetalDebugCommand* Command = new FMetalDebugCommand;
+	Command->Type = EMetalDebugCommandTypeDraw;
+	Command->Label = [Desc retain];
+	Command->PassDesc = nil;
+	Command->RenderPipeline = BSS;
+	DebugCommands.Add(Command);
+}
+
+-(void) dispatch:(NSString*)Desc withShader:(FMetalComputeShader*)Shader
+{
+	FMetalDebugCommand* Command = new FMetalDebugCommand;
+	Command->Type = EMetalDebugCommandTypeDispatch;
+	Command->Label = [Desc retain];
+	Command->PassDesc = nil;
+	Command->ComputeShader = Shader;
+	DebugCommands.Add(Command);
+}
+
+-(void) blit:(NSString*)Desc
+{
+	FMetalDebugCommand* Command = new FMetalDebugCommand;
+	Command->Type = EMetalDebugCommandTypeBlit;
+	Command->Label = [Desc retain];
+	Command->PassDesc = nil;
+	DebugCommands.Add(Command);
+}
+
+-(void) insertDebugSignpost:(NSString*)Label
+{
+	FMetalDebugCommand* Command = new FMetalDebugCommand;
+	Command->Type = EMetalDebugCommandTypeSignpost;
+	Command->Label = [Label retain];
+	Command->PassDesc = nil;
+	DebugCommands.Add(Command);
+}
+
+-(void) pushDebugGroup:(NSString*)Group
+{
+	[DebugGroup addObject:Group];
+	FMetalDebugCommand* Command = new FMetalDebugCommand;
+	Command->Type = EMetalDebugCommandTypePushGroup;
+	Command->Label = [Group retain];
+	Command->PassDesc = nil;
+	DebugCommands.Add(Command);
+}
+
+-(void) popDebugGroup
+{
+	if (DebugGroup.lastObject)
+	{
+		FMetalDebugCommand* Command = new FMetalDebugCommand;
+		Command->Type = EMetalDebugCommandTypePopGroup;
+		Command->Label = [DebugGroup.lastObject retain];
+		Command->PassDesc = nil;
+		DebugCommands.Add(Command);
+	}
+}
+
+@end
+
diff --git a/Engine/Source/Runtime/Apple/MetalRHI/Private/MetalCommandBuffer.h b/Engine/Source/Runtime/Apple/MetalRHI/Private/MetalCommandBuffer.h
new file mode 100644
--- /dev/null
+++ b/Engine/Source/Runtime/Apple/MetalRHI/Private/MetalCommandBuffer.h
@@ -0,0 +1,129 @@
+// Copyright 1998-2016 Epic Games, Inc. All Rights Reserved.
+
+#pragma once
+
+#include <Metal/Metal.h>
+#include "MetalResources.h"
+#include "MetalShaderResources.h"
+
+/**
+ * EMetalDebugCommandType: Types of command recorded in our debug command-buffer wrapper.
+ */
+enum EMetalDebugCommandType
+{
+	EMetalDebugCommandTypeRenderEncoder,
+	EMetalDebugCommandTypeComputeEncoder,
+	EMetalDebugCommandTypeBlitEncoder,
+	EMetalDebugCommandTypeEndEncoder,
+	EMetalDebugCommandTypeDraw,
+	EMetalDebugCommandTypeDispatch,
+	EMetalDebugCommandTypeBlit,
+	EMetalDebugCommandTypeSignpost,
+	EMetalDebugCommandTypePushGroup,
+	EMetalDebugCommandTypePopGroup,
+	EMetalDebugCommandTypeInvalid
+};
+
+/**
+ * EMetalDebugLevel: Level of Metal debug features to be enabled.
+ */
+enum EMetalDebugLevel
+{
+	EMetalDebugLevelOff,
+	EMetalDebugLevelValidation,
+	EMetalDebugLevelLogDebugGroups,
+	EMetalDebugLevelLogOperations,
+	EMetalDebugLevelConditionalSubmit,
+	EMetalDebugLevelWaitForComplete
+};
+
+NS_ASSUME_NONNULL_BEGIN
+/**
+ * FMetalDebugCommand: The data recorded for each command in the debug command-buffer wrapper.
+ */
+struct FMetalDebugCommand
+{
+	NSString* Label;
+	EMetalDebugCommandType Type;
+	MTLRenderPassDescriptor* PassDesc;
+	TRefCountPtr<FMetalBoundShaderState> RenderPipeline;
+	TRefCountPtr<FMetalComputeShader> ComputeShader;
+};
+
+/**
+ * FMetalDebugCommandBuffer: Wrapper around id<MTLCommandBuffer> that records information about commands.
+ * This allows reporting of substantially more information in debug modes which can be especially helpful 
+ * when debugging GPU command-buffer failures.
+ */
+@interface FMetalDebugCommandBuffer : NSObject<MTLCommandBuffer>
+{
+	TArray<FMetalDebugCommand*> DebugCommands;
+	NSMutableArray<NSString*>* DebugGroup;
+	NSString* ActiveEncoder;
+};
+
+/** The wrapped native command-buffer for which we collect debug information. */
+@property (readonly) id<MTLCommandBuffer> InnerBuffer;
+
+/** Initialise the wrapper with the provided command-buffer. */
+-(id)initWithCommandBuffer:(id<MTLCommandBuffer>)Buffer;
+
+/** Record a bgein render encoder command. */
+-(void) beginRenderCommandEncoder:(NSString*)Label withDescriptor:(MTLRenderPassDescriptor*)Desc;
+/** Record a begin compute encoder command. */
+-(void) beginComputeCommandEncoder:(NSString*)Label;
+/** Record a begin blit encoder command. */
+-(void) beginBlitCommandEncoder:(NSString*)Label;
+/** Record an end encoder command. */
+-(void) endCommandEncoder;
+
+/** Record a draw command. */
+-(void) draw:(NSString*)Desc withPipeline:(FMetalBoundShaderState*)BSS;
+/** Record a dispatch command. */
+-(void) dispatch:(NSString*)Desc withShader:(FMetalComputeShader*)Shader;
+/** Record a blit command. */
+-(void) blit:(NSString*)Desc;
+
+/** Record an signpost command. */
+-(void) insertDebugSignpost:(NSString*)Label;
+/** Record a push debug group command. */
+-(void) pushDebugGroup:(NSString*)Group;
+/** Record a pop debug group command. */
+-(void) popDebugGroup;
+
+@end
+NS_ASSUME_NONNULL_END
+
+// Debug command-buffer logging macros that simplify the calling code
+#if !UE_BUILD_SHIPPING
+#define METAL_DEBUG_COMMAND_BUFFER_DRAW_LOG(Context, LabelFormat, ...)	\
+			if (Context->GetCommandQueue().GetRuntimeDebuggingLevel() >= EMetalDebugLevelLogOperations)	\
+			{	\
+				FMetalDebugCommandBuffer* CmdBuf = (FMetalDebugCommandBuffer*)Context->GetCurrentCommandBuffer();	\
+				[CmdBuf draw:[NSString stringWithFormat: LabelFormat, __VA_ARGS__] withPipeline:Context->GetCurrentState().GetBoundShaderState()];	\
+			}
+
+#define METAL_DEBUG_COMMAND_BUFFER_DISPATCH_LOG(Context, LabelFormat, ...)	\
+			if (Context->GetCommandQueue().GetRuntimeDebuggingLevel() >= EMetalDebugLevelLogOperations)	\
+			{	\
+				FMetalDebugCommandBuffer* CmdBuf = (FMetalDebugCommandBuffer*)Context->GetCurrentCommandBuffer();	\
+				[CmdBuf dispatch:[NSString stringWithFormat: LabelFormat, __VA_ARGS__] withShader:Context->GetCurrentState().GetComputeShader()];	\
+			}
+#define METAL_DEBUG_COMMAND_BUFFER_BLIT_LOG(Context, LabelFormat, ...)	\
+			if (Context->GetCommandQueue().GetRuntimeDebuggingLevel() >= EMetalDebugLevelLogOperations)	\
+			{	\
+				FMetalDebugCommandBuffer* CmdBuf = (FMetalDebugCommandBuffer*)Context->GetCurrentCommandBuffer();	\
+				[CmdBuf blit:[NSString stringWithFormat: LabelFormat, __VA_ARGS__]];	\
+			}
+#define METAL_DEBUG_COMMAND_BUFFER_BLIT_ASYNC_LOG(Context, MtlCmdBuf, LabelFormat, ...)	\
+			if (Context->GetCommandQueue().GetRuntimeDebuggingLevel() >= EMetalDebugLevelLogOperations)	\
+			{	\
+				FMetalDebugCommandBuffer* CmdBuf = (FMetalDebugCommandBuffer*)MtlCmdBuf;	\
+				[CmdBuf blit:[NSString stringWithFormat: LabelFormat, __VA_ARGS__]];	\
+			}
+#else
+#define METAL_DEBUG_COMMAND_BUFFER_DRAW_LOG(Context, LabelFormat, ...)
+#define METAL_DEBUG_COMMAND_BUFFER_DISPATCH_LOG(Context, LabelFormat, ...)
+#define METAL_DEBUG_COMMAND_BUFFER_BLIT_LOG(Context, LabelFormat, ...)
+#define METAL_DEBUG_COMMAND_BUFFER_BLIT_ASYNC_LOG(Context, MtlCmdBuf, LabelFormat, ...)
+#endif
diff --git a/Engine/Source/Runtime/Apple/MetalRHI/Private/MetalCommandEncoder.cpp b/Engine/Source/Runtime/Apple/MetalRHI/Private/MetalCommandEncoder.cpp
--- a/Engine/Source/Runtime/Apple/MetalRHI/Private/MetalCommandEncoder.cpp
+++ b/Engine/Source/Runtime/Apple/MetalRHI/Private/MetalCommandEncoder.cpp
@@ -7,6 +7,7 @@
 #include "MetalRHIPrivate.h"
 
 #include "MetalCommandEncoder.h"
+#include "MetalCommandBuffer.h"
 #include "MetalProfiler.h"
 
 #pragma mark - Public C++ Boilerplate -
@@ -240,12 +241,26 @@ void FMetalCommandEncoder::BeginRenderCommandEncoding(void)
 	{
 		NSString* Label = [DebugGroups count] > 0 ? [DebugGroups lastObject] : (NSString*)CFSTR("InitialPass");
 		[RenderCommandEncoder setLabel:Label];
+#if !UE_BUILD_SHIPPING
+		if (CommandList.GetCommandQueue().GetRuntimeDebuggingLevel() >= EMetalDebugLevelLogOperations)
+		{
+			FMetalDebugCommandBuffer* Buf = (FMetalDebugCommandBuffer*)CommandBuffer;
+			[Buf beginRenderCommandEncoder:Label withDescriptor:RenderPassDesc];
+		}
+#endif
 		
 		if([DebugGroups count])
 		{
 			for (NSString* Group in DebugGroups)
 			{
 				[RenderCommandEncoder pushDebugGroup:Group];
+#if !UE_BUILD_SHIPPING
+				if (CommandList.GetCommandQueue().GetRuntimeDebuggingLevel() >= EMetalDebugLevelLogDebugGroups)
+				{
+					FMetalDebugCommandBuffer* Buf = (FMetalDebugCommandBuffer*)CommandBuffer;
+					[Buf pushDebugGroup:Group];
+				}
+#endif
 			}
 		}
     }
@@ -258,7 +273,7 @@ void FMetalCommandEncoder::RestoreRenderCommandEncoding(void)
 	check(IsRenderCommandEncoderActive() == false && IsComputeCommandEncoderActive() == false && IsBlitCommandEncoderActive() == false);
 	
 #if UE_BUILD_DEBUG || UE_BUILD_DEVELOPMENT
-	if ((UE_BUILD_DEBUG || GetMetalDeviceContext().IsValidationLayerEnabled()) && RenderPassDescApplied == 1)
+	if ((GetMetalDeviceContext().GetCommandQueue().GetRuntimeDebuggingLevel() >= EMetalDebugLevelValidation) && RenderPassDescApplied == 1)
 	{
 		bool bAllLoadActionsOK = true;
 		if(RenderPassDesc.colorAttachments)
@@ -402,12 +417,26 @@ void FMetalCommandEncoder::BeginComputeCommandEncoding(void)
 	{
 		NSString* Label = [DebugGroups count] > 0 ? [DebugGroups lastObject] : (NSString*)CFSTR("InitialPass");
 		[ComputeCommandEncoder setLabel:Label];
+#if !UE_BUILD_SHIPPING
+		if (CommandList.GetCommandQueue().GetRuntimeDebuggingLevel() >= EMetalDebugLevelLogOperations)
+		{
+			FMetalDebugCommandBuffer* Buf = (FMetalDebugCommandBuffer*)CommandBuffer;
+			[Buf beginComputeCommandEncoder:Label];
+		}
+#endif
 		
 		if([DebugGroups count])
 		{
 			for (NSString* Group in DebugGroups)
 			{
 				[ComputeCommandEncoder pushDebugGroup:Group];
+#if !UE_BUILD_SHIPPING
+				if (CommandList.GetCommandQueue().GetRuntimeDebuggingLevel() >= EMetalDebugLevelLogDebugGroups)
+				{
+					FMetalDebugCommandBuffer* Buf = (FMetalDebugCommandBuffer*)CommandBuffer;
+					[Buf pushDebugGroup:Group];
+				}
+#endif
 			}
 		}
 	}
@@ -469,12 +498,26 @@ void FMetalCommandEncoder::BeginBlitCommandEncoding(void)
 	{
 		NSString* Label = [DebugGroups count] > 0 ? [DebugGroups lastObject] : (NSString*)CFSTR("InitialPass");
 		[BlitCommandEncoder setLabel:Label];
+#if !UE_BUILD_SHIPPING
+		if (CommandList.GetCommandQueue().GetRuntimeDebuggingLevel() >= EMetalDebugLevelLogOperations)
+		{
+			FMetalDebugCommandBuffer* Buf = (FMetalDebugCommandBuffer*)CommandBuffer;
+			[Buf beginBlitCommandEncoder:Label];
+		}
+#endif
 		
 		if([DebugGroups count])
 		{
 			for (NSString* Group in DebugGroups)
 			{
 				[BlitCommandEncoder pushDebugGroup:Group];
+#if !UE_BUILD_SHIPPING
+				if (CommandList.GetCommandQueue().GetRuntimeDebuggingLevel() >= EMetalDebugLevelLogDebugGroups)
+				{
+					FMetalDebugCommandBuffer* Buf = (FMetalDebugCommandBuffer*)CommandBuffer;
+					[Buf pushDebugGroup:Group];
+				}
+#endif
 			}
 		}
 	}
@@ -482,6 +525,13 @@ void FMetalCommandEncoder::BeginBlitCommandEncoding(void)
 
 void FMetalCommandEncoder::EndEncoding(void)
 {
+#if !UE_BUILD_SHIPPING
+	if (CommandList.GetCommandQueue().GetRuntimeDebuggingLevel() >= EMetalDebugLevelLogOperations)
+	{
+		FMetalDebugCommandBuffer* Buf = (FMetalDebugCommandBuffer*)CommandBuffer;
+		[Buf endCommandEncoder];
+	}
+#endif
 	if(IsRenderCommandEncoderActive())
 	{
 		[RenderCommandEncoder endEncoding];
@@ -506,6 +556,14 @@ void FMetalCommandEncoder::EndEncoding(void)
 
 void FMetalCommandEncoder::InsertDebugSignpost(NSString* const String)
 {
+#if !UE_BUILD_SHIPPING
+	if (CommandList.GetCommandQueue().GetRuntimeDebuggingLevel() >= EMetalDebugLevelLogDebugGroups)
+	{
+		FMetalDebugCommandBuffer* Buf = (FMetalDebugCommandBuffer*)CommandBuffer;
+		[Buf insertDebugSignpost:String];
+	}
+#endif
+	
 	if (RenderCommandEncoder)
 	{
 		[RenderCommandEncoder insertDebugSignpost:String];
@@ -522,6 +580,14 @@ void FMetalCommandEncoder::InsertDebugSignpost(NSString* const String)
 
 void FMetalCommandEncoder::PushDebugGroup(NSString* const String)
 {
+#if !UE_BUILD_SHIPPING
+	if (CommandList.GetCommandQueue().GetRuntimeDebuggingLevel() >= EMetalDebugLevelLogDebugGroups)
+	{
+		FMetalDebugCommandBuffer* Buf = (FMetalDebugCommandBuffer*)CommandBuffer;
+		[Buf pushDebugGroup:String];
+	}
+#endif
+	
 	[DebugGroups addObject:String];
 	if (RenderCommandEncoder)
 	{
@@ -539,6 +605,14 @@ void FMetalCommandEncoder::PushDebugGroup(NSString* const String)
 
 void FMetalCommandEncoder::PopDebugGroup(void)
 {
+#if !UE_BUILD_SHIPPING
+	if (CommandList.GetCommandQueue().GetRuntimeDebuggingLevel() >= EMetalDebugLevelLogDebugGroups)
+	{
+		FMetalDebugCommandBuffer* Buf = (FMetalDebugCommandBuffer*)CommandBuffer;
+		[Buf popDebugGroup];
+	}
+#endif
+	
 	[DebugGroups removeLastObject];
 	if (RenderCommandEncoder)
 	{
diff --git a/Engine/Source/Runtime/Apple/MetalRHI/Private/MetalCommandList.cpp b/Engine/Source/Runtime/Apple/MetalRHI/Private/MetalCommandList.cpp
--- a/Engine/Source/Runtime/Apple/MetalRHI/Private/MetalCommandList.cpp
+++ b/Engine/Source/Runtime/Apple/MetalRHI/Private/MetalCommandList.cpp
@@ -42,6 +42,9 @@ static void ReportMetalCommandBufferFailure(id <MTLCommandBuffer> CompletedBuffe
 	FString FailureString = FailureDesc ? FString(FailureDesc) : FString(TEXT("Unknown"));
 	FString RecoveryString = RecoveryDesc ? FString(RecoveryDesc) : FString(TEXT("Unknown"));
 	
+	NSString* Desc = CompletedBuffer.debugDescription;
+	UE_LOG(LogMetal, Warning, TEXT("%s"), *FString(Desc));
+	
 	checkf(CompletedBuffer.status != MTLCommandBufferStatusError, TEXT("Command Buffer %s Failed with %s Error! Error Domain: %s Code: %d Description %s %s %s"), *LabelString, ErrorType, *DomainString, Code, *ErrorString, *FailureString, *RecoveryString);
 }
 
@@ -162,7 +165,7 @@ void FMetalCommandList::Commit(id<MTLCommandBuffer> Buffer, bool const bWait)
 	check(Buffer);
 	
 	[Buffer addCompletedHandler : ^ (id <MTLCommandBuffer> CompletedBuffer)
-	 {
+	{
 		if (CompletedBuffer.status == MTLCommandBufferStatusError)
 		{
 			HandleMetalCommandBufferFailure(CompletedBuffer);
diff --git a/Engine/Source/Runtime/Apple/MetalRHI/Private/MetalCommandList.h b/Engine/Source/Runtime/Apple/MetalRHI/Private/MetalCommandList.h
--- a/Engine/Source/Runtime/Apple/MetalRHI/Private/MetalCommandList.h
+++ b/Engine/Source/Runtime/Apple/MetalRHI/Private/MetalCommandList.h
@@ -57,6 +57,9 @@ public:
 	 * @returns True iff the command-list submits immediately to the command-queue, false if it performs any buffering.
 	 */
 	bool IsImmediate(void) const { return bImmediate; }
+
+	/** @returns The command queue to which this command-list submits command-buffers. */
+	FMetalCommandQueue& GetCommandQueue(void) const { return CommandQueue; }
 	
 private:
 #pragma mark - Private Member Variables -
diff --git a/Engine/Source/Runtime/Apple/MetalRHI/Private/MetalCommandQueue.cpp b/Engine/Source/Runtime/Apple/MetalRHI/Private/MetalCommandQueue.cpp
--- a/Engine/Source/Runtime/Apple/MetalRHI/Private/MetalCommandQueue.cpp
+++ b/Engine/Source/Runtime/Apple/MetalRHI/Private/MetalCommandQueue.cpp
@@ -7,6 +7,7 @@
 #include "MetalRHIPrivate.h"
 
 #include "MetalCommandQueue.h"
+#include "MetalCommandBuffer.h"
 #include "MetalCommandList.h"
 #include "MetalProfiler.h"
 #if METAL_STATISTICS
@@ -21,6 +22,9 @@ FMetalCommandQueue::FMetalCommandQueue(id<MTLDevice> Device, uint32 const MaxNum
 #if METAL_STATISTICS
 , Statistics(nullptr)
 #endif
+#if !UE_BUILD_SHIPPING
+, RuntimeDebuggingLevel(EMetalDebugLevelOff)
+#endif
 {
 	if(MaxNumCommandBuffers == 0)
 	{
@@ -64,6 +68,12 @@ id<MTLCommandBuffer> FMetalCommandQueue::CreateRetainedCommandBuffer(void)
 	@autoreleasepool
 	{
 		id<MTLCommandBuffer> CmdBuffer = [[CommandQueue commandBuffer] retain];
+#if !UE_BUILD_SHIPPING
+		if (RuntimeDebuggingLevel >= EMetalDebugLevelLogDebugGroups)
+		{
+			CmdBuffer = [[FMetalDebugCommandBuffer alloc] initWithCommandBuffer:CmdBuffer];
+		}
+#endif
 		TRACK_OBJECT(STAT_MetalCommandBufferCount, CmdBuffer);
 		return CmdBuffer;
 	}
@@ -76,6 +86,12 @@ id<MTLCommandBuffer> FMetalCommandQueue::CreateUnretainedCommandBuffer(void)
 	@autoreleasepool
 	{
 		CmdBuffer = bUnretainedRefs ? [[CommandQueue commandBufferWithUnretainedReferences] retain] : [[CommandQueue commandBuffer] retain];
+#if !UE_BUILD_SHIPPING
+		if (RuntimeDebuggingLevel >= EMetalDebugLevelLogDebugGroups)
+		{
+			CmdBuffer = [[FMetalDebugCommandBuffer alloc] initWithCommandBuffer:CmdBuffer];
+		}
+#endif
 	}
 	TRACK_OBJECT(STAT_MetalCommandBufferCount, CmdBuffer);
 	return CmdBuffer;
@@ -87,6 +103,15 @@ void FMetalCommandQueue::CommitCommandBuffer(id<MTLCommandBuffer> const CommandB
 	UNTRACK_OBJECT(STAT_MetalCommandBufferCount, CommandBuffer);
 	
 	[CommandBuffer commit];
+	
+	// Wait for completion when debugging command-buffers.
+#if !UE_BUILD_SHIPPING
+	if (RuntimeDebuggingLevel >= EMetalDebugLevelWaitForComplete)
+	{
+		[CommandBuffer waitUntilCompleted];
+	}
+#endif
+	
 	[CommandBuffer release];
 }
 
@@ -133,6 +158,18 @@ void FMetalCommandQueue::InsertDebugCaptureBoundary(void)
 	[CommandQueue insertDebugCaptureBoundary];
 }
 
+#if !UE_BUILD_SHIPPING
+void FMetalCommandQueue::SetRuntimeDebuggingLevel(int32 const Level)
+{
+	RuntimeDebuggingLevel = Level;
+}
+
+int32 FMetalCommandQueue::GetRuntimeDebuggingLevel(void) const
+{
+	return RuntimeDebuggingLevel;
+}
+#endif
+
 #if METAL_STATISTICS
 #pragma mark - Public Statistics Extensions -
 
diff --git a/Engine/Source/Runtime/Apple/MetalRHI/Private/MetalCommandQueue.h b/Engine/Source/Runtime/Apple/MetalRHI/Private/MetalCommandQueue.h
--- a/Engine/Source/Runtime/Apple/MetalRHI/Private/MetalCommandQueue.h
+++ b/Engine/Source/Runtime/Apple/MetalRHI/Private/MetalCommandQueue.h
@@ -63,6 +63,14 @@ public:
 
 	/** Inserts a boundary that marks the end of a frame for the debug capture tool. */
 	void InsertDebugCaptureBoundary(void);
+	
+#if !UE_BUILD_SHIPPING
+	/** Enable or disable runtime debugging features. */
+	void SetRuntimeDebuggingLevel(int32 const Level);
+	
+	/** @returns The level of runtime debugging features enabled. */
+	int32 GetRuntimeDebuggingLevel(void) const;
+#endif
 
 #if METAL_STATISTICS
 #pragma mark - Public Statistics Extensions -
@@ -78,4 +86,7 @@ private:
 	class IMetalStatistics* Statistics;
 #endif
 	TArray<NSArray<id<MTLCommandBuffer>>*> CommandBuffers;
+#if !UE_BUILD_SHIPPING
+	int32 RuntimeDebuggingLevel;
+#endif
 };
\ No newline at end of file
diff --git a/Engine/Source/Runtime/Apple/MetalRHI/Private/MetalCommands.cpp b/Engine/Source/Runtime/Apple/MetalRHI/Private/MetalCommands.cpp
--- a/Engine/Source/Runtime/Apple/MetalRHI/Private/MetalCommands.cpp
+++ b/Engine/Source/Runtime/Apple/MetalRHI/Private/MetalCommands.cpp
@@ -14,6 +14,7 @@
 #include "SceneUtils.h"
 #include "ShaderCache.h"
 #include "MetalProfiler.h"
+#include "MetalCommandBuffer.h"
 
 static const bool GUsesInvertedZ = true;
 static FGlobalBoundShaderState GClearMRTBoundShaderState[8][2];
@@ -84,18 +85,29 @@ void FMetalRHICommandContext::RHIDispatchComputeShader(uint32 ThreadGroupCountX,
 	ThreadGroupCountX = FMath::Max(ThreadGroupCountX, 1u);
 	ThreadGroupCountY = FMath::Max(ThreadGroupCountY, 1u);
 	ThreadGroupCountZ = FMath::Max(ThreadGroupCountZ, 1u);
+	
+	METAL_DEBUG_COMMAND_BUFFER_DISPATCH_LOG(Context, @"RHIDispatchComputeShader(ThreadGroupCountX %d, ThreadGroupCountY %d, ThreadGroupCountZ %d)", ThreadGroupCountX, ThreadGroupCountY, ThreadGroupCountZ);
+	
 	Context->Dispatch(ThreadGroupCountX, ThreadGroupCountY, ThreadGroupCountZ);
 }
 
 void FMetalRHICommandContext::RHIDispatchIndirectComputeShader(FVertexBufferRHIParamRef ArgumentBufferRHI, uint32 ArgumentOffset)
 {
-#if PLATFORM_MAC
-	RHI_PROFILE_DRAW_CALL_STATS(EMTLSamplePointBeforeCompute, EMTLSamplePointAfterCompute, 1,1);
-	FMetalVertexBuffer* VertexBuffer = ResourceCast(ArgumentBufferRHI);
-	Context->DispatchIndirect(VertexBuffer, ArgumentOffset);
-#else
-	NOT_SUPPORTED("RHIDispatchIndirectComputeShader");
+#if METAL_API_1_1
+	if (GetMetalDeviceContext().SupportsFeature(EMetalFeaturesIndirectBuffer))
+	{
+		RHI_PROFILE_DRAW_CALL_STATS(EMTLSamplePointBeforeCompute, EMTLSamplePointAfterCompute, 1,1);
+		FMetalVertexBuffer* VertexBuffer = ResourceCast(ArgumentBufferRHI);
+		
+		METAL_DEBUG_COMMAND_BUFFER_DISPATCH_LOG(Context, @"RHIDispatchIndirectComputeShader(ArgumentBufferRHI %p, ArgumentOffset %d)", ArgumentBufferRHI, ArgumentOffset);
+		
+		Context->DispatchIndirect(VertexBuffer, ArgumentOffset);
+	}
+	else
 #endif
+	{
+		NOT_SUPPORTED("RHIDispatchIndirectComputeShader");
+	}
 }
 
 void FMetalRHICommandContext::RHISetViewport(uint32 MinX,uint32 MinY,float MinZ,uint32 MaxX,uint32 MaxY,float MaxZ)
@@ -482,6 +494,8 @@ void FMetalRHICommandContext::RHIDrawPrimitive(uint32 PrimitiveType, uint32 Base
 	uint32 VertexCount = GetVertexCountForPrimitiveCount(NumPrimitives,PrimitiveType);
 	RHI_PROFILE_DRAW_CALL_STATS(EMTLSamplePointBeforeDraw, EMTLSamplePointAfterDraw, NumPrimitives * NumInstances, VertexCount * NumInstances);
 	
+	METAL_DEBUG_COMMAND_BUFFER_DRAW_LOG(Context, @"RHIDrawPrimitive(PrimitiveType %d, BaseVertexIndex %d, NumPrimitives %d, NumInstances %d)", PrimitiveType, BaseVertexIndex, NumPrimitives, NumInstances);
+	
 	// draw!
 	if(!FShaderCache::IsPredrawCall())
 	{
@@ -507,6 +521,8 @@ void FMetalRHICommandContext::RHIDrawPrimitiveIndirect(uint32 PrimitiveType, FVe
 	// finalize any pending state
 	Context->PrepareToDraw(PrimitiveType);
 	
+	METAL_DEBUG_COMMAND_BUFFER_DRAW_LOG(Context, @"RHIDrawPrimitiveIndirect(PrimitiveType %d, VertexBufferRHI %p, ArgumentOffset %d)", PrimitiveType, VertexBufferRHI, ArgumentOffset);
+	
 	if(!FShaderCache::IsPredrawCall())
 	{
 		RHI_PROFILE_DRAW_CALL_STATS(EMTLSamplePointBeforeDraw, EMTLSamplePointAfterDraw, 1, 1);
@@ -525,11 +541,9 @@ void FMetalRHICommandContext::RHIDrawIndexedPrimitive(FIndexBufferRHIParamRef In
 	uint32 NumVertices, uint32 StartIndex, uint32 NumPrimitives, uint32 NumInstances)
 {
 	SCOPE_CYCLE_COUNTER(STAT_MetalDrawCallTime);
-#if PLATFORM_IOS
 	//checkf(NumInstances == 1, TEXT("Currently only 1 instance is supported"));
-	checkf(BaseVertexIndex  == 0, TEXT("BaseVertexIndex must be 0, see GRHISupportsBaseVertexIndex"));
-	checkf(FirstInstance  == 0, TEXT("FirstInstance is currently unsupported on this RHI"));
-#endif
+	checkf(GRHISupportsBaseVertexIndex || BaseVertexIndex == 0, TEXT("BaseVertexIndex must be 0, see GRHISupportsBaseVertexIndex"));
+	checkf(GRHISupportsFirstInstance || FirstInstance == 0, TEXT("FirstInstance must be 0, see GRHISupportsFirstInstance"));
 	
 	NumInstances = FMath::Max(NumInstances,1u);
 
@@ -550,89 +564,110 @@ void FMetalRHICommandContext::RHIDrawIndexedPrimitive(FIndexBufferRHIParamRef In
 
 	RHI_PROFILE_DRAW_CALL_STATS(EMTLSamplePointBeforeDraw, EMTLSamplePointAfterDraw, NumPrimitives * NumInstances, NumVertices * NumInstances);
 	
+	METAL_DEBUG_COMMAND_BUFFER_DRAW_LOG(Context, @"RHIDrawIndexedPrimitive(IndexBufferRHI %p, PrimitiveType %d, BaseVertexIndex %d, FirstInstance %d, NumVertices %d, StartIndex %d, NumPrimitives %d, NumInstances %d)", IndexBufferRHI, PrimitiveType, BaseVertexIndex, FirstInstance, NumVertices, StartIndex, NumPrimitives, NumInstances);
+	
 	if(!FShaderCache::IsPredrawCall())
 	{
-#if PLATFORM_IOS
-		[Context->GetRenderContext() drawIndexedPrimitives:TranslatePrimitiveType(PrimitiveType)
+#if METAL_API_1_1
+		if (GRHISupportsBaseVertexIndex && GRHISupportsFirstInstance)
+		{
+			[Context->GetRenderContext() drawIndexedPrimitives:TranslatePrimitiveType(PrimitiveType)
+													indexCount:NumIndices
+													 indexType:IndexBuffer->IndexType
+												   indexBuffer:IndexBuffer->Buffer
+											 indexBufferOffset:StartIndex * IndexBuffer->GetStride()
+												 instanceCount:NumInstances
+													baseVertex:BaseVertexIndex
+												  baseInstance:FirstInstance];
+		}
+		else
+#endif
+		{
+			[Context->GetRenderContext() drawIndexedPrimitives:TranslatePrimitiveType(PrimitiveType)
 												indexCount:NumIndices
 												 indexType:IndexBuffer->IndexType
 											   indexBuffer:IndexBuffer->Buffer
 										 indexBufferOffset:StartIndex * IndexBuffer->GetStride()
 											 instanceCount:NumInstances];
-#else
-		[Context->GetRenderContext() drawIndexedPrimitives:TranslatePrimitiveType(PrimitiveType)
-												indexCount:NumIndices
-												 indexType:IndexBuffer->IndexType
-											   indexBuffer:IndexBuffer->Buffer
-										 indexBufferOffset:StartIndex * IndexBuffer->GetStride()
-											 instanceCount:NumInstances
-											 baseVertex:BaseVertexIndex
-											 baseInstance:FirstInstance];
-#endif
+		}
 		FShaderCache::LogDraw(IndexBuffer->GetStride());
 	}
 }
 
 void FMetalRHICommandContext::RHIDrawIndexedIndirect(FIndexBufferRHIParamRef IndexBufferRHI, uint32 PrimitiveType, FStructuredBufferRHIParamRef VertexBufferRHI, int32 DrawArgumentsIndex, uint32 NumInstances)
 {
-#if PLATFORM_IOS
-	NOT_SUPPORTED("RHIDrawIndexedIndirect");
-#else
-	check(NumInstances > 1);
-
-	SCOPE_CYCLE_COUNTER(STAT_MetalDrawCallTime);
-	RHI_DRAW_CALL_STATS(PrimitiveType,1);
-	
-	FMetalIndexBuffer* IndexBuffer = ResourceCast(IndexBufferRHI);
-	FMetalStructuredBuffer* VertexBuffer = ResourceCast(VertexBufferRHI);
-	
-	// finalize any pending state
-	Context->PrepareToDraw(PrimitiveType);
-	
-	if(!FShaderCache::IsPredrawCall())
+#if METAL_API_1_1
+	if (GetMetalDeviceContext().SupportsFeature(EMetalFeaturesIndirectBuffer))
 	{
-		RHI_PROFILE_DRAW_CALL_STATS(EMTLSamplePointBeforeDraw, EMTLSamplePointAfterDraw, 1, 1);
+		check(NumInstances > 1);
 		
-		[Context->GetRenderContext() drawIndexedPrimitives:TranslatePrimitiveType(PrimitiveType)
+		SCOPE_CYCLE_COUNTER(STAT_MetalDrawCallTime);
+		RHI_DRAW_CALL_STATS(PrimitiveType,1);
+		
+		FMetalIndexBuffer* IndexBuffer = ResourceCast(IndexBufferRHI);
+		FMetalStructuredBuffer* VertexBuffer = ResourceCast(VertexBufferRHI);
+		
+		// finalize any pending state
+		Context->PrepareToDraw(PrimitiveType);
+		
+		METAL_DEBUG_COMMAND_BUFFER_DRAW_LOG(Context, @"RHIDrawIndexedIndirect(IndexBufferRHI %p, PrimitiveType %d, VertexBufferRHI %p, DrawArgumentsIndex %d, NumInstances %d)", IndexBufferRHI, PrimitiveType, VertexBufferRHI, DrawArgumentsIndex, NumInstances);
+		
+		if(!FShaderCache::IsPredrawCall())
+		{
+			RHI_PROFILE_DRAW_CALL_STATS(EMTLSamplePointBeforeDraw, EMTLSamplePointAfterDraw, 1, 1);
+			
+			[Context->GetRenderContext() drawIndexedPrimitives:TranslatePrimitiveType(PrimitiveType)
 													 indexType:IndexBuffer->IndexType
 												   indexBuffer:IndexBuffer->Buffer
 											 indexBufferOffset:0
 												indirectBuffer:VertexBuffer->Buffer
 										  indirectBufferOffset:(DrawArgumentsIndex * 5 * sizeof(uint32))];
-		
-		FShaderCache::LogDraw(IndexBuffer->GetStride());
+			
+			FShaderCache::LogDraw(IndexBuffer->GetStride());
+		}
 	}
+	else
 #endif
+	{
+		NOT_SUPPORTED("RHIDrawIndexedIndirect");
+	}
 }
 
 void FMetalRHICommandContext::RHIDrawIndexedPrimitiveIndirect(uint32 PrimitiveType,FIndexBufferRHIParamRef IndexBufferRHI,FVertexBufferRHIParamRef VertexBufferRHI,uint32 ArgumentOffset)
 {
-#if PLATFORM_IOS
-	NOT_SUPPORTED("RHIDrawIndexedPrimitiveIndirect");
-#else
-	SCOPE_CYCLE_COUNTER(STAT_MetalDrawCallTime);
-	RHI_DRAW_CALL_STATS(PrimitiveType,1);
-	
-	FMetalIndexBuffer* IndexBuffer = ResourceCast(IndexBufferRHI);
-	FMetalVertexBuffer* VertexBuffer = ResourceCast(VertexBufferRHI);
-	
-	// finalize any pending state
-	Context->PrepareToDraw(PrimitiveType);
-	
-	if(!FShaderCache::IsPredrawCall())
+#if METAL_API_1_1
+	if (GetMetalDeviceContext().SupportsFeature(EMetalFeaturesIndirectBuffer))
 	{
-		RHI_PROFILE_DRAW_CALL_STATS(EMTLSamplePointBeforeDraw, EMTLSamplePointAfterDraw, 1, 1);
+		SCOPE_CYCLE_COUNTER(STAT_MetalDrawCallTime);
+		RHI_DRAW_CALL_STATS(PrimitiveType,1);
+		
+		FMetalIndexBuffer* IndexBuffer = ResourceCast(IndexBufferRHI);
+		FMetalVertexBuffer* VertexBuffer = ResourceCast(VertexBufferRHI);
+		
+		// finalize any pending state
+		Context->PrepareToDraw(PrimitiveType);
 		
-		[Context->GetRenderContext() drawIndexedPrimitives:TranslatePrimitiveType(PrimitiveType)
+		METAL_DEBUG_COMMAND_BUFFER_DRAW_LOG(Context, @"RHIDrawIndexedPrimitiveIndirect(PrimitiveType %d, IndexBufferRHI %p, VertexBufferRHI %p, ArgumentOffset %d)", PrimitiveType, IndexBufferRHI, VertexBufferRHI, ArgumentOffset);
+		
+		if(!FShaderCache::IsPredrawCall())
+		{
+			RHI_PROFILE_DRAW_CALL_STATS(EMTLSamplePointBeforeDraw, EMTLSamplePointAfterDraw, 1, 1);
+			
+			[Context->GetRenderContext() drawIndexedPrimitives:TranslatePrimitiveType(PrimitiveType)
 													 indexType:IndexBuffer->IndexType
 												   indexBuffer:IndexBuffer->Buffer
 											 indexBufferOffset:0
 												indirectBuffer:VertexBuffer->Buffer
 										  indirectBufferOffset:ArgumentOffset];
-		
-		FShaderCache::LogDraw(IndexBuffer->GetStride());
+			
+			FShaderCache::LogDraw(IndexBuffer->GetStride());
+		}
 	}
+	else
 #endif
+	{
+		NOT_SUPPORTED("RHIDrawIndexedPrimitiveIndirect");
+	}
 }
 
 
@@ -653,6 +688,8 @@ void FMetalRHICommandContext::RHIBeginDrawPrimitiveUP( uint32 PrimitiveType, uin
 	PendingPrimitiveType = PrimitiveType;
 	PendingNumPrimitives = NumPrimitives;
 	PendingVertexDataStride = VertexDataStride;
+	
+	METAL_DEBUG_COMMAND_BUFFER_DRAW_LOG(Context, @"RHIBeginDrawPrimitiveUP(PrimitiveType %d, NumPrimitives %d, NumVertices %d, VertexDataStride %d)", PrimitiveType, NumPrimitives, NumVertices, VertexDataStride);
 }
 
 
@@ -670,6 +707,8 @@ void FMetalRHICommandContext::RHIEndDrawPrimitiveUP()
 
 	// last minute draw setup
 	Context->PrepareToDraw(PendingPrimitiveType);
+	
+	METAL_DEBUG_COMMAND_BUFFER_DRAW_LOG(Context, @"%@", @"RHIEndDrawPrimitiveUP()");
 
 	if(!FShaderCache::IsPredrawCall())
 	{
@@ -707,6 +746,8 @@ void FMetalRHICommandContext::RHIBeginDrawIndexedPrimitiveUP( uint32 PrimitiveTy
 	PendingIndexDataStride = IndexDataStride;
 
 	PendingVertexDataStride = VertexDataStride;
+	
+	METAL_DEBUG_COMMAND_BUFFER_DRAW_LOG(Context, @"RHIBeginDrawIndexedPrimitiveUP(PrimitiveType %d, NumPrimitives %d, NumVertices %d, VertexDataStride %d, MinVertexIndex %d, NumIndices %d, IndexDataStride %d)", PrimitiveType, NumPrimitives, NumVertices, VertexDataStride, MinVertexIndex, NumIndices, IndexDataStride);
 }
 
 void FMetalRHICommandContext::RHIEndDrawIndexedPrimitiveUP()
@@ -724,6 +765,8 @@ void FMetalRHICommandContext::RHIEndDrawIndexedPrimitiveUP()
 	// last minute draw setup
 	Context->PrepareToDraw(PendingPrimitiveType);
 	
+	METAL_DEBUG_COMMAND_BUFFER_DRAW_LOG(Context, @"%@", @"RHIEndDrawIndexedPrimitiveUP()");
+	
 	if(!FShaderCache::IsPredrawCall())
 	{
 		[Context->GetRenderContext() drawIndexedPrimitives:TranslatePrimitiveType(PendingPrimitiveType)
@@ -772,6 +815,8 @@ void FMetalRHICommandContext::RHIDrawInstancedPrimitiveUP( FRHICommandList& RHIC
 	// last minute draw setup
 	Context->PrepareToDraw(PrimitiveType);
 	
+	METAL_DEBUG_COMMAND_BUFFER_DRAW_LOG(Context, @"RHIDrawInstancedPrimitiveUP( PrimitiveType %d, NumPrimitives %d, VertexData %p, VertexDataStride %d, InstanceCount %d)", PrimitiveType, NumPrimitives, VertexData, VertexDataStride, InstanceCount);
+	
 	if(!FShaderCache::IsPredrawCall())
 	{
 		[Context->GetRenderContext() drawPrimitives:TranslatePrimitiveType(PrimitiveType)
diff --git a/Engine/Source/Runtime/Apple/MetalRHI/Private/MetalContext.cpp b/Engine/Source/Runtime/Apple/MetalRHI/Private/MetalContext.cpp
--- a/Engine/Source/Runtime/Apple/MetalRHI/Private/MetalContext.cpp
+++ b/Engine/Source/Runtime/Apple/MetalRHI/Private/MetalContext.cpp
@@ -7,6 +7,12 @@
 
 #include "MetalContext.h"
 #include "MetalProfiler.h"
+#include "MetalCommandBuffer.h"
+
+#if METAL_STATISTICS
+#include "MetalStatistics.h"
+#include "ModuleManager.h"
+#endif
 
 static int32 GMetalCommandBufferCommitThreshold = 100;
 static FAutoConsoleVariableRef CVarMetalCommandBufferCommitThreshold(
@@ -14,6 +20,28 @@ static FAutoConsoleVariableRef CVarMetalCommandBufferCommitThreshold(
 	GMetalCommandBufferCommitThreshold,
 	TEXT("When enabled (> 0) if the command buffer has more than this number of draw/dispatch command encoded then it will be committed at the next encoder boundary to keep the GPU busy. (Default: 100, set to <= 0 to disable)"));
 
+static int32 GMetalCommandQueueSize = 64;
+static FAutoConsoleVariableRef CVarMetalCommandQueueSize(
+	TEXT("rhi.Metal.CommandQueueSize"),
+	GMetalCommandQueueSize,
+	TEXT("The maximum number of command-buffers that can be allocated from each command-queue. (Default: 64)"));
+
+#if !UE_BUILD_SHIPPING
+static int32 GMetalRuntimeDebugLevel = 0;
+static FAutoConsoleVariableRef CVarMetalRuntimeDebugLevel(
+	TEXT("rhi.Metal.RuntimeDebugLevel"),
+	GMetalRuntimeDebugLevel,
+	TEXT("The level of debug validation performed by MetalRHI in addition to the underlying Metal API & validation layer.\n")
+	TEXT("Each subsequent level adds more tests and reporting in addition to the previous level.\n")
+	TEXT("*IGNORED IN SHIPPING BUILDS*. (Default: 0, off)\n")
+	TEXT("\t0: Off,\n")
+	TEXT("\t1: Enable validation checks for encoder resets,\n")
+	TEXT("\t2: Record the debug-groups issued into a command-buffer and report them on failure,\n")
+	TEXT("\t3: Record the draw, blit & dispatch commands issued into a command-buffer and report them on failure,\n")
+	TEXT("\t4: Allow rhi.Metal.CommandBufferCommitThreshold to break command-encoders (except when MSAA is enabled),\n")
+	TEXT("\t5: Wait for each command-buffer to complete immediately after submission."));
+#endif
+
 const uint32 RingBufferSize = 8 * 1024 * 1024;
 
 #if SHOULD_TRACK_OBJECTS
@@ -67,6 +95,10 @@ static id<MTLDevice> GetMTLDevice(uint32& DeviceIndex)
 	
 	DeviceIndex = 0;
 	
+#if METAL_STATISTICS
+	IMetalStatisticsModule* StatsModule = FModuleManager::Get().LoadModulePtr<IMetalStatisticsModule>(TEXT("MetalStatistics"));
+#endif
+	
 	NSArray* DeviceList = MTLCopyAllDevices();
 	[DeviceList autorelease];
 	
@@ -176,7 +208,7 @@ FMetalDeviceContext* FMetalDeviceContext::CreateDeviceContext()
 #else // @todo zebra
 	id<MTLDevice> Device = GetMTLDevice(DeviceIndex);
 #endif
-	FMetalCommandQueue* Queue = new FMetalCommandQueue(Device);
+	FMetalCommandQueue* Queue = new FMetalCommandQueue(Device, GMetalCommandQueueSize);
 	check(Queue);
 	
 	return new FMetalDeviceContext(Device, DeviceIndex, Queue);
@@ -192,24 +224,49 @@ FMetalDeviceContext::FMetalDeviceContext(id<MTLDevice> MetalDevice, uint32 InDev
 , Features(0)
 , ActiveContexts(1)
 {
+#if !UE_BUILD_SHIPPING
+	CommandQueue.SetRuntimeDebuggingLevel(GMetalRuntimeDebugLevel);
+#endif
+	
 #if PLATFORM_IOS
 	NSOperatingSystemVersion Vers = [[NSProcessInfo processInfo] operatingSystemVersion];
 	if(Vers.majorVersion >= 9)
 	{
 		Features = EMetalFeaturesSeparateStencil | EMetalFeaturesSetBufferOffset | EMetalFeaturesResourceOptions | EMetalFeaturesDepthStencilBlitOptions;
-#if !PLATFORM_TVOS
+
+#if PLATFORM_TVOS
+		Features |= EMetalFeaturesDepthClipMode;
+		if(FParse::Param(FCommandLine::Get(),TEXT("metalv2")) || [Device supportsFeatureSet:MTLFeatureSet_tvOS_GPUFamily1_v2])
+		{
+			Features |= EMetalFeaturesStencilView;
+		}
+#else
 		if ([Device supportsFeatureSet:MTLFeatureSet_iOS_GPUFamily2_v1])
-#endif
 		{
 			Features |= EMetalFeaturesDepthClipMode;
 		}
+		
+		if ([Device supportsFeatureSet:MTLFeatureSet_iOS_GPUFamily3_v1])
+		{
+			Features |= EMetalFeaturesCountingQueries | EMetalFeaturesBaseVertexInstance | EMetalFeaturesIndirectBuffer;
+		}
+		
+		if(FParse::Param(FCommandLine::Get(),TEXT("metalv2")) || [Device supportsFeatureSet:MTLFeatureSet_iOS_GPUFamily3_v2] || [Device supportsFeatureSet:MTLFeatureSet_iOS_GPUFamily2_v3] || [Device supportsFeatureSet:MTLFeatureSet_iOS_GPUFamily1_v3])
+		{
+			Features |= EMetalFeaturesStencilView;
+		}
+#endif
 	}
 	else if(Vers.majorVersion == 8 && Vers.minorVersion >= 3)
 	{
 		Features = EMetalFeaturesSeparateStencil | EMetalFeaturesSetBufferOffset;
 	}
 #else // Assume that Mac & other platforms all support these from the start. They can diverge later.
-	Features = EMetalFeaturesSeparateStencil | EMetalFeaturesSetBufferOffset | EMetalFeaturesDepthClipMode | EMetalFeaturesResourceOptions | EMetalFeaturesDepthStencilBlitOptions;
+	Features = EMetalFeaturesSeparateStencil | EMetalFeaturesSetBufferOffset | EMetalFeaturesDepthClipMode | EMetalFeaturesResourceOptions | EMetalFeaturesDepthStencilBlitOptions | EMetalFeaturesCountingQueries | EMetalFeaturesBaseVertexInstance | EMetalFeaturesIndirectBuffer | EMetalFeaturesLayeredRendering;
+    if (FParse::Param(FCommandLine::Get(),TEXT("metalv2")) || [Device supportsFeatureSet:MTLFeatureSet_OSX_GPUFamily1_v2])
+    {
+        Features |= EMetalFeaturesStencilView | EMetalFeaturesDepth16;
+    }
 #endif
 	
 	// Hook into the ios framepacer, if it's enabled for this platform.
@@ -247,30 +304,29 @@ void FMetalDeviceContext::ClearFreeList()
 	while(Index < DelayedFreeLists.Num())
 	{
 		FMetalDelayedFreeList* Pair = DelayedFreeLists[Index];
-		Pair->FrameCount--;
-		if (Pair->FrameCount == 0)
+		if(dispatch_semaphore_wait(Pair->Signal, DISPATCH_TIME_NOW) == 0)
 		{
-			if(dispatch_semaphore_wait(Pair->Signal, DISPATCH_TIME_NOW) == 0)
+			dispatch_release(Pair->Signal);
+			for( id Entry : Pair->FreeList )
 			{
-				dispatch_release(Pair->Signal);
-				for( id Entry : Pair->FreeList )
-				{
-					CommandEncoder.UnbindObject(Entry);
-					[Entry release];
-				}
-
+				CommandEncoder.UnbindObject(Entry);
+				[Entry release];
+			}
+			{
+				FScopeLock Lock(&PoolMutex);
+				for ( id<MTLBuffer> Buffer in Pair->FreeBuffers )
 				{
-					FScopeLock Lock(&PoolMutex);
-					for ( id<MTLBuffer> Buffer in Pair->FreeBuffers )
-					{
-						BufferPool.ReleasePooledResource(Buffer);
-					}
-					[Pair->FreeBuffers release];
+					DEC_MEMORY_STAT_BY(STAT_MetalUsedPooledBufferMem, Buffer.length);
+					INC_DWORD_STAT(STAT_MetalBufferFreed);
+					INC_DWORD_STAT_BY(STAT_MetalBufferMemFreed, Buffer.length);
+					INC_MEMORY_STAT_BY(STAT_MetalFreePooledBufferMem, Buffer.length);
+					
+					BufferPool.ReleasePooledResource(Buffer);
 				}
-
-				delete Pair;
-				DelayedFreeLists.RemoveAt(Index, 1, false);
+				[Pair->FreeBuffers release];
 			}
+			delete Pair;
+			DelayedFreeLists.RemoveAt(Index, 1, false);
 		}
 		else
 		{
@@ -325,7 +381,6 @@ void FMetalDeviceContext::EndDrawingViewport(FMetalViewport* Viewport, bool bPre
 	
 	// kick the whole buffer
 	FMetalDelayedFreeList* NewList = new FMetalDelayedFreeList;
-    NewList->FrameCount = 3;
 	NewList->Signal = dispatch_semaphore_create(0);
 	if(GUseRHIThread)
 	{
@@ -423,6 +478,8 @@ void FMetalDeviceContext::EndDrawingViewport(FMetalViewport* Viewport, bool bPre
 			
 			[Blitter copyFromTexture:Src sourceSlice:0 sourceLevel:0 sourceOrigin:MTLOriginMake(0, 0, 0) sourceSize:MTLSizeMake(Src.width, Src.height, 1) toTexture:Dst destinationSlice:0 destinationLevel:0 destinationOrigin:MTLOriginMake(0, 0, 0)];
 			
+			METAL_DEBUG_COMMAND_BUFFER_BLIT_LOG(this, @"Present(BackBuffer %p, Drawable %p)", Viewport->GetBackBuffer(), Drawable);
+			
 			CommandEncoder.EndEncoding();
 #endif
 			[CurrentCommandBuffer addScheduledHandler:^(id<MTLCommandBuffer>) {
@@ -437,7 +494,18 @@ void FMetalDeviceContext::EndDrawingViewport(FMetalViewport* Viewport, bool bPre
 		FrameReadyEvent->Wait();
 	}
 	
-	SubmitCommandsHint(EMetalSubmitFlagsNone);
+	// Latched update of whether to use runtime debugging features
+	uint32 SubmitFlags = EMetalSubmitFlagsNone;
+#if !UE_BUILD_SHIPPING
+	if (GMetalRuntimeDebugLevel != CommandQueue.GetRuntimeDebuggingLevel())
+	{
+		CommandQueue.SetRuntimeDebuggingLevel(GMetalRuntimeDebugLevel);
+		
+		// After change the debug features level wait on commit
+		SubmitFlags |= EMetalSubmitFlagsWaitOnCommandBuffer;
+	}
+#endif
+	SubmitCommandsHint(SubmitFlags);
 	
 #if SHOULD_TRACK_OBJECTS
 	// print out outstanding objects
@@ -518,12 +586,8 @@ void FMetalDeviceContext::ReleasePooledBuffer(FMetalPooledBuffer Buf)
 		{
 			CommandEncoder.UnbindObject(Buf.Buffer);
 		}
-		DEC_MEMORY_STAT_BY(STAT_MetalUsedPooledBufferMem, Buf.Buffer.length);
-		INC_DWORD_STAT(STAT_MetalBufferFreed);
-		INC_DWORD_STAT_BY(STAT_MetalBufferMemFreed, Buf.Buffer.length);
-		INC_MEMORY_STAT_BY(STAT_MetalFreePooledBufferMem, Buf.Buffer.length);
 		
-		FScopeLock Lock(&PoolMutex);
+		FScopeLock Lock(&FreeListMutex);
 		[FreeBuffers addObject:Buf.Buffer];
 		//[Buf.Buffer setPurgeableState:MTLPurgeableStateVolatile];
 	}
@@ -586,6 +650,14 @@ FMetalContext::FMetalContext(FMetalCommandQueue& Queue, bool const bIsImmediate)
 	FMemory::Memzero(BufferSideTable);
 	
 	bValidationEnabled = false;
+	
+#if METAL_STATISTICS
+	IMetalStatisticsModule* StatsModule = FModuleManager::Get().LoadModulePtr<IMetalStatisticsModule>(TEXT("MetalStatistics"));
+	if (StatsModule)
+	{
+		bValidationEnabled = StatsModule->IsValidationEnabled();
+	}
+#endif
 }
 
 FMetalContext::~FMetalContext()
@@ -769,7 +841,7 @@ void FMetalContext::ResetRenderCommandEncoder()
 {
 	SubmitCommandsHint();
 	
-	SetRenderTargetsInfo(StateCache.GetRenderTargetsInfo());
+	SetRenderTargetsInfo(StateCache.GetRenderTargetsInfo(), false);
 	
 	if (CommandEncoder.IsRenderCommandEncoderActive())
 	{
@@ -789,16 +861,34 @@ void FMetalContext::PrepareToDraw(uint32 PrimitiveType)
 	// Enforce calls to SetRenderTarget prior to issuing draw calls.
 	check(StateCache.GetHasValidRenderTarget());
 	
-	
-	// Enforce calls to SetRenderTarget prior to issuing draw calls.
-	check(StateCache.GetHasValidRenderTarget());
+	bool bUpdatedStrides = false;
+	MTLVertexDescriptor* Layout = CurrentBoundShaderState->VertexDeclaration->Layout;
+	if(Layout && Layout.layouts)
+	{
+		for (uint32 ElementIndex = 0; ElementIndex < CurrentBoundShaderState->VertexDeclaration->Elements.Num(); ElementIndex++)
+		{
+			const FVertexElement& Element = CurrentBoundShaderState->VertexDeclaration->Elements[ElementIndex];
+			
+			uint32 StreamStride = StateCache.GetVertexStride(UNREAL_TO_METAL_BUFFER_INDEX(Element.StreamIndex));
+			if (StreamStride > 0 && Element.Stride != StreamStride)
+			{
+				if (!bUpdatedStrides)
+				{
+					bUpdatedStrides = true;
+					Layout = [[CurrentBoundShaderState->VertexDeclaration->Layout copy] autorelease];
+				}
+				auto BufferLayout = [Layout.layouts objectAtIndexedSubscript:UNREAL_TO_METAL_BUFFER_INDEX(Element.StreamIndex)];
+				check(BufferLayout);
+				BufferLayout.stride = StreamStride;
+			}
+		}
+	}
 	
 	// Validate the vertex layout in debug mode, or when the validation layer is enabled for development builds.
 	// Other builds will just crash & burn if it is incorrect.
 #if UE_BUILD_DEBUG || UE_BUILD_DEVELOPMENT
-	if(UE_BUILD_DEBUG || bValidationEnabled)
+	if(CommandQueue.GetRuntimeDebuggingLevel() >= EMetalDebugLevelValidation)
 	{
-		MTLVertexDescriptor* Layout = CurrentBoundShaderState->VertexDeclaration->Layout;
 		if(Layout && Layout.layouts)
 		{
 			for (uint32 i = 0; i < MaxMetalStreams; i++)
@@ -863,21 +953,22 @@ void FMetalContext::PrepareToDraw(uint32 PrimitiveType)
 	}
 	
 	// make sure the BSS has a valid pipeline state object
-	CurrentBoundShaderState->PrepareToDraw(this, StateCache.GetRenderPipelineDesc());
+	CurrentBoundShaderState->PrepareToDraw(this, Layout, StateCache.GetRenderPipelineDesc());
 	
 	if(!FShaderCache::IsPredrawCall())
 	{
-		// Force a command-encoder break on Nvidia to see if this helps with intermittent command-buffer failures.
-		if (GMetalCommandBufferCommitThreshold > 0 && OutstandingOpCount >= GMetalCommandBufferCommitThreshold && IsRHIDeviceNVIDIA())
+#if !UE_BUILD_SHIPPING
+		// Force a command-encoder when GMetalRuntimeDebugLevel is enabled to help track down intermittent command-buffer failures.
+		if (GMetalCommandBufferCommitThreshold > 0 && OutstandingOpCount >= GMetalCommandBufferCommitThreshold && CommandQueue.GetRuntimeDebuggingLevel() >= EMetalDebugLevelConditionalSubmit)
 		{
 			FRHISetRenderTargetsInfo CurrentRenderTargets = StateCache.GetRenderTargetsInfo();
 			bool bCanChangeRT = true;
-				
+			
 			if (CommandEncoder.IsRenderCommandEncoderActive())
 			{
 				const bool bIsMSAAActive = StateCache.GetHasValidRenderTarget() && StateCache.GetRenderPipelineDesc().SampleCount != 1;
 				bCanChangeRT = !bIsMSAAActive;
-
+				
 				for (int32 RenderTargetIndex = 0; bCanChangeRT && RenderTargetIndex < CurrentRenderTargets.NumColorRenderTargets; RenderTargetIndex++)
 				{
 					FRHIRenderTargetView& RenderTargetView = CurrentRenderTargets.ColorRenderTarget[RenderTargetIndex];
@@ -911,7 +1002,7 @@ void FMetalContext::PrepareToDraw(uint32 PrimitiveType)
 				SubmitCommandsHint();
 				
 				// Force submit if there's enough outstanding commands to prevent the GPU going idle.
-				ConditionalSwitchToGraphics(StateCache.NeedsToSetRenderTarget(CurrentRenderTargets));
+				ConditionalSwitchToGraphics(StateCache.NeedsToSetRenderTarget(CurrentRenderTargets), true);
 				
 				if (IsFeatureLevelSupported( GMaxRHIShaderPlatform, ERHIFeatureLevel::SM4 ))
 				{
@@ -922,9 +1013,17 @@ void FMetalContext::PrepareToDraw(uint32 PrimitiveType)
 					StateCache.SetRenderTargetsInfo(CurrentRenderTargets, NULL, false);
 				}
 				
-				CommandEncoder.RestoreRenderCommandEncodingState();
+				if(!CommandEncoder.IsRenderCommandEncoderActive())
+				{
+					CommandEncoder.RestoreRenderCommandEncoding();
+				}
+				else
+				{
+					CommandEncoder.RestoreRenderCommandEncodingState();
+				}
 			}
 		}
+#endif
 		
 		CommitGraphicsResourceTables();
 		CommitNonComputeShaderConstants();
@@ -976,14 +1075,8 @@ void FMetalContext::PrepareToDraw(uint32 PrimitiveType)
 	}
 }
 
-void FMetalContext::SetRenderTargetsInfo(const FRHISetRenderTargetsInfo& RenderTargetsInfo)
+void FMetalContext::SetRenderTargetsInfo(const FRHISetRenderTargetsInfo& RenderTargetsInfo, bool const bReset)
 {
-	// Force a submit between command-encoders on Nvidia to test NV feedback.
-	if (StateCache.NeedsToSetRenderTarget(RenderTargetsInfo) && IsRHIDeviceNVIDIA())
-	{
-		SubmitCommandsHint();
-	}
-	
 	// Force submit if there's enough outstanding commands to prevent the GPU going idle.
 	ConditionalSwitchToGraphics(StateCache.NeedsToSetRenderTarget(RenderTargetsInfo), !(PLATFORM_MAC));
 	
@@ -1025,11 +1118,11 @@ void FMetalContext::SetRenderTargetsInfo(const FRHISetRenderTargetsInfo& RenderT
 	
 	if (IsFeatureLevelSupported( GMaxRHIShaderPlatform, ERHIFeatureLevel::SM4 ))
 	{
-		StateCache.SetRenderTargetsInfo(RenderTargetsInfo, QueryBuffer->GetCurrentQueryBuffer()->Buffer, true);
+		StateCache.SetRenderTargetsInfo(RenderTargetsInfo, QueryBuffer->GetCurrentQueryBuffer()->Buffer, bReset);
 	}
 	else
 	{
-		StateCache.SetRenderTargetsInfo(RenderTargetsInfo, NULL, true);
+		StateCache.SetRenderTargetsInfo(RenderTargetsInfo, NULL, bReset);
 	}
 }
 
@@ -1214,6 +1307,28 @@ void FMetalContext::SetResource(uint32 ShaderStage, uint32 BindIndex, FMetalSamp
 	}
 }
 
+void FMetalContext::SetResource(uint32 ShaderStage, uint32 BindIndex, FMetalUnorderedAccessView* RESTRICT UAV, float CurrentTime)
+{
+	switch (ShaderStage)
+	{
+		case CrossCompiler::SHADER_STAGE_PIXEL:
+			SetShaderUnorderedAccessView(SF_Pixel, BindIndex, UAV);
+			break;
+			
+		case CrossCompiler::SHADER_STAGE_VERTEX:
+			SetShaderUnorderedAccessView(SF_Vertex, BindIndex, UAV);
+			break;
+			
+		case CrossCompiler::SHADER_STAGE_COMPUTE:
+			SetShaderUnorderedAccessView(SF_Compute, BindIndex, UAV);
+			break;
+			
+		default:
+			check(0);
+			break;
+	}
+}
+
 
 template <typename MetalResourceType>
 inline int32 FMetalContext::SetShaderResourcesFromBuffer(uint32 ShaderStage, FMetalUniformBuffer* RESTRICT Buffer, const uint32* RESTRICT ResourceMap, int32 BufferIndex)
@@ -1291,6 +1406,7 @@ void FMetalContext::SetResourcesFromTables(ShaderType Shader, uint32 ShaderStage
 		SetShaderResourcesFromBuffer<FRHITexture>(ShaderStage, Buffer, Shader->Bindings.ShaderResourceTable.TextureMap.GetData(), BufferIndex);
 		SetShaderResourcesFromBuffer<FMetalShaderResourceView>(ShaderStage, Buffer, Shader->Bindings.ShaderResourceTable.ShaderResourceViewMap.GetData(), BufferIndex);
 		SetShaderResourcesFromBuffer<FMetalSamplerState>(ShaderStage, Buffer, Shader->Bindings.ShaderResourceTable.SamplerMap.GetData(), BufferIndex);
+		SetShaderResourcesFromBuffer<FMetalUnorderedAccessView>(ShaderStage, Buffer, Shader->Bindings.ShaderResourceTable.UnorderedAccessViewMap.GetData(), BufferIndex);
 	}
 	StateCache.SetDirtyUniformBuffers(Frequency, 0);
 }
@@ -1353,31 +1469,39 @@ void FMetalContext::ConditionalSubmit(bool bRTChangePending, bool bRTChangeForce
 {
 	if ((GMetalCommandBufferCommitThreshold > 0 && OutstandingOpCount >= GMetalCommandBufferCommitThreshold) || (bRTChangePending && bRTChangeForce))
 	{
+		bool const bEncoderActive = CommandEncoder.IsRenderCommandEncoderActive() || CommandEncoder.IsComputeCommandEncoderActive() || CommandEncoder.IsBlitCommandEncoderActive();
+		
+		// Normally we break only at natural encoder boundaries but in debug modes we can break at anytime.
+		bool bCanConditionallySubmit = !bEncoderActive;
+#if !UE_BUILD_SHIPPING
+		bCanConditionallySubmit |= (CommandQueue.GetRuntimeDebuggingLevel() >= EMetalDebugLevelConditionalSubmit);
+#endif
+
 		// AJB: This triggers a 'unset' of the RT. Causing a Load/Store at potentially awkward times.
 		// check that the load/store behaviour of the current RT setup will allows resumption without affecting RT content.
 		// i.e. - dont want to clear 1/2 way through a pass either.
-		bool bCanChangeRT = false;
-		if( bRTChangePending == false)
+		if (bRTChangePending == false && CommandEncoder.IsRenderCommandEncoderActive())
 		{
 			const FRHISetRenderTargetsInfo& CurrentRenderTargets = StateCache.GetRenderTargetsInfo();
 			const bool bIsMSAAActive = StateCache.GetHasValidRenderTarget() && StateCache.GetRenderPipelineDesc().SampleCount != 1;
-			bCanChangeRT = !bIsMSAAActive;
-			for (int32 RenderTargetIndex = 0; RenderTargetIndex < CurrentRenderTargets.NumColorRenderTargets && bCanChangeRT; RenderTargetIndex++)
+			bCanConditionallySubmit = !bIsMSAAActive;
+			for (int32 RenderTargetIndex = 0; RenderTargetIndex < CurrentRenderTargets.NumColorRenderTargets && bCanConditionallySubmit; RenderTargetIndex++)
 			{
 				const FRHIRenderTargetView& RenderTargetView = CurrentRenderTargets.ColorRenderTarget[RenderTargetIndex];
-				bCanChangeRT = bCanChangeRT
-				&& RenderTargetView.LoadAction == ERenderTargetLoadAction::ELoad
-				&& RenderTargetView.StoreAction == ERenderTargetStoreAction::EStore;
+				bCanConditionallySubmit = (bCanConditionallySubmit
+					&& RenderTargetView.LoadAction == ERenderTargetLoadAction::ELoad
+					&& RenderTargetView.StoreAction == ERenderTargetStoreAction::EStore);
 			}
-			bCanChangeRT = bCanChangeRT
-			&& (!CurrentRenderTargets.DepthStencilRenderTarget.Texture
+			bCanConditionallySubmit = (bCanConditionallySubmit
+				&& (!CurrentRenderTargets.DepthStencilRenderTarget.Texture
 				|| (CurrentRenderTargets.DepthStencilRenderTarget.DepthLoadAction == ERenderTargetLoadAction::ELoad
 					&& CurrentRenderTargets.DepthStencilRenderTarget.DepthStoreAction == ERenderTargetStoreAction::EStore
 					&& CurrentRenderTargets.DepthStencilRenderTarget.StencilLoadAction == ERenderTargetLoadAction::ELoad
-					&& CurrentRenderTargets.DepthStencilRenderTarget.GetStencilStoreAction() == ERenderTargetStoreAction::EStore));
+					&& CurrentRenderTargets.DepthStencilRenderTarget.GetStencilStoreAction() == ERenderTargetStoreAction::EStore)));
 		}
 
-		if(bRTChangePending || bCanChangeRT)
+		// Only permit conditional submit if one of the conditions is met.
+		if(bRTChangePending || bCanConditionallySubmit)
 		{
 			SubmitCommandsHint();
 		}
@@ -1420,7 +1544,7 @@ void FMetalContext::Dispatch(uint32 ThreadGroupCountX, uint32 ThreadGroupCountY,
 	OutstandingOpCount++;
 }
 
-#if PLATFORM_MAC
+#if METAL_API_1_1
 void FMetalContext::DispatchIndirect(FMetalVertexBuffer* ArgumentBuffer, uint32 ArgumentOffset)
 {
 	ConditionalSwitchToCompute();
diff --git a/Engine/Source/Runtime/Apple/MetalRHI/Private/MetalContext.h b/Engine/Source/Runtime/Apple/MetalRHI/Private/MetalContext.h
--- a/Engine/Source/Runtime/Apple/MetalRHI/Private/MetalContext.h
+++ b/Engine/Source/Runtime/Apple/MetalRHI/Private/MetalContext.h
@@ -32,7 +32,19 @@ enum EMetalFeatures
 	/** Support for specifying resource usage & memory options */
 	EMetalFeaturesResourceOptions = 1 << 3,
 	/** Supports texture->buffer blit options for depth/stencil blitting */
-	EMetalFeaturesDepthStencilBlitOptions = 1 << 4
+	EMetalFeaturesDepthStencilBlitOptions = 1 << 4,
+    /** Supports creating a native stencil texture view from a depth/stencil texture */
+    EMetalFeaturesStencilView = 1 << 5,
+    /** Supports a depth-16 pixel format */
+    EMetalFeaturesDepth16 = 1 << 6,
+	/** Supports NSUInteger counting visibility queries */
+	EMetalFeaturesCountingQueries = 1 << 7,
+	/** Supports base vertex/instance for draw calls */
+	EMetalFeaturesBaseVertexInstance = 1 << 8,
+	/** Supports indirect buffers for draw calls */
+	EMetalFeaturesIndirectBuffer = 1 << 9,
+	/** Supports layered rendering */
+	EMetalFeaturesLayeredRendering = 1 << 10
 };
 
 /**
@@ -89,7 +101,7 @@ public:
 	/**
 	 * Set the color, depth and stencil render targets, and then make the new command buffer/encoder
 	 */
-	void SetRenderTargetsInfo(const FRHISetRenderTargetsInfo& RenderTargetsInfo);
+	void SetRenderTargetsInfo(const FRHISetRenderTargetsInfo& RenderTargetsInfo, bool const bReset = true);
 	
 	/**
 	 * Allocate from a dynamic ring buffer - by default align to the allowed alignment for offset field when setting buffers
@@ -116,7 +128,7 @@ public:
 	void SetShaderUnorderedAccessView(EShaderFrequency ShaderStage, uint32 BindIndex, FMetalUnorderedAccessView* RESTRICT UAV);
 
 	void Dispatch(uint32 ThreadGroupCountX, uint32 ThreadGroupCountY, uint32 ThreadGroupCountZ);
-#if PLATFORM_MAC
+#if METAL_API_1_1
 	void DispatchIndirect(FMetalVertexBuffer* ArgumentBuffer, uint32 ArgumentOffset);
 #endif
 
@@ -161,6 +173,8 @@ private:
 	
 	FORCEINLINE void SetResource(uint32 ShaderStage, uint32 BindIndex, FMetalSamplerState* RESTRICT SamplerState, float CurrentTime);
 	
+	FORCEINLINE void SetResource(uint32 ShaderStage, uint32 BindIndex, FMetalUnorderedAccessView* RESTRICT UAV, float CurrentTime);
+	
 	template <typename MetalResourceType>
 	inline int32 SetShaderResourcesFromBuffer(uint32 ShaderStage, FMetalUniformBuffer* RESTRICT Buffer, const uint32* RESTRICT ResourceMap, int32 BufferIndex);
 	
@@ -203,7 +217,7 @@ protected:
 	
 	/** the slot to store a per-thread context ref */
 	static uint32 CurrentContextTLSSlot;
-
+	
 	/** The side table for runtiem validation of buffer access */
 	uint32 BufferSideTable[SF_NumFrequencies][MaxMetalStreams];
 	
@@ -275,7 +289,6 @@ private:
 		dispatch_semaphore_t Signal;
 		TSet<id> FreeList;
 		NSMutableSet<id<MTLBuffer>>* FreeBuffers;
-        int FrameCount;
 	};
 	TArray<FMetalDelayedFreeList*> DelayedFreeLists;
 	
diff --git a/Engine/Source/Runtime/Apple/MetalRHI/Private/MetalIndexBuffer.cpp b/Engine/Source/Runtime/Apple/MetalRHI/Private/MetalIndexBuffer.cpp
--- a/Engine/Source/Runtime/Apple/MetalRHI/Private/MetalIndexBuffer.cpp
+++ b/Engine/Source/Runtime/Apple/MetalRHI/Private/MetalIndexBuffer.cpp
@@ -6,6 +6,7 @@
 
 #include "MetalRHIPrivate.h"
 #include "MetalProfiler.h"
+#include "MetalCommandBuffer.h"
 
 /** Constructor */
 FMetalIndexBuffer::FMetalIndexBuffer(uint32 InStride, uint32 InSize, uint32 InUsage)
@@ -54,6 +55,7 @@ void* FMetalIndexBuffer::Lock(EResourceLockMode LockMode, uint32 Offset, uint32
 		
 		// Synchronise the buffer with the CPU
 		id<MTLBlitCommandEncoder> Blitter = GetMetalDeviceContext().GetBlitContext();
+		METAL_DEBUG_COMMAND_BUFFER_BLIT_LOG((&GetMetalDeviceContext()), @"SynchronizeResource(IndexBuffer %p)", this);
 		[Blitter synchronizeResource:Buffer];
 		
 		//kick the current command buffer.
diff --git a/Engine/Source/Runtime/Apple/MetalRHI/Private/MetalQuery.cpp b/Engine/Source/Runtime/Apple/MetalRHI/Private/MetalQuery.cpp
--- a/Engine/Source/Runtime/Apple/MetalRHI/Private/MetalQuery.cpp
+++ b/Engine/Source/Runtime/Apple/MetalRHI/Private/MetalQuery.cpp
@@ -109,11 +109,14 @@ void FMetalRenderQuery::Begin(FMetalContext* Context)
 		Result = 0;
 		bAvailable = false;
 		
-	#if PLATFORM_MAC
-		Context->GetCommandEncoder().SetVisibilityResultMode(MTLVisibilityResultModeCounting, Buffer.Offset);
-	#else
-		Context->GetCommandEncoder().SetVisibilityResultMode(MTLVisibilityResultModeBoolean, Buffer.Offset);
-	#endif
+		if ((GMaxRHIFeatureLevel >= ERHIFeatureLevel::SM4) && GetMetalDeviceContext().SupportsFeature(EMetalFeaturesCountingQueries))
+		{
+			Context->GetCommandEncoder().SetVisibilityResultMode(MTLVisibilityResultModeCounting, Buffer.Offset);
+		}
+		else
+		{
+			Context->GetCommandEncoder().SetVisibilityResultMode(MTLVisibilityResultModeBoolean, Buffer.Offset);
+		}
 	}
 }
 
diff --git a/Engine/Source/Runtime/Apple/MetalRHI/Private/MetalRHI.cpp b/Engine/Source/Runtime/Apple/MetalRHI/Private/MetalRHI.cpp
--- a/Engine/Source/Runtime/Apple/MetalRHI/Private/MetalRHI.cpp
+++ b/Engine/Source/Runtime/Apple/MetalRHI/Private/MetalRHI.cpp
@@ -97,8 +97,11 @@ FMetalDynamicRHI::FMetalDynamicRHI()
     bool bCanUseWideMRTs = true;
     bool bCanUseASTC = false;
 	bool bSupportsD24S8 = false;
+    bool bSupportsD16 = false;
 
-	if(FParse::Param(FCommandLine::Get(),TEXT("metalsm5")))
+	// Default to Metal SM5 on 10.11.5 or later. Earlier versions should still be running SM4.
+	bool const bTenElevenFiveOrLater = (FPlatformMisc::MacOSXVersionCompare(10,11,5) >= 0);
+	if(!FParse::Param(FCommandLine::Get(),TEXT("metal")) && (bTenElevenFiveOrLater || FParse::Param(FCommandLine::Get(),TEXT("metalsm5"))))
 	{
 		GMaxRHIFeatureLevel = ERHIFeatureLevel::SM5;
 		GMaxRHIShaderPlatform = SP_METAL_SM5;
@@ -116,20 +119,22 @@ FMetalDynamicRHI::FMetalDynamicRHI()
 	
 	GRHIAdapterName = FString(Device.name);
 	
+	// Mac GPUs support layer indexing.
+	GSupportsVolumeTextureRendering = true;
+	
+	// However they don't all support other features depending on the version of the OS.
 	bool bSupportsPointLights = false;
 	bool bSupportsTiledReflections = false;
+	bool bSupportsDistanceFields = false;
 	if(GRHIAdapterName.Contains("Nvidia"))
 	{
-		// Nvidia support layer indexing.
-		GSupportsVolumeTextureRendering = true;
 		bSupportsPointLights = true;
 		GRHIVendorId = 0x10DE;
 		bSupportsTiledReflections = true;
+		bSupportsDistanceFields = (FPlatformMisc::MacOSXVersionCompare(10,11,4) >= 0);
 	}
 	else if(GRHIAdapterName.Contains("ATi") || GRHIAdapterName.Contains("AMD"))
 	{
-		// AMD support layer indexing.
-		GSupportsVolumeTextureRendering = true;
 		bSupportsPointLights = true;
 		GRHIVendorId = 0x1002;
 		if(GPUDesc.GPUVendorId == GRHIVendorId)
@@ -137,10 +142,10 @@ FMetalDynamicRHI::FMetalDynamicRHI()
 			GRHIAdapterName = FString(GPUDesc.GPUName);
 		}
 		bSupportsTiledReflections = true;
+		bSupportsDistanceFields = (FPlatformMisc::MacOSXVersionCompare(10,11,4) >= 0);
 	}
 	else if(GRHIAdapterName.Contains("Intel"))
 	{
-		GSupportsVolumeTextureRendering = true;
 		bSupportsPointLights = (FPlatformMisc::MacOSXVersionCompare(10,11,4) >= 0);
 		GRHIVendorId = 0x8086;
 	}
@@ -195,6 +200,22 @@ FMetalDynamicRHI::FMetalDynamicRHI()
 		}
 	}
 	
+	// Disable the distance field AO & shadowing effects on GPU drivers that don't currently execute the shaders correctly.
+	if (GMaxRHIShaderPlatform == SP_METAL_SM5 && !bSupportsDistanceFields && !FParse::Param(FCommandLine::Get(),TEXT("metaldistancefields")))
+	{
+		static auto CVarDistanceFieldAO = IConsoleManager::Get().FindConsoleVariable(TEXT("r.DistanceFieldAO"));
+		if(CVarDistanceFieldAO && CVarDistanceFieldAO->GetInt() != 0)
+		{
+			CVarDistanceFieldAO->Set(0);
+		}
+		
+		static auto CVarDistanceFieldShadowing = IConsoleManager::Get().FindConsoleVariable(TEXT("r.DistanceFieldShadowing"));
+		if(CVarDistanceFieldShadowing && CVarDistanceFieldShadowing->GetInt() != 0)
+		{
+			CVarDistanceFieldShadowing->Set(0);
+		}
+	}
+	
 	GRHISupportsRHIThread = false;
 	if (GMaxRHIShaderPlatform == SP_METAL_SM5)
 	{
@@ -242,16 +263,22 @@ FMetalDynamicRHI::FMetalDynamicRHI()
 
 #if PLATFORM_MAC
 	check([Device supportsFeatureSet:MTLFeatureSet_OSX_GPUFamily1_v1]);
-	GRHISupportsBaseVertexIndex = FPlatformMisc::MacOSXVersionCompare(10,11,2) >= 0 || !IsRHIDeviceAMD(); // Supported on OS X but not iOS - broken on AMD prior to 10.11.2
-	GRHISupportsFirstInstance = true; // Supported on OS X but not iOS.
+	GRHISupportsBaseVertexIndex = FPlatformMisc::MacOSXVersionCompare(10,11,2) >= 0 || !IsRHIDeviceAMD(); // Supported on macOS & iOS but not tvOS - broken on AMD prior to 10.11.2
+	GRHISupportsFirstInstance = true; // Supported on macOS & iOS but not tvOS.
 	GMaxTextureDimensions = 16384;
 	GMaxCubeTextureDimensions = 16384;
 	GMaxTextureArrayLayers = 2048;
 	GMaxShadowDepthBufferSizeX = 16384;
 	GMaxShadowDepthBufferSizeY = 16384;
+    bSupportsD16 = FParse::Param(FCommandLine::Get(),TEXT("metalv2")) || [Device supportsFeatureSet:MTLFeatureSet_OSX_GPUFamily1_v2];
 #else
+#if PLATFORM_TVOS
 	GRHISupportsBaseVertexIndex = false;
-	GRHISupportsFirstInstance = false; // Supported on OS X but not iOS.
+	GRHISupportsFirstInstance = false; // Supported on macOS & iOS but not tvOS.
+#else
+	GRHISupportsBaseVertexIndex = [Device supportsFeatureSet:MTLFeatureSet_iOS_GPUFamily3_v1];
+	GRHISupportsFirstInstance = GRHISupportsBaseVertexIndex;
+#endif
 	GMaxTextureDimensions = 4096;
 	GMaxCubeTextureDimensions = 4096;
 	GMaxTextureArrayLayers = 2048;
@@ -318,8 +345,16 @@ FMetalDynamicRHI::FMetalDynamicRHI()
 		GPixelFormats[PF_DepthStencil	].PlatformFormat	= MTLPixelFormatDepth32Float_Stencil8;
 	}
     GPixelFormats[PF_DepthStencil		].BlockBytes		= 4;
-    GPixelFormats[PF_ShadowDepth		].PlatformFormat	= MTLPixelFormatDepth32Float;
-    GPixelFormats[PF_ShadowDepth		].BlockBytes		= 4;
+    if (bSupportsD16)
+    {
+        GPixelFormats[PF_ShadowDepth		].PlatformFormat	= MTLPixelFormatDepth16Unorm;
+        GPixelFormats[PF_ShadowDepth		].BlockBytes		= 2;
+    }
+    else
+    {
+        GPixelFormats[PF_ShadowDepth		].PlatformFormat	= MTLPixelFormatDepth32Float;
+        GPixelFormats[PF_ShadowDepth		].BlockBytes		= 4;
+    }
 #endif
     GPixelFormats[PF_X24_G8				].PlatformFormat	= MTLPixelFormatStencil8;
     GPixelFormats[PF_X24_G8				].BlockBytes		= 1;
@@ -410,15 +445,18 @@ FMetalDynamicRHI::FMetalDynamicRHI()
 	GDynamicRHI = this;
 	
 #if PLATFORM_DESKTOP
-	static const auto CVar = IConsoleManager::Get().FindConsoleVariable(TEXT("r.Shaders.Optimize"));
-	
-	if (CVar->GetInt() == 0)
-	{
-		FShaderCache::InitShaderCache(SCO_NoShaderPreload, 128);
-	}
-	else
+	if (GMaxRHIFeatureLevel < ERHIFeatureLevel::SM5)
 	{
-		FShaderCache::InitShaderCache(SCO_Default, 128);
+		static const auto CVar = IConsoleManager::Get().FindConsoleVariable(TEXT("r.Shaders.Optimize"));
+		
+		if (CVar->GetInt() == 0)
+		{
+			FShaderCache::InitShaderCache(SCO_NoShaderPreload, 128);
+		}
+		else
+		{
+			FShaderCache::InitShaderCache(SCO_Default, 128);
+		}
 	}
 #endif
 
diff --git a/Engine/Source/Runtime/Apple/MetalRHI/Private/MetalRHIPrivate.h b/Engine/Source/Runtime/Apple/MetalRHI/Private/MetalRHIPrivate.h
--- a/Engine/Source/Runtime/Apple/MetalRHI/Private/MetalRHIPrivate.h
+++ b/Engine/Source/Runtime/Apple/MetalRHI/Private/MetalRHIPrivate.h
@@ -15,6 +15,7 @@ const uint32 MaxMetalRenderTargets = 6;
 const uint32 BufferOffsetAlignment = 256;
 
 #define METAL_API_1_1 (__IPHONE_9_0 || __MAC_10_11)
+#define METAL_API_1_2 ((__IPHONE_10_0 && __IPHONE_OS_VERSION_MIN_REQUIRED >= __IPHONE_10_0) || (__MAC_10_12 && __MAC_OS_X_VERSION_MIN_REQUIRED >= __MAC_10_12))
 
 #if METAL_API_1_1
 #define BUFFER_CACHE_MODE MTLResourceCPUCacheModeWriteCombined
@@ -23,9 +24,9 @@ const uint32 BufferOffsetAlignment = 256;
 #endif
 
 #if PLATFORM_MAC
-#define BUFFER_MANAGED_MEM MTLResourceStorageModeShared
-#define BUFFER_STORAGE_MODE MTLStorageModeShared
-#define BUFFER_RESOURCE_STORAGE_MANAGED MTLResourceStorageModeShared
+#define BUFFER_MANAGED_MEM MTLResourceStorageModeManaged
+#define BUFFER_STORAGE_MODE MTLStorageModeManaged
+#define BUFFER_RESOURCE_STORAGE_MANAGED MTLResourceStorageModeManaged
 #define BUFFER_DYNAMIC_REALLOC BUF_AnyDynamic
 // How many possible vertex streams are allowed
 const uint32 MaxMetalStreams = 31;
@@ -48,11 +49,26 @@ const uint32 MaxMetalStreams = 30;
 
 // Dependencies
 #include "MetalRHI.h"
-#include "MetalGlobalUniformBuffer.h"
 #include "RHI.h"
 #import <Metal/Metal.h>
 #import <QuartzCore/CAMetalLayer.h>
 
+#if !METAL_API_1_1
+#define MTLVisibilityResultModeCounting ((MTLVisibilityResultMode)2)
+#endif
+
+#if !METAL_API_1_2
+#define MTLFeatureSet_iOS_GPUFamily3_v1 ((MTLFeatureSet)4)
+#define MTLFeatureSet_iOS_GPUFamily1_v3 ((MTLFeatureSet)5)
+#define MTLFeatureSet_iOS_GPUFamily2_v3 ((MTLFeatureSet)6)
+#define MTLFeatureSet_iOS_GPUFamily3_v2 ((MTLFeatureSet)7)
+#define MTLFeatureSet_tvOS_GPUFamily1_v2 ((MTLFeatureSet)30001)
+#define MTLFeatureSet_OSX_GPUFamily1_v2 ((MTLFeatureSet)10001)
+#define MTLPixelFormatDepth16Unorm ((MTLPixelFormat)250)
+#define MTLPixelFormatX24_Stencil8 ((MTLPixelFormat)262)
+#define MTLPixelFormatX32_Stencil8 ((MTLPixelFormat)261)
+#endif
+
 // Access the internal context for the device-owning DynamicRHI object
 FMetalDeviceContext& GetMetalDeviceContext();
 
diff --git a/Engine/Source/Runtime/Apple/MetalRHI/Private/MetalRenderPipelineDesc.cpp b/Engine/Source/Runtime/Apple/MetalRHI/Private/MetalRenderPipelineDesc.cpp
--- a/Engine/Source/Runtime/Apple/MetalRHI/Private/MetalRenderPipelineDesc.cpp
+++ b/Engine/Source/Runtime/Apple/MetalRHI/Private/MetalRenderPipelineDesc.cpp
@@ -23,7 +23,7 @@ FMetalRenderPipelineDesc::FMetalRenderPipelineDesc()
 	}
 }
 
-id<MTLRenderPipelineState> FMetalRenderPipelineDesc::CreatePipelineStateForBoundShaderState(FMetalBoundShaderState* BSS) const
+id<MTLRenderPipelineState> FMetalRenderPipelineDesc::CreatePipelineStateForBoundShaderState(FMetalBoundShaderState* BSS, MTLVertexDescriptor* VertexDesc) const
 {
 	SCOPE_CYCLE_COUNTER(STAT_MetalPipelineStateTime);
 	
@@ -48,7 +48,7 @@ id<MTLRenderPipelineState> FMetalRenderPipelineDesc::CreatePipelineStateForBound
     }
 	
 	// set the bound shader state settings
-	PipelineDescriptor.vertexDescriptor = BSS->VertexDeclaration->Layout;
+	PipelineDescriptor.vertexDescriptor = VertexDesc;
 	PipelineDescriptor.vertexFunction = BSS->VertexShader->Function;
 	PipelineDescriptor.fragmentFunction = BSS->PixelShader ? BSS->PixelShader->Function : nil;
 	PipelineDescriptor.sampleCount = SampleCount;
diff --git a/Engine/Source/Runtime/Apple/MetalRHI/Private/MetalRenderPipelineDesc.h b/Engine/Source/Runtime/Apple/MetalRHI/Private/MetalRenderPipelineDesc.h
--- a/Engine/Source/Runtime/Apple/MetalRHI/Private/MetalRenderPipelineDesc.h
+++ b/Engine/Source/Runtime/Apple/MetalRHI/Private/MetalRenderPipelineDesc.h
@@ -64,7 +64,7 @@ struct FMetalRenderPipelineDesc
 	/**
 	 * @return an pipeline state object that matches the current state and the BSS
 	 */
-	id<MTLRenderPipelineState> CreatePipelineStateForBoundShaderState(FMetalBoundShaderState* BSS) const;
+	id<MTLRenderPipelineState> CreatePipelineStateForBoundShaderState(FMetalBoundShaderState* BSS, MTLVertexDescriptor* VertexDesc) const;
 
 	MTLRenderPipelineDescriptor* PipelineDescriptor;
 	uint32 SampleCount;
diff --git a/Engine/Source/Runtime/Apple/MetalRHI/Private/MetalRenderTarget.cpp b/Engine/Source/Runtime/Apple/MetalRHI/Private/MetalRenderTarget.cpp
--- a/Engine/Source/Runtime/Apple/MetalRHI/Private/MetalRenderTarget.cpp
+++ b/Engine/Source/Runtime/Apple/MetalRHI/Private/MetalRenderTarget.cpp
@@ -7,6 +7,7 @@
 #include "MetalRHIPrivate.h"
 #include "ScreenRendering.h"
 #include "MetalProfiler.h"
+#include "MetalCommandBuffer.h"
 
 int32 GMetalUseTexGetBytes = 1;
 static FAutoConsoleVariableRef CVarMetalUseTexGetBytes(
@@ -84,6 +85,8 @@ void FMetalRHICommandContext::RHICopyToResolveTarget(FTextureRHIParamRef SourceT
 			Profiler->RegisterGPUWork();
 		}
 		
+		METAL_DEBUG_COMMAND_BUFFER_BLIT_LOG(Context, @"RHICopyToResolveTarget(SourceTextureRHI %p, DestTextureRHI %p, bKeepOriginalSurface %d)", SourceTextureRHI, DestTextureRHI, (uint32)bKeepOriginalSurface);
+		
 		[Blitter copyFromTexture:Source->Texture sourceSlice:SrcIndex sourceLevel:ResolveParams.MipIndex sourceOrigin:Origin sourceSize:Size toTexture:Destination->Texture destinationSlice:DestIndex destinationLevel:ResolveParams.MipIndex destinationOrigin:Origin];
 	}
 }
@@ -487,6 +490,8 @@ void FMetalDynamicRHI::RHIReadSurfaceData(FTextureRHIParamRef TextureRHI, FIntRe
 			
 			TempTexture = [GetMetalDeviceContext().GetDevice() newTextureWithDescriptor:Desc];
 			
+			METAL_DEBUG_COMMAND_BUFFER_BLIT_LOG(Context, @"RHIReadSurfaceData(copyFromTexture:toTexture:(TextureRHI %p))", TextureRHI);
+			
 			id<MTLBlitCommandEncoder> Blitter = Context->GetBlitContext();
 			[Blitter copyFromTexture:Texture sourceSlice:0 sourceLevel:0 sourceOrigin:Region.origin sourceSize:Region.size toTexture:TempTexture destinationSlice:0 destinationLevel:0 destinationOrigin:MTLOriginMake(0, 0, 0)];
 			
@@ -498,6 +503,7 @@ void FMetalDynamicRHI::RHIReadSurfaceData(FTextureRHIParamRef TextureRHI, FIntRe
 		{
 			// Synchronise the texture with the CPU
 			id<MTLBlitCommandEncoder> Blitter = Context->GetBlitContext();
+			METAL_DEBUG_COMMAND_BUFFER_BLIT_LOG(Context, @"RHIReadSurfaceData(synchronizeTexture(TextureRHI %p))", TextureRHI);
 			[Blitter synchronizeTexture:Texture slice:0 level:0];
 			
 			//kick the current command buffer.
@@ -537,6 +543,7 @@ void FMetalDynamicRHI::RHIReadSurfaceData(FTextureRHIParamRef TextureRHI, FIntRe
 			id<MTLBlitCommandEncoder> Blitter = Context->GetBlitContext();
 			if (Surface->PixelFormat != PF_DepthStencil)
 			{
+				METAL_DEBUG_COMMAND_BUFFER_BLIT_LOG(Context, @"RHIReadSurfaceData(copyFromTexture:toBuffer:(TextureRHI %p))", TextureRHI);
 				[Blitter copyFromTexture:Texture sourceSlice:0 sourceLevel:0 sourceOrigin:Region.origin sourceSize:Region.size toBuffer:Buffer.Buffer destinationOffset:0 destinationBytesPerRow:AlignedStride destinationBytesPerImage:BytesPerImage];
 			}
 #if METAL_API_1_1
@@ -544,10 +551,12 @@ void FMetalDynamicRHI::RHIReadSurfaceData(FTextureRHIParamRef TextureRHI, FIntRe
 			{
 				if (!InFlags.GetOutputStencil())
 				{
+					METAL_DEBUG_COMMAND_BUFFER_BLIT_LOG(Context, @"RHIReadSurfaceData(copyFromTexture:toBuffer:(TextureRHI %p))", TextureRHI);
 					[Blitter copyFromTexture:Texture sourceSlice:0 sourceLevel:0 sourceOrigin:Region.origin sourceSize:Region.size toBuffer:Buffer.Buffer destinationOffset:0 destinationBytesPerRow:AlignedStride destinationBytesPerImage:BytesPerImage options:MTLBlitOptionDepthFromDepthStencil];
 				}
 				else
 				{
+					METAL_DEBUG_COMMAND_BUFFER_BLIT_LOG(Context, @"RHIReadSurfaceData(copyFromTexture:toBuffer:(TextureRHI %p))", TextureRHI);
 					[Blitter copyFromTexture:Texture sourceSlice:0 sourceLevel:0 sourceOrigin:Region.origin sourceSize:Region.size toBuffer:Buffer.Buffer destinationOffset:0 destinationBytesPerRow:AlignedStride destinationBytesPerImage:BytesPerImage options:MTLBlitOptionStencilFromDepthStencil];
 				}
 			}
@@ -635,6 +644,7 @@ void FMetalDynamicRHI::RHIReadSurfaceFloatData(FTextureRHIParamRef TextureRHI, F
 		SCOPE_CYCLE_COUNTER(STAT_MetalTexturePageOffTime);
 		
 		id<MTLBlitCommandEncoder> Blitter = Context->GetBlitContext();
+		METAL_DEBUG_COMMAND_BUFFER_BLIT_LOG(Context, @"RHIReadSurfaceFloatData(copyFromTexture:toBuffer:(TextureRHI %p))", TextureRHI);
 		[Blitter copyFromTexture:Texture sourceSlice:ArrayIndex sourceLevel:MipIndex sourceOrigin:Region.origin sourceSize:Region.size toBuffer:Buffer.Buffer destinationOffset:0 destinationBytesPerRow:AlignedStride destinationBytesPerImage:BytesPerImage];
 		
 		//kick the current command buffer.
@@ -702,6 +712,7 @@ void FMetalDynamicRHI::RHIRead3DSurfaceFloatData(FTextureRHIParamRef TextureRHI,
 		SCOPE_CYCLE_COUNTER(STAT_MetalTexturePageOffTime);
 		
 		id<MTLBlitCommandEncoder> Blitter = Context->GetBlitContext();
+		METAL_DEBUG_COMMAND_BUFFER_BLIT_LOG(Context, @"RHIRead3DSurfaceFloatData(copyFromTexture:toBuffer:(TextureRHI %p))", TextureRHI);
 		[Blitter copyFromTexture:Texture sourceSlice:0 sourceLevel:0 sourceOrigin:Region.origin sourceSize:Region.size toBuffer:Buffer.Buffer destinationOffset:0 destinationBytesPerRow:AlignedStride destinationBytesPerImage:BytesPerImage];
 		
 		//kick the current command buffer.
diff --git a/Engine/Source/Runtime/Apple/MetalRHI/Private/MetalShaders.cpp b/Engine/Source/Runtime/Apple/MetalRHI/Private/MetalShaders.cpp
--- a/Engine/Source/Runtime/Apple/MetalRHI/Private/MetalShaders.cpp
+++ b/Engine/Source/Runtime/Apple/MetalRHI/Private/MetalShaders.cpp
@@ -95,8 +95,57 @@ TMetalBaseShader<BaseResourceType, ShaderType>::TMetalBaseShader(const TArray<ui
 //	if (!Resource)
 	{
 		id<MTLLibrary> Library;
+		
+		GlslCodeNSString = @"OFFLINE";
+		
+		bool bOfflineCompile = (OfflineCompiledFlag > 0);
+		
+		if (bOfflineCompile && (Header.ShaderCode.Len() > 0))
+		{
+			NSString* ShaderSource = Header.ShaderCode.GetNSString();
+			
+			GlslCodeNSString = ShaderSource;
+			[GlslCodeNSString retain];
+			
+			// For debug/dev/test builds we can use the stored code for debugging - but shipping builds shouldn't have this as it is inappropriate.
+#if !UE_BUILD_SHIPPING
+			// For iOS/tvOS we must use runtime compilation to make the shaders debuggable, but
+			bool bSavedSource = false;
+			
+#if PLATFORM_MAC
+			// on Mac if we have a path for the shader we can access the shader code
+			if (Header.ShaderPath.Len() > 0)
+			{
+				if (IFileManager::Get().MakeDirectory(*FPaths::GetPath(Header.ShaderPath), true))
+				{
+					bSavedSource = FFileHelper::SaveStringToFile(FString(ShaderSource), *Header.ShaderPath);
+				}
+				
+				static bool bAttemptedAuth = false;
+				if (!bSavedSource && !bAttemptedAuth)
+				{
+					bAttemptedAuth = true;
+					
+					if (IFileManager::Get().MakeDirectory(*FPaths::GetPath(Header.ShaderPath), true))
+					{
+						bSavedSource = FFileHelper::SaveStringToFile(FString(ShaderSource), *Header.ShaderPath);
+					}
+					
+					if (!bSavedSource)
+					{
+						FPlatformMisc::MessageBoxExt(EAppMsgType::Ok,
+													*NSLOCTEXT("MetalRHI", "ShaderDebugAuthFail", "Could not access directory required for debugging optimised Metal shaders. Falling back to slower runtime compilation of shaders for debugging.").ToString(), TEXT("Error"));
+					}
+				}
+			}
+#endif
+			// Switch the compile mode so we get debuggable shaders even if we failed to save - if we didn't want
+			// shader debugging we wouldn't have included the code...
+			bOfflineCompile = bSavedSource;
+#endif
+		}
 
-		if (OfflineCompiledFlag)
+		if (bOfflineCompile)
 		{
 			// allow GCD to copy the data into its own buffer
 			//		dispatch_data_t GCDBuffer = dispatch_data_create(InShaderCode.GetTypedData() + CodeOffset, ShaderCode.GetActualShaderCodeSize() - CodeOffset, nil, DISPATCH_DATA_DESTRUCTOR_DEFAULT);
@@ -113,15 +162,14 @@ TMetalBaseShader<BaseResourceType, ShaderType>::TMetalBaseShader(const TArray<ui
 				NSLog(@"Failed to create library: %@", Error);
 			}
 			
-			GlslCodeNSString = @"OFFLINE";
 			dispatch_release(GCDBuffer);
 		}
 		else
 		{
 			NSError* Error;
-			NSString* ShaderSource = [NSString stringWithUTF8String:SourceCode];
+			NSString* ShaderSource = ((OfflineCompiledFlag == 0) ? [NSString stringWithUTF8String:SourceCode] : GlslCodeNSString);
 
-			if(Header.ShaderName.Len())
+			if(Header.ShaderName.Len() && (OfflineCompiledFlag == 0))
 			{
 				ShaderSource = [NSString stringWithFormat:@"// %@\n%@", Header.ShaderName.GetNSString(), ShaderSource];
 			}
@@ -150,16 +198,8 @@ TMetalBaseShader<BaseResourceType, ShaderType>::TMetalBaseShader(const TArray<ui
                 NSLog(@"*********** Error\n%@", ShaderSource);
 			}
 
-			GlslCode.Empty(CodeLength + 1);
-			GlslCode.AddUninitialized(CodeLength + 1);
-			FMemory::Memcpy(GlslCode.GetData(), SourceCode, CodeLength + 1);
-			GlslCodeString = (ANSICHAR*)GlslCode.GetData();
 			GlslCodeNSString = ShaderSource;
-
-#if !UE_BUILD_SHIPPING
 			[GlslCodeNSString retain];
-			//TRACK_OBJECT(GlslCodeNSString);
-#endif
 		}
 
 		// assume there's only one function called 'Main', and use that to get the function from the library
@@ -189,7 +229,6 @@ TMetalBaseShader<BaseResourceType, ShaderType>::~TMetalBaseShader()
 {
 	UNTRACK_OBJECT(STAT_MetalFunctionCount, Function);
 	[Function release];
-	//UNTRACK_OBJECT(GlslCodeNSString);
 	[GlslCodeNSString release];
 }
 
@@ -335,10 +374,12 @@ FMetalBoundShaderState::~FMetalBoundShaderState()
 	}
 }
 
-void FMetalBoundShaderState::PrepareToDraw(FMetalContext* Context, const FMetalRenderPipelineDesc& RenderPipelineDesc)
+void FMetalBoundShaderState::PrepareToDraw(FMetalContext* Context, MTLVertexDescriptor* VertexDesc, const FMetalRenderPipelineDesc& RenderPipelineDesc)
 {
 	// generate a key for the current statez
-	FMetalRenderPipelineHash Hash = RenderPipelineDesc.GetHash();
+	FMetalPipelineHash Hash;
+	Hash.RenderPipelineHash = RenderPipelineDesc.GetHash();
+	Hash.VertexDescHash = [VertexDesc hash];
 	
 	if(GUseRHIThread)
 	{
@@ -356,7 +397,7 @@ void FMetalBoundShaderState::PrepareToDraw(FMetalContext* Context, const FMetalR
 	// make one if not
 	if (PipelineState == nil)
 	{
-		PipelineState = RenderPipelineDesc.CreatePipelineStateForBoundShaderState(this);
+		PipelineState = RenderPipelineDesc.CreatePipelineStateForBoundShaderState(this, VertexDesc);
 		check(PipelineState);
 		
 		if(GUseRHIThread)
@@ -375,9 +416,9 @@ void FMetalBoundShaderState::PrepareToDraw(FMetalContext* Context, const FMetalR
 		if (GFrameCounter > 3)
 		{
 #if PLATFORM_MAC
-			NSLog(@"Created a hitchy pipeline state for hash %llx%llx (this = %p)", (uint64)Hash, (uint64)(Hash >> 64), this);
+			NSLog(@"Created a hitchy pipeline state for hash %llx%llx%llx (this = %p)", (uint64)Hash.RenderPipelineHash, (uint64)(Hash.RenderPipelineHash >> 64), (uint64)Hash.VertexDescHash, this);
 #else
-			NSLog(@"Created a hitchy pipeline state for hash %llx (this = %p)", (uint64)Hash, this);
+			NSLog(@"Created a hitchy pipeline state for hash %llx%llx (this = %p)", (uint64)Hash.RenderPipelineHash, (uint64)Hash.VertexDescHash, this);
 #endif
 		}
 #endif
diff --git a/Engine/Source/Runtime/Apple/MetalRHI/Private/MetalStateCache.cpp b/Engine/Source/Runtime/Apple/MetalRHI/Private/MetalStateCache.cpp
--- a/Engine/Source/Runtime/Apple/MetalRHI/Private/MetalStateCache.cpp
+++ b/Engine/Source/Runtime/Apple/MetalRHI/Private/MetalStateCache.cpp
@@ -430,21 +430,30 @@ void FMetalStateCache::SetRenderTargetsInfo(FRHISetRenderTargetsInfo const& InRe
 		
 		RenderTargetArraySize = 1;
 		
-#if METAL_API_1_1 && PLATFORM_MAC
-		
 		if(ArrayTargets)
 		{
-			if (ArrayTargets == BoundTargets)
+			if (!GetMetalDeviceContext().SupportsFeature(EMetalFeaturesLayeredRendering))
 			{
-				RenderTargetArraySize = ArrayRenderLayers;
-				RenderPass.renderTargetArrayLength = ArrayRenderLayers;
+				if (ArrayRenderLayers != 1)
+				{
+					UE_LOG(LogMetal, Fatal, TEXT("Layered rendering is unsupported on this device."));
+				}
 			}
+#if PLATFORM_MAC
 			else
 			{
-				UE_LOG(LogMetal, Fatal, TEXT("All color render targets must be layered when performing multi-layered rendering under Metal."));
+				if (ArrayTargets == BoundTargets)
+				{
+					RenderTargetArraySize = ArrayRenderLayers;
+					RenderPass.renderTargetArrayLength = ArrayRenderLayers;
+				}
+				else
+				{
+					UE_LOG(LogMetal, Fatal, TEXT("All color render targets must be layered when performing multi-layered rendering under Metal."));
+				}
 			}
-		}
 #endif
+		}
 	
 		// default to invalid
 		PipelineDesc.PipelineDescriptor.depthAttachmentPixelFormat = MTLPixelFormatInvalid;
@@ -458,7 +467,6 @@ void FMetalStateCache::SetRenderTargetsInfo(FRHISetRenderTargetsInfo const& InRe
 		{
 			FMetalSurface& Surface = *GetMetalSurfaceFromRHITexture(RenderTargetsInfo.DepthStencilRenderTarget.Texture);
 			
-#if METAL_API_1_1 && PLATFORM_MAC
 			switch(Surface.Type)
 			{
 				case RRT_Texture2DArray:
@@ -472,10 +480,18 @@ void FMetalStateCache::SetRenderTargetsInfo(FRHISetRenderTargetsInfo const& InRe
 			}
 			if(!ArrayTargets && ArrayRenderLayers > 1)
 			{
-				RenderTargetArraySize = ArrayRenderLayers;
-				RenderPass.renderTargetArrayLength = ArrayRenderLayers;
-			}
+				if (!GetMetalDeviceContext().SupportsFeature(EMetalFeaturesLayeredRendering))
+				{
+					UE_LOG(LogMetal, Fatal, TEXT("Layered rendering is unsupported on this device."));
+				}
+#if PLATFORM_MAC
+				else
+				{
+					RenderTargetArraySize = ArrayRenderLayers;
+					RenderPass.renderTargetArrayLength = ArrayRenderLayers;
+				}
 #endif
+			}
 			
 			if(!bFramebufferSizeSet)
 			{
@@ -585,8 +601,6 @@ void FMetalStateCache::SetRenderTargetsInfo(FRHISetRenderTargetsInfo const& InRe
 					PipelineDesc.SampleCount = DepthAttachment.texture.sampleCount;
 				}
 				
-				bHasValidRenderTarget = true;
-	
 				bHasValidRenderTarget = true;
 				
 				// and assign it
@@ -831,6 +845,7 @@ bool FMetalStateCache::NeedsToSetRenderTarget(const FRHISetRenderTargetsInfo& In
 			//    If we switch to Load, no need to switch as we can re-use what we already have
 			//    If we switch to Clear, we have to always switch to a new RT to force the clear
 			//    If we switch to DontCare, there's definitely no need to switch
+			//    If we switch *from* Clear then we must change target as we *don't* want to clear again.
 			if (RenderTargetView.LoadAction == ERenderTargetLoadAction::EClear || PreviousRenderTargetView.LoadAction == ERenderTargetLoadAction::EClear)
 			{
 				bAllChecksPassed = false;
@@ -848,17 +863,21 @@ bool FMetalStateCache::NeedsToSetRenderTarget(const FRHISetRenderTargetsInfo& In
 //			}
 		}
 		
-		if (InRenderTargetsInfo.DepthStencilRenderTarget.Texture && (InRenderTargetsInfo.DepthStencilRenderTarget.DepthLoadAction == ERenderTargetLoadAction::EClear || InRenderTargetsInfo.DepthStencilRenderTarget.StencilLoadAction == ERenderTargetLoadAction::EClear))
+		if (InRenderTargetsInfo.DepthStencilRenderTarget.Texture)
 		{
-			bAllChecksPassed = false;
-		}
-		
+			if ((InRenderTargetsInfo.DepthStencilRenderTarget.DepthLoadAction == ERenderTargetLoadAction::EClear || RenderTargetsInfo.DepthStencilRenderTarget.DepthLoadAction == ERenderTargetLoadAction::EClear)
+				|| (InRenderTargetsInfo.DepthStencilRenderTarget.StencilLoadAction == ERenderTargetLoadAction::EClear || RenderTargetsInfo.DepthStencilRenderTarget.StencilLoadAction == ERenderTargetLoadAction::EClear))
+			{
+				bAllChecksPassed = false;
+			}
+
 #if PLATFORM_MAC
-		if (!(InRenderTargetsInfo.DepthStencilRenderTarget.GetDepthStencilAccess() == RenderTargetsInfo.DepthStencilRenderTarget.GetDepthStencilAccess()))
-		{
-			bAllChecksPassed = false;
-		}
+			if (!(InRenderTargetsInfo.DepthStencilRenderTarget.GetDepthStencilAccess() == RenderTargetsInfo.DepthStencilRenderTarget.GetDepthStencilAccess()))
+			{
+				bAllChecksPassed = false;
+			}
 #endif
+		}
 	}
 
 	// if we are setting them to nothing, then this is probably end of frame, and we can't make a framebuffer
