Copy CL #3237188:
Duplicate & amend for Mac CL #3237031 from //UE4/Dev-Platform:
A call to BeginDrawingViewport has been bypassed in recent changes, this means that the Metal viewport's back-buffer queue isn't being popped correctly so avoid it unless the RHI thread is enabled.

diff --git a/Engine/Source/Runtime/Apple/MetalRHI/Private/MetalViewport.cpp b/Engine/Source/Runtime/Apple/MetalRHI/Private/MetalViewport.cpp
--- a/Engine/Source/Runtime/Apple/MetalRHI/Private/MetalViewport.cpp
+++ b/Engine/Source/Runtime/Apple/MetalRHI/Private/MetalViewport.cpp
@@ -108,7 +108,7 @@ uint32 FMetalViewport::GetViewportIndex(EMetalViewportAccessFlag Accessor) const
 	{
 		case EMetalViewportAccessRHI:
 			check(IsInRHIThread() || IsInRenderingThread());
-			return Accessor;
+            return GRHISupportsRHIThread ? Accessor : EMetalViewportAccessRenderer;
 		case EMetalViewportAccessRenderer:
 			check(IsInRenderingThread());
 			return Accessor;
@@ -173,6 +173,7 @@ void FMetalViewport::Resize(uint32 InSizeX, uint32 InSizeY, bool bInIsFullscreen
 		  FMath::TruncToInt(ScalingFactor * ViewFrame.size.height) == InSizeY);
 #endif
 	BackBuffer[Index]->Surface.Viewport = this;
+    if (GRHISupportsRHIThread)
 	{
 		FScopeLock Lock(&Mutex);
 		BackBuffersQueue.Insert(BackBuffer[Index], 0);
