Merging fix for cvar warnings from Release-4.16.

diff --git a/Engine/Config/BaseDeviceProfiles.ini b/Engine/Config/BaseDeviceProfiles.ini
--- a/Engine/Config/BaseDeviceProfiles.ini
+++ b/Engine/Config/BaseDeviceProfiles.ini
@@ -475,7 +475,6 @@ BaseProfileName=
 +CVars=r.DFShadowScatterTileCulling=0
 +CVars=r.Shaders.ZeroInitialise=1
 +CVars=r.Shaders.BoundsChecking=1
-+CVars=r.RHICmdStateCacheEnable=0
 
 [MacClient DeviceProfile]
 DeviceType=Mac
diff --git a/Engine/Source/Runtime/Apple/MetalRHI/Private/MetalRHI.cpp b/Engine/Source/Runtime/Apple/MetalRHI/Private/MetalRHI.cpp
--- a/Engine/Source/Runtime/Apple/MetalRHI/Private/MetalRHI.cpp
+++ b/Engine/Source/Runtime/Apple/MetalRHI/Private/MetalRHI.cpp
@@ -89,13 +89,6 @@ FMetalDynamicRHI::FMetalDynamicRHI(ERHIFeatureLevel::Type RequestedFeatureLevel)
 	GRHIAdapterName = TEXT("Metal");
 	GRHIVendorId = 1; // non-zero to avoid asserts
 	
-	// Disable RHI CommandList State Caching: MetalRHI needs to do this itself internally to function.
-	static auto CVarRHICmdStateCacheEnable = IConsoleManager::Get().FindConsoleVariable(TEXT("r.RHICmdStateCacheEnable"));
-	if(CVarRHICmdStateCacheEnable && CVarRHICmdStateCacheEnable->GetInt() != 0)
-	{
-		CVarRHICmdStateCacheEnable->Set(0);
-	}
-	
 	bool const bRequestedFeatureLevel = (RequestedFeatureLevel != ERHIFeatureLevel::Num);
 	bool bSupportsPointLights = false;
 	bool bSupportsRHIThread = false;
