Copying //UE4/Ocean-Staging to //UE4/Dev-Main (Source: //Ocean/Main @ 3480471)

==========================
MAJOR FEATURES + CHANGES
==========================

Change 3478011 on 2017/06/07 by [REDACTED_PERSON]

	SEQUENCER: Glorious and exhalted fix from Max Chen that ammends a lot of the sequencer not taking control of characters issues

Change 3474946 on 2017/06/05 by [REDACTED_PERSON]

	#ocean - Asset registry changes.  All data assets have a unique ID and specific location.

Change 3470578 on 2017/06/01 by [REDACTED_PERSON]

	Hotfix from CL 3470510 for starting in full screen mode instead of using the game user setting

Change 3470577 on 2017/06/01 by [REDACTED_PERSON]

	Fixing iOS compile error

Change 3469728 on 2017/06/01 by [REDACTED_PERSON]

	Hotfix from CL 3469556 for issue where GameObjects with no world set to tick in editor would fail to tick if viewports were set to not be realtime

Change 3465183 on 2017/05/30 by [REDACTED_PERSON]

	Hotfix from CL 3456916 for font scaling bug

Change 3465157 on 2017/05/30 by [REDACTED_PERSON]

	#ocean, #mcp - Fi x daily leads not setting requirements properly.

Change 3461677 on 2017/05/26 by [REDACTED_PERSON]

	#ocn - Hotfix merge from CL 3459902. Fix for hang on iOS waiting for texture mips to load.

	//UE4/Release-4.16/Engine/Source/Runtime/Apple/MetalRHI/Private/MetalTexture.cpp

	to //Ocean/Main/Engine/Source/Runtime/Apple/MetalRHI/Private/MetalTexture.cpp

Change 3449359 on 2017/05/19 by [REDACTED_PERSON]

	Hotfix from CL 3437541 - Increase the Slate draw buffer count from 3 to 4 to hopefully help with editor performance.

Change 3449352 on 2017/05/19 by [REDACTED_PERSON]

	Hotfix from CL3435138 - Make sure Aftermath is disabled to hopefully help with editor performance.

Change 3447137 on 2017/05/18 by [REDACTED_PERSON]

	Hotfix from CL 3447058 - Editor Performance fix from Matt Kuhlenschmidt

Change 3446315 on 2017/05/18 by [REDACTED_PERSON]

	Hotfix for CL 3440745 to clear up {ExportClass} Cook warning

Change 3445554 on 2017/05/17 by [REDACTED_PERSON]

	Hotfix from CL 3440999 to fix performance issues

Change 3441530 on 2017/05/16 by [REDACTED_PERSON]

	#ocn - MovieSequencePlayer - Do final update before Stop is called. Also add a Finalize event to allow BP to know when the sequence finishes naturally.
	Max Chen provided this to me in a shelf. He will submit the exact same changes in dev-editor.

Change 3441528 on 2017/05/16 by [REDACTED_PERSON]

	Hotfix from CL 3425483 for (OCN-7073) (OCN-7075)

Change 3428762 on 2017/05/08 by [REDACTED_PERSON]

	Hotfix from CL 3372554 for Mac/iOS compile fixes for XCode 8.3

Change 3426592 on 2017/05/05 by [REDACTED_PERSON]

	integrating 3372554 from release-4.15 to ocean/main

	compile fixes for Xcode 8.3
	(original Peter.Sauerbrei)

	#rn major Chad's Test Category/Another Subcategory

Change 3425183 on 2017/05/05 by [REDACTED_PERSON]

	merging 3424104 to Ocean/Main

	Ignore Clang's nonportable-include-path warning on Mac

Change 3410456 on 2017/04/26 by [REDACTED_PERSON]

	Convert Ocean run.bat over to launch fat jar, and use updated script which

Change 3402905 on 2017/04/20 by [REDACTED_PERSON]

	Adjusting the IsEditorWorld function to include PIE

Change 3402649 on 2017/04/20 by [REDACTED_PERSON]

	Fixing a cook failure

Change 3367928 on 2017/03/28 by [REDACTED_PERSON]

	Hotfix merged from BattleBreakers.  CL 3175413.

Change 3365963 on 2017/03/27 by [REDACTED_PERSON]

	#ocn - Deferred loader for customizer.  Prevents ~200 MB memory spike during customization at level start.

diff --git a/Engine/Source/Editor/PropertyEditor/Private/SDetailSingleItemRow.cpp b/Engine/Source/Editor/PropertyEditor/Private/SDetailSingleItemRow.cpp
--- a/Engine/Source/Editor/PropertyEditor/Private/SDetailSingleItemRow.cpp
+++ b/Engine/Source/Editor/PropertyEditor/Private/SDetailSingleItemRow.cpp
@@ -490,6 +490,12 @@ bool SDetailSingleItemRow::CanPasteProperty() const
 
 const FSlateBrush* SDetailSingleItemRow::GetBorderImage() const
 {
+	/*bool UseDarkTheme = false;
+	if (Customization->GetWidgetRow().IsDarkThemeAttr.IsSet())
+	{
+		UseDarkTheme = Customization->GetWidgetRow().IsDarkThemeAttr.Get();
+	}*/
+
 	if( IsHighlighted() )
 	{
 		return FEditorStyle::GetBrush("DetailsView.CategoryMiddle_Highlighted");
diff --git a/Engine/Source/Runtime/Apple/MetalRHI/Private/MetalTexture.cpp b/Engine/Source/Runtime/Apple/MetalRHI/Private/MetalTexture.cpp
--- a/Engine/Source/Runtime/Apple/MetalRHI/Private/MetalTexture.cpp
+++ b/Engine/Source/Runtime/Apple/MetalRHI/Private/MetalTexture.cpp
@@ -1756,53 +1756,54 @@ struct FMetalRHICommandAsyncReallocateTexture2D : public FRHICommand<FMetalRHICo
 	, NewSizeX(InNewSizeX)
 	, NewSizeY(InNewSizeY)
 	, RequestStatus(InRequestStatus)
-{
-}
+	{
+	}
 
 	void Execute(FRHICommandListBase& CmdList)
-{
+	{
 		CopyMips(Context, OldTexture, NewTexture, NewMipCount, NewSizeX, NewSizeY, RequestStatus);
-}
+	}
 
 	static void CopyMips(FMetalContext& Context, FMetalTexture2D* OldTexture, FMetalTexture2D* NewTexture, int32 NewMipCount, int32 NewSizeX, int32 NewSizeY, FThreadSafeCounter* RequestStatus)
-{
-	// figure out what mips to schedule
-	const uint32 NumSharedMips = FMath::Min(OldTexture->GetNumMips(), NewTexture->GetNumMips());
-	const uint32 SourceMipOffset = OldTexture->GetNumMips() - NumSharedMips;
-	const uint32 DestMipOffset = NewTexture->GetNumMips() - NumSharedMips;
-	
-	const uint32 BlockSizeX = GPixelFormats[OldTexture->GetFormat()].BlockSizeX;
-	const uint32 BlockSizeY = GPixelFormats[OldTexture->GetFormat()].BlockSizeY;
+	{
+		// figure out what mips to schedule
+		const uint32 NumSharedMips = FMath::Min(OldTexture->GetNumMips(), NewTexture->GetNumMips());
+		const uint32 SourceMipOffset = OldTexture->GetNumMips() - NumSharedMips;
+		const uint32 DestMipOffset = NewTexture->GetNumMips() - NumSharedMips;
 
-	// only handling straight 2D textures here
-	uint32 SliceIndex = 0;
-	MTLOrigin Origin = MTLOriginMake(0,0,0);
-	
-	id<MTLTexture> Tex = OldTexture->Surface.Texture;
-	[Tex retain];
+		const uint32 BlockSizeX = GPixelFormats[OldTexture->GetFormat()].BlockSizeX;
+		const uint32 BlockSizeY = GPixelFormats[OldTexture->GetFormat()].BlockSizeY;
 
-	// DXT/BC formats on Mac actually do have mip-tails that are smaller than the block size, they end up being uncompressed.
+		// only handling straight 2D textures here
+		uint32 SliceIndex = 0;
+		MTLOrigin Origin = MTLOriginMake(0,0,0);
+
+		id<MTLTexture> Tex = OldTexture->Surface.Texture;
+		[Tex retain];
+
+		// DXT/BC formats on Mac actually do have mip-tails that are smaller than the block size, they end up being uncompressed.
 		bool const bPixelFormatASTC = IsPixelFormatASTCCompressed(OldTexture->GetFormat());
-	
-	for (uint32 MipIndex = 0; MipIndex < NumSharedMips; ++MipIndex)
-	{
+
+		for (uint32 MipIndex = 0; MipIndex < NumSharedMips; ++MipIndex)
+		{
 			const uint32 UnalignedMipSizeX = FMath::Max<uint32>(1, NewSizeX >> (MipIndex + DestMipOffset));
 			const uint32 UnalignedMipSizeY = FMath::Max<uint32>(1, NewSizeY >> (MipIndex + DestMipOffset));
 			const uint32 MipSizeX = (bPixelFormatASTC) ? AlignArbitrary(UnalignedMipSizeX, BlockSizeX) : UnalignedMipSizeX;
 			const uint32 MipSizeY = (bPixelFormatASTC) ? AlignArbitrary(UnalignedMipSizeY, BlockSizeY) : UnalignedMipSizeY;
 
 			Context.AsyncCopyFromTextureToTexture(OldTexture->Surface.Texture, SliceIndex, MipIndex + SourceMipOffset, Origin, MTLSizeMake(MipSizeX, MipSizeY, 1), NewTexture->Surface.Texture, SliceIndex, MipIndex + DestMipOffset, Origin);
-	}
+		}
 
-	// when done, decrement the counter to indicate it's safe
+		// when done, decrement the counter to indicate it's safe
 		MTLCommandBufferHandler CompletionHandler = ^(id <MTLCommandBuffer> Buffer)
-	{
-		[Tex release];
-		RequestStatus->Decrement();
+		{
+			[Tex release];
 		};
 
-    // kick it off!
+		// kick it off!
 		Context.SubmitAsyncCommands(nil, CompletionHandler, false);
+
+		RequestStatus->Decrement();
 	}
 };
 
diff --git a/Engine/Source/Runtime/CoreUObject/Private/UObject/LinkerLoad.cpp b/Engine/Source/Runtime/CoreUObject/Private/UObject/LinkerLoad.cpp
--- a/Engine/Source/Runtime/CoreUObject/Private/UObject/LinkerLoad.cpp
+++ b/Engine/Source/Runtime/CoreUObject/Private/UObject/LinkerLoad.cpp
@@ -2263,7 +2263,7 @@ FLinkerLoad::EVerifyResult FLinkerLoad::VerifyImport(int32 ImportIndex)
 						if (DependsList.Contains(ImportPackageIndex))
 						{
 							TokenizedMessage->AddToken(FTextToken::Create(
-								FText::Format(LOCTEXT("ImportFailureExportReference", "Referenced by export {ExportClass}"),
+								FText::Format(LOCTEXT("ImportFailureExportReference", "Referenced by export {0}"),
 									FText::FromName(GetExportClassName(CurrentExportIndex)))));
 							TokenizedMessage->AddToken(FAssetNameToken::Create(GetExportPathName(CurrentExportIndex)));
 						}
diff --git a/Engine/Source/Runtime/Engine/Classes/Engine/Texture2D.h b/Engine/Source/Runtime/Engine/Classes/Engine/Texture2D.h
--- a/Engine/Source/Runtime/Engine/Classes/Engine/Texture2D.h
+++ b/Engine/Source/Runtime/Engine/Classes/Engine/Texture2D.h
@@ -265,7 +265,7 @@ public:
 	/**
 	 * Returns the number of mips in this texture that are not able to be streamed.
 	 */
-	int32 GetNumNonStreamingMips() const;
+	ENGINE_API int32 GetNumNonStreamingMips() const;
 
 	/**
 	 * Computes the minimum and maximum allowed mips for a texture.
diff --git a/Engine/Source/Runtime/Engine/Classes/Engine/World.h b/Engine/Source/Runtime/Engine/Classes/Engine/World.h
--- a/Engine/Source/Runtime/Engine/Classes/Engine/World.h
+++ b/Engine/Source/Runtime/Engine/Classes/Engine/World.h
@@ -3029,7 +3029,10 @@ public:
 	/** Returns true if this world is any kind of game world (including PIE worlds) */
 	bool IsGameWorld() const;
 
-	/** Returns true if this world is a preview game world (blueprint editor) */
+	/** Returns true if this world is any kind of editor world (including editor preview worlds) */
+	bool IsEditorWorld() const;
+
+	/** Returns true if this world is a preview game world (editor or game) */
 	bool IsPreviewWorld() const;
 
 	/** Returns true if this world should look at game hidden flags instead of editor hidden flags for the purposes of rendering */
diff --git a/Engine/Source/Runtime/Engine/Private/LevelTick.cpp b/Engine/Source/Runtime/Engine/Private/LevelTick.cpp
--- a/Engine/Source/Runtime/Engine/Private/LevelTick.cpp
+++ b/Engine/Source/Runtime/Engine/Private/LevelTick.cpp
@@ -1172,21 +1172,19 @@ void FTickableGameObject::TickObjects(UWorld* World, const int32 InTickType, con
 				// If it is tickable and in this world
 				if (TickableObject->IsTickable() && (TickableObject->GetTickableGameObjectWorld() == World))
 				{
-					// If this is a game world or the object is tickable in editor
 					const bool bIsGameWorld = InTickType == LEVELTICK_All || (World && World->IsGameWorld());
-
-					if (bIsGameWorld || (GIsEditor && TickableObject->IsTickableInEditor()))
+					// If we are in editor and it is editor tickable, always tick
+					// If this is a game world then tick if we are not doing a time only (paused) update and we are not paused or the object is tickable when paused
+					if ((GIsEditor && TickableObject->IsTickableInEditor()) ||
+						(bIsGameWorld && ((!bIsPaused && TickType != LEVELTICK_TimeOnly) || (bIsPaused && TickableObject->IsTickableWhenPaused()))))
 					{
-						// Then tick if we are not doing a time only (paused) update and we are not paused or the object is tickable when paused
-						if ((!bIsPaused && TickType != LEVELTICK_TimeOnly) || (bIsPaused && TickableObject->IsTickableWhenPaused()))
-						{
-							STAT(FScopeCycleCounter Context(TickableObject->GetStatId());)
-							TickableObject->Tick(DeltaSeconds);
+						STAT(FScopeCycleCounter Context(TickableObject->GetStatId());)
+						TickableObject->Tick(DeltaSeconds);
 
-							if (TickableObjects[i] == nullptr)
-							{
-								bNeedsCleanup = true;
-							}
+						// In case it was removed during tick
+						if (TickableObjects[i] == nullptr)
+						{
+							bNeedsCleanup = true;
 						}
 					}
 				}
diff --git a/Engine/Source/Runtime/Engine/Private/UnrealClient.cpp b/Engine/Source/Runtime/Engine/Private/UnrealClient.cpp
--- a/Engine/Source/Runtime/Engine/Private/UnrealClient.cpp
+++ b/Engine/Source/Runtime/Engine/Private/UnrealClient.cpp
@@ -19,6 +19,7 @@
 #include "Matinee/MatineeActor.h"
 #include "EditorSupportDelegates.h"
 #include "HighResScreenshot.h"
+#include "GameFramework/GameUserSettings.h"
 
 #include "Framework/Notifications/NotificationManager.h"
 #include "Widgets/Notifications/SNotificationList.h"
@@ -799,7 +800,7 @@ FViewport::FViewport(FViewportClient* InViewportClient):
 	ViewportClient(InViewportClient),
 	SizeX(0),
 	SizeY(0),
-	WindowMode(EWindowMode::Windowed),
+	WindowMode(IsRunningGame() ? GEngine->GetGameUserSettings()->GetDefaultWindowMode() : EWindowMode::Windowed),
 	bHitProxiesCached(false),
 	bHasRequestedToggleFreeze(false),
 	bIsSlateViewport(false),
diff --git a/Engine/Source/Runtime/Engine/Private/World.cpp b/Engine/Source/Runtime/Engine/Private/World.cpp
--- a/Engine/Source/Runtime/Engine/Private/World.cpp
+++ b/Engine/Source/Runtime/Engine/Private/World.cpp
@@ -5718,6 +5718,11 @@ bool UWorld::IsGameWorld() const
 	return WorldType == EWorldType::Game || WorldType == EWorldType::PIE || WorldType == EWorldType::GamePreview;
 }
 
+bool UWorld::IsEditorWorld() const
+{
+	return WorldType == EWorldType::Editor || WorldType == EWorldType::EditorPreview || WorldType == EWorldType::PIE;
+}
+
 bool UWorld::IsPreviewWorld() const
 {
 	return WorldType == EWorldType::EditorPreview || WorldType == EWorldType::GamePreview;
diff --git a/Engine/Source/Runtime/Engine/Public/TextureResource.h b/Engine/Source/Runtime/Engine/Public/TextureResource.h
--- a/Engine/Source/Runtime/Engine/Public/TextureResource.h
+++ b/Engine/Source/Runtime/Engine/Public/TextureResource.h
@@ -249,7 +249,7 @@ public:
 	 * Called from the game thread to kick off a change in ResidentMips after modifying RequestedMips.
 	 * @param bShouldPrioritizeAsyncIORequest	- Whether the Async I/O request should have higher priority
 	 */
-	void BeginUpdateMipCount( bool bShouldPrioritizeAsyncIORequest );
+	ENGINE_API void BeginUpdateMipCount( bool bShouldPrioritizeAsyncIORequest );
 	/**
 	 * Called from the game thread to kick off async I/O to load in new mips.
 	 */
diff --git a/Engine/Source/Runtime/LevelSequence/Private/LevelSequenceBindingReference.cpp b/Engine/Source/Runtime/LevelSequence/Private/LevelSequenceBindingReference.cpp
--- a/Engine/Source/Runtime/LevelSequence/Private/LevelSequenceBindingReference.cpp
+++ b/Engine/Source/Runtime/LevelSequence/Private/LevelSequenceBindingReference.cpp
@@ -6,12 +6,13 @@
 #include "UObject/ObjectMacros.h"
 #include "MovieSceneFwd.h"
 #include "PackageName.h"
+#include "Engine/World.h"
 
 FLevelSequenceBindingReference::FLevelSequenceBindingReference(UObject* InObject, UObject* InContext)
 {
 	check(InContext && InObject);
 
-	if (InObject->IsIn(InContext))
+	if (!InContext->IsA<UWorld>() && InObject->IsIn(InContext))
 	{
 		ObjectPath = InObject->GetPathName(InContext);
 	}
@@ -72,25 +73,6 @@ UObject* FLevelSequenceLegacyObjectReference::Resolve(UObject* InContext) const
 		int32 PIEInstanceID = InContext->GetOutermost()->PIEInstanceID;
 		FUniqueObjectGuid FixedUpId = PIEInstanceID == -1 ? ObjectId : ObjectId.FixupForPIE(PIEInstanceID);
 
-		if (PIEInstanceID != -1 && FixedUpId == ObjectId)
-		{
-			if (!ObjectPath.IsEmpty())
-			{
-				if (UObject* FoundObject = FindObject<UObject>(InContext, *ObjectPath, false))
-				{
-					return FoundObject;
-				}
-
-				if (UObject* FoundObject = FindObject<UObject>(ANY_PACKAGE, *ObjectPath, false))
-				{
-					return FoundObject;
-				}
-			}
-
-			UE_LOG(LogMovieScene, Warning, TEXT("Attempted to resolve object with a PIE instance that has not been fixed up yet. This is probably due to a streamed level not being available yet. %s"), *ObjectPath);
-			return nullptr;
-		}
-
 		FLazyObjectPtr LazyPtr;
 		LazyPtr = FixedUpId;
 
diff --git a/Engine/Source/Runtime/MovieScene/Private/MovieSceneSequencePlayer.cpp b/Engine/Source/Runtime/MovieScene/Private/MovieSceneSequencePlayer.cpp
--- a/Engine/Source/Runtime/MovieScene/Private/MovieSceneSequencePlayer.cpp
+++ b/Engine/Source/Runtime/MovieScene/Private/MovieSceneSequencePlayer.cpp
@@ -349,15 +349,20 @@ void UMovieSceneSequencePlayer::UpdateTimeCursorPosition(float NewPosition, TOpt
 		// stop playback
 		else
 		{
+			FMovieSceneEvaluationRange Range = PlayPosition.PlayTo(GetSequencePosition(), FixedFrameInterval);
+
+			UpdateMovieSceneInstance(Range, OptionalStatus);
+
 			Stop();
 
 			// When playback stops naturally, the time cursor is put at the boundary that was crossed to make ping-pong playback easy
 			TimeCursorPosition = bReversePlayback ? 0.f : GetLength();
 			PlayPosition.Reset(TimeCursorPosition);
 
-			FMovieSceneEvaluationRange Range = PlayPosition.PlayTo(GetSequencePosition(), FixedFrameInterval);
-
-			UpdateMovieSceneInstance(Range, OptionalStatus);
+			if (OnFinished.IsBound())
+			{
+				OnFinished.Broadcast();
+			}
 		}
 	}
 	else
diff --git a/Engine/Source/Runtime/MovieScene/Public/MovieSceneSequencePlayer.h b/Engine/Source/Runtime/MovieScene/Public/MovieSceneSequencePlayer.h
--- a/Engine/Source/Runtime/MovieScene/Public/MovieSceneSequencePlayer.h
+++ b/Engine/Source/Runtime/MovieScene/Public/MovieSceneSequencePlayer.h
@@ -205,6 +205,10 @@ public:
 	UPROPERTY(BlueprintAssignable, Category="Game|Cinematic")
 	FOnMovieSceneSequencePlayerEvent OnPause;
 
+	/** Event triggered when the level sequence player finishes naturally (without explicitly calling stop) */
+	UPROPERTY(BlueprintAssignable, Category = "Game|Cinematic")
+	FOnMovieSceneSequencePlayerEvent OnFinished;
+
 
 public:
 
diff --git a/Engine/Source/Runtime/SlateCore/Private/Fonts/FontCacheFreeType.cpp b/Engine/Source/Runtime/SlateCore/Private/Fonts/FontCacheFreeType.cpp
--- a/Engine/Source/Runtime/SlateCore/Private/Fonts/FontCacheFreeType.cpp
+++ b/Engine/Source/Runtime/SlateCore/Private/Fonts/FontCacheFreeType.cpp
@@ -369,12 +369,16 @@ bool FFreeTypeAdvanceCache::FindOrCache(FT_Face InFace, const uint32 InGlyphInde
 		}
 	}
 
-	FreeTypeUtils::ApplySizeAndScale(InFace, InFontSize, InFontScale);
+	FreeTypeUtils::ApplySizeAndScale(InFace, InFontSize, 1.0f);
 
 	// No cached data, go ahead and add an entry for it...
 	FT_Error Error = FT_Get_Advance(InFace, InGlyphIndex, InLoadFlags, &OutCachedAdvance);
 	if (Error == 0)
 	{
+		// We apply our own scaling as FreeType doesn't always produce the correct results for all fonts when applying the scale via the transform matrix
+		const FT_Long FixedFontScale = FreeTypeUtils::ConvertPixelTo16Dot16<FT_Long>(InFontScale);
+		OutCachedAdvance = FT_MulFix(OutCachedAdvance, FixedFontScale);
+
 		CachedAdvanceMap.Add(CachedAdvanceKey, OutCachedAdvance);
 		return true;
 	}
