Duplicate 3627043:
Change Metal pipeline bindings warning to Verbose as it is a debugging aid, not necessarily an error.

diff --git a/Engine/Source/Runtime/Apple/MetalRHI/Private/MetalPipeline.cpp b/Engine/Source/Runtime/Apple/MetalRHI/Private/MetalPipeline.cpp
--- a/Engine/Source/Runtime/Apple/MetalRHI/Private/MetalPipeline.cpp
+++ b/Engine/Source/Runtime/Apple/MetalRHI/Private/MetalPipeline.cpp
@@ -331,7 +331,7 @@ static FMetalShaderPipeline* CreateMTLRenderPipeline(const FGraphicsPipelineStat
 	{
 		check((PixelShader->Bindings.InOutMask & 0x8000) || (PixelShader->Bindings.InOutMask & 0x7fff) > 0 || PixelShader->Bindings.NumUAVs > 0);
 		
-		UE_CLOG((NumActiveTargets < __builtin_popcount(PixelShader->Bindings.InOutMask & 0x7fff)), LogMetal, Warning, TEXT("NumActiveTargets doesn't match pipeline's pixel shader output mask: %u, %hx"), NumActiveTargets, PixelShader->Bindings.InOutMask);
+		UE_CLOG((NumActiveTargets < __builtin_popcount(PixelShader->Bindings.InOutMask & 0x7fff)), LogMetal, Verbose, TEXT("NumActiveTargets doesn't match pipeline's pixel shader output mask: %u, %hx"), NumActiveTargets, PixelShader->Bindings.InOutMask);
 	}
 	
 	FMetalBlendState* BlendState = (FMetalBlendState*)Init.BlendState;
