Fix launching games into fullscreen on macOS:

Duplicate 3636725:
Workaround a crash on fullscreen startup in packaged builds of Showdown for Mac - something's not quite right with the code that determines when it is OK to restart a render pass that should clear the target. Need to investigate further, but avoiding some Load action fiddling that emulates D3D behaviour bypasses the issue for now.

Duplicate 3637867:
Retain the bound render-targets inside MetalRHI until we bind something else - with the RHI threading enabled it isn't possible to safely unbind them on deletion across all contexts at the moment.

diff --git a/Engine/Source/Runtime/Apple/MetalRHI/Private/MetalContext.cpp b/Engine/Source/Runtime/Apple/MetalRHI/Private/MetalContext.cpp
--- a/Engine/Source/Runtime/Apple/MetalRHI/Private/MetalContext.cpp
+++ b/Engine/Source/Runtime/Apple/MetalRHI/Private/MetalContext.cpp
@@ -1113,7 +1113,7 @@ void FMetalContext::ResetRenderCommandEncoder()
 	
 	StateCache.InvalidateRenderTargets();
 	
-	SetRenderTargetsInfo(StateCache.GetRenderTargetsInfo(), false);
+	SetRenderTargetsInfo(StateCache.GetRenderTargetsInfo(), true);
 }
 
 bool FMetalContext::PrepareToDraw(uint32 PrimitiveType, EMetalIndexType IndexType)
@@ -1231,7 +1231,7 @@ bool FMetalContext::PrepareToDraw(uint32 PrimitiveType, EMetalIndexType IndexTyp
 			}
 		}
 		
-		if (StateCache.SetRenderTargetsInfo(Info, StateCache.GetVisibilityResultsBuffer(), false))
+		if (StateCache.SetRenderTargetsInfo(Info, StateCache.GetVisibilityResultsBuffer(), true))
 		{
 			RenderPass.RestartRenderPass(StateCache.GetRenderPassDescriptor());
 		}
@@ -1256,7 +1256,7 @@ bool FMetalContext::PrepareToDraw(uint32 PrimitiveType, EMetalIndexType IndexTyp
 		RenderPass.EndRenderPass();
 		
 		StateCache.SetRenderTargetsActive(false);
-		StateCache.SetRenderTargetsInfo(Info, StateCache.GetVisibilityResultsBuffer(), false);
+		StateCache.SetRenderTargetsInfo(Info, StateCache.GetVisibilityResultsBuffer(), true);
 		
 		RenderPass.BeginRenderPass(StateCache.GetRenderPassDescriptor());
 		
@@ -1274,7 +1274,7 @@ bool FMetalContext::PrepareToDraw(uint32 PrimitiveType, EMetalIndexType IndexTyp
 	return true;
 }
 
-void FMetalContext::SetRenderTargetsInfo(const FRHISetRenderTargetsInfo& RenderTargetsInfo, bool const bReset)
+void FMetalContext::SetRenderTargetsInfo(const FRHISetRenderTargetsInfo& RenderTargetsInfo, bool const bRestart)
 {
 #if UE_BUILD_DEBUG || UE_BUILD_DEVELOPMENT
 	if (!CommandList.IsImmediate())
@@ -1326,16 +1326,16 @@ void FMetalContext::SetRenderTargetsInfo(const FRHISetRenderTargetsInfo& RenderT
 			check(IsValidRef(FallbackDepthStencilSurface));
 			Info.DepthStencilRenderTarget = FRHIDepthRenderTargetView(FallbackDepthStencilSurface, ERenderTargetLoadAction::ELoad, ERenderTargetStoreAction::ENoAction, FExclusiveDepthStencil::DepthRead_StencilRead);
 
-			bSet = StateCache.SetRenderTargetsInfo(Info, QueryBuffer->GetCurrentQueryBuffer()->Buffer, bReset);
+			bSet = StateCache.SetRenderTargetsInfo(Info, QueryBuffer->GetCurrentQueryBuffer()->Buffer, bRestart);
 		}
 		else
 		{
-			bSet = StateCache.SetRenderTargetsInfo(RenderTargetsInfo, QueryBuffer->GetCurrentQueryBuffer()->Buffer, bReset);
+			bSet = StateCache.SetRenderTargetsInfo(RenderTargetsInfo, QueryBuffer->GetCurrentQueryBuffer()->Buffer, bRestart);
 		}
 	}
 	else
 	{
-		bSet = StateCache.SetRenderTargetsInfo(RenderTargetsInfo, NULL, bReset);
+		bSet = StateCache.SetRenderTargetsInfo(RenderTargetsInfo, NULL, bRestart);
 	}
 	
 	if (bSet && StateCache.GetHasValidRenderTarget())
diff --git a/Engine/Source/Runtime/Apple/MetalRHI/Private/MetalContext.h b/Engine/Source/Runtime/Apple/MetalRHI/Private/MetalContext.h
--- a/Engine/Source/Runtime/Apple/MetalRHI/Private/MetalContext.h
+++ b/Engine/Source/Runtime/Apple/MetalRHI/Private/MetalContext.h
@@ -50,7 +50,7 @@ public:
 	/**
 	 * Set the color, depth and stencil render targets, and then make the new command buffer/encoder
 	 */
-	void SetRenderTargetsInfo(const FRHISetRenderTargetsInfo& RenderTargetsInfo, bool const bReset = true);
+	void SetRenderTargetsInfo(const FRHISetRenderTargetsInfo& RenderTargetsInfo, bool const bRestart = false);
 	
 	/**
 	 * Allocate from a dynamic ring buffer - by default align to the allowed alignment for offset field when setting buffers
diff --git a/Engine/Source/Runtime/Apple/MetalRHI/Private/MetalStateCache.cpp b/Engine/Source/Runtime/Apple/MetalRHI/Private/MetalStateCache.cpp
--- a/Engine/Source/Runtime/Apple/MetalRHI/Private/MetalStateCache.cpp
+++ b/Engine/Source/Runtime/Apple/MetalRHI/Private/MetalStateCache.cpp
@@ -270,7 +270,7 @@ void FMetalStateCache::SetComputeShader(FMetalComputeShader* InComputeShader)
 	}
 }
 
-bool FMetalStateCache::SetRenderTargetsInfo(FRHISetRenderTargetsInfo const& InRenderTargets, id<MTLBuffer> const QueryBuffer, bool const bReset)
+bool FMetalStateCache::SetRenderTargetsInfo(FRHISetRenderTargetsInfo const& InRenderTargets, id<MTLBuffer> const QueryBuffer, bool const bRestart)
 {
 	bool bNeedsSet = false;
 	
@@ -338,6 +338,8 @@ bool FMetalStateCache::SetRenderTargetsInfo(FRHISetRenderTargetsInfo const& InRe
 			if (RenderTargetIndex < RenderTargetsInfo.NumColorRenderTargets && RenderTargetsInfo.ColorRenderTarget[RenderTargetIndex].Texture != nullptr)
 			{
 				const FRHIRenderTargetView& RenderTargetView = RenderTargetsInfo.ColorRenderTarget[RenderTargetIndex];
+				ColorTargets[RenderTargetIndex] = RenderTargetView.Texture;
+				
 				FMetalSurface& Surface = *GetMetalSurfaceFromRHITexture(RenderTargetView.Texture);
 				FormatKey = Surface.FormatKey;
 				
@@ -433,7 +435,7 @@ bool FMetalStateCache::SetRenderTargetsInfo(FRHISetRenderTargetsInfo const& InRe
 					ColorAttachment.slice = ArraySliceIndex;
 				}
 				
-				ColorAttachment.loadAction = (Surface.Written || !bImmediate) ? GetMetalRTLoadAction(RenderTargetView.LoadAction) : MTLLoadActionClear;
+				ColorAttachment.loadAction = (Surface.Written || !bImmediate || bRestart) ? GetMetalRTLoadAction(RenderTargetView.LoadAction) : MTLLoadActionClear;
 				FPlatformAtomics::InterlockedExchange(&Surface.Written, 1);
 				
 				bNeedsClear |= (ColorAttachment.loadAction == MTLLoadActionClear);
@@ -456,6 +458,10 @@ bool FMetalStateCache::SetRenderTargetsInfo(FRHISetRenderTargetsInfo const& InRe
 				bHasValidRenderTarget = true;
 				bHasValidColorTarget = true;
 			}
+			else
+			{
+				ColorTargets[RenderTargetIndex].SafeRelease();
+			}
 		}
 		
 		RenderTargetArraySize = 1;
@@ -761,16 +767,6 @@ bool FMetalStateCache::SetRenderTargetsInfo(FRHISetRenderTargetsInfo const& InRe
 			DepthStencilSurface.SafeRelease();
 		}
 		
-//		if (bReset)
-//		{
-//			// Reset any existing state as that must be fully reinitialised by the caller.
-//			StencilRef = 0;
-//			BlendFactor = FLinearColor::Transparent;
-//			DepthStencilState.SafeRelease();
-//			RasterizerState.SafeRelease();
-//			GraphicsPSO.SafeRelease();
-//		}
-		
 		[RenderPassDesc release];
 		RenderPassDesc = nil;
 		RenderPassDesc = [RenderPass retain];
@@ -1642,7 +1638,7 @@ bool FMetalStateCache::PrepareToRestart(void)
 			}
 			
 			InvalidateRenderTargets();
-			return SetRenderTargetsInfo(Info, GetVisibilityResultsBuffer(), false) && CanRestartRenderPass();
+			return SetRenderTargetsInfo(Info, GetVisibilityResultsBuffer(), true) && CanRestartRenderPass();
 		}
 		else
 		{
diff --git a/Engine/Source/Runtime/Apple/MetalRHI/Private/MetalStateCache.h b/Engine/Source/Runtime/Apple/MetalRHI/Private/MetalStateCache.h
--- a/Engine/Source/Runtime/Apple/MetalRHI/Private/MetalStateCache.h
+++ b/Engine/Source/Runtime/Apple/MetalRHI/Private/MetalStateCache.h
@@ -36,7 +36,7 @@ public:
 	void SetBlendFactor(FLinearColor const& InBlendFactor);
 	void SetStencilRef(uint32 const InStencilRef);
 	void SetComputeShader(FMetalComputeShader* InComputeShader);
-	bool SetRenderTargetsInfo(FRHISetRenderTargetsInfo const& InRenderTargets, id<MTLBuffer> const QueryBuffer, bool const bReset);
+	bool SetRenderTargetsInfo(FRHISetRenderTargetsInfo const& InRenderTargets, id<MTLBuffer> const QueryBuffer, bool const bRestart);
 	void InvalidateRenderTargets(void);
 	void SetRenderTargetsActive(bool const bActive);
 	void SetViewport(const MTLViewport& InViewport);
@@ -250,6 +250,7 @@ private:
 	uint32 ActiveScissors;
 	
 	FRHISetRenderTargetsInfo RenderTargetsInfo;
+	FTextureRHIRef ColorTargets[MaxSimultaneousRenderTargets];
 	FTextureRHIRef DepthStencilSurface;
 	/** A fallback depth-stencil surface for draw calls that write to depth without a depth-stencil surface bound. */
 	FTexture2DRHIRef FallbackDepthStencilSurface;
