Duplicate 3640982:
Tweak the code for Xcode 9 GM and fix some missing tvOS markup.

Duplicate 3641205:
Add new Metal compute command-encoder functions too.

diff --git a/Engine/Source/Runtime/Apple/MetalRHI/Private/MetalCaptureManager.cpp b/Engine/Source/Runtime/Apple/MetalRHI/Private/MetalCaptureManager.cpp
--- a/Engine/Source/Runtime/Apple/MetalRHI/Private/MetalCaptureManager.cpp
+++ b/Engine/Source/Runtime/Apple/MetalRHI/Private/MetalCaptureManager.cpp
@@ -9,7 +9,7 @@ FMetalCaptureManager::FMetalCaptureManager(id<MTLDevice> InDevice, FMetalCommand
 , Queue(InQueue)
 {
 #if METAL_SUPPORTS_CAPTURE_MANAGER
-	if (@available(iOS 11.0, macOS 10.13, *))
+	if (@available(iOS 11.0, tvOS 11.0, macOS 10.13, *))
 	{
 		MTLCaptureManager* Manager = [MTLCaptureManager sharedCaptureManager];
 		
@@ -47,7 +47,7 @@ FMetalCaptureManager::~FMetalCaptureManager()
 void FMetalCaptureManager::PresentFrame(uint32 FrameNumber)
 {
 #if METAL_SUPPORTS_CAPTURE_MANAGER
-	if (@available(iOS 11.0, macOS 10.13, *))
+	if (@available(iOS 11.0, tvOS 11.0, macOS 10.13, *))
 	{
 		for (FMetalCaptureScope& Scope : ActiveScopes)
 		{
@@ -79,7 +79,7 @@ void FMetalCaptureManager::PresentFrame(uint32 FrameNumber)
 void FMetalCaptureManager::BeginCapture(void)
 {
 #if METAL_SUPPORTS_CAPTURE_MANAGER
-	if (@available(iOS 11.0, macOS 10.13, *))
+	if (@available(iOS 11.0, tvOS 11.0, macOS 10.13, *))
 	{
 		[[MTLCaptureManager sharedCaptureManager] startCaptureWithDevice:Device];
 	}
@@ -89,7 +89,7 @@ void FMetalCaptureManager::BeginCapture(void)
 void FMetalCaptureManager::EndCapture(void)
 {
 #if METAL_SUPPORTS_CAPTURE_MANAGER
-	if (@available(iOS 11.0, macOS 10.13, *))
+	if (@available(iOS 11.0, tvOS 11.0, macOS 10.13, *))
 	{
 		[[MTLCaptureManager sharedCaptureManager] stopCapture];
 	}
diff --git a/Engine/Source/Runtime/Apple/MetalRHI/Private/MetalComputeCommandEncoder.cpp b/Engine/Source/Runtime/Apple/MetalRHI/Private/MetalComputeCommandEncoder.cpp
--- a/Engine/Source/Runtime/Apple/MetalRHI/Private/MetalComputeCommandEncoder.cpp
+++ b/Engine/Source/Runtime/Apple/MetalRHI/Private/MetalComputeCommandEncoder.cpp
@@ -546,7 +546,7 @@ APPLE_PLATFORM_OBJECT_ALLOC_OVERRIDES(FMetalDebugComputeCommandEncoder)
 #if METAL_SUPPORTS_INDIRECT_ARGUMENT_BUFFERS
 - (void)useResource:(id <MTLResource>)resource usage:(MTLResourceUsage)usage
 {
-	if(@available(iOS 11.0, macOS 10.13, *))
+	if(@available(iOS 11.0, tvOS 11.0, macOS 10.13, *))
 	{
 		[Inner useResource:resource usage:usage];
 	}
@@ -554,7 +554,7 @@ APPLE_PLATFORM_OBJECT_ALLOC_OVERRIDES(FMetalDebugComputeCommandEncoder)
 
 - (void)useResources:(const id <MTLResource> [])resources count:(NSUInteger)count usage:(MTLResourceUsage)usage
 {
-	if(@available(iOS 11.0, macOS 10.13, *))
+	if(@available(iOS 11.0, tvOS 11.0, macOS 10.13, *))
 	{
 		[Inner useResources:resources count:count usage:usage];
 	}
@@ -562,7 +562,7 @@ APPLE_PLATFORM_OBJECT_ALLOC_OVERRIDES(FMetalDebugComputeCommandEncoder)
 
 - (void)useHeap:(id <MTLHeap>)heap
 {
-	if(@available(iOS 11.0, macOS 10.13, *))
+	if(@available(iOS 11.0, tvOS 11.0, macOS 10.13, *))
 	{
 		[Inner useHeap:heap];
 	}
@@ -570,7 +570,7 @@ APPLE_PLATFORM_OBJECT_ALLOC_OVERRIDES(FMetalDebugComputeCommandEncoder)
 
 - (void)useHeaps:(const id <MTLHeap> [])heaps count:(NSUInteger)count
 {
-	if(@available(iOS 11.0, macOS 10.13, *))
+	if(@available(iOS 11.0, tvOS 11.0, macOS 10.13, *))
 	{
 		[Inner useHeaps:heaps count:count];
 	}
@@ -722,6 +722,26 @@ APPLE_PLATFORM_OBJECT_ALLOC_OVERRIDES(FMetalDebugComputeCommandEncoder)
 	return self;
 }
 
+#if METAL_SUPPORTS_TILE_SHADERS
+- (void)setImageblockWidth:(NSUInteger)width height:(NSUInteger)height
+{
+	if(@available(iOS 11.0, tvOS 11.0, *))
+	{
+		[Inner setImageblockWidth:width height:height];
+	}
+}
+#endif
+
+#if METAL_SUPPORTS_CAPTURE_MANAGER
+- (void)dispatchThreads:(MTLSize)threadsPerGrid threadsPerThreadgroup:(MTLSize)threadsPerThreadgroup
+{
+	if(@available(iOS 11.0, tvOS 11.0, macOS 10.13, *))
+	{
+		[Inner dispatchThreads:threadsPerGrid threadsPerThreadgroup:threadsPerThreadgroup];
+	}
+}
+#endif
+
 @end
 
 #if !METAL_SUPPORTS_HEAPS
diff --git a/Engine/Source/Runtime/Apple/MetalRHI/Private/MetalParallelRenderCommandEncoder.cpp b/Engine/Source/Runtime/Apple/MetalRHI/Private/MetalParallelRenderCommandEncoder.cpp
--- a/Engine/Source/Runtime/Apple/MetalRHI/Private/MetalParallelRenderCommandEncoder.cpp
+++ b/Engine/Source/Runtime/Apple/MetalRHI/Private/MetalParallelRenderCommandEncoder.cpp
@@ -106,7 +106,7 @@ APPLE_PLATFORM_OBJECT_ALLOC_OVERRIDES(FMetalDebugParallelRenderCommandEncoder)
 #if (METAL_NEW_NONNULL_DECL)
 - (void)setColorStoreActionOptions:(MTLStoreActionOptions)storeActionOptions atIndex:(NSUInteger)colorAttachmentIndex
 {
-	if(@available(iOS 11.0, macOS 10.13, *))
+	if(@available(iOS 11.0, tvOS 11.0, macOS 10.13, *))
 	{
 		[Inner setColorStoreActionOptions:storeActionOptions atIndex:colorAttachmentIndex];
 	}
@@ -114,7 +114,7 @@ APPLE_PLATFORM_OBJECT_ALLOC_OVERRIDES(FMetalDebugParallelRenderCommandEncoder)
 
 - (void)setDepthStoreActionOptions:(MTLStoreActionOptions)storeActionOptions
 {
-	if(@available(iOS 11.0, macOS 10.13, *))
+	if(@available(iOS 11.0, tvOS 11.0, macOS 10.13, *))
 	{
 		[Inner setDepthStoreActionOptions:storeActionOptions];
 	}
@@ -122,7 +122,7 @@ APPLE_PLATFORM_OBJECT_ALLOC_OVERRIDES(FMetalDebugParallelRenderCommandEncoder)
 
 - (void)setStencilStoreActionOptions:(MTLStoreActionOptions)storeActionOptions
 {
-	if(@available(iOS 11.0, macOS 10.13, *))
+	if(@available(iOS 11.0, tvOS 11.0, macOS 10.13, *))
 	{
 		[Inner setStencilStoreActionOptions:storeActionOptions];
 	}
diff --git a/Engine/Source/Runtime/Apple/MetalRHI/Private/MetalRHIPrivate.h b/Engine/Source/Runtime/Apple/MetalRHI/Private/MetalRHIPrivate.h
--- a/Engine/Source/Runtime/Apple/MetalRHI/Private/MetalRHIPrivate.h
+++ b/Engine/Source/Runtime/Apple/MetalRHI/Private/MetalRHIPrivate.h
@@ -48,6 +48,7 @@ enum EMTLTextureType
 // This is the right VERSION check, see Availability.h in the SDK
 #define METAL_SUPPORTS_INDIRECT_ARGUMENT_BUFFERS !PLATFORM_MAC && __IPHONE_OS_VERSION_MAX_ALLOWED >= 110000
 #define METAL_SUPPORTS_CAPTURE_MANAGER (PLATFORM_MAC && __MAC_OS_X_VERSION_MAX_ALLOWED >= 101300) || (!PLATFORM_MAC && __IPHONE_OS_VERSION_MAX_ALLOWED >= 110000)
+#define METAL_SUPPORTS_TILE_SHADERS !PLATFORM_MAC && __IPHONE_OS_VERSION_MAX_ALLOWED >= 110000
 
 #define METAL_DEBUG_OPTIONS !(UE_BUILD_SHIPPING || UE_BUILD_TEST)
 #if METAL_DEBUG_OPTIONS
diff --git a/Engine/Source/Runtime/Apple/MetalRHI/Private/MetalRenderCommandEncoder.cpp b/Engine/Source/Runtime/Apple/MetalRHI/Private/MetalRenderCommandEncoder.cpp
--- a/Engine/Source/Runtime/Apple/MetalRHI/Private/MetalRenderCommandEncoder.cpp
+++ b/Engine/Source/Runtime/Apple/MetalRHI/Private/MetalRenderCommandEncoder.cpp
@@ -453,7 +453,7 @@ APPLE_PLATFORM_OBJECT_ALLOC_OVERRIDES(FMetalDebugRenderCommandEncoder)
 - (void)setDepthClipMode:(MTLDepthClipMode)depthClipMode
 {
 #if __clang_major__ >= 9
-	if(@available(iOS 11.0, macOS 10.11, *))
+	if(@available(iOS 11.0, tvOS 11.0, macOS 10.11, *))
 #endif
 	{
     	[Inner setDepthClipMode:depthClipMode];
@@ -1274,7 +1274,7 @@ APPLE_PLATFORM_OBJECT_ALLOC_OVERRIDES(FMetalDebugRenderCommandEncoder)
 #if METAL_SUPPORTS_INDIRECT_ARGUMENT_BUFFERS
 - (void)useResource:(id <MTLResource>)resource usage:(MTLResourceUsage)usage
 {
-	if(@available(iOS 11.0, macOS 10.13, *))
+	if(@available(iOS 11.0, tvOS 11.0, macOS 10.13, *))
 	{
 		[Inner useResource:resource usage:usage];
 	}
@@ -1282,7 +1282,7 @@ APPLE_PLATFORM_OBJECT_ALLOC_OVERRIDES(FMetalDebugRenderCommandEncoder)
 
 - (void)useResources:(const id <MTLResource> [])resources count:(NSUInteger)count usage:(MTLResourceUsage)usage
 {
-	if(@available(iOS 11.0, macOS 10.13, *))
+	if(@available(iOS 11.0, tvOS 11.0, macOS 10.13, *))
 	{
 		[Inner useResources:resources count:count usage:usage];
 	}
@@ -1290,7 +1290,7 @@ APPLE_PLATFORM_OBJECT_ALLOC_OVERRIDES(FMetalDebugRenderCommandEncoder)
 
 - (void)useHeap:(id <MTLHeap>)heap
 {
-	if(@available(iOS 11.0, macOS 10.13, *))
+	if(@available(iOS 11.0, tvOS 11.0, macOS 10.13, *))
 	{
 		[Inner useHeap:heap];
 	}
@@ -1298,7 +1298,7 @@ APPLE_PLATFORM_OBJECT_ALLOC_OVERRIDES(FMetalDebugRenderCommandEncoder)
 
 - (void)useHeaps:(const id <MTLHeap> [])heaps count:(NSUInteger)count
 {
-	if(@available(iOS 11.0, macOS 10.13, *))
+	if(@available(iOS 11.0, tvOS 11.0, macOS 10.13, *))
 	{
 		[Inner useHeaps:heaps count:count];
 	}
@@ -1500,7 +1500,7 @@ APPLE_PLATFORM_OBJECT_ALLOC_OVERRIDES(FMetalDebugRenderCommandEncoder)
 #if (METAL_NEW_NONNULL_DECL)
 - (void)setColorStoreActionOptions:(MTLStoreActionOptions)storeActionOptions atIndex:(NSUInteger)colorAttachmentIndex
 {
-	if(@available(iOS 11.0, macOS 10.13, *))
+	if(@available(iOS 11.0, tvOS 11.0, macOS 10.13, *))
 	{
 		[Inner setColorStoreActionOptions:storeActionOptions atIndex:colorAttachmentIndex];
 	}
@@ -1508,7 +1508,7 @@ APPLE_PLATFORM_OBJECT_ALLOC_OVERRIDES(FMetalDebugRenderCommandEncoder)
 
 - (void)setDepthStoreActionOptions:(MTLStoreActionOptions)storeActionOptions
 {
-	if(@available(iOS 11.0, macOS 10.13, *))
+	if(@available(iOS 11.0, tvOS 11.0, macOS 10.13, *))
 	{
 		[Inner setDepthStoreActionOptions:storeActionOptions];
 	}
@@ -1516,13 +1516,111 @@ APPLE_PLATFORM_OBJECT_ALLOC_OVERRIDES(FMetalDebugRenderCommandEncoder)
 
 - (void)setStencilStoreActionOptions:(MTLStoreActionOptions)storeActionOptions
 {
-	if(@available(iOS 11.0, macOS 10.13, *))
+	if(@available(iOS 11.0, tvOS 11.0, macOS 10.13, *))
 	{
 		[Inner setStencilStoreActionOptions:storeActionOptions];
 	}
 }
 #endif //(METAL_NEW_NONNULL_DECL)
 
+#if METAL_SUPPORTS_TILE_SHADERS
+- (void)setTileBytes:(const void *)bytes length:(NSUInteger)length atIndex:(NSUInteger)index
+{
+	if(@available(iOS 11.0, tvOS 11.0, *))
+	{
+		[Inner setTileBytes:bytes length:length atIndex:index];
+	}
+}
+
+- (void)setTileBuffer:(nullable id <MTLBuffer>)buffer offset:(NSUInteger)offset atIndex:(NSUInteger)index
+{
+	if(@available(iOS 11.0, tvOS 11.0, *))
+	{
+		[Inner setTileBuffer:buffer offset:offset atIndex:index];
+	}
+}
+
+- (void)setTileBufferOffset:(NSUInteger)offset atIndex:(NSUInteger)index
+{
+	if(@available(iOS 11.0, tvOS 11.0, *))
+	{
+		[Inner setTileBufferOffset:offset atIndex:index];
+	}
+}
+
+- (void)setTileBuffers:(const id <MTLBuffer> __nullable [__nonnull])buffers offsets:(const NSUInteger [__nonnull])offset withRange:(NSRange)range
+{
+	if(@available(iOS 11.0, tvOS 11.0, *))
+	{
+		[Inner setTileBuffers:buffers offsets:offset withRange:range];
+	}
+}
+
+- (void)setTileTexture:(nullable id <MTLTexture>)texture atIndex:(NSUInteger)index
+{
+	if(@available(iOS 11.0, tvOS 11.0, *))
+	{
+		[Inner setTileTexture:texture atIndex:index];
+	}
+}
+
+- (void)setTileTextures:(const id <MTLTexture> __nullable [__nonnull])textures withRange:(NSRange)range
+{
+	if(@available(iOS 11.0, tvOS 11.0, *))
+	{
+		[Inner setTileTextures:textures withRange:range];
+	}
+}
+
+- (void)setTileSamplerState:(nullable id <MTLSamplerState>)sampler atIndex:(NSUInteger)index
+{
+	if(@available(iOS 11.0, tvOS 11.0, *))
+	{
+		[Inner setTileSamplerState:sampler atIndex:index];
+	}
+}
+
+- (void)setTileSamplerStates:(const id <MTLSamplerState> __nullable [__nonnull])samplers withRange:(NSRange)range
+{
+	if(@available(iOS 11.0, tvOS 11.0, *))
+	{
+		[Inner setTileSamplerStates:samplers withRange:range];
+	}
+}
+
+- (void)setTileSamplerState:(nullable id <MTLSamplerState>)sampler lodMinClamp:(float)lodMinClamp lodMaxClamp:(float)lodMaxClamp atIndex:(NSUInteger)index
+{
+	if(@available(iOS 11.0, tvOS 11.0, *))
+	{
+		[Inner setTileSamplerState:sampler lodMinClamp:lodMinClamp lodMaxClamp:lodMaxClamp atIndex:index];
+	}
+}
+
+- (void)setTileSamplerStates:(const id <MTLSamplerState> __nullable [__nonnull])samplers lodMinClamps:(const float [__nonnull])lodMinClamps lodMaxClamps:(const float [__nonnull])lodMaxClamps withRange:(NSRange)range
+{
+	if(@available(iOS 11.0, tvOS 11.0, *))
+	{
+		[Inner setTileSamplerStates:samplers lodMinClamps:lodMinClamps lodMaxClamps:lodMaxClamps withRange:range];
+	}
+}
+
+- (void)dispatchThreadsPerTile:(MTLSize)threadsPerTile
+{
+	if(@available(iOS 11.0, tvOS 11.0, *))
+	{
+		[Inner dispatchThreadsPerTile:threadsPerTile];
+	}
+}
+
+- (void)setThreadgroupMemoryLength:(NSUInteger)length offset:(NSUInteger)offset atIndex:(NSUInteger)index
+{
+	if(@available(iOS 11.0, tvOS 11.0, *))
+	{
+		[Inner setThreadgroupMemoryLength:length offset:offset atIndex:index];
+	}
+}
+
+#endif
 @end
 
 #if !METAL_SUPPORTS_HEAPS
