More technology for debuggung "[CrashReport] Generic callstack for MetalCommandBufferFailure crashes on Mac"

Duplicate 3641738:
Fix-up the internal validation layer in MetalRHI which can now be enabled from the OS command-line (-MetalRuntimeDebugLevel=X - where X is the numerical level you want) and extend the validation (level 5) to write out the current operation index on calls to popDebugGroup which will give us a better idea of where the GPU was when it failed. This is only useful for debugging a reproducible error as the overhead is non-trivial and it will currently only work on Mac as on iOS the driver rearranges draw & dispatch calls to suit the TBDR GPU.

Duplicate 3642047:
Fix Xcode 8 compile errors.

diff --git a/Engine/Source/Runtime/Apple/MetalRHI/Private/MetalCommandBuffer.cpp b/Engine/Source/Runtime/Apple/MetalRHI/Private/MetalCommandBuffer.cpp
--- a/Engine/Source/Runtime/Apple/MetalRHI/Private/MetalCommandBuffer.cpp
+++ b/Engine/Source/Runtime/Apple/MetalRHI/Private/MetalCommandBuffer.cpp
@@ -57,6 +57,11 @@ APPLE_PLATFORM_OBJECT_ALLOC_OVERRIDES(FMetalDebugCommandBuffer)
 		InnerBuffer = Buffer;
 		DebugGroup = [NSMutableArray new];
 		ActiveEncoder = nil;
+		DebugInfoBuffer = nil;
+		if (DebugLevel >= EMetalDebugLevelValidation)
+		{
+			DebugInfoBuffer = [Buffer.device newBufferWithLength:BufferOffsetAlignment options:0];
+		}
 	}
 	return Self;
 }
@@ -72,6 +77,8 @@ APPLE_PLATFORM_OBJECT_ALLOC_OVERRIDES(FMetalDebugCommandBuffer)
 	
 	[InnerBuffer release];
 	[DebugGroup release];
+	[DebugInfoBuffer release];
+	DebugInfoBuffer = nil;
 	
 	[super dealloc];
 }
@@ -173,7 +180,7 @@ APPLE_PLATFORM_OBJECT_ALLOC_OVERRIDES(FMetalDebugCommandBuffer)
 - (id <MTLRenderCommandEncoder>)renderCommandEncoderWithDescriptor:(MTLRenderPassDescriptor *)renderPassDescriptor
 {
     [self beginRenderCommandEncoder:DebugGroup.lastObject ? DebugGroup.lastObject : @"Render" withDescriptor:renderPassDescriptor];
-    return [[[FMetalDebugRenderCommandEncoder alloc] initWithEncoder:[InnerBuffer renderCommandEncoderWithDescriptor:renderPassDescriptor] andCommandBuffer:self] autorelease];
+	return [[[FMetalDebugRenderCommandEncoder alloc] initWithEncoder:[InnerBuffer renderCommandEncoderWithDescriptor:renderPassDescriptor] fromDescriptor:renderPassDescriptor andCommandBuffer:self] autorelease];
 }
 
 - (id <MTLComputeCommandEncoder>)computeCommandEncoder
@@ -202,9 +209,23 @@ APPLE_PLATFORM_OBJECT_ALLOC_OVERRIDES(FMetalDebugCommandBuffer)
 	NSString* Label = self.label ? self.label : @"Unknown";
 	[String appendFormat:@"Command Buffer %p %@:", self, Label];
 
+	uint32 Index = 0;
+	if (DebugInfoBuffer)
+	{
+		Index = *((uint32*)DebugInfoBuffer.contents);
+	}
+	
+	uint32 Count = 1;
 	for (FMetalDebugCommand* Command : DebugCommands)
 	{
-		[String appendFormat:@"\n\t%@: %@", GMetalDebugCommandTypeNames[Command->Type], Command->Label];
+		if (Index == Count++)
+		{
+			[String appendFormat:@"\n\t--> %@: %@", GMetalDebugCommandTypeNames[Command->Type], Command->Label];
+		}
+		else
+		{
+			[String appendFormat:@"\n\t%@: %@", GMetalDebugCommandTypeNames[Command->Type], Command->Label];
+		}
 	}
 	
 	[String appendFormat:@"\nResources:"];
@@ -411,5 +432,69 @@ APPLE_PLATFORM_OBJECT_ALLOC_OVERRIDES(FMetalDebugCommandBuffer)
     }
 }
 
+-(CFTimeInterval) kernelStartTime
+{
+#if METAL_NEW_NONNULL_DECL
+	if (@available(iOS 10.3, tvOS 10.3, macOS 10.13, *))
+#endif
+	{
+		return ((id<IMetalCommandBufferExtensions>)InnerBuffer).kernelStartTime;
+	}
+#if METAL_NEW_NONNULL_DECL
+	else
+	{
+		return 0;
+	}
+#endif
+}
+
+-(CFTimeInterval) kernelEndTime
+{
+#if METAL_NEW_NONNULL_DECL
+	if (@available(iOS 10.3, tvOS 10.3, macOS 10.13, *))
+#endif
+	{
+		return ((id<IMetalCommandBufferExtensions>)InnerBuffer).kernelEndTime;
+	}
+#if METAL_NEW_NONNULL_DECL
+	else
+	{
+		return 0;
+	}
+#endif
+}
+
+-(CFTimeInterval) GPUStartTime
+{
+#if METAL_NEW_NONNULL_DECL
+	if (@available(iOS 10.3, tvOS 10.3, macOS 10.13, *))
+#endif
+	{
+		return ((id<IMetalCommandBufferExtensions>)InnerBuffer).GPUStartTime;
+	}
+#if METAL_NEW_NONNULL_DECL
+	else
+	{
+		return 0;
+	}
+#endif
+}
+
+-(CFTimeInterval) GPUEndTime
+{
+#if METAL_NEW_NONNULL_DECL
+	if (@available(iOS 10.3, tvOS 10.3, macOS 10.13, *))
+#endif
+	{
+		return ((id<IMetalCommandBufferExtensions>)InnerBuffer).GPUEndTime;
+	}
+#if METAL_NEW_NONNULL_DECL
+	else
+	{
+		return 0;
+	}
+#endif
+}
+
 @end
 
diff --git a/Engine/Source/Runtime/Apple/MetalRHI/Private/MetalCommandBuffer.h b/Engine/Source/Runtime/Apple/MetalRHI/Private/MetalCommandBuffer.h
--- a/Engine/Source/Runtime/Apple/MetalRHI/Private/MetalCommandBuffer.h
+++ b/Engine/Source/Runtime/Apple/MetalRHI/Private/MetalCommandBuffer.h
@@ -74,13 +74,14 @@ struct FMetalDebugCommand
  */
 @interface FMetalDebugCommandBuffer : FApplePlatformObject<MTLCommandBuffer>
 {
-	TArray<FMetalDebugCommand*> DebugCommands;
 	NSMutableArray<NSString*>* DebugGroup;
 	NSString* ActiveEncoder;
 	TSet<id<MTLResource>> Resources;
 	TSet<id> States;
 	@public
+	TArray<FMetalDebugCommand*> DebugCommands;
     EMetalDebugLevel DebugLevel;
+	id<MTLBuffer> DebugInfoBuffer;
 };
 
 /** The wrapped native command-buffer for which we collect debug information. */
diff --git a/Engine/Source/Runtime/Apple/MetalRHI/Private/MetalComputeCommandEncoder.cpp b/Engine/Source/Runtime/Apple/MetalRHI/Private/MetalComputeCommandEncoder.cpp
--- a/Engine/Source/Runtime/Apple/MetalRHI/Private/MetalComputeCommandEncoder.cpp
+++ b/Engine/Source/Runtime/Apple/MetalRHI/Private/MetalComputeCommandEncoder.cpp
@@ -13,6 +13,32 @@
 
 NS_ASSUME_NONNULL_BEGIN
 
+#if METAL_DEBUG_OPTIONS
+static NSString* GMetalDebugComputeShader = @"#include <metal_stdlib>\n"
+@"using namespace metal;\n"
+@"kernel void WriteCommandIndexCS(constant uint* Input [[ buffer(0) ]], device atomic_uint* Output [[ buffer(1) ]])\n"
+@"{\n"
+@"	atomic_store_explicit(Output, Input[0], memory_order_relaxed);\n"
+@"}\n";
+
+static id <MTLComputePipelineState> GetDebugComputeShaderState(id<MTLDevice> Device)
+{
+	static id<MTLComputePipelineState> State = nil;
+	if (!State)
+	{
+		id<MTLLibrary> Lib = [Device newLibraryWithSource:GMetalDebugComputeShader options:nil error:nullptr];
+		check(Lib);
+		id<MTLFunction> Func = [Lib newFunctionWithName:@"WriteCommandIndexCS"];
+		check(Func);
+		State = [Device newComputePipelineStateWithFunction:Func error:nil];
+		[Func release];
+		[Lib release];
+	}
+	check(State);
+	return State;
+}
+#endif
+
 @implementation FMetalDebugComputeCommandEncoder
 
 @synthesize Inner;
@@ -74,9 +100,60 @@ APPLE_PLATFORM_OBJECT_ALLOC_OVERRIDES(FMetalDebugComputeCommandEncoder)
     [Inner pushDebugGroup:string];
 }
 
+#if METAL_DEBUG_OPTIONS
+- (void)insertDebugDispatch
+{
+	switch (Buffer->DebugLevel)
+	{
+		case EMetalDebugLevelConditionalSubmit:
+		case EMetalDebugLevelWaitForComplete:
+		case EMetalDebugLevelLogOperations:
+		case EMetalDebugLevelValidation:
+		{
+			uint32 const Index = Buffer->DebugCommands.Num();
+			[Inner setBytes:&Index length:sizeof(Index) atIndex:0];
+			[Inner setBuffer:Buffer->DebugInfoBuffer offset:0 atIndex:1];
+			[Inner setComputePipelineState:GetDebugComputeShaderState(Inner.device)];
+			
+			[Inner dispatchThreadgroups:MTLSizeMake(1, 1, 1) threadsPerThreadgroup:MTLSizeMake(1, 1, 1)];
+
+			if (Pipeline && Pipeline.ComputePipelineState)
+			{
+				[Inner setComputePipelineState:Pipeline.ComputePipelineState];
+			}
+			
+			if (ShaderBuffers.Buffers[0])
+			{
+				[Inner setBuffer:ShaderBuffers.Buffers[0] offset:ShaderBuffers.Offsets[0] atIndex:0];
+			}
+			else if (ShaderBuffers.Bytes[0])
+			{
+				[Inner setBytes:ShaderBuffers.Bytes[0] length:ShaderBuffers.Offsets[0] atIndex:0];
+			}
+			
+			if (ShaderBuffers.Buffers[1])
+			{
+				[Inner setBuffer:ShaderBuffers.Buffers[1] offset:ShaderBuffers.Offsets[1] atIndex:1];
+			}
+			else if (ShaderBuffers.Bytes[1])
+			{
+				[Inner setBytes:ShaderBuffers.Bytes[1] length:ShaderBuffers.Offsets[1] atIndex:1];
+			}
+		}
+		default:
+		{
+			break;
+		}
+	}
+}
+#endif
+
 - (void)popDebugGroup
 {
     [Buffer popDebugGroup];
+#if METAL_DEBUG_OPTIONS
+	[self insertDebugDispatch];
+#endif
     [Inner popDebugGroup];
 }
 
diff --git a/Engine/Source/Runtime/Apple/MetalRHI/Private/MetalContext.cpp b/Engine/Source/Runtime/Apple/MetalRHI/Private/MetalContext.cpp
--- a/Engine/Source/Runtime/Apple/MetalRHI/Private/MetalContext.cpp
+++ b/Engine/Source/Runtime/Apple/MetalRHI/Private/MetalContext.cpp
@@ -338,6 +338,13 @@ FMetalDeviceContext* FMetalDeviceContext::CreateDeviceContext()
 	FMetalCommandQueue* Queue = new FMetalCommandQueue(Device, GMetalCommandQueueSize);
 	check(Queue);
 	
+	uint32 MetalDebug = GMetalRuntimeDebugLevel;
+	const bool bOverridesMetalDebug = FParse::Value( FCommandLine::Get(), TEXT( "MetalRuntimeDebugLevel=" ), MetalDebug );
+	if (bOverridesMetalDebug)
+	{
+		GMetalRuntimeDebugLevel = MetalDebug;
+	}
+	
 	return new FMetalDeviceContext(Device, DeviceIndex, Queue);
 }
 
diff --git a/Engine/Source/Runtime/Apple/MetalRHI/Private/MetalParallelRenderCommandEncoder.cpp b/Engine/Source/Runtime/Apple/MetalRHI/Private/MetalParallelRenderCommandEncoder.cpp
--- a/Engine/Source/Runtime/Apple/MetalRHI/Private/MetalParallelRenderCommandEncoder.cpp
+++ b/Engine/Source/Runtime/Apple/MetalRHI/Private/MetalParallelRenderCommandEncoder.cpp
@@ -85,7 +85,7 @@ APPLE_PLATFORM_OBJECT_ALLOC_OVERRIDES(FMetalDebugParallelRenderCommandEncoder)
 - (id <MTLRenderCommandEncoder>)renderCommandEncoder
 #endif
 {
-    return [[[FMetalDebugRenderCommandEncoder alloc] initWithEncoder:[Inner renderCommandEncoder] andCommandBuffer:Buffer] autorelease];
+	return [[[FMetalDebugRenderCommandEncoder alloc] initWithEncoder:[Inner renderCommandEncoder] fromDescriptor:RenderPassDescriptor andCommandBuffer:Buffer] autorelease];
 }
 
 - (void)setColorStoreAction:(MTLStoreAction)storeAction atIndex:(NSUInteger)colorAttachmentIndex
diff --git a/Engine/Source/Runtime/Apple/MetalRHI/Private/MetalRenderCommandEncoder.cpp b/Engine/Source/Runtime/Apple/MetalRHI/Private/MetalRenderCommandEncoder.cpp
--- a/Engine/Source/Runtime/Apple/MetalRHI/Private/MetalRenderCommandEncoder.cpp
+++ b/Engine/Source/Runtime/Apple/MetalRHI/Private/MetalRenderCommandEncoder.cpp
@@ -14,6 +14,68 @@
 
 NS_ASSUME_NONNULL_BEGIN
 
+#if METAL_DEBUG_OPTIONS
+static NSString* GMetalDebugVertexShader = @"#include <metal_stdlib>\n"
+@"using namespace metal;\n"
+@"struct VertexInput\n"
+@"{\n"
+@"};\n"
+@"vertex void WriteCommandIndexVS(VertexInput StageIn [[stage_in]], constant uint* Input [[ buffer(0) ]], device uint* Output  [[ buffer(1) ]])\n"
+@"{\n"
+@"	Output[0] = Input[0];\n"
+@"}\n";
+
+static id <MTLRenderPipelineState> GetDebugVertexShaderState(id<MTLDevice> Device, MTLRenderPassDescriptor* PassDesc)
+{
+	static id<MTLFunction> Func = nil;
+	static FCriticalSection Mutex;
+	static NSMutableDictionary* Dict = [NSMutableDictionary new];
+	if(!Func)
+	{
+		id<MTLLibrary> Lib = [Device newLibraryWithSource:GMetalDebugVertexShader options:nil error:nullptr];
+		check(Lib);
+		Func = [Lib newFunctionWithName:@"WriteCommandIndexVS"];
+		check(Func);
+		[Lib release];
+	}
+	
+	FScopeLock Lock(&Mutex);
+	id<MTLRenderPipelineState> State = [Dict objectForKey:PassDesc];
+	if (!State)
+	{
+		MTLRenderPipelineDescriptor* Desc = [MTLRenderPipelineDescriptor new];
+		
+		Desc.vertexFunction = Func;
+		
+		if (PassDesc.depthAttachment)
+		{
+			Desc.depthAttachmentPixelFormat = PassDesc.depthAttachment.texture.pixelFormat;
+		}
+		if (PassDesc.stencilAttachment)
+		{
+			Desc.stencilAttachmentPixelFormat = PassDesc.stencilAttachment.texture.pixelFormat;
+		}
+		if (PassDesc.colorAttachments)
+		{
+			MTLRenderPassColorAttachmentDescriptor* CD = [PassDesc.colorAttachments objectAtIndexedSubscript:0];
+			MTLRenderPipelineColorAttachmentDescriptor* CD0 = [[MTLRenderPipelineColorAttachmentDescriptor new] autorelease];
+			CD0.pixelFormat = CD.texture.pixelFormat;
+			[Desc.colorAttachments setObject:CD0 atIndexedSubscript:0];
+		}
+		Desc.rasterizationEnabled = false;
+		
+		State = [[Device newRenderPipelineStateWithDescriptor:Desc error:nil] autorelease];
+		check(State);
+		
+		[Dict setObject:State forKey:PassDesc];
+		
+		[Desc release];
+	}
+	check(State);
+	return State;
+}
+#endif
+
 @implementation FMetalDebugRenderCommandEncoder
 
 @synthesize Inner;
@@ -22,13 +84,17 @@ NS_ASSUME_NONNULL_BEGIN
 
 APPLE_PLATFORM_OBJECT_ALLOC_OVERRIDES(FMetalDebugRenderCommandEncoder)
 
--(id)initWithEncoder:(id<MTLRenderCommandEncoder>)Encoder andCommandBuffer:(FMetalDebugCommandBuffer*)SourceBuffer
+-(id)initWithEncoder:(id<MTLRenderCommandEncoder>)Encoder fromDescriptor:(MTLRenderPassDescriptor*)Desc andCommandBuffer:(FMetalDebugCommandBuffer*)SourceBuffer
 {
 	id Self = [super init];
 	if (Self)
 	{
         Inner = [Encoder retain];
 		Buffer = [SourceBuffer retain];
+		RenderPassDesc = [Desc retain];
+#if METAL_DEBUG_OPTIONS
+		DebugState = Buffer->DebugLevel >= EMetalDebugLevelValidation ? [GetDebugVertexShaderState(Buffer.device, Desc) retain] : nil;
+#endif
         Pipeline = nil;
 	}
 	return Self;
@@ -38,6 +104,10 @@ APPLE_PLATFORM_OBJECT_ALLOC_OVERRIDES(FMetalDebugRenderCommandEncoder)
 {
 	[Inner release];
 	[Buffer release];
+	[RenderPassDesc release];
+#if METAL_DEBUG_OPTIONS
+	[DebugState release];
+#endif
 	[Pipeline release];
 	[super dealloc];
 }
@@ -75,9 +145,64 @@ APPLE_PLATFORM_OBJECT_ALLOC_OVERRIDES(FMetalDebugRenderCommandEncoder)
     [Inner pushDebugGroup:string];
 }
 
+#if METAL_DEBUG_OPTIONS
+- (void)insertDebugDraw
+{
+	switch (Buffer->DebugLevel)
+	{
+		case EMetalDebugLevelConditionalSubmit:
+		case EMetalDebugLevelWaitForComplete:
+		case EMetalDebugLevelLogOperations:
+		case EMetalDebugLevelValidation:
+		{
+			uint32 const Index = Buffer->DebugCommands.Num();
+#if PLATFORM_MAC
+			[Inner textureBarrier];
+#endif
+			[Inner setVertexBytes:&Index length:sizeof(Index) atIndex:0];
+			[Inner setVertexBuffer:Buffer->DebugInfoBuffer offset:0 atIndex:1];
+			[Inner setRenderPipelineState:DebugState];
+			[Inner drawPrimitives:MTLPrimitiveTypePoint vertexStart:0 vertexCount:1];
+#if PLATFORM_MAC
+			[Inner textureBarrier];
+#endif
+			if (Pipeline && Pipeline.RenderPipelineState)
+			{
+				[Inner setRenderPipelineState:Pipeline.RenderPipelineState];
+			}
+			
+			if (ShaderBuffers[EMetalShaderVertex].Buffers[0])
+			{
+				[Inner setVertexBuffer:ShaderBuffers[EMetalShaderVertex].Buffers[0] offset:ShaderBuffers[EMetalShaderVertex].Offsets[0] atIndex:0];
+			}
+			else if (ShaderBuffers[EMetalShaderVertex].Bytes[0])
+			{
+				[Inner setVertexBytes:ShaderBuffers[EMetalShaderVertex].Bytes[0] length:ShaderBuffers[EMetalShaderVertex].Offsets[0] atIndex:0];
+			}
+			
+			if (ShaderBuffers[EMetalShaderVertex].Buffers[1])
+			{
+				[Inner setVertexBuffer:ShaderBuffers[EMetalShaderVertex].Buffers[1] offset:ShaderBuffers[EMetalShaderVertex].Offsets[1] atIndex:1];
+			}
+			else if (ShaderBuffers[EMetalShaderVertex].Bytes[1])
+			{
+				[Inner setVertexBytes:ShaderBuffers[EMetalShaderVertex].Bytes[1] length:ShaderBuffers[EMetalShaderVertex].Offsets[1] atIndex:1];
+			}
+		}
+		default:
+		{
+			break;
+		}
+	}
+}
+#endif
+
 - (void)popDebugGroup
 {
-    [Buffer popDebugGroup];
+	[Buffer popDebugGroup];
+#if METAL_DEBUG_OPTIONS
+	[self insertDebugDraw];
+#endif
     [Inner popDebugGroup];
 }
 
diff --git a/Engine/Source/Runtime/Apple/MetalRHI/Private/MetalRenderCommandEncoder.h b/Engine/Source/Runtime/Apple/MetalRHI/Private/MetalRenderCommandEncoder.h
--- a/Engine/Source/Runtime/Apple/MetalRHI/Private/MetalRenderCommandEncoder.h
+++ b/Engine/Source/Runtime/Apple/MetalRHI/Private/MetalRenderCommandEncoder.h
@@ -22,7 +22,9 @@ NS_ASSUME_NONNULL_BEGIN
     FMetalDebugBufferBindings ShaderBuffers[EMetalShaderRenderNum];
     FMetalDebugTextureBindings ShaderTextures[EMetalShaderRenderNum];
     FMetalDebugSamplerBindings ShaderSamplers[EMetalShaderRenderNum];
+	id<MTLRenderPipelineState> DebugState;
 #endif
+	MTLRenderPassDescriptor* RenderPassDesc;
 }
 
 /** The wrapped native command-encoder for which we collect debug information. */
@@ -31,7 +33,7 @@ NS_ASSUME_NONNULL_BEGIN
 @property (nonatomic, retain) FMetalShaderPipeline* Pipeline;
 
 /** Initialise the wrapper with the provided command-buffer. */
--(id)initWithEncoder:(id<MTLRenderCommandEncoder>)Encoder andCommandBuffer:(FMetalDebugCommandBuffer*)Buffer;
+-(id)initWithEncoder:(id<MTLRenderCommandEncoder>)Encoder fromDescriptor:(MTLRenderPassDescriptor*)RenderPassDesc andCommandBuffer:(FMetalDebugCommandBuffer*)Buffer;
 
 /** Validates the pipeline/binding state */
 -(bool)validateFunctionBindings:(EMetalShaderFrequency)Frequency;
