Duplicate 3642606:

Fix for Pixel Inspector Not Working on Mac -  Added missing CopySubTextureRegion to MetalRHI.  Changed RHI Function def'n not to return a bool - I don't think it makes sense.

diff --git a/Engine/Source/Runtime/Apple/MetalRHI/Private/MetalTexture.cpp b/Engine/Source/Runtime/Apple/MetalRHI/Private/MetalTexture.cpp
--- a/Engine/Source/Runtime/Apple/MetalRHI/Private/MetalTexture.cpp
+++ b/Engine/Source/Runtime/Apple/MetalRHI/Private/MetalTexture.cpp
@@ -2456,3 +2456,92 @@ void FMetalDynamicRHI::RHISetResourceAliasability_RenderThread(class FRHICommand
 		}
 	}
 }
+
+struct FRHICopySubTextureRegion : public FRHICommand<FRHICopySubTextureRegion>
+{
+	FTexture2DRHIParamRef SourceTexture;
+	FTexture2DRHIParamRef DestinationTexture;
+	FBox2D SourceBox;
+	FBox2D DestinationBox;
+	
+	FORCEINLINE_DEBUGGABLE FRHICopySubTextureRegion(FTexture2DRHIParamRef InSourceTexture, FTexture2DRHIParamRef InDestinationTexture, FBox2D InSourceBox, FBox2D InDestinationBox)
+	: SourceTexture(InSourceTexture)
+	, DestinationTexture(InDestinationTexture)
+	, SourceBox(InSourceBox)
+	, DestinationBox(InDestinationBox)
+	{
+	}
+	
+	void Execute(FRHICommandListBase& CmdList)
+	{
+		GDynamicRHI->RHICopySubTextureRegion(SourceTexture, DestinationTexture, SourceBox, DestinationBox);
+	}
+};
+
+void FMetalDynamicRHI::RHICopySubTextureRegion_RenderThread(class FRHICommandListImmediate& RHICmdList, FTexture2DRHIParamRef SourceTexture, FTexture2DRHIParamRef DestinationTexture, FBox2D SourceBox, FBox2D DestinationBox)
+{
+	@autoreleasepool
+	{
+		if (RHICmdList.Bypass() || !IsRunningRHIInSeparateThread())
+		{
+			GDynamicRHI->RHICopySubTextureRegion(SourceTexture, DestinationTexture, SourceBox, DestinationBox);
+		}
+		else
+		{
+			new (RHICmdList.AllocCommand<FRHICopySubTextureRegion>()) FRHICopySubTextureRegion(SourceTexture, DestinationTexture, SourceBox, DestinationBox);
+		}
+	}
+}
+
+void FMetalDynamicRHI::RHICopySubTextureRegion(FTexture2DRHIParamRef SourceTexture, FTexture2DRHIParamRef DestinationTexture, FBox2D SourceBox, FBox2D DestinationBox)
+{
+	check(SourceTexture);
+	check(DestinationTexture);
+	
+	if(SourceTexture->GetFormat() == DestinationTexture->GetFormat())
+	{
+		FMetalTexture2D* MetalSrcTexture = ResourceCast(SourceTexture);
+		FMetalTexture2D* MetalDestTexture = ResourceCast(DestinationTexture);
+		
+		FVector2D SourceSizeVector = SourceBox.GetSize();
+		FVector2D DestinatioSizeVector = DestinationBox.GetSize();
+
+		MTLOrigin SourceOrigin = MTLOriginMake(SourceBox.Min.X, SourceBox.Min.Y, 0);
+		MTLSize SourceSize = MTLSizeMake(SourceSizeVector.X, SourceSizeVector.Y, 1);
+		
+		MTLOrigin DestinationOrigin = MTLOriginMake(DestinationBox.Min.X, DestinationBox.Min.Y, 0);
+		MTLSize DestinationSize = MTLSizeMake(DestinatioSizeVector.X, DestinatioSizeVector.Y, 1);
+		
+		check(DestinationSize.width == SourceSize.width);
+		check(DestinationSize.height == SourceSize.height);
+		
+		// Account for create with TexCreate_SRGB flag which could make these different
+		if(MetalSrcTexture->Surface.Texture.pixelFormat == MetalDestTexture->Surface.Texture.pixelFormat)
+		{
+			ImmediateContext.GetInternalContext().CopyFromTextureToTexture(MetalSrcTexture->Surface.Texture, 0, 0, SourceOrigin,SourceSize,MetalDestTexture->Surface.Texture, 0, 0, DestinationOrigin);
+		}
+		else
+		{
+			// Linear and sRGB mismatch then try to go via metal buffer
+			// Modified clone of logic from MetalRenderTarget.cpp
+			uint32 BytesPerPixel = (MetalSrcTexture->Surface.PixelFormat != PF_DepthStencil) ? GPixelFormats[MetalSrcTexture->Surface.PixelFormat].BlockBytes : 1;
+			const uint32 Stride = BytesPerPixel * SourceSize.width;
+			const uint32 Alignment = PLATFORM_MAC ? 1u : 64u;
+			const uint32 AlignedStride = ((Stride - 1) & ~(Alignment - 1)) + Alignment;
+			const uint32 BytesPerImage = AlignedStride *  SourceSize.height;
+			
+			id<MTLBuffer> Buffer = GetMetalDeviceContext().CreatePooledBuffer(FMetalPooledBufferArgs(ImmediateContext.Context->GetDevice(), BytesPerImage, MTLStorageModeShared));
+			
+			check(Buffer != nil);
+			
+			ImmediateContext.GetInternalContext().CopyFromTextureToBuffer(MetalSrcTexture->Surface.Texture, 0, 0, SourceOrigin, SourceSize, Buffer, 0, AlignedStride, BytesPerImage, MTLBlitOptionNone);
+			ImmediateContext.GetInternalContext().CopyFromBufferToTexture(Buffer, 0, Stride, BytesPerImage, SourceSize, MetalDestTexture->Surface.Texture, 0, 0, DestinationOrigin);
+			
+			GetMetalDeviceContext().ReleasePooledBuffer(Buffer);
+		}
+	}
+	else
+	{
+		UE_LOG(LogMetal, Warning, TEXT("RHICopySubTextureRegion Source <-> Destination texture format mismatch"));
+	}
+}
diff --git a/Engine/Source/Runtime/Apple/MetalRHI/Public/MetalRHI.h b/Engine/Source/Runtime/Apple/MetalRHI/Public/MetalRHI.h
--- a/Engine/Source/Runtime/Apple/MetalRHI/Public/MetalRHI.h
+++ b/Engine/Source/Runtime/Apple/MetalRHI/Public/MetalRHI.h
@@ -181,6 +181,10 @@ public:
 	virtual FComputeFenceRHIRef RHICreateComputeFence(const FName& Name) final override;
 	
 	virtual void RHISetResourceAliasability_RenderThread(class FRHICommandListImmediate& RHICmdList, EResourceAliasability AliasMode, FTextureRHIParamRef* InTextures, int32 NumTextures) final override;
+	
+	/* Copy the source box pixels in the destination box texture, return true if implemented for the current platform*/
+	virtual void RHICopySubTextureRegion_RenderThread(class FRHICommandListImmediate& RHICmdList, FTexture2DRHIParamRef SourceTexture, FTexture2DRHIParamRef DestinationTexture, FBox2D SourceBox, FBox2D DestinationBox) final override;
+	virtual void RHICopySubTextureRegion(FTexture2DRHIParamRef SourceTexture, FTexture2DRHIParamRef DestinationTexture, FBox2D SourceBox, FBox2D DestinationBox) final override;
 
 private:
 	FTextureMemoryStats MemoryStats;
diff --git a/Engine/Source/Runtime/RHI/Private/RHICommandList.cpp b/Engine/Source/Runtime/RHI/Private/RHICommandList.cpp
--- a/Engine/Source/Runtime/RHI/Private/RHICommandList.cpp
+++ b/Engine/Source/Runtime/RHI/Private/RHICommandList.cpp
@@ -2487,4 +2487,9 @@ void FRHICommandListImmediate::UpdateTextureReference(FTextureReferenceRHIParamR
 	}
 }
 
+void FDynamicRHI::RHICopySubTextureRegion_RenderThread(class FRHICommandListImmediate& RHICmdList, FTexture2DRHIParamRef SourceTexture, FTexture2DRHIParamRef DestinationTexture, FBox2D SourceBox, FBox2D DestinationBox)
+{
+	FScopedRHIThreadStaller StallRHIThread(RHICmdList);
+	return GDynamicRHI->RHICopySubTextureRegion(SourceTexture, DestinationTexture, SourceBox, DestinationBox);
+}
 
diff --git a/Engine/Source/Runtime/RHI/Public/DynamicRHI.h b/Engine/Source/Runtime/RHI/Public/DynamicRHI.h
--- a/Engine/Source/Runtime/RHI/Public/DynamicRHI.h
+++ b/Engine/Source/Runtime/RHI/Public/DynamicRHI.h
@@ -889,8 +889,9 @@ public:
 	virtual bool CheckGpuHeartbeat() const { return true; }
 
 	/* Copy the source box pixels in the destination box texture, return true if implemented for the current platform*/
-	virtual bool RHICopySubTextureRegion(FTexture2DRHIParamRef SourceTexture, FTexture2DRHIParamRef DestinationTexture, FBox2D SourceBox, FBox2D DestinationBox) { return false; }
-
+	virtual void RHICopySubTextureRegion_RenderThread(class FRHICommandListImmediate& RHICmdList, FTexture2DRHIParamRef SourceTexture, FTexture2DRHIParamRef DestinationTexture, FBox2D SourceBox, FBox2D DestinationBox);
+	virtual void RHICopySubTextureRegion(FTexture2DRHIParamRef SourceTexture, FTexture2DRHIParamRef DestinationTexture, FBox2D SourceBox, FBox2D DestinationBox) { }
+	
 protected:
 	TArray<uint32> PixelFormatBlockBytes;
 };
diff --git a/Engine/Source/Runtime/RHI/Public/RHICommandList.h b/Engine/Source/Runtime/RHI/Public/RHICommandList.h
--- a/Engine/Source/Runtime/RHI/Public/RHICommandList.h
+++ b/Engine/Source/Runtime/RHI/Public/RHICommandList.h
@@ -4266,7 +4266,7 @@ public:
 
 	FORCEINLINE void CopySubTextureRegion(FTexture2DRHIParamRef SourceTexture, FTexture2DRHIParamRef DestinationTexture, FBox2D SourceBox, FBox2D DestinationBox)
 	{
-		GDynamicRHI->RHICopySubTextureRegion(SourceTexture, DestinationTexture, SourceBox, DestinationBox);
+		GDynamicRHI->RHICopySubTextureRegion_RenderThread(*this, SourceTexture, DestinationTexture, SourceBox, DestinationBox);
 	}
 	
 	FORCEINLINE void ExecuteCommandList(FRHICommandList* CmdList)
diff --git a/Engine/Source/Runtime/Windows/D3D11RHI/Private/D3D11RHIPrivate.h b/Engine/Source/Runtime/Windows/D3D11RHI/Private/D3D11RHIPrivate.h
--- a/Engine/Source/Runtime/Windows/D3D11RHI/Private/D3D11RHIPrivate.h
+++ b/Engine/Source/Runtime/Windows/D3D11RHI/Private/D3D11RHIPrivate.h
@@ -540,7 +540,7 @@ public:
 	virtual void RHIPushEvent(const TCHAR* Name, FColor Color) final override;
 	virtual void RHIPopEvent() final override;
 
-	virtual bool RHICopySubTextureRegion(FTexture2DRHIParamRef SourceTexture, FTexture2DRHIParamRef DestinationTexture, FBox2D SourceBox, FBox2D DestinationBox) final override;
+	virtual void RHICopySubTextureRegion(FTexture2DRHIParamRef SourceTexture, FTexture2DRHIParamRef DestinationTexture, FBox2D SourceBox, FBox2D DestinationBox) final override;
 
 	virtual FTexture2DRHIRef RHICreateTexture2DFromResource(EPixelFormat Format, uint32 TexCreateFlags, const FClearValueBinding& ClearValueBinding, ID3D11Texture2D* Resource);
 	virtual FTextureCubeRHIRef RHICreateTextureCubeFromResource(EPixelFormat Format, uint32 TexCreateFlags, const FClearValueBinding& ClearValueBinding, ID3D11Texture2D* Resource);
diff --git a/Engine/Source/Runtime/Windows/D3D11RHI/Private/D3D11Texture.cpp b/Engine/Source/Runtime/Windows/D3D11RHI/Private/D3D11Texture.cpp
--- a/Engine/Source/Runtime/Windows/D3D11RHI/Private/D3D11Texture.cpp
+++ b/Engine/Source/Runtime/Windows/D3D11RHI/Private/D3D11Texture.cpp
@@ -1657,7 +1657,7 @@ FTextureReferenceRHIRef FD3D11DynamicRHI::RHICreateTextureReference(FLastRenderT
 	return new FD3D11TextureReference(this,LastRenderTime);
 }
 
-bool FD3D11DynamicRHI::RHICopySubTextureRegion(FTexture2DRHIParamRef SourceTextureRHI, FTexture2DRHIParamRef DestinationTextureRHI, FBox2D SourceBox, FBox2D DestinationBox)
+void FD3D11DynamicRHI::RHICopySubTextureRegion(FTexture2DRHIParamRef SourceTextureRHI, FTexture2DRHIParamRef DestinationTextureRHI, FBox2D SourceBox, FBox2D DestinationBox)
 {
 	FD3D11Texture2D* SourceTexture = ResourceCast(SourceTextureRHI);
 	FD3D11Texture2D* DestinationTexture = ResourceCast(DestinationTextureRHI);
@@ -1708,7 +1708,6 @@ bool FD3D11DynamicRHI::RHICopySubTextureRegion(FTexture2DRHIParamRef SourceTextu
 	check(GPixelFormats[DestinationTexture->GetFormat()].BlockSizeY == 1);
 	ID3D11Texture2D* DestinationRessource = DestinationTexture->GetResource();
 	Direct3DDeviceIMContext->CopySubresourceRegion(DestinationRessource, 0, DestinationBox.Min.X + DestinationOffsetX, DestinationBox.Min.Y + DestinationOffsetY, 0, SourceTexture->GetResource(), 0, &SourceBoxAdjust);
-	return true;
 }
 
 void FD3D11DynamicRHI::RHIUpdateTextureReference(FTextureReferenceRHIParamRef TextureRefRHI, FTextureRHIParamRef NewTextureRHI)
