Disable WWDC hack to enable parallel rendering in the Editor on Metal which is causing problems in texture streaming.

diff --git a/Engine/Source/Runtime/Apple/MetalRHI/Private/MetalRHI.cpp b/Engine/Source/Runtime/Apple/MetalRHI/Private/MetalRHI.cpp
--- a/Engine/Source/Runtime/Apple/MetalRHI/Private/MetalRHI.cpp
+++ b/Engine/Source/Runtime/Apple/MetalRHI/Private/MetalRHI.cpp
@@ -320,7 +320,7 @@ FMetalDynamicRHI::FMetalDynamicRHI(ERHIFeatureLevel::Type RequestedFeatureLevel)
 	{
 #if METAL_SUPPORTS_PARALLEL_RHI_EXECUTE
 #if WITH_EDITORONLY_DATA
-		GRHISupportsRHIThread = true;
+		GRHISupportsRHIThread = false;
 #else
 		GRHISupportsRHIThread = bSupportsRHIThread;
 #endif
