Fix tvOS compilation.

diff --git a/Engine/Source/Runtime/Apple/MetalRHI/Private/MetalComputeCommandEncoder.cpp b/Engine/Source/Runtime/Apple/MetalRHI/Private/MetalComputeCommandEncoder.cpp
--- a/Engine/Source/Runtime/Apple/MetalRHI/Private/MetalComputeCommandEncoder.cpp
+++ b/Engine/Source/Runtime/Apple/MetalRHI/Private/MetalComputeCommandEncoder.cpp
@@ -809,7 +809,7 @@ APPLE_PLATFORM_OBJECT_ALLOC_OVERRIDES(FMetalDebugComputeCommandEncoder)
 }
 #endif
 
-#if METAL_SUPPORTS_CAPTURE_MANAGER
+#if METAL_SUPPORTS_CAPTURE_MANAGER && !PLATFORM_TVOS
 - (void)dispatchThreads:(MTLSize)threadsPerGrid threadsPerThreadgroup:(MTLSize)threadsPerThreadgroup
 {
 	if(@available(iOS 11.0, tvOS 11.0, macOS 10.13, *))
diff --git a/Engine/Source/Runtime/Apple/MetalRHI/Private/MetalRHIPrivate.h b/Engine/Source/Runtime/Apple/MetalRHI/Private/MetalRHIPrivate.h
--- a/Engine/Source/Runtime/Apple/MetalRHI/Private/MetalRHIPrivate.h
+++ b/Engine/Source/Runtime/Apple/MetalRHI/Private/MetalRHIPrivate.h
@@ -48,7 +48,7 @@ enum EMTLTextureType
 // This is the right VERSION check, see Availability.h in the SDK
 #define METAL_SUPPORTS_INDIRECT_ARGUMENT_BUFFERS !PLATFORM_MAC && __IPHONE_OS_VERSION_MAX_ALLOWED >= 110000
 #define METAL_SUPPORTS_CAPTURE_MANAGER (PLATFORM_MAC && __MAC_OS_X_VERSION_MAX_ALLOWED >= 101300) || (!PLATFORM_MAC && __IPHONE_OS_VERSION_MAX_ALLOWED >= 110000)
-#define METAL_SUPPORTS_TILE_SHADERS !PLATFORM_MAC && __IPHONE_OS_VERSION_MAX_ALLOWED >= 110000
+#define METAL_SUPPORTS_TILE_SHADERS !PLATFORM_MAC && !PLATFORM_TVOS && __IPHONE_OS_VERSION_MAX_ALLOWED >= 110000
 
 #define METAL_DEBUG_OPTIONS !(UE_BUILD_SHIPPING || UE_BUILD_TEST)
 #if METAL_DEBUG_OPTIONS
