Duplicate CL 3769940: Fix crash on startup on 2013 Mac Pro's.

diff --git a/Engine/Source/Runtime/Apple/MetalRHI/Private/MetalVertexBuffer.cpp b/Engine/Source/Runtime/Apple/MetalRHI/Private/MetalVertexBuffer.cpp
--- a/Engine/Source/Runtime/Apple/MetalRHI/Private/MetalVertexBuffer.cpp
+++ b/Engine/Source/Runtime/Apple/MetalRHI/Private/MetalVertexBuffer.cpp
@@ -149,9 +149,6 @@ void FMetalVertexBuffer::Alloc(uint32 InSize)
 		
 		if (FMetalCommandQueue::SupportsFeature(EMetalFeaturesLinearTextures) && (GetUsage() & (BUF_UnorderedAccess|BUF_ShaderResource)))
 		{
-			FMetalPooledBufferArgs Args(GetMetalDeviceContext().GetDevice(), InSize, MTLStorageModePrivate);
-			Buffer = GetMetalDeviceContext().CreatePooledBuffer(Args);
-			
 			for (TPair<EPixelFormat, id<MTLTexture>>& Pair : LinearTextures)
 			{
 				SafeReleaseMetalObject(Pair.Value);
@@ -256,11 +253,11 @@ void* FMetalVertexBuffer::Lock(EResourceLockMode LockMode, uint32 Offset, uint32
 		return ((uint8*)Data->Data) + Offset;
 	}
 	
-	id<MTLBuffer>& theBufferToUse = CPUBuffer ? CPUBuffer : Buffer;
 	
 	// In order to properly synchronise the buffer access, when a dynamic buffer is locked for writing, discard the old buffer & create a new one. This prevents writing to a buffer while it is being read by the GPU & thus causing corruption. This matches the logic of other RHIs.
 	if ((GetUsage() & BUFFER_DYNAMIC_REALLOC) && LockMode == RLM_WriteOnly)
 	{
+		id<MTLBuffer>& theBufferToUse = CPUBuffer ? CPUBuffer : Buffer;
 		INC_MEMORY_STAT_BY(STAT_MetalVertexMemAlloc, GetSize());
 		INC_MEMORY_STAT_BY(STAT_MetalVertexMemFreed, GetSize());
 		
@@ -277,6 +274,7 @@ void* FMetalVertexBuffer::Lock(EResourceLockMode LockMode, uint32 Offset, uint32
 		Alloc(InSize);
 	}
 	
+	id<MTLBuffer>& theBufferToUse = CPUBuffer ? CPUBuffer : Buffer;
 	if(LockMode != RLM_ReadOnly)
 	{
 		LockSize = Size;
