Duplicate 4091485
PR - Fixed potential memory leak.
#4751

diff --git a/Engine/Source/Runtime/Apple/MetalRHI/Private/MetalViewport.cpp b/Engine/Source/Runtime/Apple/MetalRHI/Private/MetalViewport.cpp
--- a/Engine/Source/Runtime/Apple/MetalRHI/Private/MetalViewport.cpp
+++ b/Engine/Source/Runtime/Apple/MetalRHI/Private/MetalViewport.cpp
@@ -301,9 +301,10 @@ mtlpp::Drawable FMetalViewport::GetDrawable(EMetalViewportAccessFlag Accessor)
 
 	#else
 			CGSize Size;
+			IOSAppDelegate* AppDelegate = [IOSAppDelegate GetDelegate];
 			do
 			{
-				Drawable = [[IOSAppDelegate GetDelegate].IOSView MakeDrawable];
+				Drawable = [AppDelegate.IOSView MakeDrawable];
 				Size.width = ((id<CAMetalDrawable>)Drawable).texture.width;
 				Size.height = ((id<CAMetalDrawable>)Drawable).texture.height;
 			}
