Duplicate 4075868

Metal Vega drivers in 10.13.5 fix blitting with non-zero offsets and we should be using the same buffer blitting path on 2013 Mac Pro's now that we aren't paying an exorbitant cost for doing so.

diff --git a/Engine/Source/Runtime/Apple/MetalRHI/Private/MetalCommandQueue.cpp b/Engine/Source/Runtime/Apple/MetalRHI/Private/MetalCommandQueue.cpp
--- a/Engine/Source/Runtime/Apple/MetalRHI/Private/MetalCommandQueue.cpp
+++ b/Engine/Source/Runtime/Apple/MetalRHI/Private/MetalCommandQueue.cpp
@@ -112,14 +112,12 @@ FMetalCommandQueue::FMetalCommandQueue(mtlpp::Device InDevice, uint32 const MaxN
 		
 		Features |= EMetalFeaturesLinearTextures;
 		
-		// Using Private Memory & BlitEncoders for Vertex & Index data is slower on the Mac Pro's DXXX GPUs
-		// Everywhere else it should be *much* faster.
-		if (!FString(Device.GetName()).Contains(TEXT("FirePro")))
-		{
-			Features |= EMetalFeaturesEfficientBufferBlits;
-		}
-		if (!FString(Device.GetName()).Contains(TEXT("Vega")))
-		{
+		// Using Private Memory & BlitEncoders for Vertex & Index data should be *much* faster.
+        Features |= EMetalFeaturesEfficientBufferBlits;
+		
+        // On earlier OS versions Vega didn't like non-zero blit offsets
+        if (!FString(Device.GetName()).Contains(TEXT("Vega")) || FPlatformMisc::MacOSXVersionCompare(10,13,5) >= 0)
+        {
 			Features |= EMetalFeaturesPrivateBufferSubAllocation;
 		}
     }
