Duplicate 4078874 & 4079396

Make LPV visualisation work on macOS to aid debugging.

Export the constant-buffer bindings that are actually used by Metal shaders to prevent overwriting SRVs & UAVs with unused uniform-buffer data.

diff --git a/Engine/Shaders/Private/LPVVisualise.usf b/Engine/Shaders/Private/LPVVisualise.usf
--- a/Engine/Shaders/Private/LPVVisualise.usf
+++ b/Engine/Shaders/Private/LPVVisualise.usf
@@ -33,11 +33,15 @@
 //-------------------------------------------------------------------------------------------------
 
 // GS to VS
+#if !USING_VERTEX_SHADER_LAYER
 struct VIn
 {
     float4 position : SV_POSITION;
 	float4 params	: TEXCOORD0;
 };
+#else
+#define VIn VOut
+#endif
 
 struct VOut 
 {
@@ -48,20 +52,52 @@ struct VOut
 //-------------------------------------------------------------------------------------------------
 
 // Visualization
-VIn VShader(uint id : SV_InstanceID)
+VIn VShader(uint InstanceId : SV_InstanceID
+#if USING_VERTEX_SHADER_LAYER
+, uint VertexId : SV_VertexID
+#endif
+)
 {
-	uint sliceIndex = id % 32; 
-	uint orientation = id / 32;
+	uint sliceIndex = InstanceId % 32;
+	uint orientation = InstanceId / 32;
 
 	VIn output;
 	output.position.xyz = float3( 0, (float(sliceIndex)/31.0f)*2-1, 0 );
 	output.position.w = 1;
 
+#if !USING_VERTEX_SHADER_LAYER
 	output.params = float4( sliceIndex, orientation,0,0 ) ;
+#else
+	float3 verts[4] = {
+	float3(-1,-1,0 ),
+	float3( 1,-1,0 ),
+	float3( 1, 1,0 ),
+	float3(-1, 1,0 ) };
+
+	float3x3 rot[3] =
+	{
+	float3x3( float3(1,0,0), float3(0,1,0), float3(0,0,1) ),
+	float3x3( float3(0,1,0), float3(1,0,0), float3(0,0,1) ),
+	float3x3( float3(1,0,0), float3(0,0,1), float3(0,1,0) )
+	};
+
+	float3 centreWorld = output.position.xyz;
+	int indices[6] = { 0,1,2, 2,3,0 };
+	int index = indices[VertexId];
+
+	float3 p = ( centreWorld + verts[ index ].xzy );
+	p = mul( p, rot[ (int)orientation ] );
+	float4 vert = float4( p, 1 );
+	output.texCoord = p.xyz * 0.5 + 0.5;
+	vert.xyz *= CUBESIZE;
+	vert.xyz += View.ViewForward * 100;
+	output.position = mul( vert, View.TranslatedWorldToClip);
+#endif
     return output;
 }
 
 
+#if !USING_VERTEX_SHADER_LAYER
 [maxvertexcount(12)]
 void GShader( point VIn input[1], inout TriangleStream<VOut> TriStream )
 {
@@ -99,10 +135,10 @@ void GShader( point VIn input[1], inout TriangleStream<VOut> TriStream )
 	}
 	TriStream.RestartStrip();
 }
+#endif
 
 
-
-float4 PShader(VOut IN) : SV_Target
+void PShader(VOut IN, out float4 OutColor : SV_Target0)
 {
 	float3 pos;
 	pos.xyz = IN.texCoord.xyz * 31.0f;
@@ -144,5 +180,5 @@ float4 PShader(VOut IN) : SV_Target
 		colour = float3(0,0,50);
 	}
 #endif 
-	return float4( colour / 3, 1 );
+	OutColor = float4( colour / 3, 1 );
 }
diff --git a/Engine/Source/Developer/Apple/MetalShaderFormat/Private/MetalBackend.cpp b/Engine/Source/Developer/Apple/MetalShaderFormat/Private/MetalBackend.cpp
--- a/Engine/Source/Developer/Apple/MetalShaderFormat/Private/MetalBackend.cpp
+++ b/Engine/Source/Developer/Apple/MetalShaderFormat/Private/MetalBackend.cpp
@@ -1164,6 +1164,11 @@ protected:
 							buffer,
 							" [[ buffer(%d) ]]", BufferIndex
 							);
+						
+						if (!bNeedsPointer)
+						{
+							Backend.ConstantBuffers |= 1 << BufferIndex;
+						}
 					}
 				}
 				else if (IsMain && var->mode == ir_var_in)
@@ -5890,6 +5895,7 @@ FMetalCodeBackend::FMetalCodeBackend(FMetalTessellationOutputs& TessOutputAttrib
 	InvariantBuffers(0),
 	TypedBuffers(0),
     TypedUAVs(0),
+	ConstantBuffers(0),
 	bExplicitDepthWrites(false)
 {
     Version = InVersion;
diff --git a/Engine/Source/Developer/Apple/MetalShaderFormat/Private/MetalBackend.h b/Engine/Source/Developer/Apple/MetalShaderFormat/Private/MetalBackend.h
--- a/Engine/Source/Developer/Apple/MetalShaderFormat/Private/MetalBackend.h
+++ b/Engine/Source/Developer/Apple/MetalShaderFormat/Private/MetalBackend.h
@@ -120,6 +120,7 @@ struct FMetalCodeBackend : public FCodeBackend
 	uint32 InvariantBuffers;
 	uint32 TypedBuffers;
     uint32 TypedUAVs;
+	uint32 ConstantBuffers;
     
     uint8 Version;
 	EMetalGPUSemantics bIsDesktop;
diff --git a/Engine/Source/Developer/Apple/MetalShaderFormat/Private/MetalDerivedData.cpp b/Engine/Source/Developer/Apple/MetalShaderFormat/Private/MetalDerivedData.cpp
--- a/Engine/Source/Developer/Apple/MetalShaderFormat/Private/MetalDerivedData.cpp
+++ b/Engine/Source/Developer/Apple/MetalShaderFormat/Private/MetalDerivedData.cpp
@@ -46,6 +46,7 @@ extern void BuildMetalShaderOutput(
 	uint32 TypedBuffers,
 	uint32 InvariantBuffers,
 	uint32 TypedUAVs,
+	uint32 ConstantBuffers,
 	TArray<uint8> const& TypedBufferFormats,
 	bool bAllowFastIntriniscs
 );
@@ -437,7 +438,7 @@ bool FMetalShaderOutputCooker::Build(TArray<uint8>& OutData)
 	if (Result != 0)
 	{
 		Output.Target = Input.Target;
-		BuildMetalShaderOutput(Output, Input, GUIDHash, MetalShaderSource, SourceLen, CRCLen, CRC, VersionEnum, *Standard, *MinOSVersion, TypeMode, Output.Errors, Attribs, MetalBackEnd.TypedBuffers, MetalBackEnd.InvariantBuffers, MetalBackEnd.TypedUAVs, MetalBackEnd.TypedBufferFormats, bAllowFastIntriniscs);
+		BuildMetalShaderOutput(Output, Input, GUIDHash, MetalShaderSource, SourceLen, CRCLen, CRC, VersionEnum, *Standard, *MinOSVersion, TypeMode, Output.Errors, Attribs, MetalBackEnd.TypedBuffers, MetalBackEnd.InvariantBuffers, MetalBackEnd.TypedUAVs, MetalBackEnd.ConstantBuffers, MetalBackEnd.TypedBufferFormats, bAllowFastIntriniscs);
 
 		FMemoryWriter Ar(OutData);
 		Ar << Output;
diff --git a/Engine/Source/Developer/Apple/MetalShaderFormat/Private/MetalShaderCompiler.cpp b/Engine/Source/Developer/Apple/MetalShaderFormat/Private/MetalShaderCompiler.cpp
--- a/Engine/Source/Developer/Apple/MetalShaderFormat/Private/MetalShaderCompiler.cpp
+++ b/Engine/Source/Developer/Apple/MetalShaderFormat/Private/MetalShaderCompiler.cpp
@@ -947,6 +947,7 @@ void BuildMetalShaderOutput(
 	uint32 TypedBuffers,
 	uint32 InvariantBuffers,
 	uint32 TypedUAVs,
+	uint32 ConstantBuffers,
 	TArray<uint8> const& TypedBufferFormats,
 	bool bAllowFastIntriniscs
 	)
@@ -1001,6 +1002,7 @@ void BuildMetalShaderOutput(
 	Header.SourceLen = SourceCRCLen;
 	Header.SourceCRC = SourceCRC;
     Header.Bindings.bDiscards = false;
+	Header.Bindings.ConstantBuffers = ConstantBuffers;
 	if (Version >= 2)
 	{
 		Header.Bindings.TypedBufferFormats.SetNumZeroed(METAL_MAX_BUFFERS);
diff --git a/Engine/Source/Developer/Apple/MetalShaderFormat/Private/MetalShaderFormat.cpp b/Engine/Source/Developer/Apple/MetalShaderFormat/Private/MetalShaderFormat.cpp
--- a/Engine/Source/Developer/Apple/MetalShaderFormat/Private/MetalShaderFormat.cpp
+++ b/Engine/Source/Developer/Apple/MetalShaderFormat/Private/MetalShaderFormat.cpp
@@ -111,7 +111,7 @@ class FMetalShaderFormat : public IShaderFormat
 public:
 	enum
 	{
-		HEADER_VERSION = 55,
+		HEADER_VERSION = 56,
 	};
 	
 	struct FVersion
diff --git a/Engine/Source/Runtime/Apple/MetalRHI/Private/MetalCommands.cpp b/Engine/Source/Runtime/Apple/MetalRHI/Private/MetalCommands.cpp
--- a/Engine/Source/Runtime/Apple/MetalRHI/Private/MetalCommands.cpp
+++ b/Engine/Source/Runtime/Apple/MetalRHI/Private/MetalCommands.cpp
@@ -498,7 +498,7 @@ void FMetalRHICommandContext::RHISetShaderUniformBuffer(FVertexShaderRHIParamRef
 
 	auto& Bindings = VertexShader->Bindings;
 	check(BufferIndex < Bindings.NumUniformBuffers);
-	if (Bindings.bHasRegularUniformBuffers && !(Bindings.LinearBuffer & 1 << BufferIndex))
+	if (Bindings.bHasRegularUniformBuffers && (Bindings.ConstantBuffers & 1 << BufferIndex))
 	{
 		auto* UB = (FMetalUniformBuffer*)BufferRHI;
 		Context->GetCurrentState().SetShaderBuffer(SF_Vertex, UB->Buffer, UB->Data, 0, UB->GetSize(), BufferIndex);
@@ -514,7 +514,7 @@ void FMetalRHICommandContext::RHISetShaderUniformBuffer(FHullShaderRHIParamRef H
 
 	auto& Bindings = HullShader->Bindings;
 	check(BufferIndex < Bindings.NumUniformBuffers);
-	if (Bindings.bHasRegularUniformBuffers && !(Bindings.LinearBuffer & 1 << BufferIndex))
+	if (Bindings.bHasRegularUniformBuffers && (Bindings.ConstantBuffers & 1 << BufferIndex))
 	{
 		auto* UB = (FMetalUniformBuffer*)BufferRHI;
 		Context->GetCurrentState().SetShaderBuffer(SF_Hull, UB->Buffer, UB->Data, 0, UB->GetSize(), BufferIndex);
@@ -530,7 +530,7 @@ void FMetalRHICommandContext::RHISetShaderUniformBuffer(FDomainShaderRHIParamRef
 
 	auto& Bindings = DomainShader->Bindings;
 	check(BufferIndex < Bindings.NumUniformBuffers);
-	if (Bindings.bHasRegularUniformBuffers && !(Bindings.LinearBuffer & 1 << BufferIndex))
+	if (Bindings.bHasRegularUniformBuffers && (Bindings.ConstantBuffers & 1 << BufferIndex))
 	{
 		auto* UB = (FMetalUniformBuffer*)BufferRHI;
 		Context->GetCurrentState().SetShaderBuffer(SF_Domain, UB->Buffer, UB->Data, 0, UB->GetSize(), BufferIndex);
@@ -551,7 +551,7 @@ void FMetalRHICommandContext::RHISetShaderUniformBuffer(FPixelShaderRHIParamRef
 
 	auto& Bindings = PixelShader->Bindings;
 	check(BufferIndex < Bindings.NumUniformBuffers);
-	if (Bindings.bHasRegularUniformBuffers && !(Bindings.LinearBuffer & 1 << BufferIndex))
+	if (Bindings.bHasRegularUniformBuffers && (Bindings.ConstantBuffers & 1 << BufferIndex))
 	{
 		auto* UB = (FMetalUniformBuffer*)BufferRHI;
 		Context->GetCurrentState().SetShaderBuffer(SF_Pixel, UB->Buffer, UB->Data, 0, UB->GetSize(), BufferIndex);
@@ -567,7 +567,7 @@ void FMetalRHICommandContext::RHISetShaderUniformBuffer(FComputeShaderRHIParamRe
 
 	auto& Bindings = ComputeShader->Bindings;
 	check(BufferIndex < Bindings.NumUniformBuffers);
-	if (Bindings.bHasRegularUniformBuffers && !(Bindings.LinearBuffer & 1 << BufferIndex))
+	if (Bindings.bHasRegularUniformBuffers && (Bindings.ConstantBuffers & 1 << BufferIndex))
 	{
 		auto* UB = (FMetalUniformBuffer*)BufferRHI;
 		Context->GetCurrentState().SetShaderBuffer(SF_Compute, UB->Buffer, UB->Data, 0, UB->GetSize(), BufferIndex);
diff --git a/Engine/Source/Runtime/Apple/MetalRHI/Public/MetalShaderResources.h b/Engine/Source/Runtime/Apple/MetalRHI/Public/MetalShaderResources.h
--- a/Engine/Source/Runtime/Apple/MetalRHI/Public/MetalShaderResources.h
+++ b/Engine/Source/Runtime/Apple/MetalRHI/Public/MetalShaderResources.h
@@ -137,6 +137,7 @@ struct FMetalShaderBindings
     uint32  LinearBuffer;
 	uint32	TypedBuffers;
 	uint32 	InvariantBuffers;
+	uint32 	ConstantBuffers;
 	uint16	InOutMask;
 	uint8	NumSamplers;
 	uint8	NumUniformBuffers;
@@ -148,6 +149,7 @@ struct FMetalShaderBindings
 		LinearBuffer(0),
         TypedBuffers(0),
 		InvariantBuffers(0),
+		ConstantBuffers(0),
 		InOutMask(0),
 		NumSamplers(0),
 		NumUniformBuffers(0),
@@ -167,6 +169,7 @@ inline FArchive& operator<<(FArchive& Ar, FMetalShaderBindings& Bindings)
     Ar << Bindings.LinearBuffer;
     Ar << Bindings.TypedBuffers;
 	Ar << Bindings.InvariantBuffers;
+	Ar << Bindings.ConstantBuffers;
 	Ar << Bindings.InOutMask;
 	Ar << Bindings.NumSamplers;
 	Ar << Bindings.NumUniformBuffers;
diff --git a/Engine/Source/Runtime/Renderer/Private/LightPropagationVolumeVisualisation.cpp b/Engine/Source/Runtime/Renderer/Private/LightPropagationVolumeVisualisation.cpp
--- a/Engine/Source/Runtime/Renderer/Private/LightPropagationVolumeVisualisation.cpp
+++ b/Engine/Source/Runtime/Renderer/Private/LightPropagationVolumeVisualisation.cpp
@@ -208,23 +208,34 @@ void FLightPropagationVolume::Visualise(FRHICommandList& RHICmdList, const FView
 	GraphicsPSOInit.BlendState = TStaticBlendState<CW_RGB, BO_Add, BF_One, BF_One>::GetRHI();
 
 	TShaderMapRef<FLpvVisualiseVS> VertexShader(View.ShaderMap);
-	TShaderMapRef<FLpvVisualiseGS> GeometryShader(View.ShaderMap);
+	TOptionalShaderMapRef<FLpvVisualiseGS> GeometryShader(View.ShaderMap);
 	TShaderMapRef<FLpvVisualisePS> PixelShader(View.ShaderMap);
 
+	EPrimitiveType PrimType = PT_PointList;
+	uint32 NumPrims = 1;
+	if (!RHISupportsGeometryShaders(View.GetShaderPlatform()))
+	{
+		PrimType = PT_TriangleList;
+		NumPrims = 2;
+	}
+	
 	GraphicsPSOInit.BoundShaderState.VertexDeclarationRHI = GSimpleElementVertexDeclaration.VertexDeclarationRHI;
 	GraphicsPSOInit.BoundShaderState.VertexShaderRHI = GETSAFERHISHADER_VERTEX(*VertexShader);
 	GraphicsPSOInit.BoundShaderState.GeometryShaderRHI = GETSAFERHISHADER_GEOMETRY(*GeometryShader);
 	GraphicsPSOInit.BoundShaderState.PixelShaderRHI = GETSAFERHISHADER_PIXEL(*PixelShader);
-	GraphicsPSOInit.PrimitiveType = PT_PointList;
+	GraphicsPSOInit.PrimitiveType = PrimType;
 
 	SetGraphicsPipelineState(RHICmdList, GraphicsPSOInit);
 
 	VertexShader->SetParameters(RHICmdList, View);
-	GeometryShader->SetParameters(RHICmdList, View);
+	if (GeometryShader.IsValid())
+	{
+		GeometryShader->SetParameters(RHICmdList, View);
+	}
 	PixelShader->SetParameters(RHICmdList, this, View);
 
 	RHICmdList.SetStreamSource(0, NULL, 0);
-	RHICmdList.DrawPrimitive(PT_PointList, 0, 1, 32 * 3);
+	RHICmdList.DrawPrimitive(PrimType, 0, NumPrims, 32 * 3);
 
 	PixelShader->UnbindBuffers(RHICmdList);
 }
