Duplicate 4073368

Linear allocator for <2MB Managed buffers allocated in 4MB chunks - hopefully this reduces the total number of buffers and thus reduces the churn in the driver further. These should all be Volatile buffers meaning that we should eventually get semi-decent reuse of these buffers. Fixes some buffer alignment and sizing gotcha's along the way.

diff --git a/Engine/Source/Runtime/Apple/MetalRHI/Private/MetalBuffer.cpp b/Engine/Source/Runtime/Apple/MetalRHI/Private/MetalBuffer.cpp
--- a/Engine/Source/Runtime/Apple/MetalRHI/Private/MetalBuffer.cpp
+++ b/Engine/Source/Runtime/Apple/MetalRHI/Private/MetalBuffer.cpp
@@ -24,6 +24,7 @@ FMetalBuffer::~FMetalBuffer()
 FMetalBuffer::FMetalBuffer(ns::Protocol<id<MTLBuffer>>::type handle, ns::Ownership retain)
 : mtlpp::Buffer(handle, nullptr, retain)
 , Heap(nullptr)
+, Linear(nullptr)
 , Magazine(nullptr)
 , bPooled(false)
 , bSingleUse(false)
@@ -33,6 +34,18 @@ FMetalBuffer::FMetalBuffer(ns::Protocol<id<MTLBuffer>>::type handle, ns::Ownersh
 FMetalBuffer::FMetalBuffer(mtlpp::Buffer&& rhs, FMetalSubBufferHeap* heap)
 : mtlpp::Buffer((mtlpp::Buffer&&)rhs)
 , Heap(heap)
+, Linear(nullptr)
+, Magazine(nullptr)
+, bPooled(false)
+, bSingleUse(false)
+{
+}
+
+
+FMetalBuffer::FMetalBuffer(mtlpp::Buffer&& rhs, FMetalSubBufferLinear* heap)
+: mtlpp::Buffer((mtlpp::Buffer&&)rhs)
+, Heap(nullptr)
+, Linear(heap)
 , Magazine(nullptr)
 , bPooled(false)
 , bSingleUse(false)
@@ -42,6 +55,7 @@ FMetalBuffer::FMetalBuffer(mtlpp::Buffer&& rhs, FMetalSubBufferHeap* heap)
 FMetalBuffer::FMetalBuffer(mtlpp::Buffer&& rhs, FMetalSubBufferMagazine* magazine)
 : mtlpp::Buffer((mtlpp::Buffer&&)rhs)
 , Heap(nullptr)
+, Linear(nullptr)
 , Magazine(magazine)
 , bPooled(false)
 , bSingleUse(false)
@@ -51,6 +65,7 @@ FMetalBuffer::FMetalBuffer(mtlpp::Buffer&& rhs, FMetalSubBufferMagazine* magazin
 FMetalBuffer::FMetalBuffer(mtlpp::Buffer&& rhs, bool bInPooled)
 : mtlpp::Buffer((mtlpp::Buffer&&)rhs)
 , Heap(nullptr)
+, Linear(nullptr)
 , Magazine(nullptr)
 , bPooled(bInPooled)
 , bSingleUse(false)
@@ -60,6 +75,7 @@ FMetalBuffer::FMetalBuffer(mtlpp::Buffer&& rhs, bool bInPooled)
 FMetalBuffer::FMetalBuffer(const FMetalBuffer& rhs)
 : mtlpp::Buffer(rhs)
 , Heap(rhs.Heap)
+, Linear(rhs.Linear)
 , Magazine(rhs.Magazine)
 , bPooled(rhs.bPooled)
 , bSingleUse(false)
@@ -69,6 +85,7 @@ FMetalBuffer::FMetalBuffer(const FMetalBuffer& rhs)
 FMetalBuffer::FMetalBuffer(FMetalBuffer&& rhs)
 : mtlpp::Buffer((mtlpp::Buffer&&)rhs)
 , Heap(rhs.Heap)
+, Linear(rhs.Linear)
 , Magazine(rhs.Magazine)
 , bPooled(rhs.bPooled)
 , bSingleUse(false)
@@ -81,6 +98,7 @@ FMetalBuffer& FMetalBuffer::operator=(const FMetalBuffer& rhs)
 	{
 		mtlpp::Buffer::operator=(rhs);
         Heap = rhs.Heap;
+		Linear = rhs.Linear;
 		Magazine = rhs.Magazine;
 		bPooled = rhs.bPooled;
 		bSingleUse = rhs.bSingleUse;
@@ -92,6 +110,7 @@ FMetalBuffer& FMetalBuffer::operator=(FMetalBuffer&& rhs)
 {
 	mtlpp::Buffer::operator=((mtlpp::Buffer&&)rhs);
 	Heap = rhs.Heap;
+	Linear = rhs.Linear;
 	Magazine = rhs.Magazine;
 	bPooled = rhs.bPooled;
 	bSingleUse = rhs.bSingleUse;
@@ -105,6 +124,11 @@ void FMetalBuffer::Release()
 		Heap->FreeRange(ns::Range(GetOffset(), GetLength()));
 		Heap = nullptr;
 	}
+	else if (Linear)
+	{
+		Linear->FreeRange(ns::Range(GetOffset(), GetLength()));
+		Linear = nullptr;
+	}
 	else if (Magazine)
 	{
 		Magazine->FreeRange(ns::Range(GetOffset(), GetLength()));
@@ -304,6 +328,136 @@ mtlpp::PurgeableState FMetalSubBufferHeap::SetPurgeableState(mtlpp::PurgeableSta
 	return ParentBuffer.SetPurgeableState(state);
 }
 
+#pragma mark --
+
+FMetalSubBufferLinear::FMetalSubBufferLinear(NSUInteger Size, NSUInteger Alignment, mtlpp::ResourceOptions Options, FCriticalSection& InPoolMutex)
+: PoolMutex(InPoolMutex)
+, MinAlign(Alignment)
+, WriteHead(0)
+, UsedSize(0)
+, FreedSize(0)
+{
+	NSUInteger FullSize = Align(Size, Alignment);
+	METAL_GPUPROFILE(FScopedMetalCPUStats CPUStat(FString::Printf(TEXT("AllocBuffer: %llu, %llu"), FullSize, Options)));
+	
+	mtlpp::StorageMode Storage = (mtlpp::StorageMode)((Options & mtlpp::ResourceStorageModeMask) >> mtlpp::ResourceStorageModeShift);
+	ParentBuffer = MTLPP_VALIDATE(mtlpp::Device, GetMetalDeviceContext().GetDevice(), SafeGetRuntimeDebuggingLevel() >= EMetalDebugLevelValidation, NewBuffer(FullSize, Options));
+	check(ParentBuffer.GetPtr() && ParentBuffer.GetLength() >= FullSize);
+#if STATS || ENABLE_LOW_LEVEL_MEM_TRACKER
+	MetalLLM::LogAllocBuffer(GetMetalDeviceContext().GetDevice(), ParentBuffer);
+#endif
+	INC_MEMORY_STAT_BY(STAT_MetalBufferUnusedMemory, FullSize);
+	INC_MEMORY_STAT_BY(STAT_MetalHeapBufferUnusedMemory, FullSize);
+}
+
+FMetalSubBufferLinear::~FMetalSubBufferLinear()
+{
+	DEC_MEMORY_STAT_BY(STAT_MetalBufferUnusedMemory, ParentBuffer.GetLength());
+	DEC_MEMORY_STAT_BY(STAT_MetalHeapBufferUnusedMemory, ParentBuffer.GetLength());
+}
+
+void FMetalSubBufferLinear::FreeRange(ns::Range const& Range)
+{
+#if METAL_DEBUG_OPTIONS
+	if (GIsRHIInitialized)
+	{
+		MTLPP_VALIDATE_ONLY(mtlpp::Buffer, ParentBuffer, SafeGetRuntimeDebuggingLevel() >= EMetalDebugLevelValidation, ReleaseRange(Range));
+		FMetalBuffer Buf(ParentBuffer.NewBuffer(Range), false);
+		GetMetalDeviceContext().ValidateIsInactiveBuffer(Buf);
+	}
+#endif
+	
+	FScopeLock Lock(&PoolMutex);
+	{
+		FreedSize += Range.Length;
+		INC_MEMORY_STAT_BY(STAT_MetalBufferUnusedMemory, Range.Length);
+		INC_MEMORY_STAT_BY(STAT_MetalHeapBufferUnusedMemory, Range.Length);
+		if (FreedSize == UsedSize)
+		{
+			UsedSize = 0;
+			FreedSize = 0;
+			WriteHead = 0;
+		}
+	}
+}
+
+ns::String FMetalSubBufferLinear::GetLabel() const
+{
+	return ParentBuffer.GetLabel();
+}
+
+mtlpp::Device FMetalSubBufferLinear::GetDevice() const
+{
+	return ParentBuffer.GetDevice();
+}
+
+mtlpp::StorageMode FMetalSubBufferLinear::GetStorageMode() const
+{
+	return ParentBuffer.GetStorageMode();
+}
+
+mtlpp::CpuCacheMode FMetalSubBufferLinear::GetCpuCacheMode() const
+{
+	return ParentBuffer.GetCpuCacheMode();
+}
+
+NSUInteger FMetalSubBufferLinear::GetSize() const
+{
+	return ParentBuffer.GetLength();
+}
+
+NSUInteger FMetalSubBufferLinear::GetUsedSize() const
+{
+	return UsedSize;
+}
+
+void FMetalSubBufferLinear::SetLabel(const ns::String& label)
+{
+	ParentBuffer.SetLabel(label);
+}
+
+bool FMetalSubBufferLinear::CanAllocateSize(NSUInteger Size) const
+{
+	if (WriteHead < GetSize())
+	{
+		NSUInteger Alignment = FMath::Max(NSUInteger(MinAlign), NSUInteger(Size & ~(Size - 1llu)));
+		NSUInteger NewWriteHead = Align(WriteHead, Alignment);
+		return (GetSize() - NewWriteHead) > Size;
+	}
+	else
+	{
+		return 0;
+	}
+}
+
+FMetalBuffer FMetalSubBufferLinear::NewBuffer(NSUInteger length)
+{
+	FScopeLock Lock(&PoolMutex);
+	NSUInteger Alignment = FMath::Max(NSUInteger(MinAlign), NSUInteger(length & ~(length - 1llu)));
+	NSUInteger Size = Align(length, Alignment);
+	NSUInteger NewWriteHead = Align(WriteHead, Alignment);
+	
+	FMetalBuffer Result;
+	if ((GetSize() - NewWriteHead) > Size)
+	{
+		ns::Range Range(NewWriteHead, Size);
+		DEC_MEMORY_STAT_BY(STAT_MetalBufferUnusedMemory, Range.Length);
+		DEC_MEMORY_STAT_BY(STAT_MetalHeapBufferUnusedMemory, Range.Length);
+		Result = FMetalBuffer(MTLPP_VALIDATE(mtlpp::Buffer, ParentBuffer, SafeGetRuntimeDebuggingLevel() >= EMetalDebugLevelValidation, NewBuffer(Range)), this);
+		UsedSize += Size;
+		WriteHead = NewWriteHead + Size;
+	}
+	
+	return Result;
+}
+
+mtlpp::PurgeableState FMetalSubBufferLinear::SetPurgeableState(mtlpp::PurgeableState state)
+{
+	return ParentBuffer.SetPurgeableState(state);
+}
+
+#pragma mark --
+
 FMetalSubBufferMagazine::FMetalSubBufferMagazine(NSUInteger Size, NSUInteger ChunkSize, mtlpp::ResourceOptions Options)
 : MinAlign(ChunkSize)
 , UsedSize(0)
@@ -908,9 +1062,33 @@ FMetalBuffer FMetalResourceHeap::CreateBuffer(uint32 Size, uint32 Alignment, mtl
 			case mtlpp::StorageMode::Managed:
 			{
 				FScopeLock Lock(&Mutex);
-				Buffer = ManagedBuffers.CreatePooledResource(FMetalPooledBufferArgs(Queue->GetDevice(), BlockSize, StorageMode));
-				DEC_MEMORY_STAT_BY(STAT_MetalBufferUnusedMemory, Buffer.GetLength());
-				DEC_MEMORY_STAT_BY(STAT_MetalPooledBufferUnusedMemory, Buffer.GetLength());
+				
+				if (!bForceUnique && BlockSize <= HeapSizes[NumHeapSizes - 1])
+				{
+					FMetalSubBufferLinear* Found = nullptr;
+					for (FMetalSubBufferLinear* Heap : ManagedSubHeaps)
+					{
+						if (Heap->CanAllocateSize(BlockSize))
+						{
+							Found = Heap;
+							break;
+						}
+					}
+					if (!Found)
+					{
+						Found = new FMetalSubBufferLinear(HeapAllocSizes[NumHeapSizes - 1], BufferOffsetAlignment, mtlpp::ResourceOptions((NSUInteger)Options & mtlpp::ResourceStorageModeMask), Mutex);
+						ManagedSubHeaps.Add(Found);
+					}
+					check(Found);
+					
+					return Found->NewBuffer(BlockSize);
+				}
+				else
+				{
+					Buffer = ManagedBuffers.CreatePooledResource(FMetalPooledBufferArgs(Queue->GetDevice(), BlockSize, StorageMode));
+					DEC_MEMORY_STAT_BY(STAT_MetalBufferUnusedMemory, Buffer.GetLength());
+					DEC_MEMORY_STAT_BY(STAT_MetalPooledBufferUnusedMemory, Buffer.GetLength());
+				}
 				break;
 			}
 	#endif
@@ -1113,6 +1291,15 @@ void FMetalResourceHeap::Compact(bool const bForce)
 	Buffers[AllocPrivate].DrainPool(bForce);
 #if PLATFORM_MAC
 	ManagedBuffers.DrainPool(bForce);
+	for (auto It = ManagedSubHeaps.CreateIterator(); It; ++It)
+	{
+		FMetalSubBufferLinear* Data = *It;
+		if (Data->GetUsedSize() == 0 || bForce)
+		{
+			It.RemoveCurrent();
+			delete Data;
+		}
+	}
 #endif
 	TexturePool.Drain(bForce);
 	TargetPool.Drain(bForce);
diff --git a/Engine/Source/Runtime/Apple/MetalRHI/Private/MetalBuffer.h b/Engine/Source/Runtime/Apple/MetalRHI/Private/MetalBuffer.h
--- a/Engine/Source/Runtime/Apple/MetalRHI/Private/MetalBuffer.h
+++ b/Engine/Source/Runtime/Apple/MetalRHI/Private/MetalBuffer.h
@@ -53,6 +53,35 @@ private:
 	TArray<ns::Range> FreeRanges;
 };
 
+class FMetalSubBufferLinear
+{
+public:
+	FMetalSubBufferLinear(NSUInteger Size, NSUInteger Alignment, mtlpp::ResourceOptions, FCriticalSection& PoolMutex);
+	~FMetalSubBufferLinear();
+	
+	ns::String   GetLabel() const;
+	mtlpp::Device       GetDevice() const;
+	mtlpp::StorageMode  GetStorageMode() const;
+	mtlpp::CpuCacheMode GetCpuCacheMode() const;
+	NSUInteger     GetSize() const;
+	NSUInteger     GetUsedSize() const;
+	bool	 CanAllocateSize(NSUInteger Size) const;
+	
+	void SetLabel(const ns::String& label);
+	
+	FMetalBuffer NewBuffer(NSUInteger length);
+	mtlpp::PurgeableState SetPurgeableState(mtlpp::PurgeableState state);
+	void FreeRange(ns::Range const& Range);
+	
+private:
+	FCriticalSection& PoolMutex;
+	NSUInteger MinAlign;
+	NSUInteger WriteHead;
+	NSUInteger UsedSize;
+	NSUInteger FreedSize;
+	mtlpp::Buffer ParentBuffer;
+};
+
 class FMetalSubBufferMagazine
 {
 public:
@@ -371,6 +400,7 @@ private:
 	FMetalBufferPool Buffers[NumAllocTypes];
 #if PLATFORM_MAC // All managed buffers are bucketed & pooled rather than sub-allocated to avoid memory consistency complexities
 	FMetalBufferPool ManagedBuffers;
+	TArray<FMetalSubBufferLinear*> ManagedSubHeaps;
 #endif
 	/** Anything else is just allocated directly from the device! */
 	
diff --git a/Engine/Source/Runtime/Apple/MetalRHI/Private/MetalVertexBuffer.cpp b/Engine/Source/Runtime/Apple/MetalRHI/Private/MetalVertexBuffer.cpp
--- a/Engine/Source/Runtime/Apple/MetalRHI/Private/MetalVertexBuffer.cpp
+++ b/Engine/Source/Runtime/Apple/MetalRHI/Private/MetalVertexBuffer.cpp
@@ -120,18 +120,14 @@ FMetalRHIBuffer::FMetalRHIBuffer(uint32 InSize, uint32 InUsage, ERHIResourceType
 		
 		if (InUsage & EMetalBufferUsage_LinearTex)
 		{
-			if (FMetalCommandQueue::SupportsFeature(EMetalFeaturesLinearTextures) && (InUsage & (BUF_UnorderedAccess|BUF_ShaderResource)))
-			{
-				AllocSize = Align(InSize, 1024);
-			}
-			
 			if ((InUsage & BUF_UnorderedAccess) && ((InSize - AllocSize) < 512))
 			{
 				// Padding for write flushing when not using linear texture bindings for buffers
 				AllocSize = Align(AllocSize + 512, 1024);
 			}
 			
-			if (FMetalCommandQueue::SupportsFeature(EMetalFeaturesLinearTextures) && (InUsage & (BUF_UnorderedAccess)))
+			if ((FMetalCommandQueue::SupportsFeature(EMetalFeaturesLinearTextures) && (InUsage & (BUF_ShaderResource)))
+			|| (FMetalCommandQueue::SupportsFeature(EMetalFeaturesLinearTextureUAVs) && (InUsage & (BUF_UnorderedAccess))))
 			{
 				uint32 NumElements = AllocSize;
 				uint32 SizeX = NumElements;
@@ -155,6 +151,8 @@ FMetalRHIBuffer::FMetalRHIBuffer(uint32 InSize, uint32 InUsage, ERHIResourceType
 						SizeX = NumElements;
 					}
 				}
+				
+				AllocSize = Align(AllocSize, 1024);
 			}
 		}
 		
diff --git a/Engine/Source/Runtime/Apple/MetalRHI/Public/MetalResources.h b/Engine/Source/Runtime/Apple/MetalRHI/Public/MetalResources.h
--- a/Engine/Source/Runtime/Apple/MetalRHI/Public/MetalResources.h
+++ b/Engine/Source/Runtime/Apple/MetalRHI/Public/MetalResources.h
@@ -333,15 +333,17 @@ private:
 };
 
 class FMetalSubBufferHeap;
+class FMetalSubBufferLinear;
 class FMetalSubBufferMagazine;
 
 class FMetalBuffer : public mtlpp::Buffer
 {
 public:
-	FMetalBuffer(ns::Ownership retain = ns::Ownership::Retain) : mtlpp::Buffer(retain), Heap(nullptr), Magazine(nullptr), bPooled(false) { }
+	FMetalBuffer(ns::Ownership retain = ns::Ownership::Retain) : mtlpp::Buffer(retain), Heap(nullptr), Linear(nullptr), Magazine(nullptr), bPooled(false) { }
 	FMetalBuffer(ns::Protocol<id<MTLBuffer>>::type handle, ns::Ownership retain = ns::Ownership::Retain);
 	
 	FMetalBuffer(mtlpp::Buffer&& rhs, FMetalSubBufferHeap* heap);
+	FMetalBuffer(mtlpp::Buffer&& rhs, FMetalSubBufferLinear* heap);
 	FMetalBuffer(mtlpp::Buffer&& rhs, FMetalSubBufferMagazine* magazine);
 	FMetalBuffer(mtlpp::Buffer&& rhs, bool bInPooled);
 	
@@ -369,6 +371,7 @@ public:
 	
 private:
 	FMetalSubBufferHeap* Heap;
+	FMetalSubBufferLinear* Linear;
 	FMetalSubBufferMagazine* Magazine;
 	bool bPooled;
 	bool bSingleUse;
