Duplicate 4072923 & 4081236 & 4082081 & 4084046 & 4084647

- Fixes LLM double-counting in MetalRHI.
- MetalProfiler improvements.

diff --git a/Engine/Source/Runtime/Apple/MetalRHI/Private/MetalLLM.cpp b/Engine/Source/Runtime/Apple/MetalRHI/Private/MetalLLM.cpp
--- a/Engine/Source/Runtime/Apple/MetalRHI/Private/MetalLLM.cpp
+++ b/Engine/Source/Runtime/Apple/MetalRHI/Private/MetalLLM.cpp
@@ -200,7 +200,6 @@ void MetalLLM::LogAllocTexture(mtlpp::Device& Device, mtlpp::TextureDescriptor c
 	INC_MEMORY_STAT_BY(STAT_MetalTextureMemory, Size);
 	INC_DWORD_STAT(STAT_MetalTextureCount);
 	
-	LLM(FLowLevelMemTracker::Get().OnLowLevelAlloc(ELLMTracker::Default, Ptr, Size, ELLMTag::Untagged, ELLMAllocType::System));
 	LLM(FLowLevelMemTracker::Get().OnLowLevelAlloc(ELLMTracker::Platform, Ptr, Size, ELLMTag::Untagged, ELLMAllocType::System));
 	// Assign a dealloc handler to untrack the memory - but don't track the dispatch block!
 	{
@@ -208,10 +207,8 @@ void MetalLLM::LogAllocTexture(mtlpp::Device& Device, mtlpp::TextureDescriptor c
 		
 		objc_setAssociatedObject(Texture.GetPtr(), (void*)&MetalLLM::LogAllocTexture,
 		[[[FMetalDeallocHandler alloc] initWithBlock:^{
-			LLM_SCOPE_METAL(ELLMTagMetal::Textures);
 			LLM_PLATFORM_SCOPE_METAL(ELLMTagMetal::Textures);
 			
-			LLM(FLowLevelMemTracker::Get().OnLowLevelFree(ELLMTracker::Default, Ptr, ELLMAllocType::System));
 			LLM(FLowLevelMemTracker::Get().OnLowLevelFree(ELLMTracker::Platform, Ptr, ELLMAllocType::System));
 			
 			DEC_MEMORY_STAT_BY(STAT_MetalTextureMemory, Size);
@@ -229,7 +226,6 @@ void MetalLLM::LogAllocBuffer(mtlpp::Device& Device, mtlpp::Buffer const& Buffer
 	INC_MEMORY_STAT_BY(STAT_MetalBufferMemory, Size);
 	INC_DWORD_STAT(STAT_MetalBufferCount);
 	
-	LLM(FLowLevelMemTracker::Get().OnLowLevelAlloc(ELLMTracker::Default, Ptr, Size, ELLMTag::Untagged, ELLMAllocType::System));
 	LLM(FLowLevelMemTracker::Get().OnLowLevelAlloc(ELLMTracker::Platform, Ptr, Size, ELLMTag::Untagged, ELLMAllocType::System));
 	// Assign a dealloc handler to untrack the memory - but don't track the dispatch block!
 	{
@@ -237,10 +233,8 @@ void MetalLLM::LogAllocBuffer(mtlpp::Device& Device, mtlpp::Buffer const& Buffer
 		
 		objc_setAssociatedObject(Buffer.GetPtr(), (void*)&MetalLLM::LogAllocBuffer,
 		[[[FMetalDeallocHandler alloc] initWithBlock:^{
-			LLM_SCOPE_METAL(ELLMTagMetal::Buffers);
 			LLM_PLATFORM_SCOPE_METAL(ELLMTagMetal::Buffers);
 			
-			LLM(FLowLevelMemTracker::Get().OnLowLevelFree(ELLMTracker::Default, Ptr, ELLMAllocType::System));
 			LLM(FLowLevelMemTracker::Get().OnLowLevelFree(ELLMTracker::Platform, Ptr, ELLMAllocType::System));
 			
 			DEC_MEMORY_STAT_BY(STAT_MetalBufferMemory, Size);
diff --git a/Engine/Source/Runtime/Apple/MetalRHI/Private/MetalPipeline.cpp b/Engine/Source/Runtime/Apple/MetalRHI/Private/MetalPipeline.cpp
--- a/Engine/Source/Runtime/Apple/MetalRHI/Private/MetalPipeline.cpp
+++ b/Engine/Source/Runtime/Apple/MetalRHI/Private/MetalPipeline.cpp
@@ -67,6 +67,8 @@ APPLE_PLATFORM_OBJECT_ALLOC_OVERRIDES(FMetalShaderPipeline)
 	{
 		RenderPipelineReflection = mtlpp::RenderPipelineReflection(nil);
 		ComputePipelineReflection = mtlpp::ComputePipelineReflection(nil);
+		RenderDesc = mtlpp::RenderPipelineDescriptor(nil);
+		ComputeDesc = mtlpp::ComputePipelineDescriptor(nil);
 	}
 	return Self;
 }
@@ -497,7 +499,7 @@ static FMetalShaderPipeline* CreateMTLRenderPipeline(bool const bSync, FMetalGra
 #if ENABLE_METAL_GPUPROFILE
 			ns::String VertexName = vertexFunction.GetName();
 			ns::String FragmentName = fragmentFunction ? fragmentFunction.GetName() : @"";
-			RenderPipelineDesc.SetLabel([NSString stringWithFormat:@"%@ : %@", VertexName.GetPtr(), FragmentName.GetPtr()]);
+			RenderPipelineDesc.SetLabel([NSString stringWithFormat:@"%@+%@", VertexName.GetPtr(), FragmentName.GetPtr()]);
 #endif
         }
         else
@@ -519,7 +521,7 @@ static FMetalShaderPipeline* CreateMTLRenderPipeline(bool const bSync, FMetalGra
 			{
 				ns::String VertexName = domainFunction.GetName();
 				ns::String FragmentName = fragmentFunction ? fragmentFunction.GetName() : @"";
-				RenderPipelineDesc.SetLabel([NSString stringWithFormat:@"%@ : %@", VertexName.GetPtr(), FragmentName.GetPtr()]);
+				RenderPipelineDesc.SetLabel([NSString stringWithFormat:@"%@+%@", VertexName.GetPtr(), FragmentName.GetPtr()]);
 			}
 #endif
             
@@ -675,7 +677,7 @@ static FMetalShaderPipeline* CreateMTLRenderPipeline(bool const bSync, FMetalGra
 				if (GetMetalDeviceContext().GetCommandQueue().GetRuntimeDebuggingLevel() >= EMetalDebugLevelFastValidation METAL_STATISTIC(|| GetMetalDeviceContext().GetCommandQueue().GetStatistics()))
 				{
 					mtlpp::AutoReleasedComputePipelineReflection Reflection;
-					ComputeOption = mtlpp::PipelineOption::ArgumentInfo|mtlpp::PipelineOption::BufferTypeInfo;
+					ComputeOption = mtlpp::PipelineOption::ArgumentInfo|mtlpp::PipelineOption::BufferTypeInfo METAL_STATISTIC(|NSUInteger(EMTLPipelineStats));
 					Pipeline->ComputePipelineState = Device.NewComputePipelineState(ComputePipelineDesc, (mtlpp::PipelineOption)ComputeOption, &Reflection, &AutoError);
 					Pipeline->ComputePipelineReflection = Reflection;
 				}
@@ -697,6 +699,8 @@ static FMetalShaderPipeline* CreateMTLRenderPipeline(bool const bSync, FMetalGra
 #if METAL_DEBUG_OPTIONS
 				if (Pipeline->ComputePipelineReflection)
 				{
+					Pipeline->ComputeDesc = ComputePipelineDesc;
+					
 					bool found__HSTFOut = false;
 					for(mtlpp::Argument arg : Pipeline->ComputePipelineReflection.GetArguments())
 					{
@@ -798,7 +802,7 @@ static FMetalShaderPipeline* CreateMTLRenderPipeline(bool const bSync, FMetalGra
 		Reflection = &OutReflection;
         if (GetMetalDeviceContext().GetCommandQueue().GetRuntimeDebuggingLevel() >= EMetalDebugLevelFastValidation METAL_STATISTIC(|| GetMetalDeviceContext().GetCommandQueue().GetStatistics()))
         {
-        	RenderOption = mtlpp::PipelineOption::ArgumentInfo|mtlpp::PipelineOption::BufferTypeInfo;
+        	RenderOption = mtlpp::PipelineOption::ArgumentInfo|mtlpp::PipelineOption::BufferTypeInfo METAL_STATISTIC(|NSUInteger(EMTLPipelineStats));
         }
 #endif
 
@@ -810,6 +814,7 @@ static FMetalShaderPipeline* CreateMTLRenderPipeline(bool const bSync, FMetalGra
 			if (Reflection)
 			{
 				Pipeline->RenderPipelineReflection = *Reflection;
+				Pipeline->RenderDesc = RenderPipelineDesc;
 			}
 #endif
 			Error = RenderError;
diff --git a/Engine/Source/Runtime/Apple/MetalRHI/Private/MetalPipeline.h b/Engine/Source/Runtime/Apple/MetalRHI/Private/MetalPipeline.h
--- a/Engine/Source/Runtime/Apple/MetalRHI/Private/MetalPipeline.h
+++ b/Engine/Source/Runtime/Apple/MetalRHI/Private/MetalPipeline.h
@@ -81,6 +81,8 @@ struct FMetalTessellationPipelineDesc
 	ns::String VertexSource;
 	ns::String FragmentSource;
 	ns::String ComputeSource;
+	mtlpp::RenderPipelineDescriptor RenderDesc;
+	mtlpp::ComputePipelineDescriptor ComputeDesc;
 #endif
 }
 #if METAL_DEBUG_OPTIONS
diff --git a/Engine/Source/Runtime/Apple/MetalRHI/Private/MetalProfiler.cpp b/Engine/Source/Runtime/Apple/MetalRHI/Private/MetalProfiler.cpp
--- a/Engine/Source/Runtime/Apple/MetalRHI/Private/MetalProfiler.cpp
+++ b/Engine/Source/Runtime/Apple/MetalRHI/Private/MetalProfiler.cpp
@@ -7,6 +7,10 @@
 #include "MetalCommandBuffer.h"
 #include "HAL/FileManager.h"
 
+// The Metal standard library extensions we need for UE4.
+extern unsigned int ue4_stdlib_metal_len;
+extern unsigned char ue4_stdlib_metal[];
+
 DEFINE_STAT(STAT_MetalUniformMemAlloc);
 DEFINE_STAT(STAT_MetalUniformMemFreed);
 DEFINE_STAT(STAT_MetalVertexMemAlloc);
@@ -461,12 +465,26 @@ FString IMetalStatsScope::GetJSONRepresentation(uint32 Pid)
 			else
 			{
 				FString CustomCounters;
+				TMap<FString, FMetalProfiler::EMTLCounterType> const& CounterTypes = FMetalProfiler::GetProfiler()->GetCounterTypes();
 				for(auto const& Pair : DrawStat.Counters)
 				{
 					NSString* CounterName = Pair.Key;
 					TPair<uint64, uint64> Vals = Pair.Value;
-					
-					CustomCounters += FString::Printf(TEXT(",\"%s\":\"%llu:%llu\""), *FString(CounterName), Vals.Key, Vals.Value);
+					const FMetalProfiler::EMTLCounterType* TypePtr = CounterTypes.Find(FString(CounterName));
+					FMetalProfiler::EMTLCounterType Type = TypePtr ? *TypePtr : FMetalProfiler::EMTLCounterTypeStartEnd;
+					switch(Type)
+					{
+						case FMetalProfiler::EMTLCounterTypeLast:
+							CustomCounters += FString::Printf(TEXT(",\"%s\":%llu"), *FString(CounterName), Vals.Value);
+							break;
+						case FMetalProfiler::EMTLCounterTypeDifference:
+							CustomCounters += FString::Printf(TEXT(",\"%s\":%llu"), *FString(CounterName), Vals.Value - Vals.Key);
+							break;
+						case FMetalProfiler::EMTLCounterTypeStartEnd:
+						default:
+							CustomCounters += FString::Printf(TEXT(",\"%s\":\"%llu:%llu\""), *FString(CounterName), Vals.Key, Vals.Value);
+							break;
+					}
 				}
 				
 				JSONOutput += FString::Printf(TEXT("{\"pid\":%d, \"tid\":%d, \"ph\": \"X\", \"name\": \"%s\", \"ts\": %llu, \"dur\": %llu, \"args\":{\"num_child\":%u,\"shade_cost\":%llu,\"rhi_prims\":%llu,\"ia_prims\":%llu,\"rhi_verts\":%llu,\"ia_verts\":%llu,\"vert_invoc\":%llu,\"vert_percent\":%llu,\"clip_invoc\":%llu,\"clip_prims\":%llu,\"frag_invoc\":%llu,\"frag_percent\":%llu,\"comp_invoc\":%llu,\"comp_percent\":%llu %s}},\n"),
@@ -589,11 +607,18 @@ FMetalShaderPipelineStats::FMetalShaderPipelineStats(FMetalShaderPipeline* Pipel
 	Pipeline = PipelineStat;
 	check(Pipeline);
 	
+	CmdBufferStats = nullptr;
+	
 	StartSample = nullptr;
 	
 	if (Pipeline->RenderPipelineState)
 	{
 		Name = Pipeline->RenderPipelineState.GetLabel().GetPtr();
+		
+		if (Pipeline->ComputePipelineState)
+		{
+			Name += TEXT("+") + FString(Pipeline->ComputePipelineState.GetLabel().GetPtr());
+		}
 	}
 	else if (Pipeline->ComputePipelineState)
 	{
@@ -621,8 +646,9 @@ FMetalShaderPipelineStats::~FMetalShaderPipelineStats()
 
 void FMetalShaderPipelineStats::Start(mtlpp::CommandBuffer const& Buffer)
 {
+	check(CmdBufferStats);
 	IMetalStatistics* Stats = FMetalProfiler::GetStatistics();
-	StartSample = [Stats->RegisterEncoderStatistics(Buffer.GetPtr(), EMetalSamplePipelineChange) retain];
+	StartSample = [Stats->RegisterEncoderStatistics(CmdBufferStats, EMetalSamplePipelineChange) retain];
 }
 	
 void FMetalShaderPipelineStats::End(mtlpp::CommandBuffer const& Buffer)
@@ -632,8 +658,6 @@ void FMetalShaderPipelineStats::End(mtlpp::CommandBuffer const& Buffer)
 	
 void FMetalShaderPipelineStats::GetStats(FMetalPipelineStats& PipelineStats)
 {
-	
-	
 	IMetalStatistics* Stats = FMetalProfiler::GetStatistics();
 	if (!GPUStartTime && !GPUEndTime)
 	{
@@ -674,11 +698,14 @@ void FMetalShaderPipelineStats::GetStats(FMetalPipelineStats& PipelineStats)
 			PipelineStats.PSOPerformanceStats = [NSDictionary dictionaryWithObject:Dict forKey:@"Compute Shader"];
 		}
 	}
+	
+	FMetalProfiler::GetProfiler()->DumpPipeline(Pipeline);
 }
 
 FMetalOperationStats::FMetalOperationStats(char const* DrawCall, uint64 InGPUThreadIndex, uint32 InStartPoint, uint32 InEndPoint, uint32 InRHIPrimitives, uint32 InRHIVertices, uint32 InRHIInstances)
 {
 	Name = DrawCall;
+	CmdBufferStats = nullptr;
 	
 	CPUThreadIndex = FPlatformTLS::GetCurrentThreadId();
 	GPUThreadIndex = InGPUThreadIndex;
@@ -700,6 +727,7 @@ FMetalOperationStats::FMetalOperationStats(char const* DrawCall, uint64 InGPUThr
 FMetalOperationStats::FMetalOperationStats(char const* DrawCall, uint64 InGPUThreadIndex, uint32 InStartPoint, uint32 InEndPoint)
 {
 	Name = DrawCall;
+	CmdBufferStats = nullptr;
 	
 	CPUThreadIndex = FPlatformTLS::GetCurrentThreadId();
 	GPUThreadIndex = InGPUThreadIndex;
@@ -726,8 +754,9 @@ FMetalOperationStats::~FMetalOperationStats()
 void FMetalOperationStats::Start(mtlpp::CommandBuffer const& Buffer)
 {
 	check(!DrawStats);
+	check(CmdBufferStats);
 	IMetalStatistics* Stats = FMetalProfiler::GetStatistics();
-	DrawStats = Stats->CreateDrawStats(Buffer.GetPtr(), (EMetalSamples)StartPoint, (EMetalSamples)EndPoint, RHIPrimitives, RHIVertices);
+	DrawStats = Stats->CreateDrawStats(CmdBufferStats, (EMetalSamples)StartPoint, (EMetalSamples)EndPoint, RHIPrimitives, RHIVertices);
 	check(DrawStats);
 }
 	
@@ -750,6 +779,7 @@ void FMetalOperationStats::GetStats(FMetalPipelineStats& PipelineStats)
 FMetalEncoderStats::FMetalEncoderStats(mtlpp::RenderCommandEncoder const& Encoder, uint64 InGPUThreadIndex)
 {
 	CmdBuffer = nullptr;
+	CmdBufferStats = nullptr;
 	
 	StartPoint = EMetalSampleRenderEncoderStart;
 	EndPoint = EMetalSampleRenderEncoderEnd;
@@ -771,6 +801,7 @@ FMetalEncoderStats::FMetalEncoderStats(mtlpp::RenderCommandEncoder const& Encode
 FMetalEncoderStats::FMetalEncoderStats(mtlpp::BlitCommandEncoder const& Encoder, uint64 InGPUThreadIndex)
 {
 	CmdBuffer = nullptr;
+	CmdBufferStats = nullptr;
 	
 	StartPoint = EMetalSampleBlitEncoderStart;
 	EndPoint = EMetalSampleBlitEncoderEnd;
@@ -792,6 +823,7 @@ FMetalEncoderStats::FMetalEncoderStats(mtlpp::BlitCommandEncoder const& Encoder,
 FMetalEncoderStats::FMetalEncoderStats(mtlpp::ComputeCommandEncoder const& Encoder, uint64 InGPUThreadIndex)
 {
 	CmdBuffer = nullptr;
+	CmdBufferStats = nullptr;
 	
 	StartPoint = EMetalSampleComputeEncoderStart;
 	EndPoint = EMetalSampleComputeEncoderEnd;
@@ -821,24 +853,27 @@ void FMetalEncoderStats::Start(mtlpp::CommandBuffer const& Buffer)
 	check(!StartSample);
 	check(!CmdBuffer);
 	check(Buffer);
+	check(CmdBufferStats);
 	CmdBuffer = Buffer;
 	IMetalStatistics* Stats = FMetalProfiler::GetStatistics();
-	StartSample = [Stats->RegisterEncoderStatistics(Buffer.GetPtr(), (EMetalSamples)StartPoint) retain];
+	StartSample = [Stats->RegisterEncoderStatistics(CmdBufferStats, (EMetalSamples)StartPoint) retain];
 }
 
 void FMetalEncoderStats::End(mtlpp::CommandBuffer const& Buffer)
 {
 	check(!EndSample);
 	check(Buffer.GetPtr() == CmdBuffer.GetPtr());
+	check(CmdBufferStats);
 	CPUEndTime = FPlatformTime::ToMilliseconds64(mach_absolute_time()) * 1000.0;
 	IMetalStatistics* Stats = FMetalProfiler::GetStatistics();
-	EndSample = [Stats->RegisterEncoderStatistics(Buffer.GetPtr(), (EMetalSamples)EndPoint) retain];
+	EndSample = [Stats->RegisterEncoderStatistics(CmdBufferStats, (EMetalSamples)EndPoint) retain];
 }
 
 void FMetalEncoderStats::EncodeDraw(char const* DrawCall, uint32 RHIPrimitives, uint32 RHIVertices, uint32 RHIInstances)
 {
 	check(CmdBuffer);
 	FMetalOperationStats* Draw = new FMetalOperationStats(DrawCall, GPUThreadIndex, EMetalSampleBeforeDraw, EMetalSampleAfterDraw, RHIPrimitives, RHIVertices, RHIInstances);
+	Draw->CmdBufferStats = CmdBufferStats;
 	Children.Add(Draw);
 	Draw->Start(CmdBuffer);
 	Draw->End(CmdBuffer);
@@ -848,6 +883,7 @@ void FMetalEncoderStats::EncodeBlit(char const* DrawCall)
 {
 	check(CmdBuffer);
 	FMetalOperationStats* Draw = new FMetalOperationStats(DrawCall, GPUThreadIndex, EMetalSampleBeforeBlit, EMetalSampleAfterBlit);
+	Draw->CmdBufferStats = CmdBufferStats;
 	Children.Add(Draw);
 	Draw->Start(CmdBuffer);
 	Draw->End(CmdBuffer);
@@ -858,6 +894,7 @@ void FMetalEncoderStats::EncodeDispatch(char const* DrawCall)
 {
 	check(CmdBuffer);
 	FMetalOperationStats* Draw = new FMetalOperationStats(DrawCall, GPUThreadIndex, EMetalSampleBeforeCompute, EMetalSampleAfterCompute);
+	Draw->CmdBufferStats = CmdBufferStats;
 	Children.Add(Draw);
 	Draw->Start(CmdBuffer);
 	Draw->End(CmdBuffer);
@@ -867,6 +904,7 @@ void FMetalEncoderStats::EncodePipeline(FMetalShaderPipeline* PipelineStat)
 {
 	check(CmdBuffer);
 	FMetalShaderPipelineStats* Draw = new FMetalShaderPipelineStats(PipelineStat, GPUThreadIndex);
+	Draw->CmdBufferStats = CmdBufferStats;
 	Children.Add(Draw);
 	Draw->Start(CmdBuffer);
 	Draw->End(CmdBuffer);
@@ -887,7 +925,13 @@ FMetalCommandBufferStats::FMetalCommandBufferStats(mtlpp::CommandBuffer const& B
 {
 	CmdBuffer = Buffer;
 #if METAL_STATISTICS
+	CmdBufferStats = nullptr;
 	ActiveEncoderStats = nullptr;
+	IMetalStatistics* Stats = FMetalProfiler::GetStatistics();
+	if (Stats)
+	{
+		CmdBufferStats = Stats->BeginCommandBufferStatistics(CmdBuffer.GetPtr());
+	}
 #endif
 	
 	Name = FString::Printf(TEXT("CommandBuffer: %s"), *FString(CmdBuffer.GetLabel().GetPtr()));
@@ -902,6 +946,8 @@ FMetalCommandBufferStats::~FMetalCommandBufferStats()
 {
 #if METAL_STATISTICS
 	check(!ActiveEncoderStats);
+	[CmdBufferStats release];
+	CmdBufferStats = nil;
 #endif
 }
 
@@ -961,6 +1007,7 @@ void FMetalCommandBufferStats::BeginEncoder(mtlpp::RenderCommandEncoder const& E
 	check(!ActiveEncoderStats);
 	
 	ActiveEncoderStats = new FMetalEncoderStats(Encoder, GPUThreadIndex);
+	ActiveEncoderStats->CmdBufferStats = CmdBufferStats;
 	Children.Add(ActiveEncoderStats);
 	ActiveEncoderStats->Start(CmdBuffer);
 }
@@ -970,6 +1017,7 @@ void FMetalCommandBufferStats::BeginEncoder(mtlpp::BlitCommandEncoder const& Enc
 	check(!ActiveEncoderStats);
 	
 	ActiveEncoderStats = new FMetalEncoderStats(Encoder, GPUThreadIndex);
+	ActiveEncoderStats->CmdBufferStats = CmdBufferStats;
 	Children.Add(ActiveEncoderStats);
 	ActiveEncoderStats->Start(CmdBuffer);
 }
@@ -979,6 +1027,7 @@ void FMetalCommandBufferStats::BeginEncoder(mtlpp::ComputeCommandEncoder const&
 	check(!ActiveEncoderStats);
 	
 	ActiveEncoderStats = new FMetalEncoderStats(Encoder, GPUThreadIndex);
+	ActiveEncoderStats->CmdBufferStats = CmdBufferStats;
 	Children.Add(ActiveEncoderStats);
 	ActiveEncoderStats->Start(CmdBuffer);
 }
@@ -1329,12 +1378,13 @@ void FMetalProfiler::EndEncoder(FMetalCommandBufferStats* CmdBufStats, mtlpp::Co
 		CmdBufStats->EndEncoder(Encoder);
 }
 
-void FMetalProfiler::AddCounter(NSString* Counter)
+void FMetalProfiler::AddCounter(NSString* Counter, EMTLCounterType Type)
 {
 	check (StatisticsAPI);
 	if (![NewCounters containsObject:Counter])
 	{
 		[NewCounters addObject:Counter];
+		CounterTypes.Add(FString(Counter), Type);
 	}
 }
 
@@ -1342,6 +1392,12 @@ void FMetalProfiler::RemoveCounter(NSString* Counter)
 {
 	check (StatisticsAPI);
 	[NewCounters removeObject:Counter];
+	CounterTypes.Remove(FString(Counter));
+}
+
+void FMetalProfiler::DumpPipeline(FMetalShaderPipeline* PipelineStat)
+{
+	Pipelines.Add(PipelineStat);
 }
 
 #endif
@@ -1600,6 +1656,110 @@ void FMetalProfiler::SaveTrace()
 		WriteString(OutputFile, "{}]}");
 		
 		OutputFile->Close();
+		
+#if METAL_STATISTICS
+		FString OutputDir = TracingRootPath + Filename + TEXT("/Pipelines/");
+		if (Pipelines.Num())
+		{
+			FString FileName = OutputDir + TEXT("ue4_stdlib.metal");
+			FArchive* PipelineFile = IFileManager::Get().CreateFileWriter(*FileName);
+			PipelineFile->Serialize((void*)ue4_stdlib_metal, sizeof(ANSICHAR)*ue4_stdlib_metal_len);
+			PipelineFile->Close();
+		}
+		for (auto Ptr : Pipelines)
+		{
+			FString PipelineName;
+			if (Ptr->RenderPipelineState)
+			{
+				PipelineName = Ptr->RenderPipelineState.GetLabel().GetPtr();
+				if (Ptr->ComputePipelineState)
+				{
+					PipelineName += TEXT("+") + FString(Ptr->ComputePipelineState.GetLabel().GetPtr());
+				}
+			}
+			else if (Ptr->ComputePipelineState)
+			{
+				PipelineName = FString(Ptr->ComputePipelineState.GetLabel().GetPtr());
+			}
+			
+			FString FileName = OutputDir + PipelineName + TEXT(".txt");
+			FArchive* PipelineFile = IFileManager::Get().CreateFileWriter(*FileName);
+			
+			WriteString(PipelineFile, TCHAR_TO_UTF8(*PipelineName));
+			WriteString(PipelineFile, "\n");
+			
+			if (Ptr->RenderDesc)
+			{
+				WriteString(PipelineFile, "\n\n******************* Render Pipeline Descriptor:\n");
+				WriteString(PipelineFile, [[Ptr->RenderDesc.GetPtr() description] UTF8String]);
+			}
+			if (Ptr->VertexSource)
+			{
+				FString Name;
+				if (Ptr->RenderDesc)
+				{
+					Name = Ptr->RenderDesc.GetVertexFunction().GetName().GetPtr();
+					Name += TEXT(".metal");
+				}
+				else
+				{
+					Name = Ptr->RenderPipelineState.GetLabel().GetPtr();
+					Name += TEXT(".vertex.metal");
+				}
+				
+				FString ShaderName = OutputDir + Name;
+				FArchive* ShaderFile = IFileManager::Get().CreateFileWriter(*ShaderName);
+				WriteString(ShaderFile, [Ptr->VertexSource.GetPtr() UTF8String]);
+				ShaderFile->Close();
+			}
+			if (Ptr->FragmentSource)
+			{
+				FString Name;
+				if (Ptr->RenderDesc)
+				{
+					Name = Ptr->RenderDesc.GetFragmentFunction().GetName().GetPtr();
+					Name += TEXT(".metal");
+				}
+				else
+				{
+					Name = Ptr->RenderPipelineState.GetLabel().GetPtr();
+					Name += TEXT(".fragment.metal");
+				}
+				
+				FString ShaderName = OutputDir + Name;
+				FArchive* ShaderFile = IFileManager::Get().CreateFileWriter(*ShaderName);
+				WriteString(ShaderFile, [Ptr->FragmentSource.GetPtr() UTF8String]);
+				ShaderFile->Close();
+			}
+			if (Ptr->ComputeDesc)
+			{
+				WriteString(PipelineFile, "\n\n******************* Compute Pipeline Descriptor:\n");
+				WriteString(PipelineFile, [[Ptr->ComputeDesc.GetPtr() description] UTF8String]);
+			}
+			if (Ptr->ComputeSource)
+			{
+				FString Name;
+				if (Ptr->ComputeDesc)
+				{
+					Name = Ptr->ComputeDesc.GetComputeFunction().GetName().GetPtr();
+					Name += TEXT(".metal");
+				}
+				else
+				{
+					Name = Ptr->ComputePipelineState.GetLabel().GetPtr();
+					Name += TEXT(".compute.metal");
+				}
+				
+				FString ShaderName = OutputDir + Name;
+				FArchive* ShaderFile = IFileManager::Get().CreateFileWriter(*ShaderName);
+				WriteString(ShaderFile, [Ptr->ComputeSource.GetPtr() UTF8String]);
+				ShaderFile->Close();
+			}
+			
+			PipelineFile->Close();
+		}
+		Pipelines.Empty();
+#endif
 	}
 }
 
@@ -1662,7 +1822,17 @@ static void HandleMetalProfileCommand(const TArray<FString>& Args, UWorld*, FOut
 			FString NewCounter = Args[1];
 			if (![Array containsObject:NewCounter.GetNSString()])
 			{
-				FMetalProfiler::GetProfiler()->AddCounter(NewCounter.GetNSString());
+				FString TypeName = Args.Num() > 2 ? Args[2] : TEXT("");
+				FMetalProfiler::EMTLCounterType Type = FMetalProfiler::EMTLCounterTypeStartEnd;
+				if (TypeName == TEXT("LAST"))
+				{
+					Type = FMetalProfiler::EMTLCounterTypeLast;
+				}
+				else if (TypeName == TEXT("DIFF"))
+				{
+					Type = FMetalProfiler::EMTLCounterTypeDifference;
+				}
+				FMetalProfiler::GetProfiler()->AddCounter(NewCounter.GetNSString(), Type);
 			}
 		}
 	}
diff --git a/Engine/Source/Runtime/Apple/MetalRHI/Private/MetalProfiler.h b/Engine/Source/Runtime/Apple/MetalRHI/Private/MetalProfiler.h
--- a/Engine/Source/Runtime/Apple/MetalRHI/Private/MetalProfiler.h
+++ b/Engine/Source/Runtime/Apple/MetalRHI/Private/MetalProfiler.h
@@ -237,6 +237,7 @@ struct FMetalOperationStats : public IMetalStatsScope
 	virtual void End(mtlpp::CommandBuffer const& Buffer) final override;
 	virtual void GetStats(FMetalPipelineStats& PipelineStats) final override;
 	
+	id<IMetalCommandBufferStats> CmdBufferStats;
 	uint32 StartPoint;
 	uint32 EndPoint;
 	IMetalDrawStats* DrawStats;
@@ -254,6 +255,7 @@ struct FMetalShaderPipelineStats : public IMetalStatsScope
 	virtual void End(mtlpp::CommandBuffer const& Buffer) final override;
 	virtual void GetStats(FMetalPipelineStats& PipelineStats) final override;
 	
+	id<IMetalCommandBufferStats> CmdBufferStats;
 	id<IMetalStatisticsSamples> StartSample;
 	FMetalShaderPipeline* Pipeline;
 };
@@ -274,6 +276,7 @@ struct FMetalEncoderStats : public IMetalStatsScope
 	void EncodeDispatch(char const* DrawCall);
 	void EncodePipeline(FMetalShaderPipeline* PipelineStat);
 	
+	id<IMetalCommandBufferStats> CmdBufferStats;
 	ns::AutoReleased<mtlpp::CommandBuffer> CmdBuffer;
 	uint32 StartPoint;
 	uint32 EndPoint;
@@ -304,6 +307,7 @@ struct FMetalCommandBufferStats : public IMetalStatsScope
 	void EndEncoder(mtlpp::BlitCommandEncoder const& Encoder);
 	void EndEncoder(mtlpp::ComputeCommandEncoder const& Encoder);
 
+	id<IMetalCommandBufferStats> CmdBufferStats;
 	FMetalEncoderStats* ActiveEncoderStats;
 #endif
 };
@@ -397,8 +401,18 @@ public:
 	void EndEncoder(FMetalCommandBufferStats* CmdBufStats, mtlpp::BlitCommandEncoder const& Encoder);
 	void EndEncoder(FMetalCommandBufferStats* CmdBufStats, mtlpp::ComputeCommandEncoder const& Encoder);
 	
-	void AddCounter(NSString* Counter);
+	enum EMTLCounterType
+	{
+		EMTLCounterTypeStartEnd,
+		EMTLCounterTypeLast,
+		EMTLCounterTypeDifference
+	};
+	
+	void AddCounter(NSString* Counter, EMTLCounterType Type);
 	void RemoveCounter(NSString* Counter);
+	TMap<FString, EMTLCounterType> const& GetCounterTypes() const { return CounterTypes; }
+	
+	void DumpPipeline(FMetalShaderPipeline* PipelineStat);
 #endif
 	
 	FMetalCPUStats* AddCPUStat(FString const& Name);
@@ -413,9 +427,11 @@ private:
 	FCriticalSection Mutex;
 #if METAL_STATISTICS
 	NSMutableArray* NewCounters;
+	TMap<FString, EMTLCounterType> CounterTypes;
 	IMetalStatistics* StatisticsAPI;
 	TArray<FMetalEventStats*> FrameEvents;
 	TArray<FMetalEventStats*> ActiveEvents;
+	TSet<FMetalShaderPipeline*> Pipelines;
 #endif
 	
 	TArray<FMetalCommandBufferStats*> TracedBuffers;
diff --git a/Engine/Source/Runtime/Apple/MetalRHI/Private/MetalShaders.cpp b/Engine/Source/Runtime/Apple/MetalRHI/Private/MetalShaders.cpp
--- a/Engine/Source/Runtime/Apple/MetalRHI/Private/MetalShaders.cpp
+++ b/Engine/Source/Runtime/Apple/MetalRHI/Private/MetalShaders.cpp
@@ -643,12 +643,13 @@ FMetalShaderPipeline* FMetalComputeShader::GetPipeline(EPixelFormat const* const
 		
 		METAL_GPUPROFILE(FScopedMetalCPUStats CPUStat(FString::Printf(TEXT("NewComputePipeline: %d_%d"), SourceLen, SourceCRC)));
     #if METAL_DEBUG_OPTIONS
-        if (GetMetalDeviceContext().GetCommandQueue().GetRuntimeDebuggingLevel() >= EMetalDebugLevelFastValidation)
+        if (GetMetalDeviceContext().GetCommandQueue().GetRuntimeDebuggingLevel() >= EMetalDebugLevelFastValidation METAL_STATISTIC(|| GetMetalDeviceContext().GetCommandQueue().GetStatistics()))
         {
 			ns::AutoReleasedError ComputeError;
             mtlpp::AutoReleasedComputePipelineReflection ComputeReflection;
 			
-			Kernel = GetMetalDeviceContext().GetDevice().NewComputePipelineState(Func, mtlpp::PipelineOption::ArgumentInfo, ComputeReflection, &ComputeError);
+			NSUInteger ComputeOption = mtlpp::PipelineOption::ArgumentInfo|mtlpp::PipelineOption::BufferTypeInfo METAL_STATISTIC(|NSUInteger(EMTLPipelineStats));
+			Kernel = GetMetalDeviceContext().GetDevice().NewComputePipelineState(Func, mtlpp::PipelineOption(ComputeOption), ComputeReflection, &ComputeError);
 			Error = ComputeError;
 			Reflection = ComputeReflection;
         }
@@ -671,6 +672,12 @@ FMetalShaderPipeline* FMetalComputeShader::GetPipeline(EPixelFormat const* const
 #if METAL_DEBUG_OPTIONS
         Pipeline[BT]->ComputePipelineReflection = Reflection;
         Pipeline[BT]->ComputeSource = GetSourceCode();
+		if (Reflection)
+		{
+			Pipeline[BT]->ComputeDesc = mtlpp::ComputePipelineDescriptor();
+			Pipeline[BT]->ComputeDesc.SetLabel(Func.GetName());
+			Pipeline[BT]->ComputeDesc.SetComputeFunction(Func);
+		}
 #endif
         METAL_DEBUG_OPTION(FMemory::Memzero(Pipeline[BT]->ResourceMask, sizeof(Pipeline[BT]->ResourceMask)));
 	}
