Duplicate 4107661:

Refactor the MetalRHI debugging tools to move away from reimplementing the Objective-C protocols to adding separate debugging classes. This makes it easier to build on/for future OSes without having to interpose new functionality we don't yet use.

diff --git a/Engine/Source/Runtime/Apple/MetalRHI/Private/MetalBlitCommandEncoder.cpp b/Engine/Source/Runtime/Apple/MetalRHI/Private/MetalBlitCommandEncoder.cpp
--- a/Engine/Source/Runtime/Apple/MetalRHI/Private/MetalBlitCommandEncoder.cpp
+++ b/Engine/Source/Runtime/Apple/MetalRHI/Private/MetalBlitCommandEncoder.cpp
@@ -12,87 +12,95 @@
 
 extern int32 GMetalRuntimeDebugLevel;
 
-@implementation FMetalDebugBlitCommandEncoder
+@interface FMetalDebugBlitCommandEncoder : FMetalDebugCommandEncoder
+{
+	@public
+	id<MTLBlitCommandEncoder> Inner;
+	FMetalCommandBufferDebugging Buffer;
+}
+/** Initialise the wrapper with the provided command-buffer. */
+-(id)initWithEncoder:(id<MTLBlitCommandEncoder>)Encoder andCommandBuffer:(FMetalCommandBufferDebugging const&)Buffer;
 
-@synthesize Inner;
-@synthesize Buffer;
+@end
+
+@implementation FMetalDebugBlitCommandEncoder
 
 APPLE_PLATFORM_OBJECT_ALLOC_OVERRIDES(FMetalDebugBlitCommandEncoder)
 
--(id)initWithEncoder:(id<MTLBlitCommandEncoder>)Encoder andCommandBuffer:(FMetalDebugCommandBuffer*)SourceBuffer
+-(id)initWithEncoder:(id<MTLBlitCommandEncoder>)Encoder andCommandBuffer:(FMetalCommandBufferDebugging const&)SourceBuffer
 {
 	id Self = [super init];
 	if (Self)
 	{
-        Inner = [Encoder retain];
-		Buffer = [SourceBuffer retain];
+        Inner = Encoder;
+		Buffer = SourceBuffer;
 	}
 	return Self;
 }
 
 -(void)dealloc
 {
-	[Inner release];
-	[Buffer release];
 	[super dealloc];
 }
 
--(id <MTLDevice>) device
+@end
+
+FMetalBlitCommandEncoderDebugging::FMetalBlitCommandEncoderDebugging()
 {
-	return Inner.device;
+	
 }
-
--(NSString *)label
+FMetalBlitCommandEncoderDebugging::FMetalBlitCommandEncoderDebugging(mtlpp::BlitCommandEncoder& Encoder, FMetalCommandBufferDebugging& Buffer)
+: FMetalCommandEncoderDebugging((FMetalDebugCommandEncoder*)[[FMetalDebugBlitCommandEncoder alloc] initWithEncoder:Encoder.GetPtr() andCommandBuffer:Buffer])
 {
-	return Inner.label;
+	Buffer.BeginBlitCommandEncoder([NSString stringWithFormat:@"Blit: %@", Encoder.GetLabel().GetPtr()]);
+	Encoder.SetAssociatedObject((void const*)&FMetalBlitCommandEncoderDebugging::Get, (FMetalCommandEncoderDebugging const&)*this);
 }
-
--(void)setLabel:(NSString *)Text
+FMetalBlitCommandEncoderDebugging::FMetalBlitCommandEncoderDebugging(FMetalDebugCommandEncoder* handle)
+: FMetalCommandEncoderDebugging((FMetalDebugCommandEncoder*)handle)
 {
-	Inner.label = Text;
+	
 }
 
-- (void)endEncoding
+FMetalBlitCommandEncoderDebugging FMetalBlitCommandEncoderDebugging::Get(mtlpp::BlitCommandEncoder& Encoder)
 {
-    [Buffer endCommandEncoder];
-    [Inner endEncoding];
+	return Encoder.GetAssociatedObject<FMetalBlitCommandEncoderDebugging>((void const*)&FMetalBlitCommandEncoderDebugging::Get);
 }
 
-- (void)insertDebugSignpost:(NSString *)string
+void FMetalBlitCommandEncoderDebugging::InsertDebugSignpost(ns::String const& Label)
 {
-    [Buffer insertDebugSignpost:string];
-    [Inner insertDebugSignpost:string];
+	((FMetalDebugBlitCommandEncoder*)m_ptr)->Buffer.InsertDebugSignpost(Label);
 }
-
-- (void)pushDebugGroup:(NSString *)string
+void FMetalBlitCommandEncoderDebugging::PushDebugGroup(ns::String const& Group)
 {
-    [Buffer pushDebugGroup:string];
-    [Inner pushDebugGroup:string];
+	((FMetalDebugBlitCommandEncoder*)m_ptr)->Buffer.PushDebugGroup(Group);
+}
+void FMetalBlitCommandEncoderDebugging::PopDebugGroup()
+{
+	((FMetalDebugBlitCommandEncoder*)m_ptr)->Buffer.PopDebugGroup();
 }
 
-- (void)popDebugGroup
+void FMetalBlitCommandEncoderDebugging::EndEncoder()
 {
-    [Buffer popDebugGroup];
-    [Inner popDebugGroup];
+	((FMetalDebugBlitCommandEncoder*)m_ptr)->Buffer.EndCommandEncoder();
 }
 
 #if PLATFORM_MAC
-- (void)synchronizeResource:(id<MTLResource>)resource
+void FMetalBlitCommandEncoderDebugging::Synchronize(mtlpp::Resource const& resource)
 {
 #if METAL_DEBUG_OPTIONS
-	switch(Buffer->DebugLevel)
+	switch(((FMetalDebugBlitCommandEncoder*)m_ptr)->Buffer.GetPtr()->DebugLevel)
 	{
 		case EMetalDebugLevelConditionalSubmit:
 		case EMetalDebugLevelWaitForComplete:
 		case EMetalDebugLevelLogOperations:
 		{
-			[Buffer blit:[NSString stringWithFormat:@"%s", __PRETTY_FUNCTION__]];
+			((FMetalDebugBlitCommandEncoder*)m_ptr)->Buffer.Blit([NSString stringWithFormat:@"%s", __PRETTY_FUNCTION__]);
 		}
 		case EMetalDebugLevelValidation:
 		case EMetalDebugLevelResetOnBind:
 		case EMetalDebugLevelTrackResources:
 		{
-			[Buffer trackResource:resource];
+			((FMetalDebugBlitCommandEncoder*)m_ptr)->Buffer.TrackResource(resource);
 		}
 		default:
 		{
@@ -100,25 +108,24 @@ APPLE_PLATFORM_OBJECT_ALLOC_OVERRIDES(FMetalDebugBlitCommandEncoder)
 		}
 	}
 #endif
-    [Inner synchronizeResource:(id<MTLResource>)resource];
 }
 
-- (void)synchronizeTexture:(id<MTLTexture>)texture slice:(NSUInteger)slice level:(NSUInteger)level
+void FMetalBlitCommandEncoderDebugging::Synchronize(FMetalTexture const& texture, NSUInteger slice, NSUInteger level)
 {
 #if METAL_DEBUG_OPTIONS
-	switch(Buffer->DebugLevel)
+	switch(((FMetalDebugBlitCommandEncoder*)m_ptr)->Buffer.GetPtr()->DebugLevel)
 	{
 		case EMetalDebugLevelConditionalSubmit:
 		case EMetalDebugLevelWaitForComplete:
 		case EMetalDebugLevelLogOperations:
 		{
-			[Buffer blit:[NSString stringWithFormat:@"%s", __PRETTY_FUNCTION__]];
+			((FMetalDebugBlitCommandEncoder*)m_ptr)->Buffer.Blit([NSString stringWithFormat:@"%s", __PRETTY_FUNCTION__]);
 		}
 		case EMetalDebugLevelValidation:
 		case EMetalDebugLevelResetOnBind:
 		case EMetalDebugLevelTrackResources:
 		{
-			[Buffer trackResource:texture];
+			((FMetalDebugBlitCommandEncoder*)m_ptr)->Buffer.TrackResource(texture);
 		}
 		default:
 		{
@@ -126,27 +133,26 @@ APPLE_PLATFORM_OBJECT_ALLOC_OVERRIDES(FMetalDebugBlitCommandEncoder)
 		}
 	}
 #endif
-    [Inner synchronizeTexture:(id<MTLTexture>)texture slice:(NSUInteger)slice level:(NSUInteger)level];
 }
 #endif
 
-- (void)copyFromTexture:(id<MTLTexture>)sourceTexture sourceSlice:(NSUInteger)sourceSlice sourceLevel:(NSUInteger)sourceLevel sourceOrigin:(MTLOrigin)sourceOrigin sourceSize:(MTLSize)sourceSize toTexture:(id<MTLTexture>)destinationTexture destinationSlice:(NSUInteger)destinationSlice destinationLevel:(NSUInteger)destinationLevel destinationOrigin:(MTLOrigin)destinationOrigin
+void FMetalBlitCommandEncoderDebugging::Copy(FMetalTexture const& sourceTexture, NSUInteger sourceSlice, NSUInteger sourceLevel, mtlpp::Origin const& sourceOrigin, mtlpp::Size const& sourceSize, FMetalTexture const& destinationTexture, NSUInteger destinationSlice, NSUInteger destinationLevel, mtlpp::Origin const& destinationOrigin)
 {
 #if METAL_DEBUG_OPTIONS
-	switch(Buffer->DebugLevel)
+	switch(((FMetalDebugBlitCommandEncoder*)m_ptr)->Buffer.GetPtr()->DebugLevel)
 	{
 		case EMetalDebugLevelConditionalSubmit:
 		case EMetalDebugLevelWaitForComplete:
 		case EMetalDebugLevelLogOperations:
 		{
-			[Buffer blit:[NSString stringWithFormat:@"%s", __PRETTY_FUNCTION__]];
+			((FMetalDebugBlitCommandEncoder*)m_ptr)->Buffer.Blit([NSString stringWithFormat:@"%s", __PRETTY_FUNCTION__]);
 		}
 		case EMetalDebugLevelValidation:
 		case EMetalDebugLevelResetOnBind:
 		case EMetalDebugLevelTrackResources:
 		{
-			[Buffer trackResource:sourceTexture];
-			[Buffer trackResource:destinationTexture];
+			((FMetalDebugBlitCommandEncoder*)m_ptr)->Buffer.TrackResource(sourceTexture);
+			((FMetalDebugBlitCommandEncoder*)m_ptr)->Buffer.TrackResource(destinationTexture);
 		}
 		default:
 		{
@@ -154,27 +160,25 @@ APPLE_PLATFORM_OBJECT_ALLOC_OVERRIDES(FMetalDebugBlitCommandEncoder)
 		}
 	}
 #endif
-    
-    [Inner copyFromTexture:(id<MTLTexture>)sourceTexture sourceSlice:(NSUInteger)sourceSlice sourceLevel:(NSUInteger)sourceLevel sourceOrigin:(MTLOrigin)sourceOrigin sourceSize:(MTLSize)sourceSize toTexture:(id<MTLTexture>)destinationTexture destinationSlice:(NSUInteger)destinationSlice destinationLevel:(NSUInteger)destinationLevel destinationOrigin:(MTLOrigin)destinationOrigin];
 }
 
-- (void)copyFromBuffer:(id<MTLBuffer>)sourceBuffer sourceOffset:(NSUInteger)sourceOffset sourceBytesPerRow:(NSUInteger)sourceBytesPerRow sourceBytesPerImage:(NSUInteger)sourceBytesPerImage sourceSize:(MTLSize)sourceSize toTexture:(id<MTLTexture>)destinationTexture destinationSlice:(NSUInteger)destinationSlice destinationLevel:(NSUInteger)destinationLevel destinationOrigin:(MTLOrigin)destinationOrigin
+void FMetalBlitCommandEncoderDebugging::Copy(FMetalBuffer const& sourceBuffer, NSUInteger sourceOffset, NSUInteger sourceBytesPerRow, NSUInteger sourceBytesPerImage, mtlpp::Size const& sourceSize, FMetalTexture const& destinationTexture, NSUInteger destinationSlice, NSUInteger destinationLevel, mtlpp::Origin const& destinationOrigin)
 {
 #if METAL_DEBUG_OPTIONS
-	switch(Buffer->DebugLevel)
+	switch(((FMetalDebugBlitCommandEncoder*)m_ptr)->Buffer.GetPtr()->DebugLevel)
 	{
 		case EMetalDebugLevelConditionalSubmit:
 		case EMetalDebugLevelWaitForComplete:
 		case EMetalDebugLevelLogOperations:
 		{
-			[Buffer blit:[NSString stringWithFormat:@"%s", __PRETTY_FUNCTION__]];
+			((FMetalDebugBlitCommandEncoder*)m_ptr)->Buffer.Blit([NSString stringWithFormat:@"%s", __PRETTY_FUNCTION__]);
 		}
 		case EMetalDebugLevelValidation:
 		case EMetalDebugLevelResetOnBind:
 		case EMetalDebugLevelTrackResources:
 		{
-			[Buffer trackResource:sourceBuffer];
-			[Buffer trackResource:destinationTexture];
+			((FMetalDebugBlitCommandEncoder*)m_ptr)->Buffer.TrackResource(sourceBuffer);
+			((FMetalDebugBlitCommandEncoder*)m_ptr)->Buffer.TrackResource(destinationTexture);
 		}
 		default:
 		{
@@ -182,27 +186,25 @@ APPLE_PLATFORM_OBJECT_ALLOC_OVERRIDES(FMetalDebugBlitCommandEncoder)
 		}
 	}
 #endif
-    
-    [Inner copyFromBuffer:(id<MTLBuffer>)sourceBuffer sourceOffset:(NSUInteger)sourceOffset sourceBytesPerRow:(NSUInteger)sourceBytesPerRow sourceBytesPerImage:(NSUInteger)sourceBytesPerImage sourceSize:(MTLSize)sourceSize toTexture:(id<MTLTexture>)destinationTexture destinationSlice:(NSUInteger)destinationSlice destinationLevel:(NSUInteger)destinationLevel destinationOrigin:(MTLOrigin)destinationOrigin];
 }
 
-- (void)copyFromBuffer:(id<MTLBuffer>)sourceBuffer sourceOffset:(NSUInteger)sourceOffset sourceBytesPerRow:(NSUInteger)sourceBytesPerRow sourceBytesPerImage:(NSUInteger)sourceBytesPerImage sourceSize:(MTLSize)sourceSize toTexture:(id<MTLTexture>)destinationTexture destinationSlice:(NSUInteger)destinationSlice destinationLevel:(NSUInteger)destinationLevel destinationOrigin:(MTLOrigin)destinationOrigin options:(MTLBlitOption)options
+void FMetalBlitCommandEncoderDebugging::Copy(FMetalBuffer const& sourceBuffer, NSUInteger sourceOffset, NSUInteger sourceBytesPerRow, NSUInteger sourceBytesPerImage, mtlpp::Size const& sourceSize, FMetalTexture const& destinationTexture, NSUInteger destinationSlice, NSUInteger destinationLevel, mtlpp::Origin const& destinationOrigin, mtlpp::BlitOption options)
 {
 #if METAL_DEBUG_OPTIONS
-	switch(Buffer->DebugLevel)
+	switch(((FMetalDebugBlitCommandEncoder*)m_ptr)->Buffer.GetPtr()->DebugLevel)
 	{
 		case EMetalDebugLevelConditionalSubmit:
 		case EMetalDebugLevelWaitForComplete:
 		case EMetalDebugLevelLogOperations:
 		{
-			[Buffer blit:[NSString stringWithFormat:@"%s", __PRETTY_FUNCTION__]];
+			((FMetalDebugBlitCommandEncoder*)m_ptr)->Buffer.Blit([NSString stringWithFormat:@"%s", __PRETTY_FUNCTION__]);
 		}
 		case EMetalDebugLevelValidation:
 		case EMetalDebugLevelResetOnBind:
 		case EMetalDebugLevelTrackResources:
 		{
-			[Buffer trackResource:sourceBuffer];
-			[Buffer trackResource:destinationTexture];
+			((FMetalDebugBlitCommandEncoder*)m_ptr)->Buffer.TrackResource(sourceBuffer);
+			((FMetalDebugBlitCommandEncoder*)m_ptr)->Buffer.TrackResource(destinationTexture);
 		}
 		default:
 		{
@@ -210,27 +212,25 @@ APPLE_PLATFORM_OBJECT_ALLOC_OVERRIDES(FMetalDebugBlitCommandEncoder)
 		}
 	}
 #endif
-    
-    [Inner copyFromBuffer:(id<MTLBuffer>)sourceBuffer sourceOffset:(NSUInteger)sourceOffset sourceBytesPerRow:(NSUInteger)sourceBytesPerRow sourceBytesPerImage:(NSUInteger)sourceBytesPerImage sourceSize:(MTLSize)sourceSize toTexture:(id<MTLTexture>)destinationTexture destinationSlice:(NSUInteger)destinationSlice destinationLevel:(NSUInteger)destinationLevel destinationOrigin:(MTLOrigin)destinationOrigin options:(MTLBlitOption)options];
 }
 
-- (void)copyFromTexture:(id<MTLTexture>)sourceTexture sourceSlice:(NSUInteger)sourceSlice sourceLevel:(NSUInteger)sourceLevel sourceOrigin:(MTLOrigin)sourceOrigin sourceSize:(MTLSize)sourceSize toBuffer:(id<MTLBuffer>)destinationBuffer destinationOffset:(NSUInteger)destinationOffset destinationBytesPerRow:(NSUInteger)destinationBytesPerRow destinationBytesPerImage:(NSUInteger)destinationBytesPerImage
+void FMetalBlitCommandEncoderDebugging::Copy(FMetalTexture const& sourceTexture, NSUInteger sourceSlice, NSUInteger sourceLevel, mtlpp::Origin const& sourceOrigin, mtlpp::Size const& sourceSize, FMetalBuffer const& destinationBuffer, NSUInteger destinationOffset, NSUInteger destinationBytesPerRow, NSUInteger destinationBytesPerImage)
 {
 #if METAL_DEBUG_OPTIONS
-	switch(Buffer->DebugLevel)
+	switch(((FMetalDebugBlitCommandEncoder*)m_ptr)->Buffer.GetPtr()->DebugLevel)
 	{
 		case EMetalDebugLevelConditionalSubmit:
 		case EMetalDebugLevelWaitForComplete:
 		case EMetalDebugLevelLogOperations:
 		{
-			[Buffer blit:[NSString stringWithFormat:@"%s", __PRETTY_FUNCTION__]];
+			((FMetalDebugBlitCommandEncoder*)m_ptr)->Buffer.Blit([NSString stringWithFormat:@"%s", __PRETTY_FUNCTION__]);
 		}
 		case EMetalDebugLevelValidation:
 		case EMetalDebugLevelResetOnBind:
 		case EMetalDebugLevelTrackResources:
 		{
-			[Buffer trackResource:sourceTexture];
-			[Buffer trackResource:destinationBuffer];
+			((FMetalDebugBlitCommandEncoder*)m_ptr)->Buffer.TrackResource(sourceTexture);
+			((FMetalDebugBlitCommandEncoder*)m_ptr)->Buffer.TrackResource(destinationBuffer);
 		}
 		default:
 		{
@@ -238,27 +238,25 @@ APPLE_PLATFORM_OBJECT_ALLOC_OVERRIDES(FMetalDebugBlitCommandEncoder)
 		}
 	}
 #endif
-    
-    [Inner copyFromTexture:(id<MTLTexture>)sourceTexture sourceSlice:(NSUInteger)sourceSlice sourceLevel:(NSUInteger)sourceLevel sourceOrigin:(MTLOrigin)sourceOrigin sourceSize:(MTLSize)sourceSize toBuffer:(id<MTLBuffer>)destinationBuffer destinationOffset:(NSUInteger)destinationOffset destinationBytesPerRow:(NSUInteger)destinationBytesPerRow destinationBytesPerImage:(NSUInteger)destinationBytesPerImage];
 }
 
-- (void)copyFromTexture:(id<MTLTexture>)sourceTexture sourceSlice:(NSUInteger)sourceSlice sourceLevel:(NSUInteger)sourceLevel sourceOrigin:(MTLOrigin)sourceOrigin sourceSize:(MTLSize)sourceSize toBuffer:(id<MTLBuffer>)destinationBuffer destinationOffset:(NSUInteger)destinationOffset destinationBytesPerRow:(NSUInteger)destinationBytesPerRow destinationBytesPerImage:(NSUInteger)destinationBytesPerImage options:(MTLBlitOption)options
+void FMetalBlitCommandEncoderDebugging::Copy(FMetalTexture const& sourceTexture, NSUInteger sourceSlice, NSUInteger sourceLevel, mtlpp::Origin const& sourceOrigin, mtlpp::Size const& sourceSize, FMetalBuffer const& destinationBuffer, NSUInteger destinationOffset, NSUInteger destinationBytesPerRow, NSUInteger destinationBytesPerImage, mtlpp::BlitOption options)
 {
 #if METAL_DEBUG_OPTIONS
-	switch(Buffer->DebugLevel)
+	switch(((FMetalDebugBlitCommandEncoder*)m_ptr)->Buffer.GetPtr()->DebugLevel)
 	{
 		case EMetalDebugLevelConditionalSubmit:
 		case EMetalDebugLevelWaitForComplete:
 		case EMetalDebugLevelLogOperations:
 		{
-			[Buffer blit:[NSString stringWithFormat:@"%s", __PRETTY_FUNCTION__]];
+			((FMetalDebugBlitCommandEncoder*)m_ptr)->Buffer.Blit([NSString stringWithFormat:@"%s", __PRETTY_FUNCTION__]);
 		}
 		case EMetalDebugLevelValidation:
 		case EMetalDebugLevelResetOnBind:
 		case EMetalDebugLevelTrackResources:
 		{
-			[Buffer trackResource:sourceTexture];
-			[Buffer trackResource:destinationBuffer];
+			((FMetalDebugBlitCommandEncoder*)m_ptr)->Buffer.TrackResource(sourceTexture);
+			((FMetalDebugBlitCommandEncoder*)m_ptr)->Buffer.TrackResource(destinationBuffer);
 		}
 		default:
 		{
@@ -266,26 +264,24 @@ APPLE_PLATFORM_OBJECT_ALLOC_OVERRIDES(FMetalDebugBlitCommandEncoder)
 		}
 	}
 #endif
-    
-    [Inner copyFromTexture:(id<MTLTexture>)sourceTexture sourceSlice:(NSUInteger)sourceSlice sourceLevel:(NSUInteger)sourceLevel sourceOrigin:(MTLOrigin)sourceOrigin sourceSize:(MTLSize)sourceSize toBuffer:(id<MTLBuffer>)destinationBuffer destinationOffset:(NSUInteger)destinationOffset destinationBytesPerRow:(NSUInteger)destinationBytesPerRow destinationBytesPerImage:(NSUInteger)destinationBytesPerImage options:(MTLBlitOption)options];
 }
 
-- (void)generateMipmapsForTexture:(id<MTLTexture>)texture
+void FMetalBlitCommandEncoderDebugging::GenerateMipmaps(FMetalTexture const& texture)
 {
 #if METAL_DEBUG_OPTIONS
-	switch(Buffer->DebugLevel)
+	switch(((FMetalDebugBlitCommandEncoder*)m_ptr)->Buffer.GetPtr()->DebugLevel)
 	{
 		case EMetalDebugLevelConditionalSubmit:
 		case EMetalDebugLevelWaitForComplete:
 		case EMetalDebugLevelLogOperations:
 		{
-			[Buffer blit:[NSString stringWithFormat:@"%s", __PRETTY_FUNCTION__]];
+			((FMetalDebugBlitCommandEncoder*)m_ptr)->Buffer.Blit([NSString stringWithFormat:@"%s", __PRETTY_FUNCTION__]);
 		}
 		case EMetalDebugLevelValidation:
 		case EMetalDebugLevelResetOnBind:
 		case EMetalDebugLevelTrackResources:
 		{
-			[Buffer trackResource:texture];
+			((FMetalDebugBlitCommandEncoder*)m_ptr)->Buffer.TrackResource(texture);
 		}
 		default:
 		{
@@ -293,25 +289,24 @@ APPLE_PLATFORM_OBJECT_ALLOC_OVERRIDES(FMetalDebugBlitCommandEncoder)
 		}
 	}
 #endif
-    [Inner generateMipmapsForTexture:(id<MTLTexture>)texture];
 }
 
-- (void)fillBuffer:(id <MTLBuffer>)buffer range:(NSRange)range value:(uint8_t)value
+void FMetalBlitCommandEncoderDebugging::Fill(FMetalBuffer const& buffer, ns::Range const& range, uint8_t value)
 {
 #if METAL_DEBUG_OPTIONS
-	switch(Buffer->DebugLevel)
+	switch(((FMetalDebugBlitCommandEncoder*)m_ptr)->Buffer.GetPtr()->DebugLevel)
 	{
 		case EMetalDebugLevelConditionalSubmit:
 		case EMetalDebugLevelWaitForComplete:
 		case EMetalDebugLevelLogOperations:
 		{
-			[Buffer blit:[NSString stringWithFormat:@"%s", __PRETTY_FUNCTION__]];
+			((FMetalDebugBlitCommandEncoder*)m_ptr)->Buffer.Blit([NSString stringWithFormat:@"%s", __PRETTY_FUNCTION__]);
 		}
 		case EMetalDebugLevelValidation:
 		case EMetalDebugLevelResetOnBind:
 		case EMetalDebugLevelTrackResources:
 		{
-			[Buffer trackResource:buffer];
+			((FMetalDebugBlitCommandEncoder*)m_ptr)->Buffer.TrackResource(buffer);
 		}
 		default:
 		{
@@ -319,26 +314,25 @@ APPLE_PLATFORM_OBJECT_ALLOC_OVERRIDES(FMetalDebugBlitCommandEncoder)
 		}
 	}
 #endif
-    [Inner fillBuffer:(id <MTLBuffer>)buffer range:(NSRange)range value:(uint8_t)value];
 }
 
-- (void)copyFromBuffer:(id <MTLBuffer>)sourceBuffer sourceOffset:(NSUInteger)sourceOffset toBuffer:(id <MTLBuffer>)destinationBuffer destinationOffset:(NSUInteger)destinationOffset size:(NSUInteger)size
+void FMetalBlitCommandEncoderDebugging::Copy(FMetalBuffer const& sourceBuffer, NSUInteger sourceOffset, FMetalBuffer const& destinationBuffer, NSUInteger destinationOffset, NSUInteger size)
 {
 #if METAL_DEBUG_OPTIONS
-	switch(Buffer->DebugLevel)
+	switch(((FMetalDebugBlitCommandEncoder*)m_ptr)->Buffer.GetPtr()->DebugLevel)
 	{
 		case EMetalDebugLevelConditionalSubmit:
 		case EMetalDebugLevelWaitForComplete:
 		case EMetalDebugLevelLogOperations:
 		{
-			[Buffer blit:[NSString stringWithFormat:@"%s", __PRETTY_FUNCTION__]];
+			((FMetalDebugBlitCommandEncoder*)m_ptr)->Buffer.Blit([NSString stringWithFormat:@"%s", __PRETTY_FUNCTION__]);
 		}
 		case EMetalDebugLevelValidation:
 		case EMetalDebugLevelResetOnBind:
 		case EMetalDebugLevelTrackResources:
 		{
-			[Buffer trackResource:sourceBuffer];
-			[Buffer trackResource:destinationBuffer];
+			((FMetalDebugBlitCommandEncoder*)m_ptr)->Buffer.TrackResource(sourceBuffer);
+			((FMetalDebugBlitCommandEncoder*)m_ptr)->Buffer.TrackResource(destinationBuffer);
 		}
 		default:
 		{
@@ -346,58 +340,4 @@ APPLE_PLATFORM_OBJECT_ALLOC_OVERRIDES(FMetalDebugBlitCommandEncoder)
 		}
 	}
 #endif
-    [Inner copyFromBuffer:(id <MTLBuffer>)sourceBuffer sourceOffset:(NSUInteger)sourceOffset toBuffer:(id <MTLBuffer>)destinationBuffer destinationOffset:(NSUInteger)destinationOffset size:(NSUInteger)size];
-}
-
-- (void)updateFence:(id <MTLFence>)fence
-{
-#if METAL_DEBUG_OPTIONS
-	if (fence && (EMetalDebugLevel)GMetalRuntimeDebugLevel >= EMetalDebugLevelValidation)
-	{
-		[self addUpdateFence:fence];
-        if (((FMetalDebugFence*)fence).Inner)
-        {
-            [Inner updateFence:((FMetalDebugFence*)fence).Inner];
-        }
-	}
-	else
-#endif
-	{
-		[Inner updateFence:(id <MTLFence>)fence];
-	}
-}
-
-- (void)waitForFence:(id <MTLFence>)fence
-{
-#if METAL_DEBUG_OPTIONS
-	if (fence && (EMetalDebugLevel)GMetalRuntimeDebugLevel >= EMetalDebugLevelValidation)
-	{
-		[self addWaitFence:fence];
-        if (((FMetalDebugFence*)fence).Inner)
-        {
-            [Inner waitForFence:((FMetalDebugFence*)fence).Inner];
-        }
-	}
-	else
-#endif
-	{
-		[Inner waitForFence:(id <MTLFence>)fence];
-	}
 }
-
--(NSString*) description
-{
-	return [Inner description];
-}
-
--(NSString*) debugDescription
-{
-	return [Inner debugDescription];
-}
-
--(id<MTLCommandEncoder>)commandEncoder
-{
-	return self;
-}
-
-@end
diff --git a/Engine/Source/Runtime/Apple/MetalRHI/Private/MetalBlitCommandEncoder.h b/Engine/Source/Runtime/Apple/MetalRHI/Private/MetalBlitCommandEncoder.h
--- a/Engine/Source/Runtime/Apple/MetalRHI/Private/MetalBlitCommandEncoder.h
+++ b/Engine/Source/Runtime/Apple/MetalRHI/Private/MetalBlitCommandEncoder.h
@@ -8,15 +8,44 @@
 NS_ASSUME_NONNULL_BEGIN
 
 @class FMetalDebugCommandBuffer;
+@class FMetalDebugBlitCommandEncoder;
+
+class FMetalBlitCommandEncoderDebugging : public FMetalCommandEncoderDebugging
+{
+public:
+	FMetalBlitCommandEncoderDebugging();
+	FMetalBlitCommandEncoderDebugging(mtlpp::BlitCommandEncoder& Encoder, FMetalCommandBufferDebugging& Buffer);
+	FMetalBlitCommandEncoderDebugging(FMetalDebugCommandEncoder* handle);
+	
+	static FMetalBlitCommandEncoderDebugging Get(mtlpp::BlitCommandEncoder& Buffer);
+	
+	void InsertDebugSignpost(ns::String const& Label);
+	void PushDebugGroup(ns::String const& Group);
+	void PopDebugGroup();
+	
+	void EndEncoder();
+	
+#if PLATFORM_MAC
+	void Synchronize(mtlpp::Resource const& resource);
+	
+	void Synchronize(FMetalTexture const& texture, NSUInteger slice, NSUInteger level);
+#endif
+	
+	void Copy(FMetalTexture const& sourceTexture, NSUInteger sourceSlice, NSUInteger sourceLevel, mtlpp::Origin const& sourceOrigin, mtlpp::Size const& sourceSize, FMetalTexture const& destinationTexture, NSUInteger destinationSlice, NSUInteger destinationLevel, mtlpp::Origin const& destinationOrigin);
+	
+	void Copy(FMetalBuffer const& sourceBuffer, NSUInteger sourceOffset, NSUInteger sourceBytesPerRow, NSUInteger sourceBytesPerImage, mtlpp::Size const& sourceSize, FMetalTexture const& destinationTexture, NSUInteger destinationSlice, NSUInteger destinationLevel, mtlpp::Origin const& destinationOrigin);
+	
+	void Copy(FMetalBuffer const& sourceBuffer, NSUInteger sourceOffset, NSUInteger sourceBytesPerRow, NSUInteger sourceBytesPerImage, mtlpp::Size const& sourceSize, FMetalTexture const& destinationTexture, NSUInteger destinationSlice, NSUInteger destinationLevel, mtlpp::Origin const& destinationOrigin, mtlpp::BlitOption options);
+	
+	void Copy(FMetalTexture const& sourceTexture, NSUInteger sourceSlice, NSUInteger sourceLevel, mtlpp::Origin const& sourceOrigin, mtlpp::Size const& sourceSize, FMetalBuffer const& destinationBuffer, NSUInteger destinationOffset, NSUInteger destinationBytesPerRow, NSUInteger destinationBytesPerImage);
+	
+	void Copy(FMetalTexture const& sourceTexture, NSUInteger sourceSlice, NSUInteger sourceLevel, mtlpp::Origin const& sourceOrigin, mtlpp::Size const& sourceSize, FMetalBuffer const& destinationBuffer, NSUInteger destinationOffset, NSUInteger destinationBytesPerRow, NSUInteger destinationBytesPerImage, mtlpp::BlitOption options);
+	
+	void GenerateMipmaps(FMetalTexture const& texture);
+	
+	void Fill(FMetalBuffer const& buffer, ns::Range const& range, uint8_t value);
+	
+	void Copy(FMetalBuffer const& sourceBuffer, NSUInteger sourceOffset, FMetalBuffer const& destinationBuffer, NSUInteger destinationOffset, NSUInteger size);
+};
 
-@interface FMetalDebugBlitCommandEncoder : FMetalDebugCommandEncoder<MTLBlitCommandEncoder>
-
-/** The wrapped native command-encoder for which we collect debug information. */
-@property (readonly, retain) id<MTLBlitCommandEncoder> Inner;
-@property (readonly, retain) FMetalDebugCommandBuffer* Buffer;
-
-/** Initialise the wrapper with the provided command-buffer. */
--(id)initWithEncoder:(id<MTLBlitCommandEncoder>)Encoder andCommandBuffer:(FMetalDebugCommandBuffer*)Buffer;
-
-@end
 NS_ASSUME_NONNULL_END
diff --git a/Engine/Source/Runtime/Apple/MetalRHI/Private/MetalCommandBuffer.cpp b/Engine/Source/Runtime/Apple/MetalRHI/Private/MetalCommandBuffer.cpp
--- a/Engine/Source/Runtime/Apple/MetalRHI/Private/MetalCommandBuffer.cpp
+++ b/Engine/Source/Runtime/Apple/MetalRHI/Private/MetalCommandBuffer.cpp
@@ -8,7 +8,6 @@
 
 #include "MetalCommandBuffer.h"
 #include "MetalRenderCommandEncoder.h"
-#include "MetalParallelRenderCommandEncoder.h"
 #include "MetalBlitCommandEncoder.h"
 #include "MetalComputeCommandEncoder.h"
 #include <objc/runtime.h>
@@ -47,9 +46,9 @@ uint32 SafeGetRuntimeDebuggingLevel()
 }
 @end
 
-@implementation FMetalDebugCommandBuffer
+#if MTLPP_CONFIG_VALIDATE
 
-@synthesize InnerBuffer;
+@implementation FMetalDebugCommandBuffer
 
 APPLE_PLATFORM_OBJECT_ALLOC_OVERRIDES(FMetalDebugCommandBuffer)
 
@@ -59,7 +58,7 @@ APPLE_PLATFORM_OBJECT_ALLOC_OVERRIDES(FMetalDebugCommandBuffer)
 	if (Self)
 	{
         DebugLevel = (EMetalDebugLevel)GMetalRuntimeDebugLevel;
-		InnerBuffer = [Buffer retain];
+		InnerBuffer = Buffer;
 		DebugGroup = [NSMutableArray new];
 		ActiveEncoder = nil;
 		DebugInfoBuffer = nil;
@@ -80,7 +79,6 @@ APPLE_PLATFORM_OBJECT_ALLOC_OVERRIDES(FMetalDebugCommandBuffer)
 		delete Command;
 	}
 	
-	[InnerBuffer release];
 	[DebugGroup release];
 	[DebugInfoBuffer release];
 	DebugInfoBuffer = nil;
@@ -88,140 +86,51 @@ APPLE_PLATFORM_OBJECT_ALLOC_OVERRIDES(FMetalDebugCommandBuffer)
 	[super dealloc];
 }
 
--(id <MTLDevice>) device
-{
-	return InnerBuffer.device;
-}
-
--(id <MTLCommandQueue>) commandQueue
-{
-	return InnerBuffer.commandQueue;
-}
-
--(BOOL) retainedReferences
-{
-	return InnerBuffer.retainedReferences;
-}
-
--(NSString *)label
-{
-	return InnerBuffer.label;
-}
-
--(void)setLabel:(NSString *)Text
-{
-	InnerBuffer.label = Text;
-}
-
--(MTLCommandBufferStatus) status
-{
-	return InnerBuffer.status;
-}
-
--(NSError *)error
-{
-	return InnerBuffer.error;
-}
-
-- (void)enqueue
-{
-	[InnerBuffer enqueue];
-}
-
-- (void)commit
-{
-	[InnerBuffer commit];
-}
-
-- (void)addScheduledHandler:(MTLCommandBufferHandler)block
-{
-	[InnerBuffer addScheduledHandler:^(id <MTLCommandBuffer>){
-		block(self);
-	 }];
-}
-
-- (void)presentDrawable:(id <MTLDrawable>)drawable
-{
-    [InnerBuffer presentDrawable:drawable];
-}
-
-#if !PLATFORM_MAC
-- (void)presentDrawable:(id <MTLDrawable>)drawable afterMinimumDuration:(CFTimeInterval)duration
-{
-#if (__clang_major__ > 8) || (__clang_major__ == 8 && __clang_minor__ >= 1)
-    [InnerBuffer presentDrawable:drawable afterMinimumDuration:duration];
-#endif
-}
-#endif
-
-- (void)presentDrawable:(id <MTLDrawable>)drawable atTime:(CFTimeInterval)presentationTime
-{
-    [InnerBuffer presentDrawable:drawable atTime:presentationTime];
-}
-
-- (void)waitUntilScheduled
-{
-	[InnerBuffer waitUntilScheduled];
-}
-
-- (void)addCompletedHandler:(MTLCommandBufferHandler)block
-{
-	[InnerBuffer addCompletedHandler:^(id <MTLCommandBuffer>){
-		block(self);
-	}];
-}
-
-- (void)waitUntilCompleted
-{
-	[InnerBuffer waitUntilCompleted];
-}
+@end
 
-- (id <MTLBlitCommandEncoder>)blitCommandEncoder
+FMetalCommandBufferDebugging FMetalCommandBufferDebugging::Get(mtlpp::CommandBuffer& Buffer)
 {
-	[self beginBlitCommandEncoder:DebugGroup.lastObject ? DebugGroup.lastObject : @"Blit"];
-	return [[[FMetalDebugBlitCommandEncoder alloc] initWithEncoder:[InnerBuffer blitCommandEncoder] andCommandBuffer:self] autorelease];
+	return Buffer.GetAssociatedObject<FMetalCommandBufferDebugging>((void const*)&FMetalCommandBufferDebugging::Get);
 }
 
-- (id <MTLRenderCommandEncoder>)renderCommandEncoderWithDescriptor:(MTLRenderPassDescriptor *)renderPassDescriptor
+FMetalCommandBufferDebugging::FMetalCommandBufferDebugging()
+: ns::Object<FMetalDebugCommandBuffer*>(nullptr)
 {
-    [self beginRenderCommandEncoder:DebugGroup.lastObject ? DebugGroup.lastObject : @"Render" withDescriptor:renderPassDescriptor];
-	return [[[FMetalDebugRenderCommandEncoder alloc] initWithEncoder:[InnerBuffer renderCommandEncoderWithDescriptor:renderPassDescriptor] fromDescriptor:renderPassDescriptor andCommandBuffer:self] autorelease];
+	
 }
-
-- (id <MTLComputeCommandEncoder>)computeCommandEncoder
+FMetalCommandBufferDebugging::FMetalCommandBufferDebugging(mtlpp::CommandBuffer& Buffer)
+: ns::Object<FMetalDebugCommandBuffer*>([[FMetalDebugCommandBuffer alloc] initWithCommandBuffer:Buffer.GetPtr()], ns::Ownership::Assign)
 {
-    [self beginComputeCommandEncoder:DebugGroup.lastObject ? DebugGroup.lastObject : @"Compute"];
-    return [[[FMetalDebugComputeCommandEncoder alloc] initWithEncoder:[InnerBuffer computeCommandEncoder] andCommandBuffer:self] autorelease];
+	Buffer.SetAssociatedObject((void const*)&FMetalCommandBufferDebugging::Get, *this);
 }
-
-- (id <MTLParallelRenderCommandEncoder>)parallelRenderCommandEncoderWithDescriptor:(MTLRenderPassDescriptor *)renderPassDescriptor
+FMetalCommandBufferDebugging::FMetalCommandBufferDebugging(FMetalDebugCommandBuffer* handle)
+: ns::Object<FMetalDebugCommandBuffer*>(handle)
 {
-    [self beginRenderCommandEncoder:DebugGroup.lastObject ? DebugGroup.lastObject : @"Parallel Render" withDescriptor:renderPassDescriptor];
-    return [[[FMetalDebugParallelRenderCommandEncoder alloc] initWithEncoder:[InnerBuffer parallelRenderCommandEncoderWithDescriptor:renderPassDescriptor] andCommandBuffer:self withDescriptor:renderPassDescriptor] autorelease];
+	
 }
 
--(NSString*) description
+ns::AutoReleased<ns::String> FMetalCommandBufferDebugging::GetDescription()
 {
-	NSMutableString* String = [NSMutableString new];
-	NSString* Label = self.label ? self.label : @"Unknown";
-	[String appendFormat:@"Command Buffer %p %@:", self, Label];
-	return String;
+	NSMutableString* String = [[NSMutableString new] autorelease];
+	NSString* Label = m_ptr->InnerBuffer.label ? m_ptr->InnerBuffer.label : @"Unknown";
+	[String appendFormat:@"Command Buffer %p %@:", m_ptr->InnerBuffer, Label];
+	return ns::AutoReleased<ns::String>(String);
 }
 
--(NSString*) debugDescription
+ns::AutoReleased<ns::String> FMetalCommandBufferDebugging::GetDebugDescription()
 {
-	NSMutableString* String = [NSMutableString new];
-	NSString* Label = self.label ? self.label : @"Unknown";
-	[String appendFormat:@"Command Buffer %p %@:", self, Label];
-
+	NSMutableString* String = [[NSMutableString new] autorelease];
+	NSString* Label = m_ptr->InnerBuffer.label ? m_ptr->InnerBuffer.label : @"Unknown";
+	[String appendFormat:@"Command Buffer %p %@:", m_ptr->InnerBuffer, Label];
+	
 	uint32 Index = 0;
-	if (DebugInfoBuffer)
+	if (m_ptr->DebugInfoBuffer)
 	{
-		Index = *((uint32*)DebugInfoBuffer.contents);
+		Index = *((uint32*)m_ptr->DebugInfoBuffer.contents);
 	}
 	
 	uint32 Count = 1;
-	for (FMetalDebugCommand* Command : DebugCommands)
+	for (FMetalDebugCommand* Command : m_ptr->DebugCommands)
 	{
 		if (Index == Count++)
 		{
@@ -235,271 +144,197 @@ APPLE_PLATFORM_OBJECT_ALLOC_OVERRIDES(FMetalDebugCommandBuffer)
 	
 	[String appendFormat:@"\nResources:"];
 	
-	for (id<MTLResource> Resource : Resources)
+	for (id<MTLResource> Resource : m_ptr->Resources)
 	{
 		[String appendFormat:@"\n\t%@ (%d): %@", Resource.label, (uint32)[Resource retainCount], [Resource description]];
 	}
 	
 	[String appendFormat:@"\nStates:"];
 	
-	for (id State : States)
+	for (id State : m_ptr->States)
 	{
 		[String appendFormat:@"\n\t%@ (%d): %@", ([State respondsToSelector:@selector(label)] ? [State label] : @"(null)"), (uint32)[State retainCount], [State description]];
 	}
 	
-	return String;
-}
-
--(void) trackResource:(id<MTLResource>)Resource
-{
-    if (DebugLevel >= EMetalDebugLevelValidation)
-    {
-        Resources.Add(Resource);
-    }
-}
-
--(void) trackState:(id)State
-{
-    if (DebugLevel >= EMetalDebugLevelValidation)
-    {
-        States.Add(State);
-    }
+	return ns::AutoReleased<ns::String>(String);
 }
 
--(void) beginRenderCommandEncoder:(NSString*)Label withDescriptor:(MTLRenderPassDescriptor*)RenderDesc
+void FMetalCommandBufferDebugging::TrackResource(mtlpp::Resource const& Resource)
 {
-    if (DebugLevel >= EMetalDebugLevelValidation)
-    {
-        if (DebugLevel >= EMetalDebugLevelLogOperations)
-        {
-            check(!ActiveEncoder);
-            ActiveEncoder = [Label retain];
-            FMetalDebugCommand* Command = new FMetalDebugCommand;
-            Command->Type = EMetalDebugCommandTypeRenderEncoder;
-            Command->Label = [ActiveEncoder retain];
-            Command->PassDesc = [RenderDesc retain];
-            DebugCommands.Add(Command);
-        }
-        
-        if(RenderDesc.colorAttachments)
-        {
-            for(uint i = 0; i< 8; i++)
-            {
-                MTLRenderPassColorAttachmentDescriptor* Desc = [RenderDesc.colorAttachments objectAtIndexedSubscript:i];
-                if(Desc)
-                {
-                    [self trackResource:Desc.texture];
-                }
-            }
-        }
-        if(RenderDesc.depthAttachment)
-        {
-            [self trackResource:RenderDesc.depthAttachment.texture];
-        }
-        if(RenderDesc.stencilAttachment)
-        {
-            [self trackResource:RenderDesc.stencilAttachment.texture];
-        }
-        if(RenderDesc.visibilityResultBuffer)
-        {
-            [self trackResource:RenderDesc.visibilityResultBuffer];
-        }
-    }
-}
-
--(void) beginComputeCommandEncoder:(NSString*)Label
-{
-    if (DebugLevel >= EMetalDebugLevelLogOperations)
-    {
-        check(!ActiveEncoder);
-        ActiveEncoder = [Label retain];
-        FMetalDebugCommand* Command = new FMetalDebugCommand;
-        Command->Type = EMetalDebugCommandTypeComputeEncoder;
-        Command->Label = [ActiveEncoder retain];
-        Command->PassDesc = nil;
-        DebugCommands.Add(Command);
-    }
-}
-
--(void) beginBlitCommandEncoder:(NSString*)Label
-{
-    if (DebugLevel >= EMetalDebugLevelLogOperations)
-    {
-        check(!ActiveEncoder);
-        ActiveEncoder = [Label retain];
-        FMetalDebugCommand* Command = new FMetalDebugCommand;
-        Command->Type = EMetalDebugCommandTypeBlitEncoder;
-        Command->Label = [ActiveEncoder retain];
-        Command->PassDesc = nil;
-        DebugCommands.Add(Command);
-    }
-}
-
--(void) endCommandEncoder
-{
-    if (DebugLevel >= EMetalDebugLevelLogOperations)
-    {
-        check(ActiveEncoder);
-        FMetalDebugCommand* Command = new FMetalDebugCommand;
-        Command->Type = EMetalDebugCommandTypeEndEncoder;
-        Command->Label = ActiveEncoder;
-        Command->PassDesc = nil;
-        DebugCommands.Add(Command);
-        ActiveEncoder = nil;
-    }
-}
-
--(void) setPipeline:(NSString*)Desc
-{
-    if (DebugLevel >= EMetalDebugLevelLogOperations)
-    {
-        FMetalDebugCommand* Command = new FMetalDebugCommand;
-        Command->Type = EMetalDebugCommandTypePipeline;
-        Command->Label = [Desc retain];
-        Command->PassDesc = nil;
-        DebugCommands.Add(Command);
-    }
-}
-
--(void) draw:(NSString*)Desc
-{
-    if (DebugLevel >= EMetalDebugLevelLogOperations)
-    {
-        FMetalDebugCommand* Command = new FMetalDebugCommand;
-        Command->Type = EMetalDebugCommandTypeDraw;
-        Command->Label = [Desc retain];
-        Command->PassDesc = nil;
-        DebugCommands.Add(Command);
-    }
+	if (m_ptr->DebugLevel >= EMetalDebugLevelValidation)
+	{
+		m_ptr->Resources.Add(Resource.GetPtr());
+	}
 }
-
--(void) dispatch:(NSString*)Desc
+void FMetalCommandBufferDebugging::TrackState(id State)
 {
-    if (DebugLevel >= EMetalDebugLevelLogOperations)
-    {
-        FMetalDebugCommand* Command = new FMetalDebugCommand;
-        Command->Type = EMetalDebugCommandTypeDispatch;
-        Command->Label = [Desc retain];
-        Command->PassDesc = nil;
-        DebugCommands.Add(Command);
-    }
+	if (m_ptr->DebugLevel >= EMetalDebugLevelValidation)
+	{
+		m_ptr->States.Add(State);
+	}
 }
 
--(void) blit:(NSString*)Desc
+void FMetalCommandBufferDebugging::BeginRenderCommandEncoder(ns::String const& Label, mtlpp::RenderPassDescriptor const& Desc)
 {
-    if (DebugLevel >= EMetalDebugLevelLogOperations)
-    {
-        FMetalDebugCommand* Command = new FMetalDebugCommand;
-        Command->Type = EMetalDebugCommandTypeBlit;
-        Command->Label = [Desc retain];
-        Command->PassDesc = nil;
-        DebugCommands.Add(Command);
-    }
+	if (m_ptr->DebugLevel >= EMetalDebugLevelValidation)
+	{
+		if (m_ptr->DebugLevel >= EMetalDebugLevelLogOperations)
+		{
+			check(!m_ptr->ActiveEncoder);
+			m_ptr->ActiveEncoder = [Label.GetPtr() retain];
+			FMetalDebugCommand* Command = new FMetalDebugCommand;
+			Command->Type = EMetalDebugCommandTypeRenderEncoder;
+			Command->Label = [Label.GetPtr() retain];
+			Command->PassDesc = [Desc.GetPtr() retain];
+			m_ptr->DebugCommands.Add(Command);
+		}
+		
+		ns::Array<mtlpp::RenderPassColorAttachmentDescriptor> ColorAttach = Desc.GetColorAttachments();
+		if(ColorAttach)
+		{
+			for(uint i = 0; i < 8; i++)
+			{
+				if(ColorAttach[i])
+				{
+					TrackResource(ColorAttach[i].GetTexture());
+				}
+			}
+		}
+		if(Desc.GetDepthAttachment())
+		{
+			TrackResource(Desc.GetDepthAttachment().GetTexture());
+		}
+		if(Desc.GetStencilAttachment())
+		{
+			TrackResource(Desc.GetStencilAttachment().GetTexture());
+		}
+		if(Desc.GetVisibilityResultBuffer())
+		{
+			TrackResource(Desc.GetVisibilityResultBuffer());
+		}
+	}
 }
-
--(void) insertDebugSignpost:(NSString*)Label
+void FMetalCommandBufferDebugging::BeginComputeCommandEncoder(ns::String const& Label)
 {
-    if (DebugLevel >= EMetalDebugLevelLogDebugGroups)
-    {
-        FMetalDebugCommand* Command = new FMetalDebugCommand;
-        Command->Type = EMetalDebugCommandTypeSignpost;
-        Command->Label = [Label retain];
-        Command->PassDesc = nil;
-        DebugCommands.Add(Command);
-    }
+	if (m_ptr->DebugLevel >= EMetalDebugLevelLogOperations)
+	{
+		check(!m_ptr->ActiveEncoder);
+		m_ptr->ActiveEncoder = [Label.GetPtr() retain];
+		FMetalDebugCommand* Command = new FMetalDebugCommand;
+		Command->Type = EMetalDebugCommandTypeComputeEncoder;
+		Command->Label = [m_ptr->ActiveEncoder retain];
+		Command->PassDesc = nil;
+		m_ptr->DebugCommands.Add(Command);
+	}
 }
-
--(void) pushDebugGroup:(NSString*)Group
+void FMetalCommandBufferDebugging::BeginBlitCommandEncoder(ns::String const& Label)
 {
-    if (DebugLevel >= EMetalDebugLevelLogDebugGroups)
-    {
-        [DebugGroup addObject:Group];
-        FMetalDebugCommand* Command = new FMetalDebugCommand;
-        Command->Type = EMetalDebugCommandTypePushGroup;
-        Command->Label = [Group retain];
-        Command->PassDesc = nil;
-        DebugCommands.Add(Command);
-    }
+	if (m_ptr->DebugLevel >= EMetalDebugLevelLogOperations)
+	{
+		check(!m_ptr->ActiveEncoder);
+		m_ptr->ActiveEncoder = [Label.GetPtr() retain];
+		FMetalDebugCommand* Command = new FMetalDebugCommand;
+		Command->Type = EMetalDebugCommandTypeBlitEncoder;
+		Command->Label = [m_ptr->ActiveEncoder retain];
+		Command->PassDesc = nil;
+		m_ptr->DebugCommands.Add(Command);
+	}
 }
-
--(void) popDebugGroup
+void FMetalCommandBufferDebugging::EndCommandEncoder()
 {
-    if (DebugLevel >= EMetalDebugLevelLogDebugGroups)
-    {
-        if (DebugGroup.lastObject)
-        {
-            FMetalDebugCommand* Command = new FMetalDebugCommand;
-            Command->Type = EMetalDebugCommandTypePopGroup;
-            Command->Label = [DebugGroup.lastObject retain];
-            Command->PassDesc = nil;
-            DebugCommands.Add(Command);
-        }
-    }
+	if (m_ptr->DebugLevel >= EMetalDebugLevelLogOperations)
+	{
+		check(m_ptr->ActiveEncoder);
+		FMetalDebugCommand* Command = new FMetalDebugCommand;
+		Command->Type = EMetalDebugCommandTypeEndEncoder;
+		Command->Label = m_ptr->ActiveEncoder;
+		Command->PassDesc = nil;
+		m_ptr->DebugCommands.Add(Command);
+		m_ptr->ActiveEncoder = nil;
+	}
 }
 
--(CFTimeInterval) kernelStartTime
+void FMetalCommandBufferDebugging::SetPipeline(ns::String const& Desc)
 {
-#if METAL_NEW_NONNULL_DECL
-	if (GMetalCommandBufferHasStartEndTimeAPI)
-#endif
+	if (m_ptr->DebugLevel >= EMetalDebugLevelLogOperations)
 	{
-		return ((id<IMetalCommandBufferExtensions>)InnerBuffer).kernelStartTime;
+		FMetalDebugCommand* Command = new FMetalDebugCommand;
+		Command->Type = EMetalDebugCommandTypePipeline;
+		Command->Label = [Desc.GetPtr() retain];
+		Command->PassDesc = nil;
+		m_ptr->DebugCommands.Add(Command);
 	}
-#if METAL_NEW_NONNULL_DECL
-	else
+}
+void FMetalCommandBufferDebugging::Draw(ns::String const& Desc)
+{
+	if (m_ptr->DebugLevel >= EMetalDebugLevelLogOperations)
 	{
-		return 0;
+		FMetalDebugCommand* Command = new FMetalDebugCommand;
+		Command->Type = EMetalDebugCommandTypeDraw;
+		Command->Label = [Desc.GetPtr() retain];
+		Command->PassDesc = nil;
+		m_ptr->DebugCommands.Add(Command);
 	}
-#endif
 }
-
--(CFTimeInterval) kernelEndTime
+void FMetalCommandBufferDebugging::Dispatch(ns::String const& Desc)
 {
-#if METAL_NEW_NONNULL_DECL
-	if (GMetalCommandBufferHasStartEndTimeAPI)
-#endif
+	if (m_ptr->DebugLevel >= EMetalDebugLevelLogOperations)
 	{
-		return ((id<IMetalCommandBufferExtensions>)InnerBuffer).kernelEndTime;
+		FMetalDebugCommand* Command = new FMetalDebugCommand;
+		Command->Type = EMetalDebugCommandTypeDispatch;
+		Command->Label = [Desc.GetPtr() retain];
+		Command->PassDesc = nil;
+		m_ptr->DebugCommands.Add(Command);
 	}
-#if METAL_NEW_NONNULL_DECL
-	else
+}
+void FMetalCommandBufferDebugging::Blit(ns::String const& Desc)
+{
+	if (m_ptr->DebugLevel >= EMetalDebugLevelLogOperations)
 	{
-		return 0;
+		FMetalDebugCommand* Command = new FMetalDebugCommand;
+		Command->Type = EMetalDebugCommandTypeBlit;
+		Command->Label = [Desc.GetPtr() retain];
+		Command->PassDesc = nil;
+		m_ptr->DebugCommands.Add(Command);
 	}
-#endif
 }
 
--(CFTimeInterval) GPUStartTime
+void FMetalCommandBufferDebugging::InsertDebugSignpost(ns::String const& Label)
 {
-#if METAL_NEW_NONNULL_DECL
-	if (GMetalCommandBufferHasStartEndTimeAPI)
-#endif
-	{
-		return ((id<IMetalCommandBufferExtensions>)InnerBuffer).GPUStartTime;
-	}
-#if METAL_NEW_NONNULL_DECL
-	else
+	if (m_ptr->DebugLevel >= EMetalDebugLevelLogDebugGroups)
 	{
-		return 0;
+		FMetalDebugCommand* Command = new FMetalDebugCommand;
+		Command->Type = EMetalDebugCommandTypeSignpost;
+		Command->Label = [Label.GetPtr() retain];
+		Command->PassDesc = nil;
+		m_ptr->DebugCommands.Add(Command);
 	}
-#endif
 }
-
--(CFTimeInterval) GPUEndTime
+void FMetalCommandBufferDebugging::PushDebugGroup(ns::String const& Group)
 {
-#if METAL_NEW_NONNULL_DECL
-	if (GMetalCommandBufferHasStartEndTimeAPI)
-#endif
+	if (m_ptr->DebugLevel >= EMetalDebugLevelLogDebugGroups)
 	{
-		return ((id<IMetalCommandBufferExtensions>)InnerBuffer).GPUEndTime;
+		[m_ptr->DebugGroup addObject:Group.GetPtr()];
+		FMetalDebugCommand* Command = new FMetalDebugCommand;
+		Command->Type = EMetalDebugCommandTypePushGroup;
+		Command->Label = [Group.GetPtr() retain];
+		Command->PassDesc = nil;
+		m_ptr->DebugCommands.Add(Command);
 	}
-#if METAL_NEW_NONNULL_DECL
-	else
+}
+void FMetalCommandBufferDebugging::PopDebugGroup()
+{
+	if (m_ptr->DebugLevel >= EMetalDebugLevelLogDebugGroups)
 	{
-		return 0;
+		if (m_ptr->DebugGroup.lastObject)
+		{
+			FMetalDebugCommand* Command = new FMetalDebugCommand;
+			Command->Type = EMetalDebugCommandTypePopGroup;
+			Command->Label = [m_ptr->DebugGroup.lastObject retain];
+			Command->PassDesc = nil;
+			m_ptr->DebugCommands.Add(Command);
+		}
 	}
-#endif
 }
 
-@end
-
+#endif
diff --git a/Engine/Source/Runtime/Apple/MetalRHI/Private/MetalCommandBuffer.h b/Engine/Source/Runtime/Apple/MetalRHI/Private/MetalCommandBuffer.h
--- a/Engine/Source/Runtime/Apple/MetalRHI/Private/MetalCommandBuffer.h
+++ b/Engine/Source/Runtime/Apple/MetalRHI/Private/MetalCommandBuffer.h
@@ -60,65 +60,60 @@ struct FMetalDebugCommand
 @property (nonatomic, strong) NSMutableArray<NSString*>* debugGroups;
 @end
 
-@protocol IMetalCommandBufferExtensions
-@property (readonly) CFTimeInterval kernelStartTime;
-@property (readonly) CFTimeInterval kernelEndTime;
-@property (readonly) CFTimeInterval GPUStartTime;
-@property (readonly) CFTimeInterval GPUEndTime;
-@end
+#if MTLPP_CONFIG_VALIDATE
 
 /**
  * FMetalDebugCommandBuffer: Wrapper around id<MTLCommandBuffer> that records information about commands.
- * This allows reporting of substantially more information in debug modes which can be especially helpful 
+ * This allows reporting of substantially more information in debug modes which can be especially helpful
  * when debugging GPU command-buffer failures.
  */
-@interface FMetalDebugCommandBuffer : FApplePlatformObject<MTLCommandBuffer>
+@interface FMetalDebugCommandBuffer : FApplePlatformObject
 {
+@public
 	NSMutableArray<NSString*>* DebugGroup;
 	NSString* ActiveEncoder;
 	TSet<id<MTLResource>> Resources;
 	TSet<id> States;
-	@public
+	id<MTLCommandBuffer> InnerBuffer;
 	TArray<FMetalDebugCommand*> DebugCommands;
-    EMetalDebugLevel DebugLevel;
+	EMetalDebugLevel DebugLevel;
 	id<MTLBuffer> DebugInfoBuffer;
 };
 
-/** The wrapped native command-buffer for which we collect debug information. */
-@property (readonly) id<MTLCommandBuffer> InnerBuffer;
-
 /** Initialise the wrapper with the provided command-buffer. */
 -(id)initWithCommandBuffer:(id<MTLCommandBuffer>)Buffer;
 
-/** Add the resource to be tracked in this command-buffer so we can validate lifetime on failure. */
--(void) trackResource:(id<MTLResource>)Resource;
-/** Add the state to be tracked in this command-buffer so we can validate lifetime on failure. */
--(void) trackState:(id)State;
-
-/** Record a bgein render encoder command. */
--(void) beginRenderCommandEncoder:(NSString*)Label withDescriptor:(MTLRenderPassDescriptor*)Desc;
-/** Record a begin compute encoder command. */
--(void) beginComputeCommandEncoder:(NSString*)Label;
-/** Record a begin blit encoder command. */
--(void) beginBlitCommandEncoder:(NSString*)Label;
-/** Record an end encoder command. */
--(void) endCommandEncoder;
+@end
 
-/** Record a pipeline state set. */
--(void) setPipeline:(NSString*)Desc;
-/** Record a draw command. */
--(void) draw:(NSString*)Desc;
-/** Record a dispatch command. */
--(void) dispatch:(NSString*)Desc;
-/** Record a blit command. */
--(void) blit:(NSString*)Desc;
+class FMetalCommandBufferDebugging : public ns::Object<FMetalDebugCommandBuffer*>
+{
+public:
+	FMetalCommandBufferDebugging();
+	FMetalCommandBufferDebugging(mtlpp::CommandBuffer& Buffer);
+	FMetalCommandBufferDebugging(FMetalDebugCommandBuffer* handle);
+	
+	static FMetalCommandBufferDebugging Get(mtlpp::CommandBuffer& Buffer);
+	ns::AutoReleased<ns::String> GetDescription();
+	ns::AutoReleased<ns::String> GetDebugDescription();
+	
+	void TrackResource(mtlpp::Resource const& Resource);
+	void TrackState(id State);
+	
+	void BeginRenderCommandEncoder(ns::String const& Label, mtlpp::RenderPassDescriptor const& Desc);
+	void BeginComputeCommandEncoder(ns::String const& Label);
+	void BeginBlitCommandEncoder(ns::String const& Label);
+	void EndCommandEncoder();
+	
+	void SetPipeline(ns::String const& Desc);
+	void Draw(ns::String const& Desc);
+	void Dispatch(ns::String const& Desc);
+	void Blit(ns::String const& Desc);
+	
+	void InsertDebugSignpost(ns::String const& Label);
+	void PushDebugGroup(ns::String const& Group);
+	void PopDebugGroup();
+};
 
-/** Record an signpost command. */
--(void) insertDebugSignpost:(NSString*)Label;
-/** Record a push debug group command. */
--(void) pushDebugGroup:(NSString*)Group;
-/** Record a pop debug group command. */
--(void) popDebugGroup;
+#endif
 
-@end
 NS_ASSUME_NONNULL_END
diff --git a/Engine/Source/Runtime/Apple/MetalRHI/Private/MetalCommandEncoder.cpp b/Engine/Source/Runtime/Apple/MetalRHI/Private/MetalCommandEncoder.cpp
--- a/Engine/Source/Runtime/Apple/MetalRHI/Private/MetalCommandEncoder.cpp
+++ b/Engine/Source/Runtime/Apple/MetalRHI/Private/MetalCommandEncoder.cpp
@@ -128,7 +128,8 @@ void FMetalCommandEncoder::StartCommandBuffer(void)
           && IsBlitCommandEncoderActive() == false);
 
 	CommandBuffer = CommandList.GetCommandQueue().CreateCommandBuffer();
-	
+    METAL_DEBUG_LAYER(EMetalDebugLevelFastValidation, CommandBufferDebug = FMetalCommandBufferDebugging::Get(CommandBuffer));
+    
 	if ([DebugGroups count] > 0)
 	{
 		CommandBuffer.SetLabel([DebugGroups lastObject]);
@@ -260,6 +261,7 @@ void FMetalCommandEncoder::BeginRenderCommandEncoding(void)
 	check(IsRenderCommandEncoderActive() == false && IsComputeCommandEncoderActive() == false && IsBlitCommandEncoderActive() == false);
 	
 	RenderCommandEncoder = MTLPP_VALIDATE(mtlpp::CommandBuffer, CommandBuffer, SafeGetRuntimeDebuggingLevel() >= EMetalDebugLevelValidation, RenderCommandEncoder(RenderPassDesc));
+	METAL_DEBUG_LAYER(EMetalDebugLevelFastValidation, RenderEncoderDebug = FMetalRenderCommandEncoderDebugging(RenderCommandEncoder, RenderPassDesc, CommandBufferDebug));
 	
 	check(!EncoderFence.IsValid());
 	NSString* Label = nil;
@@ -278,6 +280,7 @@ void FMetalCommandEncoder::BeginRenderCommandEncoding(void)
 					[((NSObject<MTLCommandBuffer>*)CommandBuffer.GetPtr()).debugGroups addObject:Group];
 				}
 				RenderCommandEncoder.PushDebugGroup(Group);
+				METAL_DEBUG_LAYER(EMetalDebugLevelFastValidation, RenderEncoderDebug.PushDebugGroup(Group));
 			}
 		}
 	}
@@ -292,7 +295,8 @@ void FMetalCommandEncoder::BeginComputeCommandEncoding(void)
 	check(IsRenderCommandEncoderActive() == false && IsComputeCommandEncoderActive() == false && IsBlitCommandEncoderActive() == false);
 	
 	ComputeCommandEncoder = MTLPP_VALIDATE(mtlpp::CommandBuffer, CommandBuffer, SafeGetRuntimeDebuggingLevel() >= EMetalDebugLevelValidation, ComputeCommandEncoder());
-	
+    METAL_DEBUG_LAYER(EMetalDebugLevelFastValidation, ComputeEncoderDebug = FMetalComputeCommandEncoderDebugging(ComputeCommandEncoder, CommandBufferDebug));
+
 	check(!EncoderFence.IsValid());
 	NSString* Label = nil;
 	
@@ -310,6 +314,7 @@ void FMetalCommandEncoder::BeginComputeCommandEncoding(void)
 					[((NSObject<MTLCommandBuffer>*)CommandBuffer.GetPtr()).debugGroups addObject:Group];
 				}
 				ComputeCommandEncoder.PushDebugGroup(Group);
+				METAL_DEBUG_LAYER(EMetalDebugLevelFastValidation, ComputeEncoderDebug.PushDebugGroup(Group));
 			}
 		}
 	}
@@ -323,6 +328,7 @@ void FMetalCommandEncoder::BeginBlitCommandEncoding(void)
 	check(IsRenderCommandEncoderActive() == false && IsComputeCommandEncoderActive() == false && IsBlitCommandEncoderActive() == false);
 	
 	BlitCommandEncoder = MTLPP_VALIDATE(mtlpp::CommandBuffer, CommandBuffer, SafeGetRuntimeDebuggingLevel() >= EMetalDebugLevelValidation, BlitCommandEncoder());
+	METAL_DEBUG_LAYER(EMetalDebugLevelFastValidation, BlitEncoderDebug = FMetalBlitCommandEncoderDebugging(BlitCommandEncoder, CommandBufferDebug));
 	
 	check(!EncoderFence.IsValid());
 	NSString* Label = nil;
@@ -341,6 +347,7 @@ void FMetalCommandEncoder::BeginBlitCommandEncoding(void)
 					[((NSObject<MTLCommandBuffer>*)CommandBuffer.GetPtr()).debugGroups addObject:Group];
 				}
 				BlitCommandEncoder.PushDebugGroup(Group);
+				METAL_DEBUG_LAYER(EMetalDebugLevelFastValidation, BlitEncoderDebug.PushDebugGroup(Group));
 			}
 		}
 	}
@@ -391,6 +398,7 @@ mtlpp::Fence FMetalCommandEncoder::EndEncoding(void)
 			
 			METAL_STATISTIC(FMetalProfiler::GetProfiler()->EndEncoder(CommandBufferStats, RenderCommandEncoder));
 			RenderCommandEncoder.EndEncoding();
+			METAL_DEBUG_LAYER(EMetalDebugLevelFastValidation, RenderEncoderDebug.EndEncoder());
 			RenderCommandEncoder = nil;
 			EncoderFence.Reset();
 		}
@@ -402,6 +410,7 @@ mtlpp::Fence FMetalCommandEncoder::EndEncoding(void)
 			
 			METAL_STATISTIC(FMetalProfiler::GetProfiler()->EndEncoder(CommandBufferStats, ComputeCommandEncoder));
 			ComputeCommandEncoder.EndEncoding();
+			METAL_DEBUG_LAYER(EMetalDebugLevelFastValidation, ComputeEncoderDebug.EndEncoder());
 			ComputeCommandEncoder = nil;
 			EncoderFence.Reset();
 		}
@@ -413,6 +422,7 @@ mtlpp::Fence FMetalCommandEncoder::EndEncoding(void)
 			
 			METAL_STATISTIC(FMetalProfiler::GetProfiler()->EndEncoder(CommandBufferStats, RenderCommandEncoder));
 			BlitCommandEncoder.EndEncoding();
+			METAL_DEBUG_LAYER(EMetalDebugLevelFastValidation, BlitEncoderDebug.EndEncoder());
 			BlitCommandEncoder = nil;
 			EncoderFence.Reset();
 		}
@@ -460,17 +470,21 @@ void FMetalCommandEncoder::UpdateFence(mtlpp::Fence Fence)
 	if (bSupportsFences METAL_DEBUG_OPTION(|| CommandList.GetCommandQueue().GetRuntimeDebuggingLevel() >= EMetalDebugLevelValidation))
 	{
 		check(Fence);
+		mtlpp::Fence InnerFence = METAL_DEBUG_OPTION(CommandList.GetCommandQueue().GetRuntimeDebuggingLevel() >= EMetalDebugLevelValidation ? mtlpp::Fence(((FMetalDebugFence*)Fence.GetPtr()).Inner) :) Fence;
 		if (RenderCommandEncoder)
 		{
-			RenderCommandEncoder.UpdateFence(Fence, (mtlpp::RenderStages)(mtlpp::RenderStages::Vertex | mtlpp::RenderStages::Fragment));
+			RenderCommandEncoder.UpdateFence(InnerFence, (mtlpp::RenderStages)(mtlpp::RenderStages::Vertex | mtlpp::RenderStages::Fragment));
+			METAL_DEBUG_LAYER(EMetalDebugLevelFastValidation, RenderEncoderDebug.AddUpdateFence(Fence));
 		}
 		else if (ComputeCommandEncoder)
 		{
-			ComputeCommandEncoder.UpdateFence(Fence);
+			ComputeCommandEncoder.UpdateFence(InnerFence);
+			METAL_DEBUG_LAYER(EMetalDebugLevelFastValidation, ComputeEncoderDebug.AddUpdateFence(Fence));
 		}
 		else if (BlitCommandEncoder)
 		{
-			BlitCommandEncoder.UpdateFence(Fence);
+			BlitCommandEncoder.UpdateFence(InnerFence);
+			METAL_DEBUG_LAYER(EMetalDebugLevelFastValidation, BlitEncoderDebug.AddUpdateFence(Fence));
 		}
 	}
 }
@@ -482,17 +496,21 @@ void FMetalCommandEncoder::WaitForFence(mtlpp::Fence Fence)
 	if (bSupportsFences METAL_DEBUG_OPTION(|| CommandList.GetCommandQueue().GetRuntimeDebuggingLevel() >= EMetalDebugLevelValidation))
 	{
 		check(Fence);
+		mtlpp::Fence InnerFence = METAL_DEBUG_OPTION(CommandList.GetCommandQueue().GetRuntimeDebuggingLevel() >= EMetalDebugLevelValidation ? mtlpp::Fence(((FMetalDebugFence*)Fence.GetPtr()).Inner) :) Fence;
 		if (RenderCommandEncoder)
 		{
-			RenderCommandEncoder.WaitForFence(Fence, (mtlpp::RenderStages)(mtlpp::RenderStages::Vertex | mtlpp::RenderStages::Fragment));
+			RenderCommandEncoder.WaitForFence(InnerFence, (mtlpp::RenderStages)(mtlpp::RenderStages::Vertex | mtlpp::RenderStages::Fragment));
+			METAL_DEBUG_LAYER(EMetalDebugLevelFastValidation, RenderEncoderDebug.AddWaitFence(Fence));
 		}
 		else if (ComputeCommandEncoder)
 		{
-			ComputeCommandEncoder.WaitForFence(Fence);
+			ComputeCommandEncoder.WaitForFence(InnerFence);
+			METAL_DEBUG_LAYER(EMetalDebugLevelFastValidation, ComputeEncoderDebug.AddWaitFence(Fence));
 		}
 		else if (BlitCommandEncoder)
 		{
-			BlitCommandEncoder.WaitForFence(Fence);
+			BlitCommandEncoder.WaitForFence(InnerFence);
+			METAL_DEBUG_LAYER(EMetalDebugLevelFastValidation, BlitEncoderDebug.AddWaitFence(Fence));
 		}
 	}
 }
@@ -510,14 +528,17 @@ void FMetalCommandEncoder::InsertDebugSignpost(ns::String const& String)
 		if (RenderCommandEncoder)
 		{
 			RenderCommandEncoder.InsertDebugSignpost(String);
+			METAL_DEBUG_LAYER(EMetalDebugLevelFastValidation, RenderEncoderDebug.InsertDebugSignpost(String));
 		}
 		if (ComputeCommandEncoder)
 		{
 			ComputeCommandEncoder.InsertDebugSignpost(String);
+			METAL_DEBUG_LAYER(EMetalDebugLevelFastValidation, ComputeEncoderDebug.InsertDebugSignpost(String));
 		}
 		if (BlitCommandEncoder)
 		{
 			BlitCommandEncoder.InsertDebugSignpost(String);
+			METAL_DEBUG_LAYER(EMetalDebugLevelFastValidation, BlitEncoderDebug.InsertDebugSignpost(String));
 		}
 	}
 }
@@ -534,14 +555,17 @@ void FMetalCommandEncoder::PushDebugGroup(ns::String const& String)
 		if (RenderCommandEncoder)
 		{
 			RenderCommandEncoder.PushDebugGroup(String);
+			METAL_DEBUG_LAYER(EMetalDebugLevelFastValidation, RenderEncoderDebug.PushDebugGroup(String));
 		}
 		else if (ComputeCommandEncoder)
 		{
 			ComputeCommandEncoder.PushDebugGroup(String);
+			METAL_DEBUG_LAYER(EMetalDebugLevelFastValidation, ComputeEncoderDebug.PushDebugGroup(String));
 		}
 		else if (BlitCommandEncoder)
 		{
 			BlitCommandEncoder.PushDebugGroup(String);
+			METAL_DEBUG_LAYER(EMetalDebugLevelFastValidation, BlitEncoderDebug.PushDebugGroup(String));
 		}
 	}
 }
@@ -554,14 +578,17 @@ void FMetalCommandEncoder::PopDebugGroup(void)
 		if (RenderCommandEncoder)
 		{
 			RenderCommandEncoder.PopDebugGroup();
+			METAL_DEBUG_LAYER(EMetalDebugLevelFastValidation, RenderEncoderDebug.PopDebugGroup());
 		}
 		else if (ComputeCommandEncoder)
 		{
 			ComputeCommandEncoder.PopDebugGroup();
+			METAL_DEBUG_LAYER(EMetalDebugLevelFastValidation, ComputeEncoderDebug.PopDebugGroup());
 		}
 		else if (BlitCommandEncoder)
 		{
 			BlitCommandEncoder.PopDebugGroup();
+			METAL_DEBUG_LAYER(EMetalDebugLevelFastValidation, BlitEncoderDebug.PopDebugGroup());
 		}
 	}
 }
@@ -630,7 +657,7 @@ void FMetalCommandEncoder::SetRenderPipelineState(FMetalShaderPipeline* Pipeline
 	check (RenderCommandEncoder);
 	{
 		METAL_STATISTIC(FMetalProfiler::GetProfiler()->EncodePipeline(CommandBufferStats, PipelineState));
-		METAL_SET_RENDER_REFLECTION(RenderCommandEncoder.GetPtr(), PipelineState);
+		METAL_DEBUG_LAYER(EMetalDebugLevelFastValidation, RenderEncoderDebug.SetPipeline(PipelineState));
 		RenderCommandEncoder.SetRenderPipelineState(PipelineState->RenderPipelineState);
 	}
 }
@@ -714,6 +741,7 @@ void FMetalCommandEncoder::SetDepthStencilState(mtlpp::DepthStencilState const&
     check (RenderCommandEncoder);
 	{
 		RenderCommandEncoder.SetDepthStencilState(InDepthStencilState);
+		METAL_DEBUG_LAYER(EMetalDebugLevelFastValidation, RenderEncoderDebug.SetDepthStencilState(InDepthStencilState));
 	}
 }
 
@@ -849,14 +877,17 @@ void FMetalCommandEncoder::SetShaderBufferOffset(mtlpp::FunctionType FunctionTyp
 		case mtlpp::FunctionType::Vertex:
 			check (RenderCommandEncoder);
 			RenderCommandEncoder.SetVertexBufferOffset(Offset + ShaderBuffers[uint32(FunctionType)].Buffers[index].GetOffset(), index);
+			METAL_DEBUG_LAYER(EMetalDebugLevelFastValidation, RenderEncoderDebug.SetBufferOffset(EMetalShaderVertex, Offset + ShaderBuffers[uint32(FunctionType)].Buffers[index].GetOffset(), index));
 			break;
 		case mtlpp::FunctionType::Fragment:
 			check(RenderCommandEncoder);
 			RenderCommandEncoder.SetFragmentBufferOffset(Offset + ShaderBuffers[uint32(FunctionType)].Buffers[index].GetOffset(), index);
+			METAL_DEBUG_LAYER(EMetalDebugLevelFastValidation, RenderEncoderDebug.SetBufferOffset(EMetalShaderFragment, Offset + ShaderBuffers[uint32(FunctionType)].Buffers[index].GetOffset(), index));
 			break;
 		case mtlpp::FunctionType::Kernel:
 			check (ComputeCommandEncoder);
 			ComputeCommandEncoder.SetBufferOffset(Offset + ShaderBuffers[uint32(FunctionType)].Buffers[index].GetOffset(), index);
+			METAL_DEBUG_LAYER(EMetalDebugLevelFastValidation, ComputeEncoderDebug.SetBufferOffset(Offset + ShaderBuffers[uint32(FunctionType)].Buffers[index].GetOffset(), index));
 			break;
 		default:
 			check(false);
@@ -872,16 +903,19 @@ void FMetalCommandEncoder::SetShaderTexture(mtlpp::FunctionType FunctionType, FM
 		case mtlpp::FunctionType::Vertex:
 			check (RenderCommandEncoder);
 			// MTLPP_VALIDATE(mtlpp::RenderCommandEncoder, RenderCommandEncoder, SafeGetRuntimeDebuggingLevel() >= EMetalDebugLevelValidation, UseResource(Texture, mtlpp::ResourceUsage::Read));
+			METAL_DEBUG_LAYER(EMetalDebugLevelFastValidation, RenderEncoderDebug.SetTexture(EMetalShaderVertex, Texture, index));
 			RenderCommandEncoder.SetVertexTexture(Texture, index);
 			break;
 		case mtlpp::FunctionType::Fragment:
 			check(RenderCommandEncoder);
 			// MTLPP_VALIDATE(mtlpp::RenderCommandEncoder, RenderCommandEncoder, SafeGetRuntimeDebuggingLevel() >= EMetalDebugLevelValidation, UseResource(Texture, mtlpp::ResourceUsage::Read));
+			METAL_DEBUG_LAYER(EMetalDebugLevelFastValidation, RenderEncoderDebug.SetTexture(EMetalShaderFragment, Texture, index));
 			RenderCommandEncoder.SetFragmentTexture(Texture, index);
 			break;
 		case mtlpp::FunctionType::Kernel:
 			check (ComputeCommandEncoder);
 			// MTLPP_VALIDATE(mtlpp::ComputeCommandEncoder, ComputeCommandEncoder, SafeGetRuntimeDebuggingLevel() >= EMetalDebugLevelValidation, UseResource(Texture, mtlpp::ResourceUsage::Read));
+			METAL_DEBUG_LAYER(EMetalDebugLevelFastValidation, ComputeEncoderDebug.SetTexture(Texture, index));
 			ComputeCommandEncoder.SetTexture(Texture, index);
 			break;
 		default:
@@ -902,14 +936,17 @@ void FMetalCommandEncoder::SetShaderSamplerState(mtlpp::FunctionType FunctionTyp
 	{
 		case mtlpp::FunctionType::Vertex:
        		check (RenderCommandEncoder);
+			METAL_DEBUG_LAYER(EMetalDebugLevelFastValidation, RenderEncoderDebug.SetSamplerState(EMetalShaderVertex, Sampler, index));
 			RenderCommandEncoder.SetVertexSamplerState(Sampler, index);
 			break;
 		case mtlpp::FunctionType::Fragment:
 			check (RenderCommandEncoder);
+			METAL_DEBUG_LAYER(EMetalDebugLevelFastValidation, RenderEncoderDebug.SetSamplerState(EMetalShaderFragment, Sampler, index));
 			RenderCommandEncoder.SetFragmentSamplerState(Sampler, index);
 			break;
 		case mtlpp::FunctionType::Kernel:
 			check (ComputeCommandEncoder);
+			METAL_DEBUG_LAYER(EMetalDebugLevelFastValidation, ComputeEncoderDebug.SetSamplerState(Sampler, index));
 			ComputeCommandEncoder.SetSamplerState(Sampler, index);
 			break;
 		default:
@@ -933,7 +970,7 @@ void FMetalCommandEncoder::SetComputePipelineState(FMetalShaderPipeline* State)
 	check (ComputeCommandEncoder);
 	{
 		METAL_STATISTIC(FMetalProfiler::GetProfiler()->EncodePipeline(CommandBufferStats, State));
-        METAL_SET_COMPUTE_REFLECTION(ComputeCommandEncoder.GetPtr(), State);
+        METAL_DEBUG_LAYER(EMetalDebugLevelFastValidation, ComputeEncoderDebug.SetPipeline(State));
 		ComputeCommandEncoder.SetComputePipelineState(State->ComputePipelineState);
 	}
 }
@@ -990,18 +1027,21 @@ void FMetalCommandEncoder::SetShaderBufferInternal(mtlpp::FunctionType Function,
 				ShaderBuffers[uint32(Function)].Bound |= (1 << Index);
 				check(RenderCommandEncoder);
 				// MTLPP_VALIDATE(mtlpp::RenderCommandEncoder, RenderCommandEncoder, SafeGetRuntimeDebuggingLevel() >= EMetalDebugLevelValidation, UseResource(Buffer, mtlpp::ResourceUsage::Read));
+				METAL_DEBUG_LAYER(EMetalDebugLevelFastValidation, RenderEncoderDebug.SetBuffer(EMetalShaderVertex, Buffer, Offset, Index));
 				RenderCommandEncoder.SetVertexBuffer(Buffer, Offset, Index);
 				break;
 			case mtlpp::FunctionType::Fragment:
 				ShaderBuffers[uint32(Function)].Bound |= (1 << Index);
 				check(RenderCommandEncoder);
 				// MTLPP_VALIDATE(mtlpp::RenderCommandEncoder, RenderCommandEncoder, SafeGetRuntimeDebuggingLevel() >= EMetalDebugLevelValidation, UseResource(Buffer, mtlpp::ResourceUsage::Read));
+				METAL_DEBUG_LAYER(EMetalDebugLevelFastValidation, RenderEncoderDebug.SetBuffer(EMetalShaderFragment, Buffer, Offset, Index));
 				RenderCommandEncoder.SetFragmentBuffer(Buffer, Offset, Index);
 				break;
 			case mtlpp::FunctionType::Kernel:
 				ShaderBuffers[uint32(Function)].Bound |= (1 << Index);
 				check(ComputeCommandEncoder);
 				// MTLPP_VALIDATE(mtlpp::ComputeCommandEncoder, ComputeCommandEncoder, SafeGetRuntimeDebuggingLevel() >= EMetalDebugLevelValidation, UseResource(Buffer, mtlpp::ResourceUsage::Read));
+				METAL_DEBUG_LAYER(EMetalDebugLevelFastValidation, ComputeEncoderDebug.SetBuffer(Buffer, Offset, Index));
 				ComputeCommandEncoder.SetBuffer(Buffer, Offset, Index);
 				break;
 			default:
@@ -1026,16 +1066,19 @@ void FMetalCommandEncoder::SetShaderBufferInternal(mtlpp::FunctionType Function,
 			case mtlpp::FunctionType::Vertex:
 				ShaderBuffers[uint32(Function)].Bound |= (1 << Index);
 				check(RenderCommandEncoder);
+				METAL_DEBUG_LAYER(EMetalDebugLevelFastValidation, RenderEncoderDebug.SetBytes(EMetalShaderVertex, Bytes, Len, Index));
 				RenderCommandEncoder.SetVertexData(Bytes, Len, Index);
 				break;
 			case mtlpp::FunctionType::Fragment:
 				ShaderBuffers[uint32(Function)].Bound |= (1 << Index);
 				check(RenderCommandEncoder);
+				METAL_DEBUG_LAYER(EMetalDebugLevelFastValidation, RenderEncoderDebug.SetBytes(EMetalShaderFragment, Bytes, Len, Index));
 				RenderCommandEncoder.SetFragmentData(Bytes, Len, Index);
 				break;
 			case mtlpp::FunctionType::Kernel:
 				ShaderBuffers[uint32(Function)].Bound |= (1 << Index);
 				check(ComputeCommandEncoder);
+				METAL_DEBUG_LAYER(EMetalDebugLevelFastValidation, ComputeEncoderDebug.SetBytes(Bytes, Len, Index));
 				ComputeCommandEncoder.SetBytes(Bytes, Len, Index);
 				break;
 			default:
diff --git a/Engine/Source/Runtime/Apple/MetalRHI/Private/MetalCommandEncoder.h b/Engine/Source/Runtime/Apple/MetalRHI/Private/MetalCommandEncoder.h
--- a/Engine/Source/Runtime/Apple/MetalRHI/Private/MetalCommandEncoder.h
+++ b/Engine/Source/Runtime/Apple/MetalRHI/Private/MetalCommandEncoder.h
@@ -6,6 +6,9 @@
 #include "MetalBuffer.h"
 #include "MetalQuery.h"
 #include "MetalDebugCommandEncoder.h"
+#include "MetalBlitCommandEncoder.h"
+#include "MetalComputeCommandEncoder.h"
+#include "MetalRenderCommandEncoder.h"
 #include "MetalFence.h"
 #include "MetalProfiler.h"
 #include "command_buffer.hpp"
@@ -157,6 +160,20 @@ public:
 	/* Get the command-buffer stats object. */
 	FMetalCommandBufferStats* GetCommandBufferStats(void);
 #endif
+
+#if METAL_DEBUG_OPTIONS
+	/** @returns The active render command encoder or nil if there isn't one. */
+	FMetalCommandBufferDebugging& GetCommandBufferDebugging(void) { return CommandBufferDebug; } 
+	
+	/** @returns The active render command encoder or nil if there isn't one. */
+	FMetalRenderCommandEncoderDebugging& GetRenderCommandEncoderDebugging(void) { return RenderEncoderDebug; } 
+	
+	/** @returns The active compute command encoder or nil if there isn't one. */
+	FMetalComputeCommandEncoderDebugging& GetComputeCommandEncoderDebugging(void) { return ComputeEncoderDebug; }
+	
+	/** @returns The active blit command encoder or nil if there isn't one. */
+	FMetalBlitCommandEncoderDebugging& GetBlitCommandEncoderDebugging(void) { return BlitEncoderDebug; }
+#endif
 	
 #pragma mark - Public Render State Mutators -
 
@@ -377,6 +394,12 @@ private:
 	mtlpp::RenderCommandEncoder RenderCommandEncoder;
 	mtlpp::ComputeCommandEncoder ComputeCommandEncoder;
 	mtlpp::BlitCommandEncoder BlitCommandEncoder;
+	
+	METAL_DEBUG_OPTION(FMetalCommandBufferDebugging CommandBufferDebug);
+	METAL_DEBUG_OPTION(FMetalRenderCommandEncoderDebugging RenderEncoderDebug);
+	METAL_DEBUG_OPTION(FMetalComputeCommandEncoderDebugging ComputeEncoderDebug);
+	METAL_DEBUG_OPTION(FMetalBlitCommandEncoderDebugging BlitEncoderDebug);
+	
 	FMetalFence EncoderFence;
 #if ENABLE_METAL_GPUPROFILE
 	FMetalCommandBufferStats* CommandBufferStats;
diff --git a/Engine/Source/Runtime/Apple/MetalRHI/Private/MetalCommandQueue.cpp b/Engine/Source/Runtime/Apple/MetalRHI/Private/MetalCommandQueue.cpp
--- a/Engine/Source/Runtime/Apple/MetalRHI/Private/MetalCommandQueue.cpp
+++ b/Engine/Source/Runtime/Apple/MetalRHI/Private/MetalCommandQueue.cpp
@@ -227,8 +227,8 @@ mtlpp::CommandBuffer FMetalCommandQueue::CreateCommandBuffer(void)
 		CmdBuffer = bUnretainedRefs ? MTLPP_VALIDATE(mtlpp::CommandQueue, CommandQueue, SafeGetRuntimeDebuggingLevel() >= EMetalDebugLevelValidation, CommandBufferWithUnretainedReferences()) : MTLPP_VALIDATE(mtlpp::CommandQueue, CommandQueue, SafeGetRuntimeDebuggingLevel() >= EMetalDebugLevelValidation, CommandBuffer());
 		
 		if (RuntimeDebuggingLevel > EMetalDebugLevelLogDebugGroups)
-		{
-			CmdBuffer = [[[FMetalDebugCommandBuffer alloc] initWithCommandBuffer:CmdBuffer.GetPtr()] autorelease];
+		{			
+			METAL_DEBUG_OPTION(FMetalCommandBufferDebugging AddDebugging(CmdBuffer));
 			MTLPP_VALIDATION(mtlpp::CommandBufferValidationTable ValidatedCommandBuffer(CmdBuffer));
 		}
 		else if (RuntimeDebuggingLevel == EMetalDebugLevelLogDebugGroups)
diff --git a/Engine/Source/Runtime/Apple/MetalRHI/Private/MetalComputeCommandEncoder.cpp b/Engine/Source/Runtime/Apple/MetalRHI/Private/MetalComputeCommandEncoder.cpp
--- a/Engine/Source/Runtime/Apple/MetalRHI/Private/MetalComputeCommandEncoder.cpp
+++ b/Engine/Source/Runtime/Apple/MetalRHI/Private/MetalComputeCommandEncoder.cpp
@@ -39,21 +39,36 @@ static id <MTLComputePipelineState> GetDebugComputeShaderState(id<MTLDevice> Dev
 }
 #endif
 
-@implementation FMetalDebugComputeCommandEncoder
+@interface FMetalDebugComputeCommandEncoder : FMetalDebugCommandEncoder
+{
+@public
+	id<MTLComputeCommandEncoder> Inner;
+	FMetalCommandBufferDebugging Buffer;
+	FMetalShaderPipeline* Pipeline;
 
-@synthesize Inner;
-@synthesize Buffer;
-@synthesize Pipeline;
+#pragma mark - Private Member Variables -
+#if METAL_DEBUG_OPTIONS
+	FMetalDebugShaderResourceMask ResourceMask;
+    FMetalDebugBufferBindings ShaderBuffers;
+    FMetalDebugTextureBindings ShaderTextures;
+    FMetalDebugSamplerBindings ShaderSamplers;
+#endif
+}
+/** Initialise the wrapper with the provided command-buffer. */
+-(id)initWithEncoder:(id<MTLComputeCommandEncoder>)Encoder andCommandBuffer:(FMetalCommandBufferDebugging const&)Buffer;
+@end
+
+@implementation FMetalDebugComputeCommandEncoder
 
 APPLE_PLATFORM_OBJECT_ALLOC_OVERRIDES(FMetalDebugComputeCommandEncoder)
 
--(id)initWithEncoder:(id<MTLComputeCommandEncoder>)Encoder andCommandBuffer:(FMetalDebugCommandBuffer*)SourceBuffer
+-(id)initWithEncoder:(id<MTLComputeCommandEncoder>)Encoder andCommandBuffer:(FMetalCommandBufferDebugging const&)SourceBuffer
 {
 	id Self = [super init];
 	if (Self)
 	{
-        Inner = [Encoder retain];
-		Buffer = [SourceBuffer retain];
+        Inner = Encoder;
+		Buffer = SourceBuffer;
         Pipeline = nil;
 	}
 	return Self;
@@ -61,47 +76,14 @@ APPLE_PLATFORM_OBJECT_ALLOC_OVERRIDES(FMetalDebugComputeCommandEncoder)
 
 -(void)dealloc
 {
-	[Inner release];
-	[Buffer release];
 	[Pipeline release];
 	[super dealloc];
 }
 
--(id <MTLDevice>) device
-{
-	return Inner.device;
-}
-
--(NSString *_Nullable)label
-{
-	return Inner.label;
-}
-
--(void)setLabel:(NSString *_Nullable)Text
-{
-	Inner.label = Text;
-}
-
-- (void)endEncoding
-{
-    [Buffer endCommandEncoder];
-    [Inner endEncoding];
-}
-
-- (void)insertDebugSignpost:(NSString *)string
-{
-    [Buffer insertDebugSignpost:string];
-    [Inner insertDebugSignpost:string];
-}
-
-- (void)pushDebugGroup:(NSString *)string
-{
-    [Buffer pushDebugGroup:string];
-    [Inner pushDebugGroup:string];
-}
+@end
 
 #if METAL_DEBUG_OPTIONS
-- (void)insertDebugDispatch
+void FMetalComputeCommandEncoderDebugging::InsertDebugDispatch()
 {
 //	switch (Buffer->DebugLevel)
 //	{
@@ -115,7 +97,7 @@ APPLE_PLATFORM_OBJECT_ALLOC_OVERRIDES(FMetalDebugComputeCommandEncoder)
 //			[Inner setBuffer:Buffer->DebugInfoBuffer offset:0 atIndex:1];
 //			[Inner setComputePipelineState:GetDebugComputeShaderState(Inner.device)];
 //			
-//			[Inner dispatchThreadgroups:MTLSizeMake(1, 1, 1) threadsPerThreadgroup:MTLSizeMake(1, 1, 1)];
+//			[Inner dispatchThreadgroups:mtlpp::Size const&Make(1, 1, 1) threadsPerThreadgroup:mtlpp::Size const&Make(1, 1, 1)];
 //
 //			if (Pipeline && Pipeline->ComputePipelineState)
 //			{
@@ -148,31 +130,65 @@ APPLE_PLATFORM_OBJECT_ALLOC_OVERRIDES(FMetalDebugComputeCommandEncoder)
 }
 #endif
 
-- (void)popDebugGroup
+FMetalComputeCommandEncoderDebugging::FMetalComputeCommandEncoderDebugging()
 {
-    [Buffer popDebugGroup];
+	
+}
+FMetalComputeCommandEncoderDebugging::FMetalComputeCommandEncoderDebugging(mtlpp::ComputeCommandEncoder& Encoder, FMetalCommandBufferDebugging& Buffer)
+: FMetalCommandEncoderDebugging((FMetalDebugCommandEncoder*)[[FMetalDebugComputeCommandEncoder alloc] initWithEncoder:Encoder.GetPtr() andCommandBuffer:Buffer])
+{
+	Buffer.BeginComputeCommandEncoder([NSString stringWithFormat:@"Compute: %@", Encoder.GetLabel().GetPtr()]);
+	Encoder.SetAssociatedObject((void const*)&FMetalComputeCommandEncoderDebugging::Get, (FMetalCommandEncoderDebugging const&)*this);
+}
+FMetalComputeCommandEncoderDebugging::FMetalComputeCommandEncoderDebugging(FMetalDebugCommandEncoder* handle)
+: FMetalCommandEncoderDebugging((FMetalDebugCommandEncoder*)handle)
+{
+	
+}
+
+FMetalComputeCommandEncoderDebugging FMetalComputeCommandEncoderDebugging::Get(mtlpp::ComputeCommandEncoder& Encoder)
+{
+	return Encoder.GetAssociatedObject<FMetalComputeCommandEncoderDebugging>((void const*)&FMetalComputeCommandEncoderDebugging::Get);
+}
+
+void FMetalComputeCommandEncoderDebugging::InsertDebugSignpost(ns::String const& Label)
+{
+	((FMetalDebugComputeCommandEncoder*)m_ptr)->Buffer.InsertDebugSignpost(Label);
+}
+void FMetalComputeCommandEncoderDebugging::PushDebugGroup(ns::String const& Group)
+{
+    ((FMetalDebugComputeCommandEncoder*)m_ptr)->Buffer.PushDebugGroup(Group);
+}
+void FMetalComputeCommandEncoderDebugging::PopDebugGroup()
+{
+    ((FMetalDebugComputeCommandEncoder*)m_ptr)->Buffer.PopDebugGroup();
 #if METAL_DEBUG_OPTIONS
-	[self insertDebugDispatch];
+	InsertDebugDispatch();
 #endif
-    [Inner popDebugGroup];
 }
 
-- (void)setComputePipelineState:(id <MTLComputePipelineState>)state
+void FMetalComputeCommandEncoderDebugging::EndEncoder()
+{
+	((FMetalDebugComputeCommandEncoder*)m_ptr)->Buffer.EndCommandEncoder();
+}
+
+void FMetalComputeCommandEncoderDebugging::DispatchThreadgroups(mtlpp::Size const& threadgroupsPerGrid, mtlpp::Size const& threadsPerThreadgroup)
 {
 #if METAL_DEBUG_OPTIONS
-	switch(Buffer->DebugLevel)
+	switch(((FMetalDebugComputeCommandEncoder*)m_ptr)->Buffer.GetPtr()->DebugLevel)
 	{
 		case EMetalDebugLevelConditionalSubmit:
 		case EMetalDebugLevelWaitForComplete:
 		case EMetalDebugLevelLogOperations:
 		{
-			// [Buffer pipeline:state.label];
+			((FMetalDebugComputeCommandEncoder*)m_ptr)->Buffer.Dispatch([NSString stringWithFormat:@"%s", __PRETTY_FUNCTION__]);
 		}
 		case EMetalDebugLevelValidation:
 		case EMetalDebugLevelResetOnBind:
 		case EMetalDebugLevelTrackResources:
+		case EMetalDebugLevelFastValidation:
 		{
-			[Buffer trackState:state];
+			Validate();
 		}
 		default:
 		{
@@ -180,28 +196,25 @@ APPLE_PLATFORM_OBJECT_ALLOC_OVERRIDES(FMetalDebugComputeCommandEncoder)
 		}
 	}
 #endif
-    [Inner setComputePipelineState:(id <MTLComputePipelineState>)state];
 }
 
-- (void)setBytes:(const void *)bytes length:(NSUInteger)length atIndex:(NSUInteger)index
+void FMetalComputeCommandEncoderDebugging::SetPipeline(FMetalShaderPipeline* Pipeline)
 {
 #if METAL_DEBUG_OPTIONS
-	switch (Buffer->DebugLevel)
+	((FMetalDebugComputeCommandEncoder*)m_ptr)->Pipeline = [Pipeline retain];
+	switch (((FMetalDebugComputeCommandEncoder*)m_ptr)->Buffer.GetPtr()->DebugLevel)
 	{
 		case EMetalDebugLevelConditionalSubmit:
 		case EMetalDebugLevelWaitForComplete:
 		case EMetalDebugLevelLogOperations:
-		case EMetalDebugLevelValidation:
 		{
-			ShaderBuffers.Buffers[index] = nil;
-			ShaderBuffers.Bytes[index] = bytes;
-			ShaderBuffers.Offsets[index] = length;
+			((FMetalDebugComputeCommandEncoder*)m_ptr)->Buffer.SetPipeline(Pipeline->ComputePipelineState.GetLabel());
 		}
+		case EMetalDebugLevelValidation:
 		case EMetalDebugLevelResetOnBind:
 		case EMetalDebugLevelTrackResources:
-		case EMetalDebugLevelFastValidation:
 		{
-			ResourceMask.BufferMask = bytes ? (ResourceMask.BufferMask | (1 << (index))) : (ResourceMask.BufferMask & ~(1 << (index)));
+			((FMetalDebugComputeCommandEncoder*)m_ptr)->Buffer.TrackState(Pipeline->ComputePipelineState);
 		}
 		default:
 		{
@@ -209,32 +222,27 @@ APPLE_PLATFORM_OBJECT_ALLOC_OVERRIDES(FMetalDebugComputeCommandEncoder)
 		}
 	}
 #endif
-	
-    [Inner setBytes:bytes length:length atIndex:index];
 }
 
-- (void)setBuffer:(nullable id <MTLBuffer>)buffer offset:(NSUInteger)offset atIndex:(NSUInteger)index
+void FMetalComputeCommandEncoderDebugging::SetBytes(const void * bytes, NSUInteger length, NSUInteger index)
 {
 #if METAL_DEBUG_OPTIONS
-	switch (Buffer->DebugLevel)
+	switch (((FMetalDebugComputeCommandEncoder*)m_ptr)->Buffer.GetPtr()->DebugLevel)
 	{
 		case EMetalDebugLevelConditionalSubmit:
 		case EMetalDebugLevelWaitForComplete:
 		case EMetalDebugLevelLogOperations:
 		case EMetalDebugLevelValidation:
 		{
-			ShaderBuffers.Buffers[index] = buffer;
-			ShaderBuffers.Bytes[index] = nil;
-			ShaderBuffers.Offsets[index] = offset;
+			((FMetalDebugComputeCommandEncoder*)m_ptr)->ShaderBuffers.Buffers[index] = nil;
+			((FMetalDebugComputeCommandEncoder*)m_ptr)->ShaderBuffers.Bytes[index] = bytes;
+			((FMetalDebugComputeCommandEncoder*)m_ptr)->ShaderBuffers.Offsets[index] = length;
 		}
 		case EMetalDebugLevelResetOnBind:
 		case EMetalDebugLevelTrackResources:
-		{
-			[Buffer trackResource:buffer];
-		}
 		case EMetalDebugLevelFastValidation:
 		{
-			ResourceMask.BufferMask = buffer ? (ResourceMask.BufferMask & (1 << (index))) : (ResourceMask.BufferMask & ~(1 << (index)));
+			((FMetalDebugComputeCommandEncoder*)m_ptr)->ResourceMask.BufferMask = bytes ? (((FMetalDebugComputeCommandEncoder*)m_ptr)->ResourceMask.BufferMask | (1 << (index))) : (((FMetalDebugComputeCommandEncoder*)m_ptr)->ResourceMask.BufferMask & ~(1 << (index)));
 		}
 		default:
 		{
@@ -242,27 +250,29 @@ APPLE_PLATFORM_OBJECT_ALLOC_OVERRIDES(FMetalDebugComputeCommandEncoder)
 		}
 	}
 #endif
-	
-    [Inner setBuffer:buffer offset:offset atIndex:index];
 }
-
-- (void)setBufferOffset:(NSUInteger)offset atIndex:(NSUInteger)index
+void FMetalComputeCommandEncoderDebugging::SetBuffer( FMetalBuffer const& buffer, NSUInteger offset, NSUInteger index)
 {
 #if METAL_DEBUG_OPTIONS
-	switch (Buffer->DebugLevel)
+	switch (((FMetalDebugComputeCommandEncoder*)m_ptr)->Buffer.GetPtr()->DebugLevel)
 	{
 		case EMetalDebugLevelConditionalSubmit:
 		case EMetalDebugLevelWaitForComplete:
 		case EMetalDebugLevelLogOperations:
 		case EMetalDebugLevelValidation:
 		{
-			ShaderBuffers.Offsets[index] = offset;
+			((FMetalDebugComputeCommandEncoder*)m_ptr)->ShaderBuffers.Buffers[index] = buffer;
+			((FMetalDebugComputeCommandEncoder*)m_ptr)->ShaderBuffers.Bytes[index] = nil;
+			((FMetalDebugComputeCommandEncoder*)m_ptr)->ShaderBuffers.Offsets[index] = offset;
 		}
 		case EMetalDebugLevelResetOnBind:
 		case EMetalDebugLevelTrackResources:
+		{
+			((FMetalDebugComputeCommandEncoder*)m_ptr)->Buffer.TrackResource(buffer);
+		}
 		case EMetalDebugLevelFastValidation:
 		{
-			check(ResourceMask.BufferMask | (1 << (index)));
+			((FMetalDebugComputeCommandEncoder*)m_ptr)->ResourceMask.BufferMask = buffer ? (((FMetalDebugComputeCommandEncoder*)m_ptr)->ResourceMask.BufferMask | (1 << (index))) : (((FMetalDebugComputeCommandEncoder*)m_ptr)->ResourceMask.BufferMask & ~(1 << (index)));
 		}
 		default:
 		{
@@ -270,65 +280,24 @@ APPLE_PLATFORM_OBJECT_ALLOC_OVERRIDES(FMetalDebugComputeCommandEncoder)
 		}
 	}
 #endif
-	
-    [Inner setBufferOffset:offset atIndex:index];
-}
-
-- (void)setBuffers:(const id <MTLBuffer> __nullable [__nonnull])buffers offsets:(const NSUInteger [__nonnull])offsets withRange:(NSRange)range
-{
-#if METAL_DEBUG_OPTIONS
-	for (uint32 i = 0; i < range.length; i++)
-	{
-		switch (Buffer->DebugLevel)
-		{
-			case EMetalDebugLevelConditionalSubmit:
-			case EMetalDebugLevelWaitForComplete:
-			case EMetalDebugLevelLogOperations:
-			case EMetalDebugLevelValidation:
-			{
-				ShaderBuffers.Buffers[i + range.location] = buffers[i];
-				ShaderBuffers.Bytes[i + range.location] = nil;
-				ShaderBuffers.Offsets[i + range.location] = offsets[i];
-			}
-			case EMetalDebugLevelResetOnBind:
-			case EMetalDebugLevelTrackResources:
-			{
-				[Buffer trackResource:buffers[i]];
-			}
-			case EMetalDebugLevelFastValidation:
-			{
-				ResourceMask.BufferMask = buffers[i] ? (ResourceMask.BufferMask | (1 << (i + range.location))) : (ResourceMask.BufferMask & ~(1 << (i + range.location)));
-			}
-			default:
-			{
-				break;
-			}
-		}
-    }
-#endif
-    [Inner setBuffers:buffers offsets:offsets withRange:range];
 }
-
-- (void)setTexture:(nullable id <MTLTexture>)texture atIndex:(NSUInteger)index
+void FMetalComputeCommandEncoderDebugging::SetBufferOffset(NSUInteger offset, NSUInteger index)
 {
 #if METAL_DEBUG_OPTIONS
-	switch (Buffer->DebugLevel)
+	switch (((FMetalDebugComputeCommandEncoder*)m_ptr)->Buffer.GetPtr()->DebugLevel)
 	{
 		case EMetalDebugLevelConditionalSubmit:
 		case EMetalDebugLevelWaitForComplete:
 		case EMetalDebugLevelLogOperations:
 		case EMetalDebugLevelValidation:
 		{
-			ShaderTextures.Textures[index] = texture;
+			((FMetalDebugComputeCommandEncoder*)m_ptr)->ShaderBuffers.Offsets[index] = offset;
 		}
 		case EMetalDebugLevelResetOnBind:
 		case EMetalDebugLevelTrackResources:
-		{
-			[Buffer trackResource:texture];
-		}
 		case EMetalDebugLevelFastValidation:
 		{
-			ResourceMask.TextureMask = texture ? (ResourceMask.TextureMask | (1 << (index))) : (ResourceMask.TextureMask & ~(1 << (index)));
+			check(((FMetalDebugComputeCommandEncoder*)m_ptr)->ResourceMask.BufferMask & (1 << (index)));
 		}
 		default:
 		{
@@ -336,68 +305,28 @@ APPLE_PLATFORM_OBJECT_ALLOC_OVERRIDES(FMetalDebugComputeCommandEncoder)
 		}
 	}
 #endif
-	
-    [Inner setTexture:texture atIndex:index];
 }
 
-#if METAL_NEW_NONNULL_DECL
-- (void)setTextures:(const id <MTLTexture> __nullable[__nonnull])textures withRange : (NSRange)range
-#else
-- (void)setTextures:(const id <MTLTexture> __nullable [__nullable])textures withRange:(NSRange)range
-#endif
+void FMetalComputeCommandEncoderDebugging::SetTexture( FMetalTexture const& texture, NSUInteger index)
 {
 #if METAL_DEBUG_OPTIONS
-    for (uint32 i = 0; i < range.length; i++)
-	{
-		switch (Buffer->DebugLevel)
-		{
-			case EMetalDebugLevelConditionalSubmit:
-			case EMetalDebugLevelWaitForComplete:
-			case EMetalDebugLevelLogOperations:
-			case EMetalDebugLevelValidation:
-			{
-				ShaderTextures.Textures[i + range.location] = textures[i];
-			}
-			case EMetalDebugLevelResetOnBind:
-			case EMetalDebugLevelTrackResources:
-			{
-				[Buffer trackResource:textures[i]];
-			}
-			case EMetalDebugLevelFastValidation:
-			{
-				ResourceMask.TextureMask = textures[i] ? (ResourceMask.TextureMask | (1 << (i + range.location))) : (ResourceMask.TextureMask & ~(1 << (i + range.location)));
-			}
-			default:
-			{
-				break;
-			}
-		}
-		
-    }
-#endif
-    [Inner setTextures:textures withRange:range];
-}
-
-- (void)setSamplerState:(nullable id <MTLSamplerState>)sampler atIndex:(NSUInteger)index
-{
-#if METAL_DEBUG_OPTIONS
-	switch (Buffer->DebugLevel)
+	switch (((FMetalDebugComputeCommandEncoder*)m_ptr)->Buffer.GetPtr()->DebugLevel)
 	{
 		case EMetalDebugLevelConditionalSubmit:
 		case EMetalDebugLevelWaitForComplete:
 		case EMetalDebugLevelLogOperations:
 		case EMetalDebugLevelValidation:
 		{
-			ShaderSamplers.Samplers[index] = sampler;
+			((FMetalDebugComputeCommandEncoder*)m_ptr)->ShaderTextures.Textures[index] = texture;
 		}
 		case EMetalDebugLevelResetOnBind:
 		case EMetalDebugLevelTrackResources:
 		{
-			[Buffer trackState:sampler];
+			((FMetalDebugComputeCommandEncoder*)m_ptr)->Buffer.TrackResource(texture);
 		}
 		case EMetalDebugLevelFastValidation:
 		{
-			ResourceMask.SamplerMask = sampler ? (ResourceMask.SamplerMask | (1 << (index))) : (ResourceMask.SamplerMask & ~(1 << (index)));
+			((FMetalDebugComputeCommandEncoder*)m_ptr)->ResourceMask.TextureMask = texture ? (((FMetalDebugComputeCommandEncoder*)m_ptr)->ResourceMask.TextureMask | (1 << (index))) : (((FMetalDebugComputeCommandEncoder*)m_ptr)->ResourceMask.TextureMask & ~(1 << (index)));
 		}
 		default:
 		{
@@ -405,67 +334,28 @@ APPLE_PLATFORM_OBJECT_ALLOC_OVERRIDES(FMetalDebugComputeCommandEncoder)
 		}
 	}
 #endif
-    [Inner setSamplerState:sampler atIndex:index];
-}
-
-#if METAL_NEW_NONNULL_DECL
-- (void)setSamplerStates:(const id <MTLSamplerState> __nullable[__nonnull])samplers withRange : (NSRange)range
-#else
-- (void)setSamplerStates:(const id <MTLSamplerState> __nullable [__nullable])samplers withRange:(NSRange)range
-#endif
-{
-#if METAL_DEBUG_OPTIONS
-    for(uint32 i = 0; i < range.length; i++)
-	{
-		switch (Buffer->DebugLevel)
-		{
-			case EMetalDebugLevelConditionalSubmit:
-			case EMetalDebugLevelWaitForComplete:
-			case EMetalDebugLevelLogOperations:
-			case EMetalDebugLevelValidation:
-			{
-				ShaderSamplers.Samplers[i + range.location] = samplers[i];
-			}
-			case EMetalDebugLevelResetOnBind:
-			case EMetalDebugLevelTrackResources:
-			{
-				[Buffer trackState:samplers[i]];
-			}
-			case EMetalDebugLevelFastValidation:
-			{
-				ResourceMask.SamplerMask = samplers[i] ? (ResourceMask.SamplerMask | (1 << (i + range.location))) : (ResourceMask.SamplerMask & ~(1 << (i + range.location)));
-			}
-			default:
-			{
-				break;
-			}
-		}
-		
-    }
-#endif
-    [Inner setSamplerStates:samplers withRange:range];
 }
 
-- (void)setSamplerState:(nullable id <MTLSamplerState>)sampler lodMinClamp:(float)lodMinClamp lodMaxClamp:(float)lodMaxClamp atIndex:(NSUInteger)index
+void FMetalComputeCommandEncoderDebugging::SetSamplerState( mtlpp::SamplerState const& sampler, NSUInteger index)
 {
 #if METAL_DEBUG_OPTIONS
-	switch (Buffer->DebugLevel)
+	switch (((FMetalDebugComputeCommandEncoder*)m_ptr)->Buffer.GetPtr()->DebugLevel)
 	{
 		case EMetalDebugLevelConditionalSubmit:
 		case EMetalDebugLevelWaitForComplete:
 		case EMetalDebugLevelLogOperations:
 		case EMetalDebugLevelValidation:
 		{
-			ShaderSamplers.Samplers[index] = sampler;
+			((FMetalDebugComputeCommandEncoder*)m_ptr)->ShaderSamplers.Samplers[index] = sampler;
 		}
 		case EMetalDebugLevelResetOnBind:
 		case EMetalDebugLevelTrackResources:
 		{
-			[Buffer trackState:sampler];
+			((FMetalDebugComputeCommandEncoder*)m_ptr)->Buffer.TrackState(sampler);
 		}
 		case EMetalDebugLevelFastValidation:
 		{
-			ResourceMask.SamplerMask = sampler ? (ResourceMask.SamplerMask | (1 << (index))) : (ResourceMask.SamplerMask & ~(1 << (index)));
+			((FMetalDebugComputeCommandEncoder*)m_ptr)->ResourceMask.SamplerMask = sampler ? (((FMetalDebugComputeCommandEncoder*)m_ptr)->ResourceMask.SamplerMask | (1 << (index))) : (((FMetalDebugComputeCommandEncoder*)m_ptr)->ResourceMask.SamplerMask & ~(1 << (index)));
 		}
 		default:
 		{
@@ -473,74 +363,28 @@ APPLE_PLATFORM_OBJECT_ALLOC_OVERRIDES(FMetalDebugComputeCommandEncoder)
 		}
 	}
 #endif
-    [Inner setSamplerState:sampler lodMinClamp:lodMinClamp lodMaxClamp:lodMaxClamp atIndex:index];
-}
-
-#if METAL_NEW_NONNULL_DECL
-- (void)setSamplerStates:(const id <MTLSamplerState> __nullable[__nonnull])samplers lodMinClamps : (const float[__nonnull])lodMinClamps lodMaxClamps : (const float[__nonnull])lodMaxClamps withRange : (NSRange)range
-#else
-- (void)setSamplerStates:(const id <MTLSamplerState> __nullable [__nullable])samplers lodMinClamps:(const float [__nullable])lodMinClamps lodMaxClamps:(const float [__nullable])lodMaxClamps withRange:(NSRange)range
-#endif
-{
-#if METAL_DEBUG_OPTIONS
-    for(uint32 i = 0; i < range.length; i++)
-	{
-		switch (Buffer->DebugLevel)
-		{
-			case EMetalDebugLevelConditionalSubmit:
-			case EMetalDebugLevelWaitForComplete:
-			case EMetalDebugLevelLogOperations:
-			case EMetalDebugLevelValidation:
-			{
-				ShaderSamplers.Samplers[i + range.location] = samplers[i];
-			}
-			case EMetalDebugLevelResetOnBind:
-			case EMetalDebugLevelTrackResources:
-			{
-				[Buffer trackState:samplers[i]];
-			}
-			case EMetalDebugLevelFastValidation:
-			{
-				ResourceMask.SamplerMask = samplers[i] ? (ResourceMask.SamplerMask | (1 << (i + range.location))) : (ResourceMask.SamplerMask & ~(1 << (i + range.location)));
-			}
-			default:
-			{
-				break;
-			}
-		}
-		
-    }
-#endif
-    [Inner setSamplerStates:samplers lodMinClamps:lodMinClamps lodMaxClamps:lodMaxClamps withRange:range];
 }
 
-- (void)setThreadgroupMemoryLength:(NSUInteger)length atIndex:(NSUInteger)index
-{
-    [Inner setThreadgroupMemoryLength:(NSUInteger)length atIndex:(NSUInteger)index];
-}
-
-- (void)setStageInRegion:(MTLRegion)region
-{
-    [Inner setStageInRegion:(MTLRegion)region];
-}
-
-- (void)dispatchThreadgroups:(MTLSize)threadgroupsPerGrid threadsPerThreadgroup:(MTLSize)threadsPerThreadgroup
+void FMetalComputeCommandEncoderDebugging::SetSamplerState( mtlpp::SamplerState const& sampler, float lodMinClamp, float lodMaxClamp, NSUInteger index)
 {
 #if METAL_DEBUG_OPTIONS
-	switch(Buffer->DebugLevel)
+	switch (((FMetalDebugComputeCommandEncoder*)m_ptr)->Buffer.GetPtr()->DebugLevel)
 	{
 		case EMetalDebugLevelConditionalSubmit:
 		case EMetalDebugLevelWaitForComplete:
 		case EMetalDebugLevelLogOperations:
+		case EMetalDebugLevelValidation:
 		{
-			[Buffer dispatch:[NSString stringWithFormat:@"%s", __PRETTY_FUNCTION__]];
+			((FMetalDebugComputeCommandEncoder*)m_ptr)->ShaderSamplers.Samplers[index] = sampler;
 		}
-		case EMetalDebugLevelValidation:
 		case EMetalDebugLevelResetOnBind:
 		case EMetalDebugLevelTrackResources:
+		{
+			((FMetalDebugComputeCommandEncoder*)m_ptr)->Buffer.TrackState(sampler);
+		}
 		case EMetalDebugLevelFastValidation:
 		{
-			[self validate];
+			((FMetalDebugComputeCommandEncoder*)m_ptr)->ResourceMask.SamplerMask = sampler ? (((FMetalDebugComputeCommandEncoder*)m_ptr)->ResourceMask.SamplerMask | (1 << (index))) : (((FMetalDebugComputeCommandEncoder*)m_ptr)->ResourceMask.SamplerMask & ~(1 << (index)));
 		}
 		default:
 		{
@@ -548,29 +392,28 @@ APPLE_PLATFORM_OBJECT_ALLOC_OVERRIDES(FMetalDebugComputeCommandEncoder)
 		}
 	}
 #endif
-    [Inner dispatchThreadgroups:(MTLSize)threadgroupsPerGrid threadsPerThreadgroup:(MTLSize)threadsPerThreadgroup];
 }
 
-- (void)dispatchThreadgroupsWithIndirectBuffer:(id <MTLBuffer>)indirectBuffer indirectBufferOffset:(NSUInteger)indirectBufferOffset threadsPerThreadgroup:(MTLSize)threadsPerThreadgroup
+void FMetalComputeCommandEncoderDebugging::DispatchThreadgroupsWithIndirectBuffer(FMetalBuffer const& indirectBuffer, NSUInteger indirectBufferOffset, mtlpp::Size const& threadsPerThreadgroup)
 {
 #if METAL_DEBUG_OPTIONS
-	switch(Buffer->DebugLevel)
+	switch(((FMetalDebugComputeCommandEncoder*)m_ptr)->Buffer.GetPtr()->DebugLevel)
 	{
 		case EMetalDebugLevelConditionalSubmit:
 		case EMetalDebugLevelWaitForComplete:
 		case EMetalDebugLevelLogOperations:
 		{
-			[Buffer dispatch:[NSString stringWithFormat:@"%s", __PRETTY_FUNCTION__]];
+			((FMetalDebugComputeCommandEncoder*)m_ptr)->Buffer.Dispatch([NSString stringWithFormat:@"%s", __PRETTY_FUNCTION__]);
 		}
 		case EMetalDebugLevelValidation:
 		case EMetalDebugLevelResetOnBind:
 		case EMetalDebugLevelTrackResources:
 		{
-			[Buffer trackResource:indirectBuffer];
+			((FMetalDebugComputeCommandEncoder*)m_ptr)->Buffer.TrackResource(indirectBuffer);
 		}
 		case EMetalDebugLevelFastValidation:
 		{
-			[self validate];
+			Validate();
 		}
 		default:
 		{
@@ -578,105 +421,23 @@ APPLE_PLATFORM_OBJECT_ALLOC_OVERRIDES(FMetalDebugComputeCommandEncoder)
 		}
 	}
 #endif
-	
-    [Inner dispatchThreadgroupsWithIndirectBuffer:(id <MTLBuffer>)indirectBuffer indirectBufferOffset:(NSUInteger)indirectBufferOffset threadsPerThreadgroup:(MTLSize)threadsPerThreadgroup];
 }
 
-- (void)updateFence:(id <MTLFence>)fence
-{
-#if METAL_DEBUG_OPTIONS
-	if (fence && Buffer->DebugLevel >= EMetalDebugLevelValidation)
-	{
-		[self addUpdateFence:fence];
-        if (((FMetalDebugFence*)fence).Inner)
-        {
-		[Inner updateFence:((FMetalDebugFence*)fence).Inner];
-	}
-	}
-	else
-#endif
-	{
-		[Inner updateFence:(id <MTLFence>)fence];
-	}
-}
-
-- (void)waitForFence:(id <MTLFence>)fence
-{
-#if METAL_DEBUG_OPTIONS
-	if (fence && Buffer->DebugLevel >= EMetalDebugLevelValidation)
-	{
-		[self addWaitFence:fence];
-        if (((FMetalDebugFence*)fence).Inner)
-        {
-		[Inner waitForFence:((FMetalDebugFence*)fence).Inner];
-	}
-	}
-	else
-#endif
-	{
-		[Inner waitForFence:(id <MTLFence>)fence];
-	}
-}
-
-#if METAL_SUPPORTS_INDIRECT_ARGUMENT_BUFFERS
-- (void)useResource:(id <MTLResource>)resource usage:(MTLResourceUsage)usage
-{
-	if (GMetalSupportsIndirectArgumentBuffers)
-	{
-		[Inner useResource:resource usage:usage];
-	}
-}
-
-- (void)useResources:(const id <MTLResource> [])resources count:(NSUInteger)count usage:(MTLResourceUsage)usage
-{
-	if (GMetalSupportsIndirectArgumentBuffers)
-	{
-		[Inner useResources:resources count:count usage:usage];
-	}
-}
-
-- (void)useHeap:(id <MTLHeap>)heap
-{
-	if (GMetalSupportsIndirectArgumentBuffers)
-	{
-		[Inner useHeap:heap];
-	}
-}
-
-- (void)useHeaps:(const id <MTLHeap> [])heaps count:(NSUInteger)count
-{
-	if (GMetalSupportsIndirectArgumentBuffers)
-	{
-		[Inner useHeaps:heaps count:count];
-	}
-}
-#endif // #if METAL_SUPPORTS_INDIRECT_ARGUMENT_BUFFERS
-
--(NSString*) description
-{
-	return [Inner description];
-}
-
--(NSString*) debugDescription
-{
-	return [Inner debugDescription];
-}
-
-- (void)validate
+void FMetalComputeCommandEncoderDebugging::Validate()
 {
 	bool bOK = true;
 	
 #if METAL_DEBUG_OPTIONS
-	switch (Buffer->DebugLevel)
+	switch(((FMetalDebugComputeCommandEncoder*)m_ptr)->Buffer.GetPtr()->DebugLevel)
 	{
 		case EMetalDebugLevelConditionalSubmit:
 		case EMetalDebugLevelWaitForComplete:
 		case EMetalDebugLevelLogOperations:
 		case EMetalDebugLevelValidation:
 		{
-			check(Pipeline);
+			check(((FMetalDebugComputeCommandEncoder*)m_ptr)->Pipeline);
 			
-			MTLComputePipelineReflection* Reflection = Pipeline->ComputePipelineReflection;
+			MTLComputePipelineReflection* Reflection = ((FMetalDebugComputeCommandEncoder*)m_ptr)->Pipeline->ComputePipelineReflection;
 			check(Reflection);
 	
 			NSArray<MTLArgument*>* Arguments = Reflection.arguments;
@@ -689,7 +450,7 @@ APPLE_PLATFORM_OBJECT_ALLOC_OVERRIDES(FMetalDebugComputeCommandEncoder)
 					case MTLArgumentTypeBuffer:
 					{
 						checkf(Arg.index < ML_MaxBuffers, TEXT("Metal buffer index exceeded!"));
-						if ((ShaderBuffers.Buffers[Arg.index] == nil && ShaderBuffers.Bytes[Arg.index] == nil))
+						if ((((FMetalDebugComputeCommandEncoder*)m_ptr)->ShaderBuffers.Buffers[Arg.index] == nil && ((FMetalDebugComputeCommandEncoder*)m_ptr)->ShaderBuffers.Bytes[Arg.index] == nil))
 						{
 							UE_LOG(LogMetal, Warning, TEXT("Unbound buffer at Metal index %u which will crash the driver: %s"), (uint32)Arg.index, *FString([Arg description]));
 							bOK = false;
@@ -703,14 +464,14 @@ APPLE_PLATFORM_OBJECT_ALLOC_OVERRIDES(FMetalDebugComputeCommandEncoder)
 					case MTLArgumentTypeTexture:
 					{
 						checkf(Arg.index < ML_MaxTextures, TEXT("Metal texture index exceeded!"));
-						if (ShaderTextures.Textures[Arg.index] == nil)
+						if (((FMetalDebugComputeCommandEncoder*)m_ptr)->ShaderTextures.Textures[Arg.index] == nil)
 						{
 							UE_LOG(LogMetal, Warning, TEXT("Unbound texture at Metal index %u  which will crash the driver: %s"), (uint32)Arg.index, *FString([Arg description]));
 							bOK = false;
 						}
-						else if (ShaderTextures.Textures[Arg.index].textureType != Arg.textureType)
+						else if (((FMetalDebugComputeCommandEncoder*)m_ptr)->ShaderTextures.Textures[Arg.index].textureType != Arg.textureType)
 						{
-							UE_LOG(LogMetal, Warning, TEXT("Incorrect texture type bound at Metal index %u which will crash the driver: %s\n%s"), (uint32)Arg.index, *FString([Arg description]), *FString([ShaderTextures.Textures[Arg.index] description]));
+							UE_LOG(LogMetal, Warning, TEXT("Incorrect texture type bound at Metal index %u which will crash the driver: %s\n%s"), (uint32)Arg.index, *FString([Arg description]), *FString([((FMetalDebugComputeCommandEncoder*)m_ptr)->ShaderTextures.Textures[Arg.index] description]));
 							bOK = false;
 						}
 						break;
@@ -718,7 +479,7 @@ APPLE_PLATFORM_OBJECT_ALLOC_OVERRIDES(FMetalDebugComputeCommandEncoder)
 					case MTLArgumentTypeSampler:
 					{
 						checkf(Arg.index < ML_MaxSamplers, TEXT("Metal sampler index exceeded!"));
-						if (ShaderSamplers.Samplers[Arg.index] == nil)
+						if (((FMetalDebugComputeCommandEncoder*)m_ptr)->ShaderSamplers.Samplers[Arg.index] == nil)
 						{
 							UE_LOG(LogMetal, Warning, TEXT("Unbound sampler at Metal index %u which will crash the driver: %s"), (uint32)Arg.index, *FString([Arg description]));
 							bOK = false;
@@ -736,41 +497,41 @@ APPLE_PLATFORM_OBJECT_ALLOC_OVERRIDES(FMetalDebugComputeCommandEncoder)
 		case EMetalDebugLevelTrackResources:
 		case EMetalDebugLevelFastValidation:
 		{
-			check(Pipeline);
+			check(((FMetalDebugComputeCommandEncoder*)m_ptr)->Pipeline);
 			
-			FMetalTextureMask TextureMask = (ResourceMask.TextureMask & Pipeline->ResourceMask[EMetalShaderCompute].TextureMask);
-			if (TextureMask != Pipeline->ResourceMask[EMetalShaderCompute].TextureMask)
+			FMetalTextureMask TextureMask = (((FMetalDebugComputeCommandEncoder*)m_ptr)->ResourceMask.TextureMask & ((FMetalDebugComputeCommandEncoder*)m_ptr)->Pipeline->ResourceMask[EMetalShaderCompute].TextureMask);
+			if (TextureMask != ((FMetalDebugComputeCommandEncoder*)m_ptr)->Pipeline->ResourceMask[EMetalShaderCompute].TextureMask)
 			{
 				bOK = false;
 				for (uint32 i = 0; i < ML_MaxTextures; i++)
 				{
-					if ((Pipeline->ResourceMask[EMetalShaderCompute].TextureMask & (1 < i)) && ((TextureMask & (1 < i)) != (Pipeline->ResourceMask[EMetalShaderCompute].TextureMask & (1 < i))))
+					if ((((FMetalDebugComputeCommandEncoder*)m_ptr)->Pipeline->ResourceMask[EMetalShaderCompute].TextureMask & (1 < i)) && ((TextureMask & (1 < i)) != (((FMetalDebugComputeCommandEncoder*)m_ptr)->Pipeline->ResourceMask[EMetalShaderCompute].TextureMask & (1 < i))))
 					{
 						UE_LOG(LogMetal, Warning, TEXT("Unbound texture at Metal index %u which will crash the driver"), i);
 					}
 				}
 			}
 			
-			FMetalBufferMask BufferMask = (ResourceMask.BufferMask & Pipeline->ResourceMask[EMetalShaderCompute].BufferMask);
-			if (BufferMask != Pipeline->ResourceMask[EMetalShaderCompute].BufferMask)
+			FMetalBufferMask BufferMask = (((FMetalDebugComputeCommandEncoder*)m_ptr)->ResourceMask.BufferMask & ((FMetalDebugComputeCommandEncoder*)m_ptr)->Pipeline->ResourceMask[EMetalShaderCompute].BufferMask);
+			if (BufferMask != ((FMetalDebugComputeCommandEncoder*)m_ptr)->Pipeline->ResourceMask[EMetalShaderCompute].BufferMask)
 			{
 				bOK = false;
 				for (uint32 i = 0; i < ML_MaxBuffers; i++)
 				{
-					if ((Pipeline->ResourceMask[EMetalShaderCompute].BufferMask & (1 < i)) && ((BufferMask & (1 < i)) != (Pipeline->ResourceMask[EMetalShaderCompute].BufferMask & (1 < i))))
+					if ((((FMetalDebugComputeCommandEncoder*)m_ptr)->Pipeline->ResourceMask[EMetalShaderCompute].BufferMask & (1 < i)) && ((BufferMask & (1 < i)) != (((FMetalDebugComputeCommandEncoder*)m_ptr)->Pipeline->ResourceMask[EMetalShaderCompute].BufferMask & (1 < i))))
 					{
 						UE_LOG(LogMetal, Warning, TEXT("Unbound buffer at Metal index %u which will crash the driver"), i);
 					}
 				}
 			}
 			
-			FMetalSamplerMask SamplerMask = (ResourceMask.SamplerMask & Pipeline->ResourceMask[EMetalShaderCompute].SamplerMask);
-			if (SamplerMask != Pipeline->ResourceMask[EMetalShaderCompute].SamplerMask)
+			FMetalSamplerMask SamplerMask = (((FMetalDebugComputeCommandEncoder*)m_ptr)->ResourceMask.SamplerMask & ((FMetalDebugComputeCommandEncoder*)m_ptr)->Pipeline->ResourceMask[EMetalShaderCompute].SamplerMask);
+			if (SamplerMask != ((FMetalDebugComputeCommandEncoder*)m_ptr)->Pipeline->ResourceMask[EMetalShaderCompute].SamplerMask)
 			{
 				bOK = false;
 				for (uint32 i = 0; i < ML_MaxSamplers; i++)
 				{
-					if ((Pipeline->ResourceMask[EMetalShaderCompute].SamplerMask & (1 < i)) && ((SamplerMask & (1 < i)) != (Pipeline->ResourceMask[EMetalShaderCompute].SamplerMask & (1 < i))))
+					if ((((FMetalDebugComputeCommandEncoder*)m_ptr)->Pipeline->ResourceMask[EMetalShaderCompute].SamplerMask & (1 < i)) && ((SamplerMask & (1 < i)) != (((FMetalDebugComputeCommandEncoder*)m_ptr)->Pipeline->ResourceMask[EMetalShaderCompute].SamplerMask & (1 < i))))
 					{
 						UE_LOG(LogMetal, Warning, TEXT("Unbound sampler at Metal index %u which will crash the driver"), i);
 					}
@@ -787,36 +548,9 @@ APPLE_PLATFORM_OBJECT_ALLOC_OVERRIDES(FMetalDebugComputeCommandEncoder)
 	
     if (!bOK)
     {
-        UE_LOG(LogMetal, Error, TEXT("Metal Validation failures for compute shader:\n%s"), (Pipeline && Pipeline->ComputeSource) ? *FString(Pipeline->ComputeSource) : TEXT("nil"));
+        UE_LOG(LogMetal, Error, TEXT("Metal Validation failures for compute shader:\n%s"), (((FMetalDebugComputeCommandEncoder*)m_ptr)->Pipeline && ((FMetalDebugComputeCommandEncoder*)m_ptr)->Pipeline->ComputeSource) ? *FString(((FMetalDebugComputeCommandEncoder*)m_ptr)->Pipeline->ComputeSource) : TEXT("nil"));
     }
 #endif
 }
 
--(id<MTLCommandEncoder>)commandEncoder
-{
-	return self;
-}
-
-#if METAL_SUPPORTS_TILE_SHADERS
-- (void)setImageblockWidth:(NSUInteger)width height:(NSUInteger)height
-{
-	if (GMetalSupportsTileShaders)
-	{
-		[Inner setImageblockWidth:width height:height];
-	}
-}
-#endif
-
-#if METAL_SUPPORTS_CAPTURE_MANAGER && !PLATFORM_TVOS
-- (void)dispatchThreads:(MTLSize)threadsPerGrid threadsPerThreadgroup:(MTLSize)threadsPerThreadgroup
-{
-	if (GMetalSupportsCaptureManager)
-	{
-		[Inner dispatchThreads:threadsPerGrid threadsPerThreadgroup:threadsPerThreadgroup];
-	}
-}
-#endif
-
-@end
-
 NS_ASSUME_NONNULL_END
diff --git a/Engine/Source/Runtime/Apple/MetalRHI/Private/MetalComputeCommandEncoder.h b/Engine/Source/Runtime/Apple/MetalRHI/Private/MetalComputeCommandEncoder.h
--- a/Engine/Source/Runtime/Apple/MetalRHI/Private/MetalComputeCommandEncoder.h
+++ b/Engine/Source/Runtime/Apple/MetalRHI/Private/MetalComputeCommandEncoder.h
@@ -6,41 +6,35 @@
 
 NS_ASSUME_NONNULL_BEGIN
 
-@class FMetalDebugCommandBuffer;
+@class FMetalDebugComputeCommandEncoder;
 @class FMetalShaderPipeline;
 
-@interface FMetalDebugComputeCommandEncoder : FMetalDebugCommandEncoder<MTLComputeCommandEncoder>
-{
-@private
-#pragma mark - Private Member Variables -
-#if METAL_DEBUG_OPTIONS
-	FMetalDebugShaderResourceMask ResourceMask;
-    FMetalDebugBufferBindings ShaderBuffers;
-    FMetalDebugTextureBindings ShaderTextures;
-    FMetalDebugSamplerBindings ShaderSamplers;
-#endif
-}
-
-/** The wrapped native command-encoder for which we collect debug information. */
-@property (readonly, retain) id<MTLComputeCommandEncoder> Inner;
-@property (readonly, retain) FMetalDebugCommandBuffer* Buffer;
-@property (nonatomic, retain) FMetalShaderPipeline* Pipeline;
-
-/** Initialise the wrapper with the provided command-buffer. */
--(id)initWithEncoder:(id<MTLComputeCommandEncoder>)Encoder andCommandBuffer:(FMetalDebugCommandBuffer*)Buffer;
 
-/** Validates the pipeline/binding state */
--(void)validate;
+class FMetalComputeCommandEncoderDebugging : public FMetalCommandEncoderDebugging
+{
+	void InsertDebugDispatch();
+	void Validate();
+public:
+	FMetalComputeCommandEncoderDebugging();
+	FMetalComputeCommandEncoderDebugging(mtlpp::ComputeCommandEncoder& Encoder, FMetalCommandBufferDebugging& Buffer);
+	FMetalComputeCommandEncoderDebugging(FMetalDebugCommandEncoder* handle);
+	
+	static FMetalComputeCommandEncoderDebugging Get(mtlpp::ComputeCommandEncoder& Buffer);
+	
+	void InsertDebugSignpost(ns::String const& Label);
+	void PushDebugGroup(ns::String const& Group);
+	void PopDebugGroup();
+	void EndEncoder();
+	void DispatchThreadgroups(mtlpp::Size const& threadgroupsPerGrid, mtlpp::Size const& threadsPerThreadgroup);
+	void SetPipeline(FMetalShaderPipeline* Pipeline);
+	void SetBytes(const void * bytes, NSUInteger length, NSUInteger index);
+	void SetBuffer( FMetalBuffer const& buffer, NSUInteger offset, NSUInteger index);
+	void SetBufferOffset(NSUInteger offset, NSUInteger index);
+	void SetTexture( FMetalTexture const& texture, NSUInteger index);
+	void SetSamplerState( mtlpp::SamplerState const& sampler, NSUInteger index);
+	void SetSamplerState( mtlpp::SamplerState const& sampler, float lodMinClamp, float lodMaxClamp, NSUInteger index);
+	
+	void DispatchThreadgroupsWithIndirectBuffer(FMetalBuffer const& indirectBuffer, NSUInteger indirectBufferOffset, mtlpp::Size const& threadsPerThreadgroup);
+};
 
-@end
 NS_ASSUME_NONNULL_END
-
-#if METAL_DEBUG_OPTIONS
-#define METAL_SET_COMPUTE_REFLECTION(Encoder, InPipeline)														\
-    if (GetMetalDeviceContext().GetCommandQueue().GetRuntimeDebuggingLevel() >= EMetalDebugLevelFastValidation) \
-    {																											\
-        ((FMetalDebugComputeCommandEncoder*)Encoder).Pipeline = InPipeline;										\
-    }
-#else
-#define METAL_SET_COMPUTE_REFLECTION(Encoder, InPipeline)
-#endif
diff --git a/Engine/Source/Runtime/Apple/MetalRHI/Private/MetalDebugCommandEncoder.cpp b/Engine/Source/Runtime/Apple/MetalRHI/Private/MetalDebugCommandEncoder.cpp
--- a/Engine/Source/Runtime/Apple/MetalRHI/Private/MetalDebugCommandEncoder.cpp
+++ b/Engine/Source/Runtime/Apple/MetalRHI/Private/MetalDebugCommandEncoder.cpp
@@ -29,26 +29,33 @@ APPLE_PLATFORM_OBJECT_ALLOC_OVERRIDES(FMetalDebugCommandEncoder)
 	[WaitingFences release];
 	[super dealloc];
 }
+@end
 
--(id<MTLCommandEncoder>)commandEncoder
+FMetalCommandEncoderDebugging::FMetalCommandEncoderDebugging()
 {
-	check(false);
-	return nil;
+	
 }
--(void)addUpdateFence:(id)Fence
+
+FMetalCommandEncoderDebugging::FMetalCommandEncoderDebugging(FMetalDebugCommandEncoder* handle)
+: ns::Object<FMetalDebugCommandEncoder*>(handle)
+{
+	
+}
+
+void FMetalCommandEncoderDebugging::AddUpdateFence(id Fence)
 {
 	if ((EMetalDebugLevel)GMetalRuntimeDebugLevel >= EMetalDebugLevelValidation && Fence)
 	{
-		[UpdatedFences addObject:(FMetalDebugFence*)Fence];
-		[(FMetalDebugFence*)Fence updatingEncoder:self];
+		[m_ptr->UpdatedFences addObject:(FMetalDebugFence*)Fence];
+		[(FMetalDebugFence*)Fence updatingEncoder:m_ptr];
 	}
 }
--(void)addWaitFence:(id)Fence
+
+void FMetalCommandEncoderDebugging::AddWaitFence(id Fence)
 {
 	if ((EMetalDebugLevel)GMetalRuntimeDebugLevel >= EMetalDebugLevelValidation && Fence)
 	{
-		[WaitingFences addObject:(FMetalDebugFence*)Fence];
-		[(FMetalDebugFence*)Fence waitingEncoder:self];
+		[m_ptr->WaitingFences addObject:(FMetalDebugFence*)Fence];
+		[(FMetalDebugFence*)Fence waitingEncoder:m_ptr];
 	}
 }
-@end
diff --git a/Engine/Source/Runtime/Apple/MetalRHI/Private/MetalDebugCommandEncoder.h b/Engine/Source/Runtime/Apple/MetalRHI/Private/MetalDebugCommandEncoder.h
--- a/Engine/Source/Runtime/Apple/MetalRHI/Private/MetalDebugCommandEncoder.h
+++ b/Engine/Source/Runtime/Apple/MetalRHI/Private/MetalDebugCommandEncoder.h
@@ -2,7 +2,7 @@
 
 #pragma once
 
-#include <Metal/Metal.h>
+#include "MetalCommandBuffer.h"
 
 // For some reason when including this file while building the editor Clang 9 ignores this pragma from MacPlatformCompilerPreSetup.h,
 // resulting in errors in FMetalDebugBufferBindings, FMetalDebugTextureBindings and FMetalDebugSamplerBindings. Readding it here works around this problem.
@@ -80,8 +80,8 @@ struct FMetalDebugSamplerBindings
     id<MTLSamplerState> _Nullable Samplers[ML_MaxSamplers];
 };
 
+class FMetalCommandBufferDebugging;
 @class FMetalDebugCommandBuffer;
-
 @class FMetalDebugFence;
 
 @interface FMetalDebugCommandEncoder : FApplePlatformObject
@@ -91,10 +91,17 @@ struct FMetalDebugSamplerBindings
 	NSHashTable<FMetalDebugFence*>* WaitingFences;
 }
 -(instancetype)init;
--(id<MTLCommandEncoder>)commandEncoder;
--(void)addUpdateFence:(id)Fence;
--(void)addWaitFence:(id)Fence;
 @end
 
+class FMetalCommandEncoderDebugging : public ns::Object<FMetalDebugCommandEncoder*>
+{
+public:
+	FMetalCommandEncoderDebugging();
+	FMetalCommandEncoderDebugging(FMetalDebugCommandEncoder* handle);
+	
+	void AddUpdateFence(id Fence);
+	void AddWaitFence(id Fence);
+};
+
 NS_ASSUME_NONNULL_END
 
diff --git a/Engine/Source/Runtime/Apple/MetalRHI/Private/MetalParallelRenderCommandEncoder.cpp b/Engine/Source/Runtime/Apple/MetalRHI/Private/MetalParallelRenderCommandEncoder.cpp
deleted file mode 100644
diff --git a/Engine/Source/Runtime/Apple/MetalRHI/Private/MetalParallelRenderCommandEncoder.h b/Engine/Source/Runtime/Apple/MetalRHI/Private/MetalParallelRenderCommandEncoder.h
deleted file mode 100644
diff --git a/Engine/Source/Runtime/Apple/MetalRHI/Private/MetalQuery.cpp b/Engine/Source/Runtime/Apple/MetalRHI/Private/MetalQuery.cpp
--- a/Engine/Source/Runtime/Apple/MetalRHI/Private/MetalQuery.cpp
+++ b/Engine/Source/Runtime/Apple/MetalRHI/Private/MetalQuery.cpp
@@ -411,11 +411,20 @@ void FMetalDynamicRHI::RHICalibrateTimers()
 		
 		id<IMetalStatisticsSamples> Samples = Context.GetCommandQueue().GetStatistics()->RegisterEncoderStatistics(Buffer.GetPtr(), EMetalSampleComputeEncoderStart);
 		mtlpp::ComputeCommandEncoder Encoder = Buffer.ComputeCommandEncoder();
+#if METAL_DEBUG_OPTIONS
+		FMetalComputeCommandEncoderDebugging Debugging;
+		if (SafeGetRuntimeDebuggingLevel() >= EMetalDebugLevelFastValidation)
+		{
+			FMetalCommandBufferDebugging CmdDebug = FMetalCommandBufferDebugging::Get(Buffer);
+			Debugging = FMetalComputeCommandEncoderDebugging(Encoder, CmdDebug);
+		}
+#endif
 		
 		Context.GetCommandQueue().GetStatistics()->RegisterEncoderStatistics(Buffer.GetPtr(), EMetalSampleComputeEncoderEnd);
 		check(Samples);
 		[Samples retain];
 		Encoder.EndEncoding();
+		METAL_DEBUG_LAYER(EMetalDebugLevelFastValidation, Debugging.EndEncoder());
 		
 		FMetalProfiler* Profiler = ImmediateContext.GetProfiler();
 		Buffer.AddCompletedHandler(^(const mtlpp::CommandBuffer & theBuffer) {
diff --git a/Engine/Source/Runtime/Apple/MetalRHI/Private/MetalRHIPrivate.h b/Engine/Source/Runtime/Apple/MetalRHI/Private/MetalRHIPrivate.h
--- a/Engine/Source/Runtime/Apple/MetalRHI/Private/MetalRHIPrivate.h
+++ b/Engine/Source/Runtime/Apple/MetalRHI/Private/MetalRHIPrivate.h
@@ -88,8 +88,10 @@ extern FMetalBufferFormat GMetalBufferFormats[PF_MAX];
 #define METAL_DEBUG_OPTIONS !(UE_BUILD_SHIPPING || UE_BUILD_TEST)
 #if METAL_DEBUG_OPTIONS
 #define METAL_DEBUG_OPTION(Code) Code
+#define METAL_DEBUG_LAYER(Level, Code) if (SafeGetRuntimeDebuggingLevel() >= Level) Code
 #else
 #define METAL_DEBUG_OPTION(Code)
+#define METAL_DEBUG_LAYER(Level, Code)
 #endif
 
 extern bool GMetalSupportsTileShaders;
diff --git a/Engine/Source/Runtime/Apple/MetalRHI/Private/MetalRenderCommandEncoder.cpp b/Engine/Source/Runtime/Apple/MetalRHI/Private/MetalRenderCommandEncoder.cpp
--- a/Engine/Source/Runtime/Apple/MetalRHI/Private/MetalRenderCommandEncoder.cpp
+++ b/Engine/Source/Runtime/Apple/MetalRHI/Private/MetalRenderCommandEncoder.cpp
@@ -14,6 +14,29 @@
 
 NS_ASSUME_NONNULL_BEGIN
 
+@interface FMetalDebugRenderCommandEncoder : FMetalDebugCommandEncoder
+{
+@public
+#pragma mark - Private Member Variables -
+#if METAL_DEBUG_OPTIONS
+	FMetalDebugShaderResourceMask ResourceMask[EMetalShaderRenderNum];
+	FMetalDebugBufferBindings ShaderBuffers[EMetalShaderRenderNum];
+	FMetalDebugTextureBindings ShaderTextures[EMetalShaderRenderNum];
+	FMetalDebugSamplerBindings ShaderSamplers[EMetalShaderRenderNum];
+	id<MTLRenderPipelineState> DebugState;
+#endif
+	MTLRenderPassDescriptor* RenderPassDesc;
+	id<MTLRenderCommandEncoder> Inner;
+	FMetalCommandBufferDebugging Buffer;
+	FMetalShaderPipeline* Pipeline;
+}
+
+/** Initialise the wrapper with the provided command-buffer. */
+-(id)initWithEncoder:(id<MTLRenderCommandEncoder>)Encoder fromDescriptor:(MTLRenderPassDescriptor*)RenderPassDesc andCommandBuffer:(FMetalCommandBufferDebugging const&)Buffer;
+
+@end
+
+
 #if METAL_DEBUG_OPTIONS
 static NSString* GMetalDebugVertexShader = @"#include <metal_stdlib>\n"
 @"using namespace metal;\n"
@@ -84,22 +107,18 @@ static id <MTLRenderPipelineState> GetDebugVertexShaderState(id<MTLDevice> Devic
 
 @implementation FMetalDebugRenderCommandEncoder
 
-@synthesize Inner;
-@synthesize Buffer;
-@synthesize Pipeline;
-
 APPLE_PLATFORM_OBJECT_ALLOC_OVERRIDES(FMetalDebugRenderCommandEncoder)
 
--(id)initWithEncoder:(id<MTLRenderCommandEncoder>)Encoder fromDescriptor:(MTLRenderPassDescriptor*)Desc andCommandBuffer:(FMetalDebugCommandBuffer*)SourceBuffer
+-(id)initWithEncoder:(id<MTLRenderCommandEncoder>)Encoder fromDescriptor:(MTLRenderPassDescriptor*)Desc andCommandBuffer:(FMetalCommandBufferDebugging const&)SourceBuffer
 {
 	id Self = [super init];
 	if (Self)
 	{
-        Inner = [Encoder retain];
-		Buffer = [SourceBuffer retain];
+        Inner = Encoder;
+		Buffer = SourceBuffer;
 		RenderPassDesc = [Desc retain];
 #if METAL_DEBUG_OPTIONS
-		DebugState = Buffer->DebugLevel >= EMetalDebugLevelValidation ? [GetDebugVertexShaderState(Buffer.device, Desc) retain] : nil;
+		DebugState = Buffer.GetPtr()->DebugLevel >= EMetalDebugLevelValidation ? [GetDebugVertexShaderState(Buffer.GetPtr()->InnerBuffer.device, Desc) retain] : nil;
 #endif
         Pipeline = nil;
 	}
@@ -108,8 +127,6 @@ APPLE_PLATFORM_OBJECT_ALLOC_OVERRIDES(FMetalDebugRenderCommandEncoder)
 
 -(void)dealloc
 {
-	[Inner release];
-	[Buffer release];
 	[RenderPassDesc release];
 #if METAL_DEBUG_OPTIONS
 	[DebugState release];
@@ -118,43 +135,14 @@ APPLE_PLATFORM_OBJECT_ALLOC_OVERRIDES(FMetalDebugRenderCommandEncoder)
 	[super dealloc];
 }
 
--(id <MTLDevice>) device
-{
-	return Inner.device;
-}
-
--(NSString *_Nullable)label
-{
-	return Inner.label;
-}
-
--(void)setLabel:(NSString *_Nullable)Text
-{
-	Inner.label = Text;
-}
-
-- (void)endEncoding
-{
-    [Buffer endCommandEncoder];
-    [Inner endEncoding];
-}
+@end
 
-- (void)insertDebugSignpost:(NSString *)string
-{
-    [Buffer insertDebugSignpost:string];
-    [Inner insertDebugSignpost:string];
-}
+NS_ASSUME_NONNULL_END
 
-- (void)pushDebugGroup:(NSString *)string
+void FMetalRenderCommandEncoderDebugging::InsertDebugDraw()
 {
-    [Buffer pushDebugGroup:string];
-    [Inner pushDebugGroup:string];
-}
-
 #if METAL_DEBUG_OPTIONS
-- (void)insertDebugDraw
-{
-//	switch (Buffer->DebugLevel)
+//	switch (((FMetalDebugRenderCommandEncoder*)m_ptr)->Buffer.GetPtr()->DebugLevel)
 //	{
 //		case EMetalDebugLevelConditionalSubmit:
 //		case EMetalDebugLevelWaitForComplete:
@@ -168,7 +156,7 @@ APPLE_PLATFORM_OBJECT_ALLOC_OVERRIDES(FMetalDebugRenderCommandEncoder)
 //			[Inner setVertexBytes:&Index length:sizeof(Index) atIndex:0];
 //			[Inner setVertexBuffer:Buffer->DebugInfoBuffer offset:0 atIndex:1];
 //			[Inner setRenderPipelineState:DebugState];
-//			[Inner drawPrimitives:MTLPrimitiveTypePoint vertexStart:0 vertexCount:1];
+//			[Inner drawPrimitives:mtlpp::PrimitiveTypePoint vertexStart:0 vertexCount:1];
 //#if PLATFORM_MAC
 //			[Inner textureBarrier];
 //#endif
@@ -200,323 +188,47 @@ APPLE_PLATFORM_OBJECT_ALLOC_OVERRIDES(FMetalDebugRenderCommandEncoder)
 //			break;
 //		}
 //	}
-}
-#endif
-
-- (void)popDebugGroup
-{
-    [Buffer popDebugGroup];
-#if METAL_DEBUG_OPTIONS
-	[self insertDebugDraw];
-#endif
-    [Inner popDebugGroup];
-}
-
-- (void)setRenderPipelineState:(id <MTLRenderPipelineState>)pipelineState
-{
-#if METAL_DEBUG_OPTIONS
-	switch (Buffer->DebugLevel)
-	{
-		case EMetalDebugLevelConditionalSubmit:
-		case EMetalDebugLevelWaitForComplete:
-		case EMetalDebugLevelLogOperations:
-		{
-			[Buffer setPipeline:pipelineState.label];
-		}
-		case EMetalDebugLevelValidation:
-		case EMetalDebugLevelResetOnBind:
-		case EMetalDebugLevelTrackResources:
-		{
-			[Buffer trackState:pipelineState];
-		}
-		default:
-		{
-			break;
-		}
-	}
 #endif
-    [Inner setRenderPipelineState:pipelineState];
 }
 
-- (void)setVertexBytes:(const void *)bytes length:(NSUInteger)length atIndex:(NSUInteger)index
+FMetalRenderCommandEncoderDebugging::FMetalRenderCommandEncoderDebugging()
 {
-#if METAL_DEBUG_OPTIONS
-	switch (Buffer->DebugLevel)
-	{
-		case EMetalDebugLevelConditionalSubmit:
-		case EMetalDebugLevelWaitForComplete:
-		case EMetalDebugLevelLogOperations:
-		case EMetalDebugLevelValidation:
-		{
-			ShaderBuffers[EMetalShaderVertex].Buffers[index] = nil;
-			ShaderBuffers[EMetalShaderVertex].Bytes[index] = bytes;
-			ShaderBuffers[EMetalShaderVertex].Offsets[index] = length;
-		}
-		case EMetalDebugLevelResetOnBind:
-		case EMetalDebugLevelTrackResources:
-		case EMetalDebugLevelFastValidation:
-		{
-			ResourceMask[EMetalShaderVertex].BufferMask = bytes ? (ResourceMask[EMetalShaderVertex].BufferMask | (1 << index)) : (ResourceMask[EMetalShaderVertex].BufferMask & ~(1 << index));
-		}
-		default:
-		{
-			break;
-		}
-	}
-#endif
 	
-    [Inner setVertexBytes:bytes length:length atIndex:index];
-}
-
-- (void)setVertexBuffer:(nullable id <MTLBuffer>)buffer offset:(NSUInteger)offset atIndex:(NSUInteger)index
-{
-#if METAL_DEBUG_OPTIONS
-	switch (Buffer->DebugLevel)
-	{
-		case EMetalDebugLevelConditionalSubmit:
-		case EMetalDebugLevelWaitForComplete:
-		case EMetalDebugLevelLogOperations:
-		case EMetalDebugLevelValidation:
-		{
-			ShaderBuffers[EMetalShaderVertex].Buffers[index] = buffer;
-			ShaderBuffers[EMetalShaderVertex].Bytes[index] = nil;
-			ShaderBuffers[EMetalShaderVertex].Offsets[index] = offset;
-		}
-		case EMetalDebugLevelResetOnBind:
-		case EMetalDebugLevelTrackResources:
-		{
-			[Buffer trackResource:buffer];
-		}
-		case EMetalDebugLevelFastValidation:
-		{
-			ResourceMask[EMetalShaderVertex].BufferMask = buffer ? (ResourceMask[EMetalShaderVertex].BufferMask | (1 << index)) : (ResourceMask[EMetalShaderVertex].BufferMask & ~(1 << index));
-		}
-		default:
-		{
-			break;
-		}
-	}
-#endif
-    [Inner setVertexBuffer:buffer offset:offset atIndex:index];
-}
-
-- (void)setVertexBufferOffset:(NSUInteger)offset atIndex:(NSUInteger)index
-{
-#if METAL_DEBUG_OPTIONS
-	switch (Buffer->DebugLevel)
-	{
-		case EMetalDebugLevelConditionalSubmit:
-		case EMetalDebugLevelWaitForComplete:
-		case EMetalDebugLevelLogOperations:
-		case EMetalDebugLevelValidation:
-		{
-			ShaderBuffers[EMetalShaderVertex].Offsets[index] = offset;
-		}
-		case EMetalDebugLevelResetOnBind:
-		case EMetalDebugLevelTrackResources:
-		case EMetalDebugLevelFastValidation:
-		{
-			check(ResourceMask[EMetalShaderVertex].BufferMask & 1 << index);
-		}
-		default:
-		{
-			break;
-		}
-	}
-#endif
-	[Inner setVertexBufferOffset:offset atIndex:index];
 }
-
-- (void)setVertexBuffers:(const id <MTLBuffer> __nullable [__nonnull])buffers offsets:(const NSUInteger [__nonnull])offsets withRange:(NSRange)range
+FMetalRenderCommandEncoderDebugging::FMetalRenderCommandEncoderDebugging(mtlpp::RenderCommandEncoder& Encoder, mtlpp::RenderPassDescriptor const& Desc, FMetalCommandBufferDebugging& Buffer)
+: FMetalCommandEncoderDebugging((FMetalDebugCommandEncoder*)[[FMetalDebugRenderCommandEncoder alloc] initWithEncoder:Encoder.GetPtr() fromDescriptor:Desc.GetPtr() andCommandBuffer:Buffer])
 {
-#if METAL_DEBUG_OPTIONS
-	for (uint32 i = 0; i < range.length; i++)
-	{
-		switch (Buffer->DebugLevel)
-		{
-			case EMetalDebugLevelConditionalSubmit:
-			case EMetalDebugLevelWaitForComplete:
-			case EMetalDebugLevelLogOperations:
-			case EMetalDebugLevelValidation:
-			{
-				ShaderBuffers[EMetalShaderVertex].Buffers[i + range.location] = buffers[i];
-				ShaderBuffers[EMetalShaderVertex].Bytes[i + range.location] = nil;
-				ShaderBuffers[EMetalShaderVertex].Offsets[i + range.location] = offsets[i];
-			}
-			case EMetalDebugLevelResetOnBind:
-			case EMetalDebugLevelTrackResources:
-			{
-				[Buffer trackResource:buffers[i]];
-			}
-			case EMetalDebugLevelFastValidation:
-			{
-				ResourceMask[EMetalShaderVertex].BufferMask = buffers[i] ? (ResourceMask[EMetalShaderVertex].BufferMask | (1 << (i + range.location))) : (ResourceMask[EMetalShaderVertex].BufferMask & ~(1 << (i + range.location)));
-			}
-			default:
-			{
-				break;
-			}
-		}
-    }
-#endif
-    [Inner setVertexBuffers:buffers offsets:offsets withRange:range];
+	Buffer.BeginRenderCommandEncoder([NSString stringWithFormat:@"Render: %@", Encoder.GetLabel().GetPtr()], Desc);
+	Encoder.SetAssociatedObject((void const*)&FMetalRenderCommandEncoderDebugging::Get, (FMetalCommandEncoderDebugging const&)*this);
 }
-
-- (void)setVertexTexture:(nullable id <MTLTexture>)texture atIndex:(NSUInteger)index
+FMetalRenderCommandEncoderDebugging::FMetalRenderCommandEncoderDebugging(FMetalDebugCommandEncoder* handle)
+: FMetalCommandEncoderDebugging(handle)
 {
-#if METAL_DEBUG_OPTIONS
-	switch (Buffer->DebugLevel)
-	{
-		case EMetalDebugLevelConditionalSubmit:
-		case EMetalDebugLevelWaitForComplete:
-		case EMetalDebugLevelLogOperations:
-		case EMetalDebugLevelValidation:
-		{
-			ShaderTextures[EMetalShaderVertex].Textures[index] = texture;
-		}
-		case EMetalDebugLevelResetOnBind:
-		case EMetalDebugLevelTrackResources:
-		{
-			[Buffer trackResource:texture];
-		}
-		case EMetalDebugLevelFastValidation:
-		{
-			ResourceMask[EMetalShaderVertex].TextureMask = texture ? (ResourceMask[EMetalShaderVertex].TextureMask | (1 << (index))) : (ResourceMask[EMetalShaderVertex].TextureMask & ~(1 << (index)));
-		}
-		default:
-		{
-			break;
-		}
-	}
-#endif
-    [Inner setVertexTexture:texture atIndex:index];
+	
 }
 
-#if METAL_NEW_NONNULL_DECL
-- (void)setVertexTextures:(const id <MTLTexture> __nullable[__nonnull])textures withRange : (NSRange)range
-#else
-- (void)setVertexTextures:(const id <MTLTexture> __nullable [__nullable])textures withRange:(NSRange)range
-#endif
+FMetalRenderCommandEncoderDebugging FMetalRenderCommandEncoderDebugging::Get(mtlpp::RenderCommandEncoder& Encoder)
 {
-#if METAL_DEBUG_OPTIONS
-    for (uint32 i = 0; i < range.length; i++)
-    {
-		switch (Buffer->DebugLevel)
-		{
-			case EMetalDebugLevelConditionalSubmit:
-			case EMetalDebugLevelWaitForComplete:
-			case EMetalDebugLevelLogOperations:
-			case EMetalDebugLevelValidation:
-			{
-				ShaderTextures[EMetalShaderVertex].Textures[i + range.location] = textures[i];
-			}
-			case EMetalDebugLevelResetOnBind:
-			case EMetalDebugLevelTrackResources:
-			{
-				[Buffer trackResource:textures[i]];
-			}
-			case EMetalDebugLevelFastValidation:
-			{
-				ResourceMask[EMetalShaderVertex].TextureMask = textures[i] ? (ResourceMask[EMetalShaderVertex].TextureMask | (1 << (i + range.location))) : (ResourceMask[EMetalShaderVertex].TextureMask & ~(1 << (i + range.location)));
-			}
-			default:
-			{
-				break;
-			}
-		}
-    }
-#endif
-    [Inner setVertexTextures:textures withRange:range];
+	return Encoder.GetAssociatedObject<FMetalRenderCommandEncoderDebugging>((void const*)&FMetalRenderCommandEncoderDebugging::Get);
 }
 
-- (void)setVertexSamplerState:(nullable id <MTLSamplerState>)sampler atIndex:(NSUInteger)index
+void FMetalRenderCommandEncoderDebugging::SetPipeline(FMetalShaderPipeline* Pipeline)
 {
 #if METAL_DEBUG_OPTIONS
-	switch (Buffer->DebugLevel)
+	((FMetalDebugRenderCommandEncoder*)m_ptr)->Pipeline = [Pipeline retain];
+	switch (((FMetalDebugRenderCommandEncoder*)m_ptr)->Buffer.GetPtr()->DebugLevel)
 	{
 		case EMetalDebugLevelConditionalSubmit:
 		case EMetalDebugLevelWaitForComplete:
 		case EMetalDebugLevelLogOperations:
-		case EMetalDebugLevelValidation:
-		{
-			ShaderSamplers[EMetalShaderVertex].Samplers[index] = sampler;
-		}
-		case EMetalDebugLevelResetOnBind:
-		case EMetalDebugLevelTrackResources:
 		{
-			[Buffer trackState:sampler];
-		}
-		case EMetalDebugLevelFastValidation:
-		{
-			ResourceMask[EMetalShaderVertex].SamplerMask = sampler ? (ResourceMask[EMetalShaderVertex].SamplerMask | (1 << (index))) : (ResourceMask[EMetalShaderVertex].SamplerMask & ~(1 << (index)));
-		}
-		default:
-		{
-			break;
+			((FMetalDebugRenderCommandEncoder*)m_ptr)->Buffer.SetPipeline(Pipeline->RenderPipelineState.GetLabel());
 		}
-	}
-#endif
-    [Inner setVertexSamplerState:sampler atIndex:index];
-}
-
-#if METAL_NEW_NONNULL_DECL
-- (void)setVertexSamplerStates:(const id <MTLSamplerState> __nullable[__nonnull])samplers withRange : (NSRange)range
-#else
-- (void)setVertexSamplerStates:(const id <MTLSamplerState> __nullable [__nullable])samplers withRange:(NSRange)range
-#endif
-{
-#if METAL_DEBUG_OPTIONS
-    for(uint32 i = 0; i < range.length; i++)
-    {
-		switch (Buffer->DebugLevel)
-		{
-			case EMetalDebugLevelConditionalSubmit:
-			case EMetalDebugLevelWaitForComplete:
-			case EMetalDebugLevelLogOperations:
-			case EMetalDebugLevelValidation:
-			{
-				ShaderSamplers[EMetalShaderVertex].Samplers[i + range.location] = samplers[i];
-			}
-			case EMetalDebugLevelResetOnBind:
-			case EMetalDebugLevelTrackResources:
-			{
-				[Buffer trackState:samplers[i]];
-			}
-			case EMetalDebugLevelFastValidation:
-			{
-				ResourceMask[EMetalShaderVertex].SamplerMask = samplers[i] ? (ResourceMask[EMetalShaderVertex].SamplerMask | (1 << (i + range.location))) : (ResourceMask[EMetalShaderVertex].SamplerMask & ~(1 << (i + range.location)));
-			}
-			default:
-			{
-				break;
-			}
-		}
-		
-    }
-#endif
-    [Inner setVertexSamplerStates:samplers withRange:range];
-}
-
-- (void)setVertexSamplerState:(nullable id <MTLSamplerState>)sampler lodMinClamp:(float)lodMinClamp lodMaxClamp:(float)lodMaxClamp atIndex:(NSUInteger)index
-{
-#if METAL_DEBUG_OPTIONS
-	switch (Buffer->DebugLevel)
-	{
-		case EMetalDebugLevelConditionalSubmit:
-		case EMetalDebugLevelWaitForComplete:
-		case EMetalDebugLevelLogOperations:
 		case EMetalDebugLevelValidation:
-		{
-			ShaderSamplers[EMetalShaderVertex].Samplers[index] = sampler;
-		}
 		case EMetalDebugLevelResetOnBind:
 		case EMetalDebugLevelTrackResources:
 		{
-			[Buffer trackState:sampler];
-		}
-		case EMetalDebugLevelFastValidation:
-		{
-			ResourceMask[EMetalShaderVertex].SamplerMask = sampler ? (ResourceMask[EMetalShaderVertex].SamplerMask | (1 << (index))) : (ResourceMask[EMetalShaderVertex].SamplerMask & ~(1 << (index)));
+			((FMetalDebugRenderCommandEncoder*)m_ptr)->Buffer.TrackState(Pipeline->RenderPipelineState);
 		}
 		default:
 		{
@@ -524,106 +236,27 @@ APPLE_PLATFORM_OBJECT_ALLOC_OVERRIDES(FMetalDebugRenderCommandEncoder)
 		}
 	}
 #endif
-    [Inner setVertexSamplerState:sampler lodMinClamp:lodMinClamp lodMaxClamp:lodMaxClamp atIndex:index];
-}
-
-#if METAL_NEW_NONNULL_DECL
-- (void)setVertexSamplerStates:(const id <MTLSamplerState> __nullable[__nonnull])samplers lodMinClamps : (const float[__nonnull])lodMinClamps lodMaxClamps : (const float[__nonnull])lodMaxClamps withRange : (NSRange)range
-#else
-- (void)setVertexSamplerStates:(const id <MTLSamplerState> __nullable [__nullable])samplers lodMinClamps:(const float [__nullable])lodMinClamps lodMaxClamps:(const float [__nullable])lodMaxClamps withRange:(NSRange)range
-#endif
-{
-#if METAL_DEBUG_OPTIONS
-    for(uint32 i = 0; i < range.length; i++)
-    {
-		switch (Buffer->DebugLevel)
-		{
-			case EMetalDebugLevelConditionalSubmit:
-			case EMetalDebugLevelWaitForComplete:
-			case EMetalDebugLevelLogOperations:
-			case EMetalDebugLevelValidation:
-			{
-				ShaderSamplers[EMetalShaderVertex].Samplers[i + range.location] = samplers[i];
-			}
-			case EMetalDebugLevelResetOnBind:
-			case EMetalDebugLevelTrackResources:
-			{
-				[Buffer trackState:samplers[i]];
-			}
-			case EMetalDebugLevelFastValidation:
-			{
-				ResourceMask[EMetalShaderVertex].SamplerMask = samplers[i] ? (ResourceMask[EMetalShaderVertex].SamplerMask | (1 << (i + range.location))) : (ResourceMask[EMetalShaderVertex].SamplerMask & ~(1 << (i + range.location)));
-			}
-			default:
-			{
-				break;
-			}
-		}
-		
-    }
-#endif
-    [Inner setVertexSamplerStates:samplers lodMinClamps:lodMinClamps lodMaxClamps:lodMaxClamps withRange:range];
-}
-
-- (void)setViewport:(MTLViewport)viewport
-{
-    [Inner setViewport:viewport];
-}
-
-- (void)setFrontFacingWinding:(MTLWinding)frontFacingWinding
-{
-    [Inner setFrontFacingWinding:frontFacingWinding];
-}
-
-- (void)setCullMode:(MTLCullMode)cullMode
-{
-    [Inner setCullMode:cullMode];
-}
-
-#if PLATFORM_MAC || (PLATFORM_IOS && __IPHONE_OS_VERSION_MAX_ALLOWED >= 110000)
-- (void)setDepthClipMode:(MTLDepthClipMode)depthClipMode
-{
-	if (GMetalSupportsDepthClipMode)
-	{
-    [Inner setDepthClipMode:depthClipMode];
-}
-}
-#endif
-
-- (void)setDepthBias:(float)depthBias slopeScale:(float)slopeScale clamp:(float)clamp
-{
-    [Inner setDepthBias:depthBias slopeScale:slopeScale clamp:clamp];
-}
-
-- (void)setScissorRect:(MTLScissorRect)rect
-{
-    [Inner setScissorRect:rect];
 }
 
-- (void)setTriangleFillMode:(MTLTriangleFillMode)fillMode
-{
-    [Inner setTriangleFillMode:fillMode];
-}
-
-- (void)setFragmentBytes:(const void *)bytes length:(NSUInteger)length atIndex:(NSUInteger)index
+void FMetalRenderCommandEncoderDebugging::SetBytes(EMetalShaderFrequency Freq, const void * bytes, NSUInteger length, NSUInteger index)
 {
 #if METAL_DEBUG_OPTIONS
-	switch (Buffer->DebugLevel)
+	switch (((FMetalDebugRenderCommandEncoder*)m_ptr)->Buffer.GetPtr()->DebugLevel)
 	{
 		case EMetalDebugLevelConditionalSubmit:
 		case EMetalDebugLevelWaitForComplete:
 		case EMetalDebugLevelLogOperations:
 		case EMetalDebugLevelValidation:
 		{
-			ShaderBuffers[EMetalShaderFragment].Buffers[index] = nil;
-			ShaderBuffers[EMetalShaderFragment].Bytes[index] = bytes;
-			ShaderBuffers[EMetalShaderFragment].Offsets[index] = length;
+			((FMetalDebugRenderCommandEncoder*)m_ptr)->ShaderBuffers[Freq].Buffers[index] = nil;
+			((FMetalDebugRenderCommandEncoder*)m_ptr)->ShaderBuffers[Freq].Bytes[index] = bytes;
+			((FMetalDebugRenderCommandEncoder*)m_ptr)->ShaderBuffers[Freq].Offsets[index] = length;
 		}
 		case EMetalDebugLevelResetOnBind:
 		case EMetalDebugLevelTrackResources:
 		case EMetalDebugLevelFastValidation:
 		{
-			ResourceMask[EMetalShaderFragment].BufferMask = bytes ? (ResourceMask[EMetalShaderFragment].BufferMask | (1 << (index))) : (ResourceMask[EMetalShaderFragment].BufferMask & ~(1 << (index)));
+			((FMetalDebugRenderCommandEncoder*)m_ptr)->ResourceMask[Freq].BufferMask = bytes ? (((FMetalDebugRenderCommandEncoder*)m_ptr)->ResourceMask[Freq].BufferMask | (1 << (index))) : (((FMetalDebugRenderCommandEncoder*)m_ptr)->ResourceMask[Freq].BufferMask & ~(1 << (index)));
 		}
 		default:
 		{
@@ -631,31 +264,29 @@ APPLE_PLATFORM_OBJECT_ALLOC_OVERRIDES(FMetalDebugRenderCommandEncoder)
 		}
 	}
 #endif
-    [Inner setFragmentBytes:(const void *)bytes length:(NSUInteger)length atIndex:(NSUInteger)index];
 }
-
-- (void)setFragmentBuffer:(nullable id <MTLBuffer>)buffer offset:(NSUInteger)offset atIndex:(NSUInteger)index
+void FMetalRenderCommandEncoderDebugging::SetBuffer(EMetalShaderFrequency Freq,  FMetalBuffer const& buffer, NSUInteger offset, NSUInteger index)
 {
 #if METAL_DEBUG_OPTIONS
-	switch (Buffer->DebugLevel)
+	switch (((FMetalDebugRenderCommandEncoder*)m_ptr)->Buffer.GetPtr()->DebugLevel)
 	{
 		case EMetalDebugLevelConditionalSubmit:
 		case EMetalDebugLevelWaitForComplete:
 		case EMetalDebugLevelLogOperations:
 		case EMetalDebugLevelValidation:
 		{
-			ShaderBuffers[EMetalShaderFragment].Buffers[index] = buffer;
-			ShaderBuffers[EMetalShaderFragment].Bytes[index] = nil;
-			ShaderBuffers[EMetalShaderFragment].Offsets[index] = offset;
+			((FMetalDebugRenderCommandEncoder*)m_ptr)->ShaderBuffers[Freq].Buffers[index] = buffer;
+			((FMetalDebugRenderCommandEncoder*)m_ptr)->ShaderBuffers[Freq].Bytes[index] = nil;
+			((FMetalDebugRenderCommandEncoder*)m_ptr)->ShaderBuffers[Freq].Offsets[index] = offset;
 		}
 		case EMetalDebugLevelResetOnBind:
 		case EMetalDebugLevelTrackResources:
 		{
-			[Buffer trackResource:buffer];
+			((FMetalDebugRenderCommandEncoder*)m_ptr)->Buffer.TrackResource(buffer);
 		}
 		case EMetalDebugLevelFastValidation:
 		{
-			ResourceMask[EMetalShaderFragment].BufferMask = buffer ? (ResourceMask[EMetalShaderFragment].BufferMask | (1 << (index))) : (ResourceMask[EMetalShaderFragment].BufferMask & ~(1 << (index)));
+			((FMetalDebugRenderCommandEncoder*)m_ptr)->ResourceMask[Freq].BufferMask = buffer ? (((FMetalDebugRenderCommandEncoder*)m_ptr)->ResourceMask[Freq].BufferMask | (1 << (index))) : (((FMetalDebugRenderCommandEncoder*)m_ptr)->ResourceMask[Freq].BufferMask & ~(1 << (index)));
 		}
 		default:
 		{
@@ -663,26 +294,24 @@ APPLE_PLATFORM_OBJECT_ALLOC_OVERRIDES(FMetalDebugRenderCommandEncoder)
 		}
 	}
 #endif
-    [Inner setFragmentBuffer:buffer offset:offset atIndex:index];
 }
-
-- (void)setFragmentBufferOffset:(NSUInteger)offset atIndex:(NSUInteger)index
+void FMetalRenderCommandEncoderDebugging::SetBufferOffset(EMetalShaderFrequency Freq, NSUInteger offset, NSUInteger index)
 {
 #if METAL_DEBUG_OPTIONS
-	switch (Buffer->DebugLevel)
+	switch (((FMetalDebugRenderCommandEncoder*)m_ptr)->Buffer.GetPtr()->DebugLevel)
 	{
 		case EMetalDebugLevelConditionalSubmit:
 		case EMetalDebugLevelWaitForComplete:
 		case EMetalDebugLevelLogOperations:
 		case EMetalDebugLevelValidation:
 		{
-			ShaderBuffers[EMetalShaderFragment].Offsets[index] = offset;
+			((FMetalDebugRenderCommandEncoder*)m_ptr)->ShaderBuffers[Freq].Offsets[index] = offset;
 		}
 		case EMetalDebugLevelResetOnBind:
 		case EMetalDebugLevelTrackResources:
 		case EMetalDebugLevelFastValidation:
 		{
-			check(ResourceMask[EMetalShaderFragment].BufferMask & (1 << (index)));
+			check(((FMetalDebugRenderCommandEncoder*)m_ptr)->ResourceMask[Freq].BufferMask & (1 << (index)));
 		}
 		default:
 		{
@@ -690,64 +319,28 @@ APPLE_PLATFORM_OBJECT_ALLOC_OVERRIDES(FMetalDebugRenderCommandEncoder)
 		}
 	}
 #endif
-    [Inner setFragmentBufferOffset:(NSUInteger)offset atIndex:(NSUInteger)index];
-}
-
-- (void)setFragmentBuffers:(const id <MTLBuffer> __nullable [__nonnull])buffers offsets:(const NSUInteger [__nonnull])offset withRange:(NSRange)range
-{
-#if METAL_DEBUG_OPTIONS
-    for (uint32 i = 0; i < range.length; i++)
-    {
-		switch (Buffer->DebugLevel)
-		{
-			case EMetalDebugLevelConditionalSubmit:
-			case EMetalDebugLevelWaitForComplete:
-			case EMetalDebugLevelLogOperations:
-			case EMetalDebugLevelValidation:
-			{
-				ShaderBuffers[EMetalShaderFragment].Buffers[i + range.location] = buffers[i];
-				ShaderBuffers[EMetalShaderFragment].Bytes[i + range.location] = nil;
-				ShaderBuffers[EMetalShaderFragment].Offsets[i + range.location] = offset[i];
-			}
-			case EMetalDebugLevelResetOnBind:
-			case EMetalDebugLevelTrackResources:
-			{
-				[Buffer trackResource:buffers[i]];
-			}
-			case EMetalDebugLevelFastValidation:
-			{
-				ResourceMask[EMetalShaderFragment].BufferMask = buffers[i] ? (ResourceMask[EMetalShaderFragment].BufferMask | (1 << (i + range.location))) : (ResourceMask[EMetalShaderFragment].BufferMask & ~(1 << (i + range.location)));
-			}
-			default:
-			{
-				break;
-			}
-		}
-    }
-#endif
-    [Inner setFragmentBuffers:buffers offsets:offset withRange:range];
 }
 
-- (void)setFragmentTexture:(nullable id <MTLTexture>)texture atIndex:(NSUInteger)index
+void FMetalRenderCommandEncoderDebugging::SetTexture(EMetalShaderFrequency Freq,  FMetalTexture const& texture, NSUInteger index)
 {
 #if METAL_DEBUG_OPTIONS
-	switch (Buffer->DebugLevel)
+	switch (((FMetalDebugRenderCommandEncoder*)m_ptr)->Buffer.GetPtr()->DebugLevel)
 	{
 		case EMetalDebugLevelConditionalSubmit:
 		case EMetalDebugLevelWaitForComplete:
 		case EMetalDebugLevelLogOperations:
 		case EMetalDebugLevelValidation:
 		{
-			ShaderTextures[EMetalShaderFragment].Textures[index] = texture;
+			((FMetalDebugRenderCommandEncoder*)m_ptr)->ShaderTextures[Freq].Textures[index] = texture;
 		}
 		case EMetalDebugLevelResetOnBind:
 		case EMetalDebugLevelTrackResources:
 		{
-			[Buffer trackResource:texture];
+			((FMetalDebugRenderCommandEncoder*)m_ptr)->Buffer.TrackResource(texture);
 		}
 		case EMetalDebugLevelFastValidation:
 		{
-			ResourceMask[EMetalShaderFragment].TextureMask = texture ? (ResourceMask[EMetalShaderFragment].TextureMask | (1 << (index))) : (ResourceMask[EMetalShaderFragment].TextureMask & ~(1 << (index)));
+			((FMetalDebugRenderCommandEncoder*)m_ptr)->ResourceMask[Freq].TextureMask = texture ? (((FMetalDebugRenderCommandEncoder*)m_ptr)->ResourceMask[Freq].TextureMask | (1 << (index))) : (((FMetalDebugRenderCommandEncoder*)m_ptr)->ResourceMask[Freq].TextureMask & ~(1 << (index)));
 		}
 		default:
 		{
@@ -755,67 +348,28 @@ APPLE_PLATFORM_OBJECT_ALLOC_OVERRIDES(FMetalDebugRenderCommandEncoder)
 		}
 	}
 #endif
-    [Inner setFragmentTexture:texture atIndex:index];
-}
-
-#if METAL_NEW_NONNULL_DECL
-- (void)setFragmentTextures:(const id <MTLTexture> __nullable[__nonnull])textures withRange : (NSRange)range
-#else
-- (void)setFragmentTextures:(const id <MTLTexture> __nullable [__nullable])textures withRange:(NSRange)range
-#endif
-{
-#if METAL_DEBUG_OPTIONS
-    for (uint32 i = 0; i < range.length; i++)
-    {
-		switch (Buffer->DebugLevel)
-		{
-			case EMetalDebugLevelConditionalSubmit:
-			case EMetalDebugLevelWaitForComplete:
-			case EMetalDebugLevelLogOperations:
-			case EMetalDebugLevelValidation:
-			{
-				ShaderTextures[EMetalShaderFragment].Textures[i + range.location] = textures[i];
-			}
-			case EMetalDebugLevelResetOnBind:
-			case EMetalDebugLevelTrackResources:
-			{
-				[Buffer trackResource:textures[i]];
-			}
-			case EMetalDebugLevelFastValidation:
-			{
-				ResourceMask[EMetalShaderFragment].TextureMask = textures[i] ? (ResourceMask[EMetalShaderFragment].TextureMask | (1 << (i + range.location))) : (ResourceMask[EMetalShaderFragment].TextureMask & ~(1 << (i + range.location)));
-			}
-			default:
-			{
-				break;
-			}
-		}
-		
-    }
-#endif
-    [Inner setFragmentTextures:textures withRange:(NSRange)range];
 }
 
-- (void)setFragmentSamplerState:(nullable id <MTLSamplerState>)sampler atIndex:(NSUInteger)index
+void FMetalRenderCommandEncoderDebugging::SetSamplerState(EMetalShaderFrequency Freq,  mtlpp::SamplerState const& sampler, NSUInteger index)
 {
 #if METAL_DEBUG_OPTIONS
-	switch (Buffer->DebugLevel)
+	switch (((FMetalDebugRenderCommandEncoder*)m_ptr)->Buffer.GetPtr()->DebugLevel)
 	{
 		case EMetalDebugLevelConditionalSubmit:
 		case EMetalDebugLevelWaitForComplete:
 		case EMetalDebugLevelLogOperations:
 		case EMetalDebugLevelValidation:
 		{
-			ShaderSamplers[EMetalShaderFragment].Samplers[index] = sampler;
+			((FMetalDebugRenderCommandEncoder*)m_ptr)->ShaderSamplers[Freq].Samplers[index] = sampler;
 		}
 		case EMetalDebugLevelResetOnBind:
 		case EMetalDebugLevelTrackResources:
 		{
-			[Buffer trackState:sampler];
+			((FMetalDebugRenderCommandEncoder*)m_ptr)->Buffer.TrackState(sampler);
 		}
 		case EMetalDebugLevelFastValidation:
 		{
-			ResourceMask[EMetalShaderFragment].SamplerMask = sampler ? (ResourceMask[EMetalShaderFragment].SamplerMask | (1 << (index))) : (ResourceMask[EMetalShaderFragment].SamplerMask & ~(1 << (index)));
+			((FMetalDebugRenderCommandEncoder*)m_ptr)->ResourceMask[Freq].SamplerMask = sampler ? (((FMetalDebugRenderCommandEncoder*)m_ptr)->ResourceMask[Freq].SamplerMask | (1 << (index))) : (((FMetalDebugRenderCommandEncoder*)m_ptr)->ResourceMask[Freq].SamplerMask & ~(1 << (index)));
 		}
 		default:
 		{
@@ -823,67 +377,28 @@ APPLE_PLATFORM_OBJECT_ALLOC_OVERRIDES(FMetalDebugRenderCommandEncoder)
 		}
 	}
 #endif
-    [Inner setFragmentSamplerState:sampler atIndex:(NSUInteger)index];
-}
-
-#if METAL_NEW_NONNULL_DECL
-- (void)setFragmentSamplerStates:(const id <MTLSamplerState> __nullable[__nonnull])samplers withRange : (NSRange)range
-#else
-- (void)setFragmentSamplerStates:(const id <MTLSamplerState> __nullable [__nullable])samplers withRange:(NSRange)range
-#endif
-{
-#if METAL_DEBUG_OPTIONS
-	for(uint32 i = 0; i < range.length; i++)
-    {
-		switch (Buffer->DebugLevel)
-		{
-			case EMetalDebugLevelConditionalSubmit:
-			case EMetalDebugLevelWaitForComplete:
-			case EMetalDebugLevelLogOperations:
-			case EMetalDebugLevelValidation:
-			{
-				ShaderSamplers[EMetalShaderFragment].Samplers[i + range.location] = samplers[i];
-			}
-			case EMetalDebugLevelResetOnBind:
-			case EMetalDebugLevelTrackResources:
-			{
-				[Buffer trackState:samplers[i]];
-			}
-			case EMetalDebugLevelFastValidation:
-			{
-				ResourceMask[EMetalShaderFragment].SamplerMask = samplers[i] ? (ResourceMask[EMetalShaderFragment].SamplerMask | (1 << (i + range.location))) : (ResourceMask[EMetalShaderFragment].SamplerMask & ~(1 << (i + range.location)));
-			}
-			default:
-			{
-				break;
-			}
-		}
-		
-    }
-#endif
-    [Inner setFragmentSamplerStates:samplers withRange:(NSRange)range];
 }
 
-- (void)setFragmentSamplerState:(nullable id <MTLSamplerState>)sampler lodMinClamp:(float)lodMinClamp lodMaxClamp:(float)lodMaxClamp atIndex:(NSUInteger)index
+void FMetalRenderCommandEncoderDebugging::SetSamplerState(EMetalShaderFrequency Freq,  mtlpp::SamplerState const& sampler, float lodMinClamp, float lodMaxClamp, NSUInteger index)
 {
 #if METAL_DEBUG_OPTIONS
-	switch (Buffer->DebugLevel)
+	switch (((FMetalDebugRenderCommandEncoder*)m_ptr)->Buffer.GetPtr()->DebugLevel)
 	{
 		case EMetalDebugLevelConditionalSubmit:
 		case EMetalDebugLevelWaitForComplete:
 		case EMetalDebugLevelLogOperations:
 		case EMetalDebugLevelValidation:
 		{
-			ShaderSamplers[EMetalShaderFragment].Samplers[index] = sampler;
+			((FMetalDebugRenderCommandEncoder*)m_ptr)->ShaderSamplers[Freq].Samplers[index] = sampler;
 		}
 		case EMetalDebugLevelResetOnBind:
 		case EMetalDebugLevelTrackResources:
 		{
-			[Buffer trackState:sampler];
+			((FMetalDebugRenderCommandEncoder*)m_ptr)->Buffer.TrackState(sampler);
 		}
 		case EMetalDebugLevelFastValidation:
 		{
-			ResourceMask[EMetalShaderFragment].SamplerMask = sampler ? (ResourceMask[EMetalShaderFragment].SamplerMask | (1 << (index))) : (ResourceMask[EMetalShaderFragment].SamplerMask & ~(1 << (index)));
+			((FMetalDebugRenderCommandEncoder*)m_ptr)->ResourceMask[Freq].SamplerMask = sampler ? (((FMetalDebugRenderCommandEncoder*)m_ptr)->ResourceMask[Freq].SamplerMask | (1 << (index))) : (((FMetalDebugRenderCommandEncoder*)m_ptr)->ResourceMask[Freq].SamplerMask & ~(1 << (index)));
 		}
 		default:
 		{
@@ -891,109 +406,35 @@ APPLE_PLATFORM_OBJECT_ALLOC_OVERRIDES(FMetalDebugRenderCommandEncoder)
 		}
 	}
 #endif
-    [Inner setFragmentSamplerState:sampler lodMinClamp:(float)lodMinClamp lodMaxClamp:(float)lodMaxClamp atIndex:(NSUInteger)index];
-}
-
-#if METAL_NEW_NONNULL_DECL
-- (void)setFragmentSamplerStates:(const id <MTLSamplerState> __nullable[__nonnull])samplers lodMinClamps : (const float[__nonnull])lodMinClamps lodMaxClamps : (const float[__nonnull])lodMaxClamps withRange : (NSRange)range
-#else
-- (void)setFragmentSamplerStates:(const id <MTLSamplerState> __nullable [__nullable])samplers lodMinClamps:(const float [__nullable])lodMinClamps lodMaxClamps:(const float [__nullable])lodMaxClamps withRange:(NSRange)range
-#endif
-{
-#if METAL_DEBUG_OPTIONS
-    for(uint32 i = 0; i < range.length; i++)
-    {
-		switch (Buffer->DebugLevel)
-		{
-			case EMetalDebugLevelConditionalSubmit:
-			case EMetalDebugLevelWaitForComplete:
-			case EMetalDebugLevelLogOperations:
-			case EMetalDebugLevelValidation:
-			{
-				ShaderSamplers[EMetalShaderFragment].Samplers[i + range.location] = samplers[i];
-			}
-			case EMetalDebugLevelResetOnBind:
-			case EMetalDebugLevelTrackResources:
-			{
-				[Buffer trackState:samplers[i]];
-			}
-			case EMetalDebugLevelFastValidation:
-			{
-				ResourceMask[EMetalShaderFragment].SamplerMask = samplers[i] ? (ResourceMask[EMetalShaderFragment].SamplerMask | (1 << (i + range.location))) : (ResourceMask[EMetalShaderFragment].SamplerMask & ~(1 << (i + range.location)));
-			}
-			default:
-			{
-				break;
-			}
-		}
-    }
-#endif
-    [Inner setFragmentSamplerStates:samplers lodMinClamps:lodMinClamps lodMaxClamps:lodMaxClamps withRange:(NSRange)range];
-}
-
-- (void)setBlendColorRed:(float)red green:(float)green blue:(float)blue alpha:(float)alpha
-{
-    [Inner setBlendColorRed:red green:green blue:blue alpha:alpha];
 }
 
-- (void)setDepthStencilState:(nullable id <MTLDepthStencilState>)depthStencilState
+void FMetalRenderCommandEncoderDebugging::SetDepthStencilState( mtlpp::DepthStencilState const& depthStencilState)
 {
 #if METAL_DEBUG_OPTIONS
-	if (Buffer->DebugLevel >= EMetalDebugLevelTrackResources)
+	if (((FMetalDebugRenderCommandEncoder*)m_ptr)->Buffer.GetPtr()->DebugLevel >= EMetalDebugLevelTrackResources)
 	{
-		[Buffer trackState:depthStencilState];
+		((FMetalDebugRenderCommandEncoder*)m_ptr)->Buffer.TrackState(depthStencilState);
 	}
 #endif
-    [Inner setDepthStencilState:depthStencilState];
-}
-
-- (void)setStencilReferenceValue:(uint32_t)referenceValue
-{
-    [Inner setStencilReferenceValue:referenceValue];
-}
-
-- (void)setStencilFrontReferenceValue:(uint32_t)frontReferenceValue backReferenceValue:(uint32_t)backReferenceValue
-{
-    [Inner setStencilFrontReferenceValue:frontReferenceValue backReferenceValue:backReferenceValue];
-}
-
-- (void)setVisibilityResultMode:(MTLVisibilityResultMode)mode offset:(NSUInteger)offset
-{
-    [Inner setVisibilityResultMode:mode offset:offset];
-}
-
-- (void)setColorStoreAction:(MTLStoreAction)storeAction atIndex:(NSUInteger)colorAttachmentIndex
-{
-    [Inner setColorStoreAction:storeAction atIndex:colorAttachmentIndex];
-}
-
-- (void)setDepthStoreAction:(MTLStoreAction)storeAction
-{
-    [Inner setDepthStoreAction:storeAction];
-}
-
-- (void)setStencilStoreAction:(MTLStoreAction)storeAction
-{
-    [Inner setStencilStoreAction:storeAction];
 }
 
-- (void)drawPrimitives:(MTLPrimitiveType)primitiveType vertexStart:(NSUInteger)vertexStart vertexCount:(NSUInteger)vertexCount instanceCount:(NSUInteger)instanceCount
+void FMetalRenderCommandEncoderDebugging::Draw(mtlpp::PrimitiveType primitiveType, NSUInteger vertexStart, NSUInteger vertexCount, NSUInteger instanceCount)
 {
 #if METAL_DEBUG_OPTIONS
-	switch(Buffer->DebugLevel)
+	switch(((FMetalDebugRenderCommandEncoder*)m_ptr)->Buffer.GetPtr()->DebugLevel)
 	{
 		case EMetalDebugLevelConditionalSubmit:
 		case EMetalDebugLevelWaitForComplete:
 		case EMetalDebugLevelLogOperations:
 		{
-			[Buffer draw:[NSString stringWithFormat:@"%s", __PRETTY_FUNCTION__]];
+			((FMetalDebugRenderCommandEncoder*)m_ptr)->Buffer.Draw([NSString stringWithFormat:@"%s", __PRETTY_FUNCTION__]);
 		}
 		case EMetalDebugLevelValidation:
 		case EMetalDebugLevelResetOnBind:
 		case EMetalDebugLevelTrackResources:
 		case EMetalDebugLevelFastValidation:
 		{
-			[self validate];
+			Validate();
 		}
 		default:
 		{
@@ -1001,26 +442,25 @@ APPLE_PLATFORM_OBJECT_ALLOC_OVERRIDES(FMetalDebugRenderCommandEncoder)
 		}
 	}
 #endif
-    [Inner drawPrimitives:primitiveType vertexStart:vertexStart vertexCount:vertexCount instanceCount:instanceCount];
 }
 
-- (void)drawPrimitives:(MTLPrimitiveType)primitiveType vertexStart:(NSUInteger)vertexStart vertexCount:(NSUInteger)vertexCount
+void FMetalRenderCommandEncoderDebugging::Draw(mtlpp::PrimitiveType primitiveType, NSUInteger vertexStart, NSUInteger vertexCount)
 {
 #if METAL_DEBUG_OPTIONS
-	switch(Buffer->DebugLevel)
+	switch(((FMetalDebugRenderCommandEncoder*)m_ptr)->Buffer.GetPtr()->DebugLevel)
 	{
 		case EMetalDebugLevelConditionalSubmit:
 		case EMetalDebugLevelWaitForComplete:
 		case EMetalDebugLevelLogOperations:
 		{
-			[Buffer draw:[NSString stringWithFormat:@"%s", __PRETTY_FUNCTION__]];
+			((FMetalDebugRenderCommandEncoder*)m_ptr)->Buffer.Draw([NSString stringWithFormat:@"%s", __PRETTY_FUNCTION__]);
 		}
 		case EMetalDebugLevelValidation:
 		case EMetalDebugLevelResetOnBind:
 		case EMetalDebugLevelTrackResources:
 		case EMetalDebugLevelFastValidation:
 		{
-			[self validate];
+			Validate();
 		}
 		default:
 		{
@@ -1028,29 +468,28 @@ APPLE_PLATFORM_OBJECT_ALLOC_OVERRIDES(FMetalDebugRenderCommandEncoder)
 		}
 	}
 #endif
-    [Inner drawPrimitives:primitiveType vertexStart:vertexStart vertexCount:vertexCount];
 }
 
-- (void)drawIndexedPrimitives:(MTLPrimitiveType)primitiveType indexCount:(NSUInteger)indexCount indexType:(MTLIndexType)indexType indexBuffer:(id <MTLBuffer>)indexBuffer indexBufferOffset:(NSUInteger)indexBufferOffset instanceCount:(NSUInteger)instanceCount
+void FMetalRenderCommandEncoderDebugging::DrawIndexed(mtlpp::PrimitiveType primitiveType, NSUInteger indexCount, mtlpp::IndexType indexType, FMetalBuffer const& indexBuffer, NSUInteger indexBufferOffset, NSUInteger instanceCount)
 {
 #if METAL_DEBUG_OPTIONS
-	switch(Buffer->DebugLevel)
+	switch(((FMetalDebugRenderCommandEncoder*)m_ptr)->Buffer.GetPtr()->DebugLevel)
 	{
 		case EMetalDebugLevelConditionalSubmit:
 		case EMetalDebugLevelWaitForComplete:
 		case EMetalDebugLevelLogOperations:
 		{
-			[Buffer draw:[NSString stringWithFormat:@"%s", __PRETTY_FUNCTION__]];
+			((FMetalDebugRenderCommandEncoder*)m_ptr)->Buffer.Draw([NSString stringWithFormat:@"%s", __PRETTY_FUNCTION__]);
 		}
 		case EMetalDebugLevelValidation:
 		case EMetalDebugLevelResetOnBind:
 		case EMetalDebugLevelTrackResources:
 		{
-			[Buffer trackResource:indexBuffer];
+			((FMetalDebugRenderCommandEncoder*)m_ptr)->Buffer.TrackResource(indexBuffer);
 		}
 		case EMetalDebugLevelFastValidation:
 		{
-			[self validate];
+			Validate();
 		}
 		default:
 		{
@@ -1058,29 +497,28 @@ APPLE_PLATFORM_OBJECT_ALLOC_OVERRIDES(FMetalDebugRenderCommandEncoder)
 		}
 	}
 #endif
-    [Inner drawIndexedPrimitives:(MTLPrimitiveType)primitiveType indexCount:(NSUInteger)indexCount indexType:(MTLIndexType)indexType indexBuffer:(id <MTLBuffer>)indexBuffer indexBufferOffset:(NSUInteger)indexBufferOffset instanceCount:(NSUInteger)instanceCount];
 }
 
-- (void)drawIndexedPrimitives:(MTLPrimitiveType)primitiveType indexCount:(NSUInteger)indexCount indexType:(MTLIndexType)indexType indexBuffer:(id <MTLBuffer>)indexBuffer indexBufferOffset:(NSUInteger)indexBufferOffset
+void FMetalRenderCommandEncoderDebugging::DrawIndexed(mtlpp::PrimitiveType primitiveType, NSUInteger indexCount, mtlpp::IndexType indexType, FMetalBuffer const& indexBuffer, NSUInteger indexBufferOffset)
 {
 #if METAL_DEBUG_OPTIONS
-	switch(Buffer->DebugLevel)
+	switch(((FMetalDebugRenderCommandEncoder*)m_ptr)->Buffer.GetPtr()->DebugLevel)
 	{
 		case EMetalDebugLevelConditionalSubmit:
 		case EMetalDebugLevelWaitForComplete:
 		case EMetalDebugLevelLogOperations:
 		{
-			[Buffer draw:[NSString stringWithFormat:@"%s", __PRETTY_FUNCTION__]];
+			((FMetalDebugRenderCommandEncoder*)m_ptr)->Buffer.Draw([NSString stringWithFormat:@"%s", __PRETTY_FUNCTION__]);
 		}
 		case EMetalDebugLevelValidation:
 		case EMetalDebugLevelResetOnBind:
 		case EMetalDebugLevelTrackResources:
 		{
-			[Buffer trackResource:indexBuffer];
+			((FMetalDebugRenderCommandEncoder*)m_ptr)->Buffer.TrackResource(indexBuffer);
 		}
 		case EMetalDebugLevelFastValidation:
 		{
-			[self validate];
+			Validate();
 		}
 		default:
 		{
@@ -1088,26 +526,25 @@ APPLE_PLATFORM_OBJECT_ALLOC_OVERRIDES(FMetalDebugRenderCommandEncoder)
 		}
 	}
 #endif
-    [Inner drawIndexedPrimitives:(MTLPrimitiveType)primitiveType indexCount:(NSUInteger)indexCount indexType:(MTLIndexType)indexType indexBuffer:(id <MTLBuffer>)indexBuffer indexBufferOffset:(NSUInteger)indexBufferOffset];
 }
 
-- (void)drawPrimitives:(MTLPrimitiveType)primitiveType vertexStart:(NSUInteger)vertexStart vertexCount:(NSUInteger)vertexCount instanceCount:(NSUInteger)instanceCount baseInstance:(NSUInteger)baseInstance
+void FMetalRenderCommandEncoderDebugging::Draw(mtlpp::PrimitiveType primitiveType, NSUInteger vertexStart, NSUInteger vertexCount, NSUInteger instanceCount, NSUInteger baseInstance)
 {
 #if METAL_DEBUG_OPTIONS
-	switch(Buffer->DebugLevel)
+	switch(((FMetalDebugRenderCommandEncoder*)m_ptr)->Buffer.GetPtr()->DebugLevel)
 	{
 		case EMetalDebugLevelConditionalSubmit:
 		case EMetalDebugLevelWaitForComplete:
 		case EMetalDebugLevelLogOperations:
 		{
-			[Buffer draw:[NSString stringWithFormat:@"%s:%u,%u,%u,%u,%u", __PRETTY_FUNCTION__, (uint32)primitiveType, (uint32)vertexStart, (uint32)vertexCount, (uint32)instanceCount, (uint32)baseInstance]];
+			((FMetalDebugRenderCommandEncoder*)m_ptr)->Buffer.Draw([NSString stringWithFormat:@"%s:%u,%u,%u,%u,%u", __PRETTY_FUNCTION__, (uint32)primitiveType, (uint32)vertexStart, (uint32)vertexCount, (uint32)instanceCount, (uint32)baseInstance]);
 		}
 		case EMetalDebugLevelValidation:
 		case EMetalDebugLevelResetOnBind:
 		case EMetalDebugLevelTrackResources:
 		case EMetalDebugLevelFastValidation:
 		{
-			[self validate];
+			Validate();
 		}
 		default:
 		{
@@ -1115,29 +552,28 @@ APPLE_PLATFORM_OBJECT_ALLOC_OVERRIDES(FMetalDebugRenderCommandEncoder)
 		}
 	}
 #endif
-    [Inner drawPrimitives:(MTLPrimitiveType)primitiveType vertexStart:(NSUInteger)vertexStart vertexCount:(NSUInteger)vertexCount instanceCount:(NSUInteger)instanceCount baseInstance:(NSUInteger)baseInstance];
 }
 
-- (void)drawIndexedPrimitives:(MTLPrimitiveType)primitiveType indexCount:(NSUInteger)indexCount indexType:(MTLIndexType)indexType indexBuffer:(id <MTLBuffer>)indexBuffer indexBufferOffset:(NSUInteger)indexBufferOffset instanceCount:(NSUInteger)instanceCount baseVertex:(NSInteger)baseVertex baseInstance:(NSUInteger)baseInstance
+void FMetalRenderCommandEncoderDebugging::DrawIndexed(mtlpp::PrimitiveType primitiveType, NSUInteger indexCount, mtlpp::IndexType indexType, FMetalBuffer const& indexBuffer, NSUInteger indexBufferOffset, NSUInteger instanceCount, NSInteger baseVertex, NSUInteger baseInstance)
 {
 #if METAL_DEBUG_OPTIONS
-	switch(Buffer->DebugLevel)
+	switch(((FMetalDebugRenderCommandEncoder*)m_ptr)->Buffer.GetPtr()->DebugLevel)
 	{
 		case EMetalDebugLevelConditionalSubmit:
 		case EMetalDebugLevelWaitForComplete:
 		case EMetalDebugLevelLogOperations:
 		{
-			[Buffer draw:[NSString stringWithFormat:@"%s:%u,%u,%u,%u,%u,%u,%u", __PRETTY_FUNCTION__, (uint32)primitiveType, (uint32)indexCount, (uint32)indexType, (uint32)indexBufferOffset, (uint32)instanceCount, (uint32)baseVertex, (uint32)baseInstance]];
+			((FMetalDebugRenderCommandEncoder*)m_ptr)->Buffer.Draw([NSString stringWithFormat:@"%s:%u,%u,%u,%u,%u,%u,%u", __PRETTY_FUNCTION__, (uint32)primitiveType, (uint32)indexCount, (uint32)indexType, (uint32)indexBufferOffset, (uint32)instanceCount, (uint32)baseVertex, (uint32)baseInstance]);
 		}
 		case EMetalDebugLevelValidation:
 		case EMetalDebugLevelResetOnBind:
 		case EMetalDebugLevelTrackResources:
 		{
-			[Buffer trackResource:indexBuffer];
+			((FMetalDebugRenderCommandEncoder*)m_ptr)->Buffer.TrackResource(indexBuffer);
 		}
 		case EMetalDebugLevelFastValidation:
 		{
-			[self validate];
+			Validate();
 		}
 		default:
 		{
@@ -1145,29 +581,28 @@ APPLE_PLATFORM_OBJECT_ALLOC_OVERRIDES(FMetalDebugRenderCommandEncoder)
 		}
 	}
 #endif
-    [Inner drawIndexedPrimitives:(MTLPrimitiveType)primitiveType indexCount:(NSUInteger)indexCount indexType:(MTLIndexType)indexType indexBuffer:(id <MTLBuffer>)indexBuffer indexBufferOffset:(NSUInteger)indexBufferOffset instanceCount:(NSUInteger)instanceCount baseVertex:(NSInteger)baseVertex baseInstance:(NSUInteger)baseInstance];
 }
 
-- (void)drawPrimitives:(MTLPrimitiveType)primitiveType indirectBuffer:(id <MTLBuffer>)indirectBuffer indirectBufferOffset:(NSUInteger)indirectBufferOffset
+void FMetalRenderCommandEncoderDebugging::Draw(mtlpp::PrimitiveType primitiveType, FMetalBuffer const& indirectBuffer, NSUInteger indirectBufferOffset)
 {
 #if METAL_DEBUG_OPTIONS
-	switch(Buffer->DebugLevel)
+	switch(((FMetalDebugRenderCommandEncoder*)m_ptr)->Buffer.GetPtr()->DebugLevel)
 	{
 		case EMetalDebugLevelConditionalSubmit:
 		case EMetalDebugLevelWaitForComplete:
 		case EMetalDebugLevelLogOperations:
 		{
-			[Buffer draw:[NSString stringWithFormat:@"%s", __PRETTY_FUNCTION__]];
+			((FMetalDebugRenderCommandEncoder*)m_ptr)->Buffer.Draw([NSString stringWithFormat:@"%s", __PRETTY_FUNCTION__]);
 		}
 		case EMetalDebugLevelValidation:
 		case EMetalDebugLevelResetOnBind:
 		case EMetalDebugLevelTrackResources:
 		{
-			[Buffer trackResource:indirectBuffer];
+			((FMetalDebugRenderCommandEncoder*)m_ptr)->Buffer.TrackResource(indirectBuffer);
 		}
 		case EMetalDebugLevelFastValidation:
 		{
-			[self validate];
+			Validate();
 		}
 		default:
 		{
@@ -1175,30 +610,29 @@ APPLE_PLATFORM_OBJECT_ALLOC_OVERRIDES(FMetalDebugRenderCommandEncoder)
 		}
 	}
 #endif
-    [Inner drawPrimitives:(MTLPrimitiveType)primitiveType indirectBuffer:(id <MTLBuffer>)indirectBuffer indirectBufferOffset:(NSUInteger)indirectBufferOffset];
 }
 
-- (void)drawIndexedPrimitives:(MTLPrimitiveType)primitiveType indexType:(MTLIndexType)indexType indexBuffer:(id <MTLBuffer>)indexBuffer indexBufferOffset:(NSUInteger)indexBufferOffset indirectBuffer:(id <MTLBuffer>)indirectBuffer indirectBufferOffset:(NSUInteger)indirectBufferOffset
+void FMetalRenderCommandEncoderDebugging::DrawIndexed(mtlpp::PrimitiveType primitiveType, mtlpp::IndexType indexType, FMetalBuffer const& indexBuffer, NSUInteger indexBufferOffset, FMetalBuffer const& indirectBuffer, NSUInteger indirectBufferOffset)
 {
 #if METAL_DEBUG_OPTIONS
-	switch(Buffer->DebugLevel)
+	switch(((FMetalDebugRenderCommandEncoder*)m_ptr)->Buffer.GetPtr()->DebugLevel)
 	{
 		case EMetalDebugLevelConditionalSubmit:
 		case EMetalDebugLevelWaitForComplete:
 		case EMetalDebugLevelLogOperations:
 		{
-			[Buffer draw:[NSString stringWithFormat:@"%s", __PRETTY_FUNCTION__]];
+			((FMetalDebugRenderCommandEncoder*)m_ptr)->Buffer.Draw([NSString stringWithFormat:@"%s", __PRETTY_FUNCTION__]);
 		}
 		case EMetalDebugLevelValidation:
 		case EMetalDebugLevelResetOnBind:
 		case EMetalDebugLevelTrackResources:
 		{
-			[Buffer trackResource:indexBuffer];
-			[Buffer trackResource:indirectBuffer];
+			((FMetalDebugRenderCommandEncoder*)m_ptr)->Buffer.TrackResource(indexBuffer);
+			((FMetalDebugRenderCommandEncoder*)m_ptr)->Buffer.TrackResource(indirectBuffer);
 		}
 		case EMetalDebugLevelFastValidation:
 		{
-			[self validate];
+			Validate();
 		}
 		default:
 		{
@@ -1206,89 +640,38 @@ APPLE_PLATFORM_OBJECT_ALLOC_OVERRIDES(FMetalDebugRenderCommandEncoder)
 		}
 	}
 #endif
-    [Inner drawIndexedPrimitives:(MTLPrimitiveType)primitiveType indexType:(MTLIndexType)indexType indexBuffer:(id <MTLBuffer>)indexBuffer indexBufferOffset:(NSUInteger)indexBufferOffset indirectBuffer:(id <MTLBuffer>)indirectBuffer indirectBufferOffset:(NSUInteger)indirectBufferOffset];
-}
-
-#if PLATFORM_MAC
-- (void)textureBarrier
-{
-	[Inner textureBarrier];
-}
-
-#endif
-
-- (void)updateFence:(id <MTLFence>)fence afterStages:(MTLRenderStages)stages
-{
-#if METAL_DEBUG_OPTIONS
-	if (fence && Buffer->DebugLevel >= EMetalDebugLevelValidation)
-	{
-		[self addUpdateFence:fence];
-        if (((FMetalDebugFence*)fence).Inner)
-        {
-		[Inner updateFence:((FMetalDebugFence*)fence).Inner afterStages:(MTLRenderStages)stages];
-	}
-	}
-	else
-#endif
-	{
-		[Inner updateFence:(id <MTLFence>)fence afterStages:(MTLRenderStages)stages];
-	}
-}
-
-- (void)waitForFence:(id <MTLFence>)fence beforeStages:(MTLRenderStages)stages
-{
-#if METAL_DEBUG_OPTIONS
-	if (fence && Buffer->DebugLevel >= EMetalDebugLevelValidation)
-	{
-		[self addWaitFence:fence];
-        if (((FMetalDebugFence*)fence).Inner)
-        {
-		[Inner waitForFence:((FMetalDebugFence*)fence).Inner beforeStages:(MTLRenderStages)stages];
-	}
-	}
-	else
-#endif
-	{
-		[Inner waitForFence:(id <MTLFence>)fence beforeStages:(MTLRenderStages)stages];
-	}
 }
 
--(void)setTessellationFactorBuffer:(nullable id <MTLBuffer>)buffer offset:(NSUInteger)offset instanceStride:(NSUInteger)instanceStride
+void FMetalRenderCommandEncoderDebugging::SetTessellationFactorBuffer( FMetalBuffer const& buffer, NSUInteger offset, NSUInteger instanceStride)
 {
 #if METAL_DEBUG_OPTIONS
-	if (Buffer->DebugLevel >= EMetalDebugLevelTrackResources)
+	if (((FMetalDebugRenderCommandEncoder*)m_ptr)->Buffer.GetPtr()->DebugLevel >= EMetalDebugLevelTrackResources)
 	{
-		[Buffer trackResource:buffer];
+		((FMetalDebugRenderCommandEncoder*)m_ptr)->Buffer.TrackResource(buffer);
 	}
 #endif
-    [Inner setTessellationFactorBuffer:buffer offset:(NSUInteger)offset instanceStride:(NSUInteger)instanceStride];
 }
 
--(void)setTessellationFactorScale:(float)scale
-{
-    [Inner setTessellationFactorScale:(float)scale];
-}
-
--(void)drawPatches:(NSUInteger)numberOfPatchControlPoints patchStart:(NSUInteger)patchStart patchCount:(NSUInteger)patchCount patchIndexBuffer:(nullable id <MTLBuffer>)patchIndexBuffer patchIndexBufferOffset:(NSUInteger)patchIndexBufferOffset instanceCount:(NSUInteger)instanceCount baseInstance:(NSUInteger)baseInstance
+void FMetalRenderCommandEncoderDebugging::DrawPatches(NSUInteger numberOfPatchControlPoints, NSUInteger patchStart, NSUInteger patchCount,  FMetalBuffer const& patchIndexBuffer, NSUInteger patchIndexBufferOffset, NSUInteger instanceCount, NSUInteger baseInstance)
 {
 #if METAL_DEBUG_OPTIONS
-	switch(Buffer->DebugLevel)
+	switch(((FMetalDebugRenderCommandEncoder*)m_ptr)->Buffer.GetPtr()->DebugLevel)
 	{
 		case EMetalDebugLevelConditionalSubmit:
 		case EMetalDebugLevelWaitForComplete:
 		case EMetalDebugLevelLogOperations:
 		{
-			[Buffer draw:[NSString stringWithFormat:@"%s", __PRETTY_FUNCTION__]];
+			((FMetalDebugRenderCommandEncoder*)m_ptr)->Buffer.Draw([NSString stringWithFormat:@"%s", __PRETTY_FUNCTION__]);
 		}
 		case EMetalDebugLevelValidation:
 		case EMetalDebugLevelResetOnBind:
 		case EMetalDebugLevelTrackResources:
 		{
-			[Buffer trackResource:patchIndexBuffer];
+			((FMetalDebugRenderCommandEncoder*)m_ptr)->Buffer.TrackResource(patchIndexBuffer);
 		}
 		case EMetalDebugLevelFastValidation:
 		{
-			[self validate];
+			Validate();
 		}
 		default:
 		{
@@ -1296,31 +679,29 @@ APPLE_PLATFORM_OBJECT_ALLOC_OVERRIDES(FMetalDebugRenderCommandEncoder)
 		}
 	}
 #endif
-    [Inner drawPatches:(NSUInteger)numberOfPatchControlPoints patchStart:(NSUInteger)patchStart patchCount:(NSUInteger)patchCount patchIndexBuffer:patchIndexBuffer patchIndexBufferOffset:(NSUInteger)patchIndexBufferOffset instanceCount:(NSUInteger)instanceCount baseInstance:(NSUInteger)baseInstance];
 }
 
-#if PLATFORM_MAC
--(void)drawPatches:(NSUInteger)numberOfPatchControlPoints patchIndexBuffer:(nullable id <MTLBuffer>)patchIndexBuffer patchIndexBufferOffset:(NSUInteger)patchIndexBufferOffset indirectBuffer:(id <MTLBuffer>)indirectBuffer indirectBufferOffset:(NSUInteger)indirectBufferOffset
+void FMetalRenderCommandEncoderDebugging::DrawPatches(NSUInteger numberOfPatchControlPoints,  FMetalBuffer const& patchIndexBuffer, NSUInteger patchIndexBufferOffset, FMetalBuffer const& indirectBuffer, NSUInteger indirectBufferOffset)
 {
 #if METAL_DEBUG_OPTIONS
-	switch(Buffer->DebugLevel)
+	switch(((FMetalDebugRenderCommandEncoder*)m_ptr)->Buffer.GetPtr()->DebugLevel)
 	{
 		case EMetalDebugLevelConditionalSubmit:
 		case EMetalDebugLevelWaitForComplete:
 		case EMetalDebugLevelLogOperations:
 		{
-			[Buffer draw:[NSString stringWithFormat:@"%s", __PRETTY_FUNCTION__]];
+			((FMetalDebugRenderCommandEncoder*)m_ptr)->Buffer.Draw([NSString stringWithFormat:@"%s", __PRETTY_FUNCTION__]);
 		}
 		case EMetalDebugLevelValidation:
 		case EMetalDebugLevelResetOnBind:
 		case EMetalDebugLevelTrackResources:
 		{
-			[Buffer trackResource:patchIndexBuffer];
-			[Buffer trackResource:indirectBuffer];
+			((FMetalDebugRenderCommandEncoder*)m_ptr)->Buffer.TrackResource(patchIndexBuffer);
+			((FMetalDebugRenderCommandEncoder*)m_ptr)->Buffer.TrackResource(indirectBuffer);
 		}
 		case EMetalDebugLevelFastValidation:
 		{
-			[self validate];
+			Validate();
 		}
 		default:
 		{
@@ -1328,31 +709,29 @@ APPLE_PLATFORM_OBJECT_ALLOC_OVERRIDES(FMetalDebugRenderCommandEncoder)
 		}
 	}
 #endif
-    [Inner drawPatches:(NSUInteger)numberOfPatchControlPoints patchIndexBuffer:patchIndexBuffer patchIndexBufferOffset:(NSUInteger)patchIndexBufferOffset indirectBuffer:(id <MTLBuffer>)indirectBuffer indirectBufferOffset:(NSUInteger)indirectBufferOffset];
 }
-#endif
 
--(void)drawIndexedPatches:(NSUInteger)numberOfPatchControlPoints patchStart:(NSUInteger)patchStart patchCount:(NSUInteger)patchCount patchIndexBuffer:(nullable id <MTLBuffer>)patchIndexBuffer patchIndexBufferOffset:(NSUInteger)patchIndexBufferOffset controlPointIndexBuffer:(id <MTLBuffer>)controlPointIndexBuffer controlPointIndexBufferOffset:(NSUInteger)controlPointIndexBufferOffset instanceCount:(NSUInteger)instanceCount baseInstance:(NSUInteger)baseInstance
+void FMetalRenderCommandEncoderDebugging::DrawIndexedPatches(NSUInteger numberOfPatchControlPoints, NSUInteger patchStart, NSUInteger patchCount,  FMetalBuffer const& patchIndexBuffer, NSUInteger patchIndexBufferOffset, FMetalBuffer const& controlPointIndexBuffer, NSUInteger controlPointIndexBufferOffset, NSUInteger instanceCount, NSUInteger baseInstance)
 {
 #if METAL_DEBUG_OPTIONS
-	switch(Buffer->DebugLevel)
+	switch(((FMetalDebugRenderCommandEncoder*)m_ptr)->Buffer.GetPtr()->DebugLevel)
 	{
 		case EMetalDebugLevelConditionalSubmit:
 		case EMetalDebugLevelWaitForComplete:
 		case EMetalDebugLevelLogOperations:
 		{
-			[Buffer draw:[NSString stringWithFormat:@"%s", __PRETTY_FUNCTION__]];
+			((FMetalDebugRenderCommandEncoder*)m_ptr)->Buffer.Draw([NSString stringWithFormat:@"%s", __PRETTY_FUNCTION__]);
 		}
 		case EMetalDebugLevelValidation:
 		case EMetalDebugLevelResetOnBind:
 		case EMetalDebugLevelTrackResources:
 		{
-			[Buffer trackResource:patchIndexBuffer];
-			[Buffer trackResource:controlPointIndexBuffer];
+			((FMetalDebugRenderCommandEncoder*)m_ptr)->Buffer.TrackResource(patchIndexBuffer);
+			((FMetalDebugRenderCommandEncoder*)m_ptr)->Buffer.TrackResource(controlPointIndexBuffer);
 		}
 		case EMetalDebugLevelFastValidation:
 		{
-			[self validate];
+			Validate();
 		}
 		default:
 		{
@@ -1360,32 +739,30 @@ APPLE_PLATFORM_OBJECT_ALLOC_OVERRIDES(FMetalDebugRenderCommandEncoder)
 		}
 	}
 #endif
-    [Inner drawIndexedPatches:(NSUInteger)numberOfPatchControlPoints patchStart:(NSUInteger)patchStart patchCount:(NSUInteger)patchCount patchIndexBuffer:patchIndexBuffer patchIndexBufferOffset:(NSUInteger)patchIndexBufferOffset controlPointIndexBuffer:(id <MTLBuffer>)controlPointIndexBuffer controlPointIndexBufferOffset:(NSUInteger)controlPointIndexBufferOffset instanceCount:(NSUInteger)instanceCount baseInstance:(NSUInteger)baseInstance];
 }
 
-#if PLATFORM_MAC
--(void)drawIndexedPatches:(NSUInteger)numberOfPatchControlPoints patchIndexBuffer:(nullable id <MTLBuffer>)patchIndexBuffer patchIndexBufferOffset:(NSUInteger)patchIndexBufferOffset controlPointIndexBuffer:(id <MTLBuffer>)controlPointIndexBuffer controlPointIndexBufferOffset:(NSUInteger)controlPointIndexBufferOffset indirectBuffer:(id <MTLBuffer>)indirectBuffer indirectBufferOffset:(NSUInteger)indirectBufferOffset
+void FMetalRenderCommandEncoderDebugging::DrawIndexedPatches(NSUInteger numberOfPatchControlPoints,  FMetalBuffer const& patchIndexBuffer, NSUInteger patchIndexBufferOffset, FMetalBuffer const& controlPointIndexBuffer, NSUInteger controlPointIndexBufferOffset, FMetalBuffer const& indirectBuffer, NSUInteger indirectBufferOffset)
 {
 #if METAL_DEBUG_OPTIONS
-	switch(Buffer->DebugLevel)
+	switch(((FMetalDebugRenderCommandEncoder*)m_ptr)->Buffer.GetPtr()->DebugLevel)
 	{
 		case EMetalDebugLevelConditionalSubmit:
 		case EMetalDebugLevelWaitForComplete:
 		case EMetalDebugLevelLogOperations:
 		{
-			[Buffer draw:[NSString stringWithFormat:@"%s", __PRETTY_FUNCTION__]];
+			((FMetalDebugRenderCommandEncoder*)m_ptr)->Buffer.Draw([NSString stringWithFormat:@"%s", __PRETTY_FUNCTION__]);
 		}
 		case EMetalDebugLevelValidation:
 		case EMetalDebugLevelResetOnBind:
 		case EMetalDebugLevelTrackResources:
 		{
-			[Buffer trackResource:patchIndexBuffer];
-			[Buffer trackResource:controlPointIndexBuffer];
-			[Buffer trackResource:indirectBuffer];
+			((FMetalDebugRenderCommandEncoder*)m_ptr)->Buffer.TrackResource(patchIndexBuffer);
+			((FMetalDebugRenderCommandEncoder*)m_ptr)->Buffer.TrackResource(controlPointIndexBuffer);
+			((FMetalDebugRenderCommandEncoder*)m_ptr)->Buffer.TrackResource(indirectBuffer);
 		}
 		case EMetalDebugLevelFastValidation:
 		{
-			[self validate];
+			Validate();
 		}
 		default:
 		{
@@ -1393,87 +770,23 @@ APPLE_PLATFORM_OBJECT_ALLOC_OVERRIDES(FMetalDebugRenderCommandEncoder)
 		}
 	}
 #endif
-	
-    [Inner drawIndexedPatches:(NSUInteger)numberOfPatchControlPoints patchIndexBuffer:patchIndexBuffer patchIndexBufferOffset:(NSUInteger)patchIndexBufferOffset controlPointIndexBuffer:(id <MTLBuffer>)controlPointIndexBuffer controlPointIndexBufferOffset:(NSUInteger)controlPointIndexBufferOffset indirectBuffer:(id <MTLBuffer>)indirectBuffer indirectBufferOffset:(NSUInteger)indirectBufferOffset];
-}
-#endif
-
-#if METAL_SUPPORTS_INDIRECT_ARGUMENT_BUFFERS
-- (void)useResource:(id <MTLResource>)resource usage:(MTLResourceUsage)usage
-{
-	if (GMetalSupportsIndirectArgumentBuffers)
-	{
-		[Inner useResource:resource usage:usage];
-	}
-}
-
-- (void)useResources:(const id <MTLResource> [])resources count:(NSUInteger)count usage:(MTLResourceUsage)usage
-{
-	if (GMetalSupportsIndirectArgumentBuffers)
-	{
-		[Inner useResources:resources count:count usage:usage];
-	}
-}
-
-- (void)useHeap:(id <MTLHeap>)heap
-{
-	if (GMetalSupportsIndirectArgumentBuffers)
-	{
-		[Inner useHeap:heap];
-	}
-}
-
-- (void)useHeaps:(const id <MTLHeap> [])heaps count:(NSUInteger)count
-{
-	if (GMetalSupportsIndirectArgumentBuffers)
-	{
-		[Inner useHeaps:heaps count:count];
-	}
-}
-#endif // #if METAL_SUPPORTS_INDIRECT_ARGUMENT_BUFFERS
-
-- (void)setViewports:(const MTLViewport [])viewports count:(NSUInteger)count
-{
-	if(FMetalCommandQueue::SupportsFeature(EMetalFeaturesMultipleViewports))
-	{
-		id<IMetalRenderCommandEncoder> Enc = (id<IMetalRenderCommandEncoder>)Inner;
-		[Enc setViewports:viewports count:count];
-	}
-}
-
-- (void)setScissorRects:(const MTLScissorRect [])scissorRects count:(NSUInteger)count
-{
-	if(FMetalCommandQueue::SupportsFeature(EMetalFeaturesMultipleViewports))
-	{
-		id<IMetalRenderCommandEncoder> Enc = (id<IMetalRenderCommandEncoder>)Inner;
-		[Enc setScissorRects:scissorRects count:count];
-	}
 }
 
--(NSString*) description
-{
-	return [Inner description];
-}
-
--(NSString*) debugDescription
-{
-	return [Inner debugDescription];
-}
-
-- (bool)validateFunctionBindings:(EMetalShaderFrequency)Frequency
+/** Validates the pipeline/binding state */
+bool FMetalRenderCommandEncoderDebugging::ValidateFunctionBindings(EMetalShaderFrequency Frequency)
 {
 	bool bOK = true;
 #if METAL_DEBUG_OPTIONS
-	switch (Buffer->DebugLevel)
+	switch (((FMetalDebugRenderCommandEncoder*)m_ptr)->Buffer.GetPtr()->DebugLevel)
 	{
 		case EMetalDebugLevelConditionalSubmit:
 		case EMetalDebugLevelWaitForComplete:
 		case EMetalDebugLevelLogOperations:
 		case EMetalDebugLevelValidation:
 		{
-			check(Pipeline);
-
-			MTLRenderPipelineReflection* Reflection = Pipeline->RenderPipelineReflection;
+			check(((FMetalDebugRenderCommandEncoder*)m_ptr)->Pipeline);
+			
+			MTLRenderPipelineReflection* Reflection = ((FMetalDebugRenderCommandEncoder*)m_ptr)->Pipeline->RenderPipelineReflection;
 			check(Reflection);
 			
 			NSArray<MTLArgument*>* Arguments = nil;
@@ -1503,7 +816,7 @@ APPLE_PLATFORM_OBJECT_ALLOC_OVERRIDES(FMetalDebugRenderCommandEncoder)
 					case MTLArgumentTypeBuffer:
 					{
 						checkf(Arg.index < ML_MaxBuffers, TEXT("Metal buffer index exceeded!"));
-						if ((ShaderBuffers[Frequency].Buffers[Arg.index] == nil && ShaderBuffers[Frequency].Bytes[Arg.index] == nil))
+						if ((((FMetalDebugRenderCommandEncoder*)m_ptr)->ShaderBuffers[Frequency].Buffers[Arg.index] == nil && ((FMetalDebugRenderCommandEncoder*)m_ptr)->ShaderBuffers[Frequency].Bytes[Arg.index] == nil))
 						{
 							bOK = false;
 							UE_LOG(LogMetal, Warning, TEXT("Unbound buffer at Metal index %u which will crash the driver: %s"), (uint32)Arg.index, *FString([Arg description]));
@@ -1517,22 +830,22 @@ APPLE_PLATFORM_OBJECT_ALLOC_OVERRIDES(FMetalDebugRenderCommandEncoder)
 					case MTLArgumentTypeTexture:
 					{
 						checkf(Arg.index < ML_MaxTextures, TEXT("Metal texture index exceeded!"));
-						if (ShaderTextures[Frequency].Textures[Arg.index] == nil)
+						if (((FMetalDebugRenderCommandEncoder*)m_ptr)->ShaderTextures[Frequency].Textures[Arg.index] == nil)
 						{
 							bOK = false;
 							UE_LOG(LogMetal, Warning, TEXT("Unbound texture at Metal index %u which will crash the driver: %s"), (uint32)Arg.index, *FString([Arg description]));
 						}
-						else if (ShaderTextures[Frequency].Textures[Arg.index].textureType != Arg.textureType)
+						else if (((FMetalDebugRenderCommandEncoder*)m_ptr)->ShaderTextures[Frequency].Textures[Arg.index].textureType != Arg.textureType)
 						{
 							bOK = false;
-							UE_LOG(LogMetal, Warning, TEXT("Incorrect texture type bound at Metal index %u which will crash the driver: %s\n%s"), (uint32)Arg.index, *FString([Arg description]), *FString([ShaderTextures[Frequency].Textures[Arg.index] description]));
+							UE_LOG(LogMetal, Warning, TEXT("Incorrect texture type bound at Metal index %u which will crash the driver: %s\n%s"), (uint32)Arg.index, *FString([Arg description]), *FString([((FMetalDebugRenderCommandEncoder*)m_ptr)->ShaderTextures[Frequency].Textures[Arg.index] description]));
 						}
 						break;
 					}
 					case MTLArgumentTypeSampler:
 					{
 						checkf(Arg.index < ML_MaxSamplers, TEXT("Metal sampler index exceeded!"));
-						if (ShaderSamplers[Frequency].Samplers[Arg.index] == nil)
+						if (((FMetalDebugRenderCommandEncoder*)m_ptr)->ShaderSamplers[Frequency].Samplers[Arg.index] == nil)
 						{
 							bOK = false;
 							UE_LOG(LogMetal, Warning, TEXT("Unbound sampler at Metal index %u which will crash the driver: %s"), (uint32)Arg.index, *FString([Arg description]));
@@ -1550,41 +863,41 @@ APPLE_PLATFORM_OBJECT_ALLOC_OVERRIDES(FMetalDebugRenderCommandEncoder)
 		case EMetalDebugLevelTrackResources:
 		case EMetalDebugLevelFastValidation:
 		{
-			check(Pipeline);
+			check(((FMetalDebugRenderCommandEncoder*)m_ptr)->Pipeline);
 			
-			FMetalTextureMask TextureMask = (ResourceMask[Frequency].TextureMask & Pipeline->ResourceMask[Frequency].TextureMask);
-			if (TextureMask != Pipeline->ResourceMask[Frequency].TextureMask)
+			FMetalTextureMask TextureMask = (((FMetalDebugRenderCommandEncoder*)m_ptr)->ResourceMask[Frequency].TextureMask & ((FMetalDebugRenderCommandEncoder*)m_ptr)->Pipeline->ResourceMask[Frequency].TextureMask);
+			if (TextureMask != ((FMetalDebugRenderCommandEncoder*)m_ptr)->Pipeline->ResourceMask[Frequency].TextureMask)
 			{
 				bOK = false;
 				for (uint32 i = 0; i < ML_MaxTextures; i++)
 				{
-					if ((Pipeline->ResourceMask[Frequency].TextureMask & (1 < i)) && ((TextureMask & (1 < i)) != (Pipeline->ResourceMask[Frequency].TextureMask & (1 < i))))
+					if ((((FMetalDebugRenderCommandEncoder*)m_ptr)->Pipeline->ResourceMask[Frequency].TextureMask & (1 < i)) && ((TextureMask & (1 < i)) != (((FMetalDebugRenderCommandEncoder*)m_ptr)->Pipeline->ResourceMask[Frequency].TextureMask & (1 < i))))
 					{
 						UE_LOG(LogMetal, Warning, TEXT("Unbound texture at Metal index %u which will crash the driver"), i);
 					}
 				}
 			}
 			
-			FMetalBufferMask BufferMask = (ResourceMask[Frequency].BufferMask & Pipeline->ResourceMask[Frequency].BufferMask);
-			if (BufferMask != Pipeline->ResourceMask[Frequency].BufferMask)
+			FMetalBufferMask BufferMask = (((FMetalDebugRenderCommandEncoder*)m_ptr)->ResourceMask[Frequency].BufferMask & ((FMetalDebugRenderCommandEncoder*)m_ptr)->Pipeline->ResourceMask[Frequency].BufferMask);
+			if (BufferMask != ((FMetalDebugRenderCommandEncoder*)m_ptr)->Pipeline->ResourceMask[Frequency].BufferMask)
 			{
 				bOK = false;
 				for (uint32 i = 0; i < ML_MaxBuffers; i++)
 				{
-					if ((Pipeline->ResourceMask[Frequency].BufferMask & (1 < i)) && ((BufferMask & (1 < i)) != (Pipeline->ResourceMask[Frequency].BufferMask & (1 < i))))
+					if ((((FMetalDebugRenderCommandEncoder*)m_ptr)->Pipeline->ResourceMask[Frequency].BufferMask & (1 < i)) && ((BufferMask & (1 < i)) != (((FMetalDebugRenderCommandEncoder*)m_ptr)->Pipeline->ResourceMask[Frequency].BufferMask & (1 < i))))
 					{
 						UE_LOG(LogMetal, Warning, TEXT("Unbound buffer at Metal index %u which will crash the driver"), i);
 					}
 				}
 			}
 			
-			FMetalSamplerMask SamplerMask = (ResourceMask[Frequency].SamplerMask & Pipeline->ResourceMask[Frequency].SamplerMask);
-			if (SamplerMask != Pipeline->ResourceMask[Frequency].SamplerMask)
+			FMetalSamplerMask SamplerMask = (((FMetalDebugRenderCommandEncoder*)m_ptr)->ResourceMask[Frequency].SamplerMask & ((FMetalDebugRenderCommandEncoder*)m_ptr)->Pipeline->ResourceMask[Frequency].SamplerMask);
+			if (SamplerMask != ((FMetalDebugRenderCommandEncoder*)m_ptr)->Pipeline->ResourceMask[Frequency].SamplerMask)
 			{
 				bOK = false;
 				for (uint32 i = 0; i < ML_MaxSamplers; i++)
 				{
-					if ((Pipeline->ResourceMask[Frequency].SamplerMask & (1 < i)) && ((SamplerMask & (1 < i)) != (Pipeline->ResourceMask[Frequency].SamplerMask & (1 < i))))
+					if ((((FMetalDebugRenderCommandEncoder*)m_ptr)->Pipeline->ResourceMask[Frequency].SamplerMask & (1 < i)) && ((SamplerMask & (1 < i)) != (((FMetalDebugRenderCommandEncoder*)m_ptr)->Pipeline->ResourceMask[Frequency].SamplerMask & (1 < i))))
 					{
 						UE_LOG(LogMetal, Warning, TEXT("Unbound sampler at Metal index %u which will crash the driver"), i);
 					}
@@ -1599,155 +912,42 @@ APPLE_PLATFORM_OBJECT_ALLOC_OVERRIDES(FMetalDebugRenderCommandEncoder)
 		}
 	}
 #endif
-    return bOK;
+	return bOK;
 }
-
-- (void)validate
+void FMetalRenderCommandEncoderDebugging::Validate()
 {
 #if METAL_DEBUG_OPTIONS
-    bool bOK = [self validateFunctionBindings:EMetalShaderVertex];
-    if (!bOK)
-    {
-        UE_LOG(LogMetal, Error, TEXT("Metal Validation failures for vertex shader:\n%s"), Pipeline && Pipeline->VertexSource ? *FString(Pipeline->VertexSource) : TEXT("nil"));
-    }
-	
-    bOK = [self validateFunctionBindings:EMetalShaderFragment];
-    if (!bOK)
-    {
-        UE_LOG(LogMetal, Error, TEXT("Metal Validation failures for fragment shader:\n%s"), Pipeline && Pipeline->FragmentSource ? *FString(Pipeline->FragmentSource) : TEXT("nil"));
-    }
-#endif
-}
-
--(id<MTLCommandEncoder>)commandEncoder
-{
-	return self;
-}
-
-#if (METAL_NEW_NONNULL_DECL)
--(void)setColorStoreActionOptions:(MTLStoreActionOptions)storeActionOptions atIndex : (NSUInteger)colorAttachmentIndex
-{
-	if (GMetalSupportsStoreActionOptions)
-	{
-		[Inner setColorStoreActionOptions:storeActionOptions atIndex:colorAttachmentIndex];
-	}
-}
-
-- (void)setDepthStoreActionOptions:(MTLStoreActionOptions)storeActionOptions
-{
-	if (GMetalSupportsStoreActionOptions)
-	{
-		[Inner setDepthStoreActionOptions:storeActionOptions];
-	}
-}
-
-- (void)setStencilStoreActionOptions:(MTLStoreActionOptions)storeActionOptions
-{
-	if (GMetalSupportsStoreActionOptions)
-	{
-		[Inner setStencilStoreActionOptions:storeActionOptions];
-	}
-}
-#endif //(METAL_NEW_NONNULL_DECL)
-
-#if METAL_SUPPORTS_TILE_SHADERS
-- (void)setTileBytes:(const void *)bytes length:(NSUInteger)length atIndex:(NSUInteger)index
-{
-	if (GMetalSupportsTileShaders)
-	{
-		[Inner setTileBytes:bytes length:length atIndex:index];
-	}
-}
-
-- (void)setTileBuffer:(nullable id <MTLBuffer>)buffer offset:(NSUInteger)offset atIndex:(NSUInteger)index
-{
-	if (GMetalSupportsTileShaders)
-	{
-		[Inner setTileBuffer:buffer offset:offset atIndex:index];
-	}
-}
-
-- (void)setTileBufferOffset:(NSUInteger)offset atIndex:(NSUInteger)index
-{
-	if (GMetalSupportsTileShaders)
-	{
-		[Inner setTileBufferOffset:offset atIndex:index];
-	}
-}
-
-- (void)setTileBuffers:(const id <MTLBuffer> __nullable [__nonnull])buffers offsets:(const NSUInteger [__nonnull])offset withRange:(NSRange)range
-{
-	if (GMetalSupportsTileShaders)
-	{
-		[Inner setTileBuffers:buffers offsets:offset withRange:range];
-	}
-}
-
-- (void)setTileTexture:(nullable id <MTLTexture>)texture atIndex:(NSUInteger)index
-{
-	if (GMetalSupportsTileShaders)
-	{
-		[Inner setTileTexture:texture atIndex:index];
-	}
-}
-
-- (void)setTileTextures:(const id <MTLTexture> __nullable [__nonnull])textures withRange:(NSRange)range
-{
-	if (GMetalSupportsTileShaders)
-	{
-		[Inner setTileTextures:textures withRange:range];
-	}
-}
-
-- (void)setTileSamplerState:(nullable id <MTLSamplerState>)sampler atIndex:(NSUInteger)index
-{
-	if (GMetalSupportsTileShaders)
+	bool bOK = ValidateFunctionBindings(EMetalShaderVertex);
+	if (!bOK)
 	{
-		[Inner setTileSamplerState:sampler atIndex:index];
+		UE_LOG(LogMetal, Error, TEXT("Metal Validation failures for vertex shader:\n%s"), ((FMetalDebugRenderCommandEncoder*)m_ptr)->Pipeline && ((FMetalDebugRenderCommandEncoder*)m_ptr)->Pipeline->VertexSource ? *FString(((FMetalDebugRenderCommandEncoder*)m_ptr)->Pipeline->VertexSource) : TEXT("nil"));
 	}
-}
-
-- (void)setTileSamplerStates:(const id <MTLSamplerState> __nullable [__nonnull])samplers withRange:(NSRange)range
-{
-	if (GMetalSupportsTileShaders)
+	
+	bOK = ValidateFunctionBindings(EMetalShaderFragment);
+	if (!bOK)
 	{
-		[Inner setTileSamplerStates:samplers withRange:range];
+		UE_LOG(LogMetal, Error, TEXT("Metal Validation failures for fragment shader:\n%s"), ((FMetalDebugRenderCommandEncoder*)m_ptr)->Pipeline && ((FMetalDebugRenderCommandEncoder*)m_ptr)->Pipeline->FragmentSource ? *FString(((FMetalDebugRenderCommandEncoder*)m_ptr)->Pipeline->FragmentSource) : TEXT("nil"));
 	}
+#endif
 }
 
-- (void)setTileSamplerState:(nullable id <MTLSamplerState>)sampler lodMinClamp:(float)lodMinClamp lodMaxClamp:(float)lodMaxClamp atIndex:(NSUInteger)index
+void FMetalRenderCommandEncoderDebugging::InsertDebugSignpost(ns::String const& Label)
 {
-	if (GMetalSupportsTileShaders)
-	{
-		[Inner setTileSamplerState:sampler lodMinClamp:lodMinClamp lodMaxClamp:lodMaxClamp atIndex:index];
-	}
+	((FMetalDebugRenderCommandEncoder*)m_ptr)->Buffer.InsertDebugSignpost(Label);
 }
-
-- (void)setTileSamplerStates:(const id <MTLSamplerState> __nullable [__nonnull])samplers lodMinClamps:(const float [__nonnull])lodMinClamps lodMaxClamps:(const float [__nonnull])lodMaxClamps withRange:(NSRange)range
+void FMetalRenderCommandEncoderDebugging::PushDebugGroup(ns::String const& Group)
 {
-	if (GMetalSupportsTileShaders)
-	{
-		[Inner setTileSamplerStates:samplers lodMinClamps:lodMinClamps lodMaxClamps:lodMaxClamps withRange:range];
-	}
+    ((FMetalDebugRenderCommandEncoder*)m_ptr)->Buffer.PushDebugGroup(Group);
 }
-
-- (void)dispatchThreadsPerTile:(MTLSize)threadsPerTile
+void FMetalRenderCommandEncoderDebugging::PopDebugGroup()
 {
-	if (GMetalSupportsTileShaders)
-	{
-		[Inner dispatchThreadsPerTile:threadsPerTile];
-	}
+    ((FMetalDebugRenderCommandEncoder*)m_ptr)->Buffer.PopDebugGroup();
+#if METAL_DEBUG_OPTIONS
+	InsertDebugDraw();
+#endif
 }
 
-- (void)setThreadgroupMemoryLength:(NSUInteger)length offset:(NSUInteger)offset atIndex:(NSUInteger)index
+void FMetalRenderCommandEncoderDebugging::EndEncoder()
 {
-	if (GMetalSupportsTileShaders)
-	{
-		[Inner setThreadgroupMemoryLength:length offset:offset atIndex:index];
-	}
+	((FMetalDebugRenderCommandEncoder*)m_ptr)->Buffer.EndCommandEncoder();
 }
-
-#endif
-@end
-
-NS_ASSUME_NONNULL_END
diff --git a/Engine/Source/Runtime/Apple/MetalRHI/Private/MetalRenderCommandEncoder.h b/Engine/Source/Runtime/Apple/MetalRHI/Private/MetalRenderCommandEncoder.h
--- a/Engine/Source/Runtime/Apple/MetalRHI/Private/MetalRenderCommandEncoder.h
+++ b/Engine/Source/Runtime/Apple/MetalRHI/Private/MetalRenderCommandEncoder.h
@@ -7,48 +7,64 @@
 NS_ASSUME_NONNULL_BEGIN
 
 @class FMetalShaderPipeline;
+@class FMetalDebugRenderCommandEncoder;
 
-@protocol IMetalRenderCommandEncoder <MTLRenderCommandEncoder>
-- (void)setViewports:(const MTLViewport [])viewports count:(NSUInteger)count;
-- (void)setScissorRects:(const MTLScissorRect [])scissorRects count:(NSUInteger)count;
-@end
-
-@interface FMetalDebugRenderCommandEncoder : FMetalDebugCommandEncoder<MTLRenderCommandEncoder>
+class FMetalRenderCommandEncoderDebugging : public FMetalCommandEncoderDebugging
 {
-    @private
-#pragma mark - Private Member Variables -
-#if METAL_DEBUG_OPTIONS
-	FMetalDebugShaderResourceMask ResourceMask[EMetalShaderRenderNum];
-    FMetalDebugBufferBindings ShaderBuffers[EMetalShaderRenderNum];
-    FMetalDebugTextureBindings ShaderTextures[EMetalShaderRenderNum];
-    FMetalDebugSamplerBindings ShaderSamplers[EMetalShaderRenderNum];
-	id<MTLRenderPipelineState> DebugState;
-#endif
-	MTLRenderPassDescriptor* RenderPassDesc;
-}
-
-/** The wrapped native command-encoder for which we collect debug information. */
-@property (readonly, retain) id<MTLRenderCommandEncoder> Inner;
-@property (readonly, retain) FMetalDebugCommandBuffer* Buffer;
-@property (nonatomic, retain) FMetalShaderPipeline* Pipeline;
-
-/** Initialise the wrapper with the provided command-buffer. */
--(id)initWithEncoder:(id<MTLRenderCommandEncoder>)Encoder fromDescriptor:(MTLRenderPassDescriptor*)RenderPassDesc andCommandBuffer:(FMetalDebugCommandBuffer*)Buffer;
+	void InsertDebugDraw();
+public:
+	FMetalRenderCommandEncoderDebugging();
+	FMetalRenderCommandEncoderDebugging(mtlpp::RenderCommandEncoder& Encoder, mtlpp::RenderPassDescriptor const& Desc, FMetalCommandBufferDebugging& Buffer);
+	FMetalRenderCommandEncoderDebugging(FMetalDebugCommandEncoder* handle);
+	
+	static FMetalRenderCommandEncoderDebugging Get(mtlpp::RenderCommandEncoder& Buffer);
+	
+	void SetPipeline(FMetalShaderPipeline* Pipeline);
+	
+	void SetBytes(EMetalShaderFrequency Freq, const void * bytes, NSUInteger length, NSUInteger index);
+	void SetBuffer(EMetalShaderFrequency Freq,  FMetalBuffer const& buffer, NSUInteger offset, NSUInteger index);
+	void SetBufferOffset(EMetalShaderFrequency Freq, NSUInteger offset, NSUInteger index);
+	void SetTexture(EMetalShaderFrequency Freq,  FMetalTexture const& texture, NSUInteger index);
+	void SetSamplerState(EMetalShaderFrequency Freq,  mtlpp::SamplerState const& sampler, NSUInteger index);
+	void SetSamplerState(EMetalShaderFrequency Freq,  mtlpp::SamplerState const& sampler, float lodMinClamp, float lodMaxClamp, NSUInteger index);
+	void SetDepthStencilState( mtlpp::DepthStencilState const& depthStencilState);
+	
+	void Draw(mtlpp::PrimitiveType primitiveType, NSUInteger vertexStart, NSUInteger vertexCount, NSUInteger instanceCount);
+	
+	void Draw(mtlpp::PrimitiveType primitiveType, NSUInteger vertexStart, NSUInteger vertexCount);
+	
+	void DrawIndexed(mtlpp::PrimitiveType primitiveType, NSUInteger indexCount, mtlpp::IndexType indexType, FMetalBuffer const& indexBuffer, NSUInteger indexBufferOffset, NSUInteger instanceCount);
+	
+	void DrawIndexed(mtlpp::PrimitiveType primitiveType, NSUInteger indexCount, mtlpp::IndexType indexType, FMetalBuffer const& indexBuffer, NSUInteger indexBufferOffset);
+	
+	void Draw(mtlpp::PrimitiveType primitiveType, NSUInteger vertexStart, NSUInteger vertexCount, NSUInteger instanceCount, NSUInteger baseInstance);
+	
+	void DrawIndexed(mtlpp::PrimitiveType primitiveType, NSUInteger indexCount, mtlpp::IndexType indexType, FMetalBuffer const& indexBuffer, NSUInteger indexBufferOffset, NSUInteger instanceCount, NSInteger baseVertex, NSUInteger baseInstance);
+	
+	void Draw(mtlpp::PrimitiveType primitiveType, FMetalBuffer const& indirectBuffer, NSUInteger indirectBufferOffset);
+	
+	void DrawIndexed(mtlpp::PrimitiveType primitiveType, mtlpp::IndexType indexType, FMetalBuffer const& indexBuffer, NSUInteger indexBufferOffset, FMetalBuffer const& indirectBuffer, NSUInteger indirectBufferOffset);
+	
+	void SetTessellationFactorBuffer( FMetalBuffer const& buffer, NSUInteger offset, NSUInteger instanceStride);
+	
+	void DrawPatches(NSUInteger numberOfPatchControlPoints, NSUInteger patchStart, NSUInteger patchCount,  FMetalBuffer const& patchIndexBuffer, NSUInteger patchIndexBufferOffset, NSUInteger instanceCount, NSUInteger baseInstance);
+	
+	void DrawPatches(NSUInteger numberOfPatchControlPoints,  FMetalBuffer const& patchIndexBuffer, NSUInteger patchIndexBufferOffset, FMetalBuffer const& indirectBuffer, NSUInteger indirectBufferOffset);
+	
+	void DrawIndexedPatches(NSUInteger numberOfPatchControlPoints, NSUInteger patchStart, NSUInteger patchCount,  FMetalBuffer const& patchIndexBuffer, NSUInteger patchIndexBufferOffset, FMetalBuffer const& controlPointIndexBuffer, NSUInteger controlPointIndexBufferOffset, NSUInteger instanceCount, NSUInteger baseInstance);
+	
+	void DrawIndexedPatches(NSUInteger numberOfPatchControlPoints,  FMetalBuffer const& patchIndexBuffer, NSUInteger patchIndexBufferOffset, FMetalBuffer const& controlPointIndexBuffer, NSUInteger controlPointIndexBufferOffset, FMetalBuffer const& indirectBuffer, NSUInteger indirectBufferOffset);
+	
+	/** Validates the pipeline/binding state */
+	bool ValidateFunctionBindings(EMetalShaderFrequency Frequency);
+	void Validate();
+	
+	void InsertDebugSignpost(ns::String const& Label);
+	void PushDebugGroup(ns::String const& Group);
+	void PopDebugGroup();
+	
+	void EndEncoder();
+};
 
-/** Validates the pipeline/binding state */
--(bool)validateFunctionBindings:(EMetalShaderFrequency)Frequency;
--(void)validate;
-
-@end
 NS_ASSUME_NONNULL_END
 
-#if METAL_DEBUG_OPTIONS
-#define METAL_SET_RENDER_REFLECTION(Encoder, InPipeline)														\
-    if (GetMetalDeviceContext().GetCommandQueue().GetRuntimeDebuggingLevel() >= EMetalDebugLevelFastValidation) \
-    {																											\
-        ((FMetalDebugRenderCommandEncoder*)Encoder).Pipeline = InPipeline;										\
-    }
-#else
-#define METAL_SET_RENDER_REFLECTION(Encoder, InPipeline)
-#endif
-
diff --git a/Engine/Source/Runtime/Apple/MetalRHI/Private/MetalRenderPass.cpp b/Engine/Source/Runtime/Apple/MetalRHI/Private/MetalRenderPass.cpp
--- a/Engine/Source/Runtime/Apple/MetalRHI/Private/MetalRenderPass.cpp
+++ b/Engine/Source/Runtime/Apple/MetalRHI/Private/MetalRenderPass.cpp
@@ -311,6 +311,7 @@ void FMetalRenderPass::DrawPrimitive(uint32 PrimitiveType, uint32 BaseVertexInde
 			
 			METAL_GPUPROFILE(FMetalProfiler::GetProfiler()->EncodeDraw(CurrentEncoder.GetCommandBufferStats(), __FUNCTION__, NumPrimitives, NumVertices, NumInstances));
 			CurrentEncoder.GetRenderCommandEncoder().Draw(TranslatePrimitiveType(PrimitiveType), BaseVertexIndex, NumVertices, NumInstances);
+			METAL_DEBUG_LAYER(EMetalDebugLevelFastValidation, CurrentEncoder.GetRenderCommandEncoderDebugging().Draw(TranslatePrimitiveType(PrimitiveType), BaseVertexIndex, NumVertices, NumInstances));
 		}
 	}
 	else
@@ -335,6 +336,7 @@ void FMetalRenderPass::DrawPrimitiveIndirect(uint32 PrimitiveType, FMetalVertexB
 		{
 			METAL_GPUPROFILE(FMetalProfiler::GetProfiler()->EncodeDraw(CurrentEncoder.GetCommandBufferStats(), __FUNCTION__, 1, 1, 1));
 			CurrentEncoder.GetRenderCommandEncoder().Draw(TranslatePrimitiveType(PrimitiveType), VertexBuffer->Buffer, ArgumentOffset);
+			METAL_DEBUG_LAYER(EMetalDebugLevelFastValidation, CurrentEncoder.GetRenderCommandEncoderDebugging().Draw(TranslatePrimitiveType(PrimitiveType), VertexBuffer->Buffer, ArgumentOffset));
 		}
 		ConditionalSubmit();
 	}
@@ -410,10 +412,12 @@ void FMetalRenderPass::DrawIndexedPrimitive(FMetalBuffer const& IndexBuffer, uin
 			if (GRHISupportsBaseVertexIndex && GRHISupportsFirstInstance)
 			{
 				CurrentEncoder.GetRenderCommandEncoder().DrawIndexed(TranslatePrimitiveType(PrimitiveType), NumIndices, ((IndexStride == 2) ? mtlpp::IndexType::UInt16 : mtlpp::IndexType::UInt32), IndexBuffer, StartIndex * IndexStride, NumInstances, BaseVertexIndex, FirstInstance);
+				METAL_DEBUG_LAYER(EMetalDebugLevelFastValidation, CurrentEncoder.GetRenderCommandEncoderDebugging().DrawIndexed(TranslatePrimitiveType(PrimitiveType), NumIndices, ((IndexStride == 2) ? mtlpp::IndexType::UInt16 : mtlpp::IndexType::UInt32), IndexBuffer, StartIndex * IndexStride, NumInstances, BaseVertexIndex, FirstInstance));
 			}
 			else
 			{
 				CurrentEncoder.GetRenderCommandEncoder().DrawIndexed(TranslatePrimitiveType(PrimitiveType), NumIndices, ((IndexStride == 2) ? mtlpp::IndexType::UInt16 : mtlpp::IndexType::UInt32), IndexBuffer, StartIndex * IndexStride, NumInstances);
+				METAL_DEBUG_LAYER(EMetalDebugLevelFastValidation, CurrentEncoder.GetRenderCommandEncoderDebugging().DrawIndexed(TranslatePrimitiveType(PrimitiveType), NumIndices, ((IndexStride == 2) ? mtlpp::IndexType::UInt16 : mtlpp::IndexType::UInt32), IndexBuffer, StartIndex * IndexStride, NumInstances));
 			}
 		}
 	}
@@ -442,6 +446,7 @@ void FMetalRenderPass::DrawIndexedIndirect(FMetalIndexBuffer* IndexBuffer, uint3
 		{
 			METAL_GPUPROFILE(FMetalProfiler::GetProfiler()->EncodeDraw(CurrentEncoder.GetCommandBufferStats(), __FUNCTION__, 1, 1, 1));
 			CurrentEncoder.GetRenderCommandEncoder().DrawIndexed(TranslatePrimitiveType(PrimitiveType), (mtlpp::IndexType)IndexBuffer->IndexType, IndexBuffer->Buffer, 0, VertexBuffer->Buffer, (DrawArgumentsIndex * 5 * sizeof(uint32)));
+			METAL_DEBUG_LAYER(EMetalDebugLevelFastValidation, CurrentEncoder.GetRenderCommandEncoderDebugging().DrawIndexed(TranslatePrimitiveType(PrimitiveType), (mtlpp::IndexType)IndexBuffer->IndexType, IndexBuffer->Buffer, 0, VertexBuffer->Buffer, (DrawArgumentsIndex * 5 * sizeof(uint32))));
 		}
 		ConditionalSubmit();
 	}
@@ -465,6 +470,7 @@ void FMetalRenderPass::DrawIndexedPrimitiveIndirect(uint32 PrimitiveType,FMetalI
 		{
 			METAL_GPUPROFILE(FMetalProfiler::GetProfiler()->EncodeDraw(CurrentEncoder.GetCommandBufferStats(), __FUNCTION__, 1, 1, 1));
 			CurrentEncoder.GetRenderCommandEncoder().DrawIndexed(TranslatePrimitiveType(PrimitiveType), (mtlpp::IndexType)IndexBuffer->IndexType, IndexBuffer->Buffer, 0, VertexBuffer->Buffer, ArgumentOffset);
+			METAL_DEBUG_LAYER(EMetalDebugLevelFastValidation, CurrentEncoder.GetRenderCommandEncoderDebugging().DrawIndexed(TranslatePrimitiveType(PrimitiveType), (mtlpp::IndexType)IndexBuffer->IndexType, IndexBuffer->Buffer, 0, VertexBuffer->Buffer, ArgumentOffset));
 		}
 		ConditionalSubmit();
 	}
@@ -592,6 +598,7 @@ void FMetalRenderPass::DrawPatches(uint32 PrimitiveType,FMetalBuffer const& Inde
 		{
 			METAL_GPUPROFILE(FMetalProfiler::GetProfiler()->EncodeDispatch(PrologueEncoder.GetCommandBufferStats(), __FUNCTION__));
 			computeEncoder.DispatchThreadgroups(threadgroups, threadsPerThreadgroup);
+			METAL_DEBUG_LAYER(EMetalDebugLevelFastValidation, PrologueEncoder.GetComputeCommandEncoderDebugging().DispatchThreadgroups(threadgroups, threadsPerThreadgroup));
 		}
 		
 		check(computeEncoder.GetPtr() != nil);
@@ -600,11 +607,13 @@ void FMetalRenderPass::DrawPatches(uint32 PrimitiveType,FMetalBuffer const& Inde
 		if(tessellationFactorBuffer)
 		{
 			renderEncoder.SetTessellationFactorBuffer(tessellationFactorBuffer, tessellationFactorsOffset, 0);
+			METAL_DEBUG_LAYER(EMetalDebugLevelFastValidation, CurrentEncoder.GetRenderCommandEncoderDebugging().SetTessellationFactorBuffer(tessellationFactorBuffer, tessellationFactorsOffset, 0));
 		}
 		if(GMetalTessellationRunDomainStage && !FShaderCache::IsPredrawCall(State.GetShaderCacheStateObject()))
 		{
 			METAL_GPUPROFILE(FMetalProfiler::GetProfiler()->EncodeDraw(CurrentEncoder.GetCommandBufferStats(), __FUNCTION__, NumPrimitives, GetVertexCountForPrimitiveCount(NumPrimitives, PrimitiveType), NumInstances));
 			renderEncoder.DrawPatches(boundShaderState->VertexShader->TessellationOutputControlPoints, 0, NumPrimitives * NumInstances, nil, 0, 1, 0);
+			METAL_DEBUG_LAYER(EMetalDebugLevelFastValidation, CurrentEncoder.GetRenderCommandEncoderDebugging().DrawPatches(boundShaderState->VertexShader->TessellationOutputControlPoints, 0, NumPrimitives * NumInstances, nil, 0, 1, 0));
 		}
 		
 		if(hullShaderOutputBufferSize)
@@ -643,6 +652,7 @@ void FMetalRenderPass::Dispatch(uint32 ThreadGroupCountX, uint32 ThreadGroupCoun
 	check(ComputeShader->NumThreadsX > 0 && ComputeShader->NumThreadsY > 0 && ComputeShader->NumThreadsZ > 0);
 	mtlpp::Size Threadgroups = mtlpp::Size(ThreadGroupCountX, ThreadGroupCountY, ThreadGroupCountZ);
 	CurrentEncoder.GetComputeCommandEncoder().DispatchThreadgroups(Threadgroups, ThreadgroupCounts);
+	METAL_DEBUG_LAYER(EMetalDebugLevelFastValidation, CurrentEncoder.GetComputeCommandEncoderDebugging().DispatchThreadgroups(Threadgroups, ThreadgroupCounts));
 	
 	ConditionalSubmit();
 }
@@ -665,6 +675,7 @@ void FMetalRenderPass::DispatchIndirect(FMetalVertexBuffer* ArgumentBuffer, uint
 	check(ComputeShader->NumThreadsX > 0 && ComputeShader->NumThreadsY > 0 && ComputeShader->NumThreadsZ > 0);
 	
 	CurrentEncoder.GetComputeCommandEncoder().DispatchThreadgroupsWithIndirectBuffer(ArgumentBuffer->Buffer, ArgumentOffset, ThreadgroupCounts);
+	METAL_DEBUG_LAYER(EMetalDebugLevelFastValidation, CurrentEncoder.GetComputeCommandEncoderDebugging().DispatchThreadgroupsWithIndirectBuffer(ArgumentBuffer->Buffer, ArgumentOffset, ThreadgroupCounts));
 
 	ConditionalSubmit();
 }
@@ -692,11 +703,13 @@ void FMetalRenderPass::CopyFromTextureToBuffer(FMetalTexture const& Texture, uin
 	if (CmdList.GetCommandQueue().SupportsFeature(EMetalFeaturesDepthStencilBlitOptions))
 	{
 		MTLPP_VALIDATE(mtlpp::BlitCommandEncoder, Encoder, SafeGetRuntimeDebuggingLevel() >= EMetalDebugLevelValidation, Copy(Texture, sourceSlice, sourceLevel, sourceOrigin, sourceSize, toBuffer, destinationOffset, destinationBytesPerRow, destinationBytesPerImage, options));
+		METAL_DEBUG_LAYER(EMetalDebugLevelFastValidation, CurrentEncoder.GetBlitCommandEncoderDebugging().Copy(Texture, sourceSlice, sourceLevel, sourceOrigin, sourceSize, toBuffer, destinationOffset, destinationBytesPerRow, destinationBytesPerImage, options));
 	}
 	else
 	{
 		check(options == mtlpp::BlitOption::None);
 		MTLPP_VALIDATE(mtlpp::BlitCommandEncoder, Encoder, SafeGetRuntimeDebuggingLevel() >= EMetalDebugLevelValidation, Copy(Texture, sourceSlice, sourceLevel, sourceOrigin, sourceSize, toBuffer, destinationOffset, destinationBytesPerRow, destinationBytesPerImage));
+		METAL_DEBUG_LAYER(EMetalDebugLevelFastValidation, CurrentEncoder.GetBlitCommandEncoderDebugging().Copy(Texture, sourceSlice, sourceLevel, sourceOrigin, sourceSize, toBuffer, destinationOffset, destinationBytesPerRow, destinationBytesPerImage));
 	}
 	ConditionalSubmit();
 }
@@ -711,10 +724,12 @@ void FMetalRenderPass::CopyFromBufferToTexture(FMetalBuffer const& Buffer, uint3
 	if (options == mtlpp::BlitOption::None)
 	{
 		MTLPP_VALIDATE(mtlpp::BlitCommandEncoder, Encoder, SafeGetRuntimeDebuggingLevel() >= EMetalDebugLevelValidation, Copy(Buffer, sourceOffset, sourceBytesPerRow, sourceBytesPerImage, sourceSize, toTexture, destinationSlice, destinationLevel, destinationOrigin));
+		METAL_DEBUG_LAYER(EMetalDebugLevelFastValidation, CurrentEncoder.GetBlitCommandEncoderDebugging().Copy(Buffer, sourceOffset, sourceBytesPerRow, sourceBytesPerImage, sourceSize, toTexture, destinationSlice, destinationLevel, destinationOrigin));
 	}
 	else
 	{
 		MTLPP_VALIDATE(mtlpp::BlitCommandEncoder, Encoder, SafeGetRuntimeDebuggingLevel() >= EMetalDebugLevelValidation, Copy(Buffer, sourceOffset, sourceBytesPerRow, sourceBytesPerImage, sourceSize, toTexture, destinationSlice, destinationLevel, destinationOrigin, options));
+		METAL_DEBUG_LAYER(EMetalDebugLevelFastValidation, CurrentEncoder.GetBlitCommandEncoderDebugging().Copy(Buffer, sourceOffset, sourceBytesPerRow, sourceBytesPerImage, sourceSize, toTexture, destinationSlice, destinationLevel, destinationOrigin, options));
 	}
 	ConditionalSubmit();
 }
@@ -727,6 +742,7 @@ void FMetalRenderPass::CopyFromTextureToTexture(FMetalTexture const& Texture, ui
 	
 	METAL_GPUPROFILE(FMetalProfiler::GetProfiler()->EncodeBlit(CurrentEncoder.GetCommandBufferStats(), __FUNCTION__));
 	MTLPP_VALIDATE(mtlpp::BlitCommandEncoder, Encoder, SafeGetRuntimeDebuggingLevel() >= EMetalDebugLevelValidation, Copy(Texture, sourceSlice, sourceLevel, sourceOrigin, sourceSize, toTexture, destinationSlice, destinationLevel, destinationOrigin));
+	METAL_DEBUG_LAYER(EMetalDebugLevelFastValidation, CurrentEncoder.GetBlitCommandEncoderDebugging().Copy(Texture, sourceSlice, sourceLevel, sourceOrigin, sourceSize, toTexture, destinationSlice, destinationLevel, destinationOrigin));
 	ConditionalSubmit();
 }
 
@@ -738,7 +754,8 @@ void FMetalRenderPass::CopyFromBufferToBuffer(FMetalBuffer const& SourceBuffer,
 	
 	METAL_GPUPROFILE(FMetalProfiler::GetProfiler()->EncodeBlit(CurrentEncoder.GetCommandBufferStats(), __FUNCTION__));
 	MTLPP_VALIDATE(mtlpp::BlitCommandEncoder, Encoder, SafeGetRuntimeDebuggingLevel() >= EMetalDebugLevelValidation, Copy(SourceBuffer, SourceOffset, DestinationBuffer, DestinationOffset, Size));
-	 ConditionalSubmit();
+	METAL_DEBUG_LAYER(EMetalDebugLevelFastValidation, CurrentEncoder.GetBlitCommandEncoderDebugging().Copy(SourceBuffer, SourceOffset, DestinationBuffer, DestinationOffset, Size));
+	ConditionalSubmit();
 }
 
 void FMetalRenderPass::PresentTexture(FMetalTexture const& Texture, uint32 sourceSlice, uint32 sourceLevel, mtlpp::Origin sourceOrigin, mtlpp::Size sourceSize, FMetalTexture const& toTexture, uint32 destinationSlice, uint32 destinationLevel, mtlpp::Origin destinationOrigin)
@@ -749,6 +766,7 @@ void FMetalRenderPass::PresentTexture(FMetalTexture const& Texture, uint32 sourc
 	
 	METAL_GPUPROFILE(FMetalProfiler::GetProfiler()->EncodeBlit(CurrentEncoder.GetCommandBufferStats(), __FUNCTION__));
 	MTLPP_VALIDATE(mtlpp::BlitCommandEncoder, Encoder, SafeGetRuntimeDebuggingLevel() >= EMetalDebugLevelValidation, Copy(Texture, sourceSlice, sourceLevel, sourceOrigin, sourceSize, toTexture, destinationSlice, destinationLevel, destinationOrigin));
+	METAL_DEBUG_LAYER(EMetalDebugLevelFastValidation, CurrentEncoder.GetBlitCommandEncoderDebugging().Copy(Texture, sourceSlice, sourceLevel, sourceOrigin, sourceSize, toTexture, destinationSlice, destinationLevel, destinationOrigin));
 }
 
 void FMetalRenderPass::SynchronizeTexture(FMetalTexture const& Texture, uint32 Slice, uint32 Level)
@@ -761,6 +779,7 @@ void FMetalRenderPass::SynchronizeTexture(FMetalTexture const& Texture, uint32 S
 	
 	METAL_GPUPROFILE(FMetalProfiler::GetProfiler()->EncodeBlit(CurrentEncoder.GetCommandBufferStats(), __FUNCTION__));
 	MTLPP_VALIDATE(mtlpp::BlitCommandEncoder, Encoder, SafeGetRuntimeDebuggingLevel() >= EMetalDebugLevelValidation, Synchronize(Texture, Slice, Level));
+	METAL_DEBUG_LAYER(EMetalDebugLevelFastValidation, CurrentEncoder.GetBlitCommandEncoderDebugging().Synchronize(Texture, Slice, Level));
 	ConditionalSubmit();
 #endif
 }
@@ -775,6 +794,7 @@ void FMetalRenderPass::SynchroniseResource(mtlpp::Resource const& Resource)
 	
 	METAL_GPUPROFILE(FMetalProfiler::GetProfiler()->EncodeBlit(CurrentEncoder.GetCommandBufferStats(), __FUNCTION__));
 	MTLPP_VALIDATE(mtlpp::BlitCommandEncoder, Encoder, SafeGetRuntimeDebuggingLevel() >= EMetalDebugLevelValidation, Synchronize(Resource));
+	METAL_DEBUG_LAYER(EMetalDebugLevelFastValidation, CurrentEncoder.GetBlitCommandEncoderDebugging().Synchronize(Resource));
 	ConditionalSubmit();
 #endif
 }
@@ -789,25 +809,28 @@ void FMetalRenderPass::FillBuffer(FMetalBuffer const& Buffer, ns::Range Range, u
 	
 	METAL_GPUPROFILE(FMetalProfiler::GetProfiler()->EncodeBlit(CurrentEncoder.GetCommandBufferStats(), __FUNCTION__));
 	MTLPP_VALIDATE(mtlpp::BlitCommandEncoder, Encoder, SafeGetRuntimeDebuggingLevel() >= EMetalDebugLevelValidation, Fill(Buffer, Range, Value));
+	METAL_DEBUG_LAYER(EMetalDebugLevelFastValidation, CurrentEncoder.GetBlitCommandEncoderDebugging().Fill(Buffer, Range, Value));
 	ConditionalSubmit();
 }
 
 bool FMetalRenderPass::AsyncCopyFromBufferToTexture(FMetalBuffer const& Buffer, uint32 sourceOffset, uint32 sourceBytesPerRow, uint32 sourceBytesPerImage, mtlpp::Size sourceSize, FMetalTexture const& toTexture, uint32 destinationSlice, uint32 destinationLevel, mtlpp::Origin destinationOrigin, mtlpp::BlitOption options)
 {
 	mtlpp::BlitCommandEncoder TargetEncoder;
-	
+	METAL_DEBUG_OPTION(FMetalBlitCommandEncoderDebugging Debugging);
 	bool bAsync = !CurrentEncoder.HasTextureBindingHistory(toTexture);
 	if(bAsync)
 	{
 		ConditionalSwitchToAsyncBlit();
 		TargetEncoder = PrologueEncoder.GetBlitCommandEncoder();
 		METAL_GPUPROFILE(FMetalProfiler::GetProfiler()->EncodeBlit(PrologueEncoder.GetCommandBufferStats(), __FUNCTION__));
+		METAL_DEBUG_LAYER(EMetalDebugLevelFastValidation, Debugging = PrologueEncoder.GetBlitCommandEncoderDebugging());
 	}
 	else
 	{
 		ConditionalSwitchToBlit();
 		TargetEncoder = CurrentEncoder.GetBlitCommandEncoder();
 		METAL_GPUPROFILE(FMetalProfiler::GetProfiler()->EncodeBlit(CurrentEncoder.GetCommandBufferStats(), __FUNCTION__));
+		METAL_DEBUG_LAYER(EMetalDebugLevelFastValidation, Debugging = CurrentEncoder.GetBlitCommandEncoderDebugging());
 	}
 	
 	check(TargetEncoder.GetPtr());
@@ -815,10 +838,12 @@ bool FMetalRenderPass::AsyncCopyFromBufferToTexture(FMetalBuffer const& Buffer,
 	if (options == mtlpp::BlitOption::None)
 	{
 		MTLPP_VALIDATE(mtlpp::BlitCommandEncoder, TargetEncoder, SafeGetRuntimeDebuggingLevel() >= EMetalDebugLevelValidation, Copy(Buffer, sourceOffset, sourceBytesPerRow, sourceBytesPerImage, sourceSize, toTexture, destinationSlice, destinationLevel, destinationOrigin));
+		METAL_DEBUG_LAYER(EMetalDebugLevelFastValidation, Debugging.Copy(Buffer, sourceOffset, sourceBytesPerRow, sourceBytesPerImage, sourceSize, toTexture, destinationSlice, destinationLevel, destinationOrigin));
 	}
 	else
 	{
 		MTLPP_VALIDATE(mtlpp::BlitCommandEncoder, TargetEncoder, SafeGetRuntimeDebuggingLevel() >= EMetalDebugLevelValidation, Copy(Buffer, sourceOffset, sourceBytesPerRow, sourceBytesPerImage, sourceSize, toTexture, destinationSlice, destinationLevel, destinationOrigin, options));
+		METAL_DEBUG_LAYER(EMetalDebugLevelFastValidation, Debugging.Copy(Buffer, sourceOffset, sourceBytesPerRow, sourceBytesPerImage, sourceSize, toTexture, destinationSlice, destinationLevel, destinationOrigin, options));
 	}
 	
 	return bAsync;
@@ -827,24 +852,27 @@ bool FMetalRenderPass::AsyncCopyFromBufferToTexture(FMetalBuffer const& Buffer,
 bool FMetalRenderPass::AsyncCopyFromTextureToTexture(FMetalTexture const& Texture, uint32 sourceSlice, uint32 sourceLevel, mtlpp::Origin sourceOrigin, mtlpp::Size sourceSize, FMetalTexture const& toTexture, uint32 destinationSlice, uint32 destinationLevel, mtlpp::Origin destinationOrigin)
 {
 	mtlpp::BlitCommandEncoder TargetEncoder;
-	
+	METAL_DEBUG_OPTION(FMetalBlitCommandEncoderDebugging Debugging);
 	bool bAsync = !CurrentEncoder.HasTextureBindingHistory(toTexture);
 	if(bAsync)
 	{
 		ConditionalSwitchToAsyncBlit();
 		TargetEncoder = PrologueEncoder.GetBlitCommandEncoder();
 		METAL_GPUPROFILE(FMetalProfiler::GetProfiler()->EncodeBlit(PrologueEncoder.GetCommandBufferStats(), __FUNCTION__));
+		METAL_DEBUG_LAYER(EMetalDebugLevelFastValidation, Debugging = PrologueEncoder.GetBlitCommandEncoderDebugging());
 	}
 	else
 	{
 		ConditionalSwitchToBlit();
 		TargetEncoder = CurrentEncoder.GetBlitCommandEncoder();
 		METAL_GPUPROFILE(FMetalProfiler::GetProfiler()->EncodeBlit(CurrentEncoder.GetCommandBufferStats(), __FUNCTION__));
+		METAL_DEBUG_LAYER(EMetalDebugLevelFastValidation, Debugging = CurrentEncoder.GetBlitCommandEncoderDebugging());
 	}
 	
 	check(TargetEncoder.GetPtr());
 	
 	MTLPP_VALIDATE(mtlpp::BlitCommandEncoder, TargetEncoder, SafeGetRuntimeDebuggingLevel() >= EMetalDebugLevelValidation, Copy(Texture, sourceSlice, sourceLevel, sourceOrigin, sourceSize, toTexture, destinationSlice, destinationLevel, destinationOrigin));
+	METAL_DEBUG_LAYER(EMetalDebugLevelFastValidation, Debugging.Copy(Texture, sourceSlice, sourceLevel, sourceOrigin, sourceSize, toTexture, destinationSlice, destinationLevel, destinationOrigin));
 	
 	return bAsync;
 }
@@ -852,37 +880,40 @@ bool FMetalRenderPass::AsyncCopyFromTextureToTexture(FMetalTexture const& Textur
 void FMetalRenderPass::AsyncCopyFromBufferToBuffer(FMetalBuffer const& SourceBuffer, NSUInteger SourceOffset, FMetalBuffer const& DestinationBuffer, NSUInteger DestinationOffset, NSUInteger Size)
 {
 	mtlpp::BlitCommandEncoder TargetEncoder;
-	
+	METAL_DEBUG_OPTION(FMetalBlitCommandEncoderDebugging Debugging);
 	bool bAsync = !CurrentEncoder.HasBufferBindingHistory(DestinationBuffer);
 	if(bAsync)
 	{
 		ConditionalSwitchToAsyncBlit();
 		TargetEncoder = PrologueEncoder.GetBlitCommandEncoder();
 		METAL_GPUPROFILE(FMetalProfiler::GetProfiler()->EncodeBlit(PrologueEncoder.GetCommandBufferStats(), __FUNCTION__));
+		METAL_DEBUG_LAYER(EMetalDebugLevelFastValidation, Debugging = PrologueEncoder.GetBlitCommandEncoderDebugging());
 	}
 	else
 	{
 		ConditionalSwitchToBlit();
 		TargetEncoder = CurrentEncoder.GetBlitCommandEncoder();
 		METAL_GPUPROFILE(FMetalProfiler::GetProfiler()->EncodeBlit(CurrentEncoder.GetCommandBufferStats(), __FUNCTION__));
+		METAL_DEBUG_LAYER(EMetalDebugLevelFastValidation, Debugging = CurrentEncoder.GetBlitCommandEncoderDebugging());
 	}
 	
 	check(TargetEncoder.GetPtr());
 	
     MTLPP_VALIDATE(mtlpp::BlitCommandEncoder, TargetEncoder, SafeGetRuntimeDebuggingLevel() >= EMetalDebugLevelValidation, Copy(SourceBuffer, SourceOffset, DestinationBuffer, DestinationOffset, Size));
+	METAL_DEBUG_LAYER(EMetalDebugLevelFastValidation, Debugging.Copy(SourceBuffer, SourceOffset, DestinationBuffer, DestinationOffset, Size));
 }
 
 void FMetalRenderPass::AsyncGenerateMipmapsForTexture(FMetalTexture const& Texture)
 {
 	// This must be a plain old error
 	check(!CurrentEncoder.HasTextureBindingHistory(Texture));
-
 	ConditionalSwitchToAsyncBlit();
 	mtlpp::BlitCommandEncoder Encoder = PrologueEncoder.GetBlitCommandEncoder();
 	check(Encoder.GetPtr());
 	
 	METAL_GPUPROFILE(FMetalProfiler::GetProfiler()->EncodeBlit(CurrentEncoder.GetCommandBufferStats(), __FUNCTION__));
 	MTLPP_VALIDATE(mtlpp::BlitCommandEncoder, Encoder, SafeGetRuntimeDebuggingLevel() >= EMetalDebugLevelValidation, GenerateMipmaps(Texture));
+	METAL_DEBUG_LAYER(EMetalDebugLevelFastValidation, PrologueEncoder.GetBlitCommandEncoderDebugging().GenerateMipmaps(Texture));
 }
 
 mtlpp::Fence FMetalRenderPass::End(void)
diff --git a/Engine/Source/Runtime/Apple/MetalRHI/Private/MetalStateCache.cpp b/Engine/Source/Runtime/Apple/MetalRHI/Private/MetalStateCache.cpp
--- a/Engine/Source/Runtime/Apple/MetalRHI/Private/MetalStateCache.cpp
+++ b/Engine/Source/Runtime/Apple/MetalRHI/Private/MetalStateCache.cpp
@@ -1841,6 +1841,7 @@ void FMetalStateCache::FlushVisibilityResults(FMetalCommandEncoder& CommandEncod
 
 		METAL_GPUPROFILE(FMetalProfiler::GetProfiler()->EncodeBlit(CommandEncoder.GetCommandBufferStats(), __FUNCTION__));
 		MTLPP_VALIDATE(mtlpp::BlitCommandEncoder, Encoder, SafeGetRuntimeDebuggingLevel() >= EMetalDebugLevelValidation, Synchronize(VisibilityResults->Buffer));
+		METAL_DEBUG_LAYER(EMetalDebugLevelFastValidation, CommandEncoder.GetBlitCommandEncoderDebugging().Synchronize(VisibilityResults->Buffer));
 		
 		VisibilityWritten = 0;
 	}
diff --git a/Engine/Source/Runtime/Apple/MetalRHI/Private/MetalViewport.cpp b/Engine/Source/Runtime/Apple/MetalRHI/Private/MetalViewport.cpp
--- a/Engine/Source/Runtime/Apple/MetalRHI/Private/MetalViewport.cpp
+++ b/Engine/Source/Runtime/Apple/MetalRHI/Private/MetalViewport.cpp
@@ -429,14 +429,24 @@ void FMetalViewport::Present(FMetalCommandQueue& CommandQueue, bool bLockToVsync
 							
 							mtlpp::BlitCommandEncoder Encoder = CurrentCommandBuffer.BlitCommandEncoder();
 							check(Encoder.GetPtr());
+#if METAL_DEBUG_OPTIONS
+							FMetalBlitCommandEncoderDebugging Debugging;
+							if (SafeGetRuntimeDebuggingLevel() >= EMetalDebugLevelFastValidation)
+							{
+								FMetalCommandBufferDebugging CmdDebug = FMetalCommandBufferDebugging::Get(CurrentCommandBuffer);
+								Debugging = FMetalBlitCommandEncoderDebugging(Encoder, CmdDebug);
+							}
+#endif
 							
 							METAL_STATISTIC(Profiler->BeginEncoder(Stats, Encoder));
 							METAL_GPUPROFILE(Profiler->EncodeBlit(Stats, __FUNCTION__));
 
 							Encoder.Copy(Src, 0, 0, mtlpp::Origin(0, 0, 0), mtlpp::Size(Width, Height, 1), Dst, 0, 0, mtlpp::Origin(0, 0, 0));
+							METAL_DEBUG_LAYER(EMetalDebugLevelFastValidation, Debugging.Copy(Src, 0, 0, mtlpp::Origin(0, 0, 0), mtlpp::Size(Width, Height, 1), Dst, 0, 0, mtlpp::Origin(0, 0, 0)));
 							
 							METAL_STATISTIC(Profiler->EndEncoder(Stats, Encoder));
 							Encoder.EndEncoding();
+							METAL_DEBUG_LAYER(EMetalDebugLevelFastValidation, Debugging.EndEncoder());
 							
 							mtlpp::CommandBufferHandler H = [Src, Dst](const mtlpp::CommandBuffer &) {
 							};
