Fix Shipping iOS builds.

diff --git a/Engine/Source/Runtime/Apple/MetalRHI/Private/MetalBlitCommandEncoder.cpp b/Engine/Source/Runtime/Apple/MetalRHI/Private/MetalBlitCommandEncoder.cpp
--- a/Engine/Source/Runtime/Apple/MetalRHI/Private/MetalBlitCommandEncoder.cpp
+++ b/Engine/Source/Runtime/Apple/MetalRHI/Private/MetalBlitCommandEncoder.cpp
@@ -10,6 +10,7 @@
 #include "MetalCommandBuffer.h"
 #include "MetalFence.h"
 
+#if METAL_DEBUG_OPTIONS
 extern int32 GMetalRuntimeDebugLevel;
 
 @interface FMetalDebugBlitCommandEncoder : FMetalDebugCommandEncoder
@@ -341,3 +342,5 @@ void FMetalBlitCommandEncoderDebugging::Copy(FMetalBuffer const& sourceBuffer, N
 	}
 #endif
 }
+
+#endif
diff --git a/Engine/Source/Runtime/Apple/MetalRHI/Private/MetalBlitCommandEncoder.h b/Engine/Source/Runtime/Apple/MetalRHI/Private/MetalBlitCommandEncoder.h
--- a/Engine/Source/Runtime/Apple/MetalRHI/Private/MetalBlitCommandEncoder.h
+++ b/Engine/Source/Runtime/Apple/MetalRHI/Private/MetalBlitCommandEncoder.h
@@ -5,6 +5,7 @@
 #include <Metal/Metal.h>
 #include "MetalDebugCommandEncoder.h"
 
+#if METAL_DEBUG_OPTIONS
 NS_ASSUME_NONNULL_BEGIN
 
 @class FMetalDebugCommandBuffer;
@@ -49,3 +50,4 @@ public:
 };
 
 NS_ASSUME_NONNULL_END
+#endif
diff --git a/Engine/Source/Runtime/Apple/MetalRHI/Private/MetalCommandBuffer.cpp b/Engine/Source/Runtime/Apple/MetalRHI/Private/MetalCommandBuffer.cpp
--- a/Engine/Source/Runtime/Apple/MetalRHI/Private/MetalCommandBuffer.cpp
+++ b/Engine/Source/Runtime/Apple/MetalRHI/Private/MetalCommandBuffer.cpp
@@ -46,7 +46,7 @@ uint32 SafeGetRuntimeDebuggingLevel()
 }
 @end
 
-#if MTLPP_CONFIG_VALIDATE
+#if MTLPP_CONFIG_VALIDATE && METAL_DEBUG_OPTIONS
 
 @implementation FMetalDebugCommandBuffer
 
diff --git a/Engine/Source/Runtime/Apple/MetalRHI/Private/MetalCommandBuffer.h b/Engine/Source/Runtime/Apple/MetalRHI/Private/MetalCommandBuffer.h
--- a/Engine/Source/Runtime/Apple/MetalRHI/Private/MetalCommandBuffer.h
+++ b/Engine/Source/Runtime/Apple/MetalRHI/Private/MetalCommandBuffer.h
@@ -60,7 +60,7 @@ struct FMetalDebugCommand
 @property (nonatomic, strong) NSMutableArray<NSString*>* debugGroups;
 @end
 
-#if MTLPP_CONFIG_VALIDATE
+#if MTLPP_CONFIG_VALIDATE && METAL_DEBUG_OPTIONS
 
 /**
  * FMetalDebugCommandBuffer: Wrapper around id<MTLCommandBuffer> that records information about commands.
diff --git a/Engine/Source/Runtime/Apple/MetalRHI/Private/MetalComputeCommandEncoder.cpp b/Engine/Source/Runtime/Apple/MetalRHI/Private/MetalComputeCommandEncoder.cpp
--- a/Engine/Source/Runtime/Apple/MetalRHI/Private/MetalComputeCommandEncoder.cpp
+++ b/Engine/Source/Runtime/Apple/MetalRHI/Private/MetalComputeCommandEncoder.cpp
@@ -11,6 +11,7 @@
 #include "MetalFence.h"
 #include "MetalPipeline.h"
 
+#if METAL_DEBUG_OPTIONS
 NS_ASSUME_NONNULL_BEGIN
 
 #if METAL_DEBUG_OPTIONS
@@ -554,3 +555,4 @@ void FMetalComputeCommandEncoderDebugging::Validate()
 }
 
 NS_ASSUME_NONNULL_END
+#endif
diff --git a/Engine/Source/Runtime/Apple/MetalRHI/Private/MetalComputeCommandEncoder.h b/Engine/Source/Runtime/Apple/MetalRHI/Private/MetalComputeCommandEncoder.h
--- a/Engine/Source/Runtime/Apple/MetalRHI/Private/MetalComputeCommandEncoder.h
+++ b/Engine/Source/Runtime/Apple/MetalRHI/Private/MetalComputeCommandEncoder.h
@@ -4,6 +4,7 @@
 
 #include "MetalDebugCommandEncoder.h"
 
+#if METAL_DEBUG_OPTIONS
 NS_ASSUME_NONNULL_BEGIN
 
 @class FMetalDebugComputeCommandEncoder;
@@ -38,3 +39,4 @@ public:
 };
 
 NS_ASSUME_NONNULL_END
+#endif
diff --git a/Engine/Source/Runtime/Apple/MetalRHI/Private/MetalDebugCommandEncoder.cpp b/Engine/Source/Runtime/Apple/MetalRHI/Private/MetalDebugCommandEncoder.cpp
--- a/Engine/Source/Runtime/Apple/MetalRHI/Private/MetalDebugCommandEncoder.cpp
+++ b/Engine/Source/Runtime/Apple/MetalRHI/Private/MetalDebugCommandEncoder.cpp
@@ -6,6 +6,7 @@
 #include "MetalCommandBuffer.h"
 #include "MetalFence.h"
 
+#if METAL_DEBUG_OPTIONS
 extern int32 GMetalRuntimeDebugLevel;
 
 @implementation FMetalDebugCommandEncoder
@@ -59,3 +60,5 @@ void FMetalCommandEncoderDebugging::AddWaitFence(id Fence)
 		[(FMetalDebugFence*)Fence waitingEncoder:m_ptr];
 	}
 }
+#endif
+
diff --git a/Engine/Source/Runtime/Apple/MetalRHI/Private/MetalDebugCommandEncoder.h b/Engine/Source/Runtime/Apple/MetalRHI/Private/MetalDebugCommandEncoder.h
--- a/Engine/Source/Runtime/Apple/MetalRHI/Private/MetalDebugCommandEncoder.h
+++ b/Engine/Source/Runtime/Apple/MetalRHI/Private/MetalDebugCommandEncoder.h
@@ -80,6 +80,7 @@ struct FMetalDebugSamplerBindings
     id<MTLSamplerState> _Nullable Samplers[ML_MaxSamplers];
 };
 
+#if METAL_DEBUG_OPTIONS
 class FMetalCommandBufferDebugging;
 @class FMetalDebugCommandBuffer;
 @class FMetalDebugFence;
@@ -103,5 +104,5 @@ public:
 	void AddWaitFence(id Fence);
 };
 
+#endif
 NS_ASSUME_NONNULL_END
-
diff --git a/Engine/Source/Runtime/Apple/MetalRHI/Private/MetalRenderCommandEncoder.cpp b/Engine/Source/Runtime/Apple/MetalRHI/Private/MetalRenderCommandEncoder.cpp
--- a/Engine/Source/Runtime/Apple/MetalRHI/Private/MetalRenderCommandEncoder.cpp
+++ b/Engine/Source/Runtime/Apple/MetalRHI/Private/MetalRenderCommandEncoder.cpp
@@ -12,6 +12,7 @@
 #include "MetalFence.h"
 #include "MetalPipeline.h"
 
+#if METAL_DEBUG_OPTIONS
 NS_ASSUME_NONNULL_BEGIN
 
 @interface FMetalDebugRenderCommandEncoder : FMetalDebugCommandEncoder
@@ -951,3 +952,5 @@ void FMetalRenderCommandEncoderDebugging::EndEncoder()
 {
 	((FMetalDebugRenderCommandEncoder*)m_ptr)->Buffer.EndCommandEncoder();
 }
+#endif
+
diff --git a/Engine/Source/Runtime/Apple/MetalRHI/Private/MetalRenderCommandEncoder.h b/Engine/Source/Runtime/Apple/MetalRHI/Private/MetalRenderCommandEncoder.h
--- a/Engine/Source/Runtime/Apple/MetalRHI/Private/MetalRenderCommandEncoder.h
+++ b/Engine/Source/Runtime/Apple/MetalRHI/Private/MetalRenderCommandEncoder.h
@@ -4,6 +4,7 @@
 
 #include "MetalDebugCommandEncoder.h"
 
+#if METAL_DEBUG_OPTIONS
 NS_ASSUME_NONNULL_BEGIN
 
 @class FMetalShaderPipeline;
@@ -67,4 +68,4 @@ public:
 };
 
 NS_ASSUME_NONNULL_END
-
+#endif
