Bringing CL 4104051 from Dev-Rendering (//UE4/Dev-Rendering) to Release-4.20 (//UE4/Release-4.20)

---
Some fixes for semaphore handling.

Moved the completion handler to EndFrame (out of FlushFreeList())

FlushFreeList() attached a handler to the command buffer that signaled the semaphore. This can be called through RHIFlushResources which will cause a mismatched signal()

Added BeginFrame/EndFrame to the loop in AddModalWindows so the renderer gets the expected frame delimiters. Removed call to EndFrame from EndDrawingViewport.

---

diff --git a/Engine/Source/Runtime/Apple/MetalRHI/Private/MetalContext.cpp b/Engine/Source/Runtime/Apple/MetalRHI/Private/MetalContext.cpp
--- a/Engine/Source/Runtime/Apple/MetalRHI/Private/MetalContext.cpp
+++ b/Engine/Source/Runtime/Apple/MetalRHI/Private/MetalContext.cpp
@@ -483,6 +483,21 @@ void FMetalDeviceContext::EndFrame()
 	
 	ClearFreeList();
 	
+	// A 'frame' in this context is from the beginning of encoding on the CPU
+	// to the end of all rendering operations on the GPU. So the semaphore is
+	// signalled when the last command buffer finishes GPU execution.
+	{
+		dispatch_semaphore_t CmdBufferSemaphore = CommandBufferSemaphore;
+		dispatch_retain(CmdBufferSemaphore);
+		
+		RenderPass.AddCompletionHandler(
+		[CmdBufferSemaphore](mtlpp::CommandBuffer const& cmd_buf)
+		{
+			dispatch_semaphore_signal(CmdBufferSemaphore);
+			dispatch_release(CmdBufferSemaphore);
+		});
+	}
+	
 	if (bPresented)
 	{
 		FMetalGPUProfiler::IncrementFrameIndex();
@@ -567,14 +582,6 @@ void FMetalDeviceContext::FlushFreeList()
 	ObjectFreeList.Empty(ObjectFreeList.Num());
 	FreeListMutex.Unlock();
 	
-	dispatch_semaphore_t CmdBufferSemaphore = CommandBufferSemaphore;
-	dispatch_retain(CmdBufferSemaphore);
-	
-	RenderPass.AddCompletionHandler([CmdBufferSemaphore](mtlpp::CommandBuffer const&)
-									{
-										dispatch_semaphore_signal(CmdBufferSemaphore);
-										dispatch_release(CmdBufferSemaphore);
-									});
 	DelayedFreeLists.Add(NewList);
 }
 
@@ -620,12 +627,6 @@ void FMetalDeviceContext::EndDrawingViewport(FMetalViewport* Viewport, bool bPre
 		FrameReadyEvent->Wait();
 	}
 	
-	// The editor doesn't always call EndFrame appropriately so do so here
-	if (GIsEditor)
-	{
-		EndFrame();
-	}
-	
 	Viewport->ReleaseDrawable();
 }
 
diff --git a/Engine/Source/Runtime/Slate/Private/Framework/Application/SlateApplication.cpp b/Engine/Source/Runtime/Slate/Private/Framework/Application/SlateApplication.cpp
--- a/Engine/Source/Runtime/Slate/Private/Framework/Application/SlateApplication.cpp
+++ b/Engine/Source/Runtime/Slate/Private/Framework/Application/SlateApplication.cpp
@@ -2159,6 +2159,10 @@ void FSlateApplication::AddModalWindow( TSharedRef<SWindow> InSlateWindow, const
 		{
 			PlatformApplication->Cursor->Show( true );
 		}
+		
+		// Since the engine tick is paused here we have to end any outstanding frames.
+		// That will put us in a clean state for the Slate Tick loop below.
+		Renderer->EndFrame();
 
 		//Throttle loop data
 		float LastLoopTime = (float)FPlatformTime::Seconds();
@@ -2195,9 +2199,14 @@ void FSlateApplication::AddModalWindow( TSharedRef<SWindow> InSlateWindow, const
 				{
 					break;
 				}
+				
+				// Slate's Tick does not issue Begin/EndFrame so we'll do it ourselves.
+				Renderer->BeginFrame();
 
 				// Tick and render Slate
 				TickApplication(ESlateTickType::All, DeltaTime);
+				
+				Renderer->EndFrame();
 			}
 
 			// Update Slate Stats
@@ -2206,6 +2215,10 @@ void FSlateApplication::AddModalWindow( TSharedRef<SWindow> InSlateWindow, const
 			// Synchronize the game thread and the render thread so that the render thread doesn't get too far behind.
 			Renderer->Sync();
 		}
+		
+		// When we let the engine Tick as normal we need to restore the state as before.
+		// So we'll start a frame.
+		Renderer->BeginFrame();
 	}
 }
 
diff --git a/Engine/Source/Runtime/SlateCore/Public/Rendering/SlateRenderer.h b/Engine/Source/Runtime/SlateCore/Public/Rendering/SlateRenderer.h
--- a/Engine/Source/Runtime/SlateCore/Public/Rendering/SlateRenderer.h
+++ b/Engine/Source/Runtime/SlateCore/Public/Rendering/SlateRenderer.h
@@ -348,6 +348,18 @@ public:
 	 * in a multi-threaded environment.  This function does not return until the sync is complete
 	 */
 	virtual void Sync() const {};
+	
+	/**
+	 * Indicates the start of a new frame to the Renderer. This is usually handled by the engine loop
+	 * but certain situations (ie, when the main loop is paused) may require manual calls.
+	 */
+	virtual void BeginFrame() const {};
+	
+	/**
+	 * Indicates the end of the current frame to the Renderer. This is usually handled by the engine loop
+	 * but certain situations (ie, when the main loop is paused) may require manual calls.
+	 */
+	virtual void EndFrame() const {};
 
 	/**
 	 * Reloads all texture resources from disk
diff --git a/Engine/Source/Runtime/SlateRHIRenderer/Private/SlateRHIRenderer.cpp b/Engine/Source/Runtime/SlateRHIRenderer/Private/SlateRHIRenderer.cpp
--- a/Engine/Source/Runtime/SlateRHIRenderer/Private/SlateRHIRenderer.cpp
+++ b/Engine/Source/Runtime/SlateRHIRenderer/Private/SlateRHIRenderer.cpp
@@ -1227,6 +1227,32 @@ void FSlateRHIRenderer::Sync() const
 	FrameEndSync.Sync(CVarAllowOneFrameThreadLag->GetValueOnAnyThread() != 0);
 }
 
+/**
+ * Inline issues a BeginFrame to the RHI.
+ * This is to handle cases like Modal dialogs in the UI. The game loop stops while
+ * the dialog is open but continues to issue draws. The RHI thinks there are all part of one super long
+ * frame until the Modal window is closed.
+ */
+void FSlateRHIRenderer::BeginFrame() const
+{
+	ENQUEUE_RENDER_COMMAND(SlateRHIBeginFrame)(
+	   [](FRHICommandListImmediate& RHICmdList)
+	   {
+		   RHICmdList.BeginFrame();
+	   }
+	);
+}
+
+void FSlateRHIRenderer::EndFrame() const
+{
+	ENQUEUE_RENDER_COMMAND(SlateRHIEndFrame)(
+	   [](FRHICommandListImmediate& RHICmdList)
+	   {
+		   RHICmdList.EndFrame();
+	   }
+	);
+}
+
 void FSlateRHIRenderer::ReloadTextureResources()
 {
 	ResourceManager->ReloadTextures();
@@ -1553,4 +1579,4 @@ void FSlateEndDrawingWindowsCommand::EndDrawingWindows(FRHICommandListImmediate&
 		FSlateEndDrawingWindowsCommand Cmd(Policy, DrawBuffer);
 		Cmd.Execute(RHICmdList);
 	}
-}
\ No newline at end of file
+}
diff --git a/Engine/Source/Runtime/SlateRHIRenderer/Private/SlateRHIRenderer.h b/Engine/Source/Runtime/SlateRHIRenderer/Private/SlateRHIRenderer.h
--- a/Engine/Source/Runtime/SlateRHIRenderer/Private/SlateRHIRenderer.h
+++ b/Engine/Source/Runtime/SlateRHIRenderer/Private/SlateRHIRenderer.h
@@ -201,6 +201,8 @@ public:
 	virtual int32 RegisterCurrentScene(FSceneInterface* Scene) override;
 	virtual int32 GetCurrentSceneIndex() const override;
 	virtual void ClearScenes() override;
+	virtual void BeginFrame() const override;
+	virtual void EndFrame() const override;
 	virtual void AddWidgetRendererUpdate(const struct FRenderThreadUpdateContext& Context, bool bDeferredRenderTargetUpdate) override;
 
 	/** Draws windows from a FSlateDrawBuffer on the render thread */
