CL 4107446 from Dev-Rendering.
Always force a commandbuffer submission during EndFrame.

diff --git a/Engine/Source/Runtime/Apple/MetalRHI/Private/MetalContext.cpp b/Engine/Source/Runtime/Apple/MetalRHI/Private/MetalContext.cpp
--- a/Engine/Source/Runtime/Apple/MetalRHI/Private/MetalContext.cpp
+++ b/Engine/Source/Runtime/Apple/MetalRHI/Private/MetalContext.cpp
@@ -505,9 +505,10 @@ void FMetalDeviceContext::EndFrame()
 		bPresented = false;
 	}
 	
-	// Latched update of whether to use runtime debugging features
-	uint32 SubmitFlags = EMetalSubmitFlagsResetState;
+	// Force submission so the completion handler that signals CommandBufferSemaphore fires.
+	uint32 SubmitFlags = EMetalSubmitFlagsResetState | EMetalSubmitFlagsForce;
 #if METAL_DEBUG_OPTIONS
+	// Latched update of whether to use runtime debugging features
 	if (GMetalRuntimeDebugLevel != CommandQueue.GetRuntimeDebuggingLevel())
 	{
 		CommandQueue.SetRuntimeDebuggingLevel(GMetalRuntimeDebugLevel);
