Rolling back the previous integration of CL 4080446

diff --git a/Engine/Source/Runtime/Apple/MetalRHI/Private/MetalContext.cpp b/Engine/Source/Runtime/Apple/MetalRHI/Private/MetalContext.cpp
--- a/Engine/Source/Runtime/Apple/MetalRHI/Private/MetalContext.cpp
+++ b/Engine/Source/Runtime/Apple/MetalRHI/Private/MetalContext.cpp
@@ -505,10 +505,9 @@ void FMetalDeviceContext::EndFrame()
 		bPresented = false;
 	}
 	
-	// Force submission so the completion handler that signals CommandBufferSemaphore fires.
-	uint32 SubmitFlags = EMetalSubmitFlagsResetState | EMetalSubmitFlagsForce;
-#if METAL_DEBUG_OPTIONS
 	// Latched update of whether to use runtime debugging features
+	uint32 SubmitFlags = EMetalSubmitFlagsResetState;
+#if METAL_DEBUG_OPTIONS
 	if (GMetalRuntimeDebugLevel != CommandQueue.GetRuntimeDebuggingLevel())
 	{
 		CommandQueue.SetRuntimeDebuggingLevel(GMetalRuntimeDebugLevel);
