More build fixes, this time for Debug on iOS.

diff --git a/Engine/Source/Runtime/Apple/MetalRHI/Private/MetalBlitCommandEncoder.cpp b/Engine/Source/Runtime/Apple/MetalRHI/Private/MetalBlitCommandEncoder.cpp
--- a/Engine/Source/Runtime/Apple/MetalRHI/Private/MetalBlitCommandEncoder.cpp
+++ b/Engine/Source/Runtime/Apple/MetalRHI/Private/MetalBlitCommandEncoder.cpp
@@ -10,7 +10,7 @@
 #include "MetalCommandBuffer.h"
 #include "MetalFence.h"
 
-#if METAL_DEBUG_OPTIONS
+#if MTLPP_CONFIG_VALIDATE && METAL_DEBUG_OPTIONS
 extern int32 GMetalRuntimeDebugLevel;
 
 @interface FMetalDebugBlitCommandEncoder : FMetalDebugCommandEncoder
diff --git a/Engine/Source/Runtime/Apple/MetalRHI/Private/MetalBlitCommandEncoder.h b/Engine/Source/Runtime/Apple/MetalRHI/Private/MetalBlitCommandEncoder.h
--- a/Engine/Source/Runtime/Apple/MetalRHI/Private/MetalBlitCommandEncoder.h
+++ b/Engine/Source/Runtime/Apple/MetalRHI/Private/MetalBlitCommandEncoder.h
@@ -5,7 +5,7 @@
 #include <Metal/Metal.h>
 #include "MetalDebugCommandEncoder.h"
 
-#if METAL_DEBUG_OPTIONS
+#if MTLPP_CONFIG_VALIDATE && METAL_DEBUG_OPTIONS
 NS_ASSUME_NONNULL_BEGIN
 
 @class FMetalDebugCommandBuffer;
diff --git a/Engine/Source/Runtime/Apple/MetalRHI/Private/MetalCommandEncoder.h b/Engine/Source/Runtime/Apple/MetalRHI/Private/MetalCommandEncoder.h
--- a/Engine/Source/Runtime/Apple/MetalRHI/Private/MetalCommandEncoder.h
+++ b/Engine/Source/Runtime/Apple/MetalRHI/Private/MetalCommandEncoder.h
@@ -161,7 +161,7 @@ public:
 	FMetalCommandBufferStats* GetCommandBufferStats(void);
 #endif
 
-#if METAL_DEBUG_OPTIONS
+#if MTLPP_CONFIG_VALIDATE && METAL_DEBUG_OPTIONS
 	/** @returns The active render command encoder or nil if there isn't one. */
 	FMetalCommandBufferDebugging& GetCommandBufferDebugging(void) { return CommandBufferDebug; } 
 	
@@ -395,10 +395,10 @@ private:
 	mtlpp::ComputeCommandEncoder ComputeCommandEncoder;
 	mtlpp::BlitCommandEncoder BlitCommandEncoder;
 	
-	METAL_DEBUG_OPTION(FMetalCommandBufferDebugging CommandBufferDebug);
-	METAL_DEBUG_OPTION(FMetalRenderCommandEncoderDebugging RenderEncoderDebug);
-	METAL_DEBUG_OPTION(FMetalComputeCommandEncoderDebugging ComputeEncoderDebug);
-	METAL_DEBUG_OPTION(FMetalBlitCommandEncoderDebugging BlitEncoderDebug);
+	METAL_DEBUG_ONLY(FMetalCommandBufferDebugging CommandBufferDebug);
+	METAL_DEBUG_ONLY(FMetalRenderCommandEncoderDebugging RenderEncoderDebug);
+	METAL_DEBUG_ONLY(FMetalComputeCommandEncoderDebugging ComputeEncoderDebug);
+	METAL_DEBUG_ONLY(FMetalBlitCommandEncoderDebugging BlitEncoderDebug);
 	
 	FMetalFence EncoderFence;
 #if ENABLE_METAL_GPUPROFILE
diff --git a/Engine/Source/Runtime/Apple/MetalRHI/Private/MetalCommandQueue.cpp b/Engine/Source/Runtime/Apple/MetalRHI/Private/MetalCommandQueue.cpp
--- a/Engine/Source/Runtime/Apple/MetalRHI/Private/MetalCommandQueue.cpp
+++ b/Engine/Source/Runtime/Apple/MetalRHI/Private/MetalCommandQueue.cpp
@@ -228,7 +228,7 @@ mtlpp::CommandBuffer FMetalCommandQueue::CreateCommandBuffer(void)
 		
 		if (RuntimeDebuggingLevel > EMetalDebugLevelLogDebugGroups)
 		{			
-			METAL_DEBUG_OPTION(FMetalCommandBufferDebugging AddDebugging(CmdBuffer));
+			METAL_DEBUG_ONLY(FMetalCommandBufferDebugging AddDebugging(CmdBuffer));
 			MTLPP_VALIDATION(mtlpp::CommandBufferValidationTable ValidatedCommandBuffer(CmdBuffer));
 		}
 		else if (RuntimeDebuggingLevel == EMetalDebugLevelLogDebugGroups)
diff --git a/Engine/Source/Runtime/Apple/MetalRHI/Private/MetalComputeCommandEncoder.cpp b/Engine/Source/Runtime/Apple/MetalRHI/Private/MetalComputeCommandEncoder.cpp
--- a/Engine/Source/Runtime/Apple/MetalRHI/Private/MetalComputeCommandEncoder.cpp
+++ b/Engine/Source/Runtime/Apple/MetalRHI/Private/MetalComputeCommandEncoder.cpp
@@ -11,7 +11,7 @@
 #include "MetalFence.h"
 #include "MetalPipeline.h"
 
-#if METAL_DEBUG_OPTIONS
+#if MTLPP_CONFIG_VALIDATE && METAL_DEBUG_OPTIONS
 NS_ASSUME_NONNULL_BEGIN
 
 #if METAL_DEBUG_OPTIONS
diff --git a/Engine/Source/Runtime/Apple/MetalRHI/Private/MetalComputeCommandEncoder.h b/Engine/Source/Runtime/Apple/MetalRHI/Private/MetalComputeCommandEncoder.h
--- a/Engine/Source/Runtime/Apple/MetalRHI/Private/MetalComputeCommandEncoder.h
+++ b/Engine/Source/Runtime/Apple/MetalRHI/Private/MetalComputeCommandEncoder.h
@@ -4,7 +4,7 @@
 
 #include "MetalDebugCommandEncoder.h"
 
-#if METAL_DEBUG_OPTIONS
+#if MTLPP_CONFIG_VALIDATE && METAL_DEBUG_OPTIONS
 NS_ASSUME_NONNULL_BEGIN
 
 @class FMetalDebugComputeCommandEncoder;
diff --git a/Engine/Source/Runtime/Apple/MetalRHI/Private/MetalDebugCommandEncoder.cpp b/Engine/Source/Runtime/Apple/MetalRHI/Private/MetalDebugCommandEncoder.cpp
--- a/Engine/Source/Runtime/Apple/MetalRHI/Private/MetalDebugCommandEncoder.cpp
+++ b/Engine/Source/Runtime/Apple/MetalRHI/Private/MetalDebugCommandEncoder.cpp
@@ -6,7 +6,7 @@
 #include "MetalCommandBuffer.h"
 #include "MetalFence.h"
 
-#if METAL_DEBUG_OPTIONS
+#if MTLPP_CONFIG_VALIDATE && METAL_DEBUG_OPTIONS
 extern int32 GMetalRuntimeDebugLevel;
 
 @implementation FMetalDebugCommandEncoder
diff --git a/Engine/Source/Runtime/Apple/MetalRHI/Private/MetalDebugCommandEncoder.h b/Engine/Source/Runtime/Apple/MetalRHI/Private/MetalDebugCommandEncoder.h
--- a/Engine/Source/Runtime/Apple/MetalRHI/Private/MetalDebugCommandEncoder.h
+++ b/Engine/Source/Runtime/Apple/MetalRHI/Private/MetalDebugCommandEncoder.h
@@ -80,7 +80,7 @@ struct FMetalDebugSamplerBindings
     id<MTLSamplerState> _Nullable Samplers[ML_MaxSamplers];
 };
 
-#if METAL_DEBUG_OPTIONS
+#if MTLPP_CONFIG_VALIDATE && METAL_DEBUG_OPTIONS
 class FMetalCommandBufferDebugging;
 @class FMetalDebugCommandBuffer;
 @class FMetalDebugFence;
diff --git a/Engine/Source/Runtime/Apple/MetalRHI/Private/MetalQuery.cpp b/Engine/Source/Runtime/Apple/MetalRHI/Private/MetalQuery.cpp
--- a/Engine/Source/Runtime/Apple/MetalRHI/Private/MetalQuery.cpp
+++ b/Engine/Source/Runtime/Apple/MetalRHI/Private/MetalQuery.cpp
@@ -411,7 +411,7 @@ void FMetalDynamicRHI::RHICalibrateTimers()
 		
 		id<IMetalStatisticsSamples> Samples = Context.GetCommandQueue().GetStatistics()->RegisterEncoderStatistics(Buffer.GetPtr(), EMetalSampleComputeEncoderStart);
 		mtlpp::ComputeCommandEncoder Encoder = Buffer.ComputeCommandEncoder();
-#if METAL_DEBUG_OPTIONS
+#if MTLPP_CONFIG_VALIDATE && METAL_DEBUG_OPTIONS
 		FMetalComputeCommandEncoderDebugging Debugging;
 		if (SafeGetRuntimeDebuggingLevel() >= EMetalDebugLevelFastValidation)
 		{
diff --git a/Engine/Source/Runtime/Apple/MetalRHI/Private/MetalRHIPrivate.h b/Engine/Source/Runtime/Apple/MetalRHI/Private/MetalRHIPrivate.h
--- a/Engine/Source/Runtime/Apple/MetalRHI/Private/MetalRHIPrivate.h
+++ b/Engine/Source/Runtime/Apple/MetalRHI/Private/MetalRHIPrivate.h
@@ -88,9 +88,15 @@ extern FMetalBufferFormat GMetalBufferFormats[PF_MAX];
 #define METAL_DEBUG_OPTIONS !(UE_BUILD_SHIPPING || UE_BUILD_TEST)
 #if METAL_DEBUG_OPTIONS
 #define METAL_DEBUG_OPTION(Code) Code
-#define METAL_DEBUG_LAYER(Level, Code) if (SafeGetRuntimeDebuggingLevel() >= Level) Code
 #else
 #define METAL_DEBUG_OPTION(Code)
+#endif
+
+#if MTLPP_CONFIG_VALIDATE && METAL_DEBUG_OPTIONS
+#define METAL_DEBUG_ONLY(Code) Code
+#define METAL_DEBUG_LAYER(Level, Code) if (SafeGetRuntimeDebuggingLevel() >= Level) Code
+#else
+#define METAL_DEBUG_ONLY(Code)
 #define METAL_DEBUG_LAYER(Level, Code)
 #endif
 
diff --git a/Engine/Source/Runtime/Apple/MetalRHI/Private/MetalRenderCommandEncoder.cpp b/Engine/Source/Runtime/Apple/MetalRHI/Private/MetalRenderCommandEncoder.cpp
--- a/Engine/Source/Runtime/Apple/MetalRHI/Private/MetalRenderCommandEncoder.cpp
+++ b/Engine/Source/Runtime/Apple/MetalRHI/Private/MetalRenderCommandEncoder.cpp
@@ -12,7 +12,7 @@
 #include "MetalFence.h"
 #include "MetalPipeline.h"
 
-#if METAL_DEBUG_OPTIONS
+#if MTLPP_CONFIG_VALIDATE && METAL_DEBUG_OPTIONS
 NS_ASSUME_NONNULL_BEGIN
 
 @interface FMetalDebugRenderCommandEncoder : FMetalDebugCommandEncoder
diff --git a/Engine/Source/Runtime/Apple/MetalRHI/Private/MetalRenderCommandEncoder.h b/Engine/Source/Runtime/Apple/MetalRHI/Private/MetalRenderCommandEncoder.h
--- a/Engine/Source/Runtime/Apple/MetalRHI/Private/MetalRenderCommandEncoder.h
+++ b/Engine/Source/Runtime/Apple/MetalRHI/Private/MetalRenderCommandEncoder.h
@@ -4,7 +4,7 @@
 
 #include "MetalDebugCommandEncoder.h"
 
-#if METAL_DEBUG_OPTIONS
+#if MTLPP_CONFIG_VALIDATE && METAL_DEBUG_OPTIONS
 NS_ASSUME_NONNULL_BEGIN
 
 @class FMetalShaderPipeline;
diff --git a/Engine/Source/Runtime/Apple/MetalRHI/Private/MetalRenderPass.cpp b/Engine/Source/Runtime/Apple/MetalRHI/Private/MetalRenderPass.cpp
--- a/Engine/Source/Runtime/Apple/MetalRHI/Private/MetalRenderPass.cpp
+++ b/Engine/Source/Runtime/Apple/MetalRHI/Private/MetalRenderPass.cpp
@@ -816,7 +816,7 @@ void FMetalRenderPass::FillBuffer(FMetalBuffer const& Buffer, ns::Range Range, u
 bool FMetalRenderPass::AsyncCopyFromBufferToTexture(FMetalBuffer const& Buffer, uint32 sourceOffset, uint32 sourceBytesPerRow, uint32 sourceBytesPerImage, mtlpp::Size sourceSize, FMetalTexture const& toTexture, uint32 destinationSlice, uint32 destinationLevel, mtlpp::Origin destinationOrigin, mtlpp::BlitOption options)
 {
 	mtlpp::BlitCommandEncoder TargetEncoder;
-	METAL_DEBUG_OPTION(FMetalBlitCommandEncoderDebugging Debugging);
+	METAL_DEBUG_ONLY(FMetalBlitCommandEncoderDebugging Debugging);
 	bool bAsync = !CurrentEncoder.HasTextureBindingHistory(toTexture);
 	if(bAsync)
 	{
@@ -852,7 +852,7 @@ bool FMetalRenderPass::AsyncCopyFromBufferToTexture(FMetalBuffer const& Buffer,
 bool FMetalRenderPass::AsyncCopyFromTextureToTexture(FMetalTexture const& Texture, uint32 sourceSlice, uint32 sourceLevel, mtlpp::Origin sourceOrigin, mtlpp::Size sourceSize, FMetalTexture const& toTexture, uint32 destinationSlice, uint32 destinationLevel, mtlpp::Origin destinationOrigin)
 {
 	mtlpp::BlitCommandEncoder TargetEncoder;
-	METAL_DEBUG_OPTION(FMetalBlitCommandEncoderDebugging Debugging);
+	METAL_DEBUG_ONLY(FMetalBlitCommandEncoderDebugging Debugging);
 	bool bAsync = !CurrentEncoder.HasTextureBindingHistory(toTexture);
 	if(bAsync)
 	{
@@ -880,7 +880,7 @@ bool FMetalRenderPass::AsyncCopyFromTextureToTexture(FMetalTexture const& Textur
 void FMetalRenderPass::AsyncCopyFromBufferToBuffer(FMetalBuffer const& SourceBuffer, NSUInteger SourceOffset, FMetalBuffer const& DestinationBuffer, NSUInteger DestinationOffset, NSUInteger Size)
 {
 	mtlpp::BlitCommandEncoder TargetEncoder;
-	METAL_DEBUG_OPTION(FMetalBlitCommandEncoderDebugging Debugging);
+	METAL_DEBUG_ONLY(FMetalBlitCommandEncoderDebugging Debugging);
 	bool bAsync = !CurrentEncoder.HasBufferBindingHistory(DestinationBuffer);
 	if(bAsync)
 	{
diff --git a/Engine/Source/Runtime/Apple/MetalRHI/Private/MetalViewport.cpp b/Engine/Source/Runtime/Apple/MetalRHI/Private/MetalViewport.cpp
--- a/Engine/Source/Runtime/Apple/MetalRHI/Private/MetalViewport.cpp
+++ b/Engine/Source/Runtime/Apple/MetalRHI/Private/MetalViewport.cpp
@@ -429,7 +429,7 @@ void FMetalViewport::Present(FMetalCommandQueue& CommandQueue, bool bLockToVsync
 							
 							mtlpp::BlitCommandEncoder Encoder = CurrentCommandBuffer.BlitCommandEncoder();
 							check(Encoder.GetPtr());
-#if METAL_DEBUG_OPTIONS
+#if MTLPP_CONFIG_VALIDATE && METAL_DEBUG_OPTIONS
 							FMetalBlitCommandEncoderDebugging Debugging;
 							if (SafeGetRuntimeDebuggingLevel() >= EMetalDebugLevelFastValidation)
 							{
