Fix mismatched #defines breaking test builds.

diff --git a/Engine/Source/Runtime/Apple/MetalRHI/Private/MetalProfiler.cpp b/Engine/Source/Runtime/Apple/MetalRHI/Private/MetalProfiler.cpp
--- a/Engine/Source/Runtime/Apple/MetalRHI/Private/MetalProfiler.cpp
+++ b/Engine/Source/Runtime/Apple/MetalRHI/Private/MetalProfiler.cpp
@@ -611,6 +611,7 @@ FMetalShaderPipelineStats::FMetalShaderPipelineStats(FMetalShaderPipeline* Pipel
 	
 	StartSample = nullptr;
 	
+#if METAL_DEBUG_OPTIONS
 	if (Pipeline->RenderPipelineState)
 	{
 		Name = Pipeline->RenderPipelineState.GetLabel().GetPtr();
@@ -625,6 +626,7 @@ FMetalShaderPipelineStats::FMetalShaderPipelineStats(FMetalShaderPipeline* Pipel
 		Name = Pipeline->ComputePipelineState.GetLabel().GetPtr();
 	}
 	else
+#endif
 	{
 		Name = "Unknown Pipeline";
 	}
@@ -671,6 +673,7 @@ void FMetalShaderPipelineStats::GetStats(FMetalPipelineStats& PipelineStats)
 		PipelineStats.DrawCallTime = GPUEndTime - GPUStartTime;
 	}
 	
+#if METAL_DEBUG_OPTIONS
 	if (Pipeline->RenderPipelineReflection)
 	{
 		PipelineStats.PSOPerformanceStats = Stats->GetPipelinePerformanceStats(Pipeline->RenderPipelineReflection.GetPtr());
@@ -698,6 +701,7 @@ void FMetalShaderPipelineStats::GetStats(FMetalPipelineStats& PipelineStats)
 			PipelineStats.PSOPerformanceStats = [NSDictionary dictionaryWithObject:Dict forKey:@"Compute Shader"];
 		}
 	}
+#endif
 	
 	FMetalProfiler::GetProfiler()->DumpPipeline(Pipeline);
 }
@@ -1657,7 +1661,7 @@ void FMetalProfiler::SaveTrace()
 		
 		OutputFile->Close();
 		
-#if METAL_STATISTICS
+#if METAL_STATISTICS && METAL_DEBUG_OPTIONS
 		FString OutputDir = TracingRootPath + Filename + TEXT("/Pipelines/");
 		if (Pipelines.Num())
 		{
