Don't disable shadows on Intel please - the underlying bug was several OS point releases ago and screws up Engine testing. If this workaround is still necessary for games then the they need to handle this in their own scalability code. For my sanity however we'll only unblock shadows on 10.13.5 and above because I care not for older OSes. Remove obsolete volumetric fog hack - we fixed that.

diff --git a/Engine/Source/Runtime/Apple/MetalRHI/Private/MetalRHI.cpp b/Engine/Source/Runtime/Apple/MetalRHI/Private/MetalRHI.cpp
--- a/Engine/Source/Runtime/Apple/MetalRHI/Private/MetalRHI.cpp
+++ b/Engine/Source/Runtime/Apple/MetalRHI/Private/MetalRHI.cpp
@@ -415,14 +415,8 @@ FMetalDynamicRHI::FMetalDynamicRHI(ERHIFeatureLevel::Type RequestedFeatureLevel)
 	}
 
 #if PLATFORM_MAC
-	if (IsRHIDeviceIntel())
+	if (IsRHIDeviceIntel() && FPlatformMisc::MacOSXVersionCompare(10,13,5) < 0)
 	{
-		static auto CVarVolumetricFog = IConsoleManager::Get().FindConsoleVariable(TEXT("r.VolumetricFog"));
-		if(CVarVolumetricFog && CVarVolumetricFog->GetInt() != 0)
-		{
-			CVarVolumetricFog->Set(0);
-		}
-
 		static auto CVarSGShadowQuality = IConsoleManager::Get().FindConsoleVariable((TEXT("sg.ShadowQuality")));
 		if(CVarSGShadowQuality && CVarSGShadowQuality->GetInt() != 0)
 		{
