Fix buffer zero-initialisation and buffer-scribbling debug options in MetalRHI to try and help work out what's wrong on Intel. The FillBuffer function needs to be async. to avoid an issue with order of events when flush resource is involved.

diff --git a/Engine/Source/Runtime/Apple/MetalRHI/Private/MetalBuffer.cpp b/Engine/Source/Runtime/Apple/MetalRHI/Private/MetalBuffer.cpp
--- a/Engine/Source/Runtime/Apple/MetalRHI/Private/MetalBuffer.cpp
+++ b/Engine/Source/Runtime/Apple/MetalRHI/Private/MetalBuffer.cpp
@@ -1178,6 +1178,27 @@ FMetalBuffer FMetalResourceHeap::CreateBuffer(uint32 Size, uint32 Alignment, mtl
 #endif
 	}
 	
+	if (GMetalBufferZeroFill)
+	{
+		switch(Buffer.GetStorageMode())
+		{
+			case mtlpp::StorageMode::Private:
+			{
+				mtlpp::CommandBuffer CmdBuffer = Queue->CreateCommandBuffer();
+				mtlpp::BlitCommandEncoder Encoder = CmdBuffer.BlitCommandEncoder();
+				Encoder.Fill(Buffer, ns::Range(0, Buffer.GetLength()), 0);
+				Encoder.EndEncoding();
+				Queue->CommitCommandBuffer(CmdBuffer);
+				break;
+			}
+			default:
+			{
+				FMemory::Memset(((uint8*)Buffer.GetContents()), 0, Buffer.GetLength());
+				break;
+			}
+		}
+	}
+	
     METAL_DEBUG_OPTION(GetMetalDeviceContext().ValidateIsInactiveBuffer(Buffer));
 	check(Buffer && Buffer.GetPtr());
 	return Buffer;
diff --git a/Engine/Source/Runtime/Apple/MetalRHI/Private/MetalContext.cpp b/Engine/Source/Runtime/Apple/MetalRHI/Private/MetalContext.cpp
--- a/Engine/Source/Runtime/Apple/MetalRHI/Private/MetalContext.cpp
+++ b/Engine/Source/Runtime/Apple/MetalRHI/Private/MetalContext.cpp
@@ -54,13 +54,13 @@ int32 GMetalBufferScribble = 0; // Deliberately not static, see InitFrame_Unifor
 static FAutoConsoleVariableRef CVarMetalBufferScribble(
 	TEXT("rhi.Metal.BufferScribble"),
 	GMetalBufferScribble,
-	TEXT("Debug option: when enabled will scribble over the buffer contents with 0xCD when releasing Shared & Managed buffer objects. (Default: 0, Off)"));
+	TEXT("Debug option: when enabled will scribble over the buffer contents with a single value when releasing buffer objects, or regions thereof. (Default: 0, Off)"));
 
 int32 GMetalBufferZeroFill = 0; // Deliberately not static
 static FAutoConsoleVariableRef CVarMetalBufferZeroFill(
 	TEXT("rhi.Metal.BufferZeroFill"),
 	GMetalBufferZeroFill,
-	TEXT("Debug option: when enabled will fill the buffer contents with 0 when allocating Shared & Managed buffer objects, or regions thereof. (Default: 0, Off)"));
+	TEXT("Debug option: when enabled will fill the buffer contents with 0 when allocating buffer objects, or regions thereof. (Default: 0, Off)"));
 
 static int32 GMetalResourcePurgeOnDelete = 0;
 static FAutoConsoleVariableRef CVarMetalResourcePurgeOnDelete(
@@ -405,11 +405,12 @@ void FMetalDeviceContext::BeginFrame()
 }
 
 #if METAL_DEBUG_OPTIONS
-void ScribbleBuffer(FMetalBuffer Buffer)
+void FMetalDeviceContext::ScribbleBuffer(FMetalBuffer& Buffer)
 {
+	static uint8 Fill = 0;
 	if (Buffer.GetStorageMode() != mtlpp::StorageMode::Private)
 	{
-		FMemory::Memset(Buffer.GetContents(), 0xCD, Buffer.GetLength());
+		FMemory::Memset(Buffer.GetContents(), Fill++, Buffer.GetLength());
 #if PLATFORM_MAC
 		if (Buffer.GetStorageMode() == mtlpp::StorageMode::Managed)
 		{
@@ -417,6 +418,10 @@ void ScribbleBuffer(FMetalBuffer Buffer)
 		}
 #endif
 	}
+	else
+	{
+		FillBuffer(Buffer, ns::Range(0, Buffer.GetLength()), Fill++);
+	}
 }
 #endif
 
diff --git a/Engine/Source/Runtime/Apple/MetalRHI/Private/MetalContext.h b/Engine/Source/Runtime/Apple/MetalRHI/Private/MetalContext.h
--- a/Engine/Source/Runtime/Apple/MetalRHI/Private/MetalContext.h
+++ b/Engine/Source/Runtime/Apple/MetalRHI/Private/MetalContext.h
@@ -209,7 +209,8 @@ public:
 #if METAL_DEBUG_OPTIONS
     void AddActiveBuffer(FMetalBuffer const& Buffer);
     void RemoveActiveBuffer(FMetalBuffer const& Buffer);
-    bool ValidateIsInactiveBuffer(FMetalBuffer const& Buffer);
+	bool ValidateIsInactiveBuffer(FMetalBuffer const& Buffer);
+	void ScribbleBuffer(FMetalBuffer& Buffer);
 #endif
 	
 private:
diff --git a/Engine/Source/Runtime/Apple/MetalRHI/Private/MetalRenderPass.cpp b/Engine/Source/Runtime/Apple/MetalRHI/Private/MetalRenderPass.cpp
--- a/Engine/Source/Runtime/Apple/MetalRHI/Private/MetalRenderPass.cpp
+++ b/Engine/Source/Runtime/Apple/MetalRHI/Private/MetalRenderPass.cpp
@@ -803,14 +803,33 @@ void FMetalRenderPass::FillBuffer(FMetalBuffer const& Buffer, ns::Range Range, u
 {
 	check(Buffer);
 	
-	ConditionalSwitchToBlit();
-	mtlpp::BlitCommandEncoder& Encoder = CurrentEncoder.GetBlitCommandEncoder();
-	check(Encoder.GetPtr());
+	mtlpp::BlitCommandEncoder TargetEncoder;
+	METAL_DEBUG_ONLY(FMetalBlitCommandEncoderDebugging Debugging);
+	bool bAsync = !CurrentEncoder.HasBufferBindingHistory(Buffer);
+	if(bAsync)
+	{
+		ConditionalSwitchToAsyncBlit();
+		TargetEncoder = PrologueEncoder.GetBlitCommandEncoder();
+		METAL_GPUPROFILE(FMetalProfiler::GetProfiler()->EncodeBlit(PrologueEncoder.GetCommandBufferStats(), __FUNCTION__));
+		METAL_DEBUG_LAYER(EMetalDebugLevelFastValidation, Debugging = PrologueEncoder.GetBlitCommandEncoderDebugging());
+	}
+	else
+	{
+		ConditionalSwitchToBlit();
+		TargetEncoder = CurrentEncoder.GetBlitCommandEncoder();
+		METAL_GPUPROFILE(FMetalProfiler::GetProfiler()->EncodeBlit(CurrentEncoder.GetCommandBufferStats(), __FUNCTION__));
+		METAL_DEBUG_LAYER(EMetalDebugLevelFastValidation, Debugging = CurrentEncoder.GetBlitCommandEncoderDebugging());
+	}
 	
-	METAL_GPUPROFILE(FMetalProfiler::GetProfiler()->EncodeBlit(CurrentEncoder.GetCommandBufferStats(), __FUNCTION__));
-	MTLPP_VALIDATE(mtlpp::BlitCommandEncoder, Encoder, SafeGetRuntimeDebuggingLevel() >= EMetalDebugLevelValidation, Fill(Buffer, Range, Value));
+	check(TargetEncoder.GetPtr());
+	
+	MTLPP_VALIDATE(mtlpp::BlitCommandEncoder, TargetEncoder, SafeGetRuntimeDebuggingLevel() >= EMetalDebugLevelValidation, Fill(Buffer, Range, Value));
 	METAL_DEBUG_LAYER(EMetalDebugLevelFastValidation, CurrentEncoder.GetBlitCommandEncoderDebugging().Fill(Buffer, Range, Value));
-	ConditionalSubmit();
+	
+	if (!bAsync)
+	{
+		ConditionalSubmit();
+	}
 }
 
 bool FMetalRenderPass::AsyncCopyFromBufferToTexture(FMetalBuffer const& Buffer, uint32 sourceOffset, uint32 sourceBytesPerRow, uint32 sourceBytesPerImage, mtlpp::Size sourceSize, FMetalTexture const& toTexture, uint32 destinationSlice, uint32 destinationLevel, mtlpp::Origin destinationOrigin, mtlpp::BlitOption options)
