Trivial build fix.

diff --git a/Engine/Source/Runtime/Apple/MetalRHI/Private/MetalBuffer.cpp b/Engine/Source/Runtime/Apple/MetalRHI/Private/MetalBuffer.cpp
--- a/Engine/Source/Runtime/Apple/MetalRHI/Private/MetalBuffer.cpp
+++ b/Engine/Source/Runtime/Apple/MetalRHI/Private/MetalBuffer.cpp
@@ -1178,6 +1178,7 @@ FMetalBuffer FMetalResourceHeap::CreateBuffer(uint32 Size, uint32 Alignment, mtl
 #endif
 	}
 	
+#if METAL_DEBUG_OPTIONS
 	if (GMetalBufferZeroFill)
 	{
 		switch(Buffer.GetStorageMode())
@@ -1198,6 +1199,7 @@ FMetalBuffer FMetalResourceHeap::CreateBuffer(uint32 Size, uint32 Alignment, mtl
 			}
 		}
 	}
+#endif
 	
     METAL_DEBUG_OPTION(GetMetalDeviceContext().ValidateIsInactiveBuffer(Buffer));
 	check(Buffer && Buffer.GetPtr());
