Don't use Shared->Private blitting or buffer suballocation on Intel's Broadwell GPUs when running on macOS Metal - something goes wrong and the updates don't synchronise as they do on other GPUs.

diff --git a/Engine/Source/Runtime/Apple/MetalRHI/Private/MetalBuffer.cpp b/Engine/Source/Runtime/Apple/MetalRHI/Private/MetalBuffer.cpp
--- a/Engine/Source/Runtime/Apple/MetalRHI/Private/MetalBuffer.cpp
+++ b/Engine/Source/Runtime/Apple/MetalRHI/Private/MetalBuffer.cpp
@@ -1051,6 +1051,9 @@ FMetalBuffer FMetalResourceHeap::CreateBuffer(uint32 Size, uint32 Alignment, mtl
 	LLM_SCOPE_METAL(ELLMTagMetal::Buffers);
 	LLM_PLATFORM_SCOPE_METAL(ELLMTagMetal::Buffers);
 	
+	static bool bSupportsBufferSubAllocation = FMetalCommandQueue::SupportsFeature(EMetalFeaturesBufferSubAllocation);
+	bForceUnique |= !(bSupportsBufferSubAllocation);
+	
 	FMetalBuffer Buffer;
 	uint32 BlockSize = Align(Size, Alignment);
 	mtlpp::StorageMode StorageMode = (mtlpp::StorageMode)(((NSUInteger)Options & mtlpp::ResourceStorageModeMask) >> mtlpp::ResourceStorageModeShift);
diff --git a/Engine/Source/Runtime/Apple/MetalRHI/Private/MetalCommandQueue.cpp b/Engine/Source/Runtime/Apple/MetalRHI/Private/MetalCommandQueue.cpp
--- a/Engine/Source/Runtime/Apple/MetalRHI/Private/MetalCommandQueue.cpp
+++ b/Engine/Source/Runtime/Apple/MetalRHI/Private/MetalCommandQueue.cpp
@@ -112,13 +112,20 @@ FMetalCommandQueue::FMetalCommandQueue(mtlpp::Device InDevice, uint32 const MaxN
 		
 		Features |= EMetalFeaturesLinearTextures;
 		
-		// Using Private Memory & BlitEncoders for Vertex & Index data should be *much* faster.
-        Features |= EMetalFeaturesEfficientBufferBlits;
-		
-        // On earlier OS versions Vega didn't like non-zero blit offsets
-        if (!FString(Device.GetName()).Contains(TEXT("Vega")) || FPlatformMisc::MacOSXVersionCompare(10,13,5) >= 0)
-        {
-			Features |= EMetalFeaturesPrivateBufferSubAllocation;
+		FString DeviceName(Device.GetName());
+		// On earlier OS versions Intel Broadwell couldn't suballocate properly
+		if (!(DeviceName.Contains(TEXT("Intel")) && (DeviceName.Contains(TEXT("5300")) || DeviceName.Contains(TEXT("6000")) || DeviceName.Contains(TEXT("6100")))) || FPlatformMisc::MacOSXVersionCompare(10,14,0) >= 0)
+		{
+			// Using Private Memory & BlitEncoders for Vertex & Index data should be *much* faster.
+        	Features |= EMetalFeaturesEfficientBufferBlits;
+        	
+			Features |= EMetalFeaturesBufferSubAllocation;
+					
+	        // On earlier OS versions Vega didn't like non-zero blit offsets
+	        if (!DeviceName.Contains(TEXT("Vega")) || FPlatformMisc::MacOSXVersionCompare(10,13,5) >= 0)
+	        {
+				Features |= EMetalFeaturesPrivateBufferSubAllocation;
+			}
 		}
     }
     else if ([Device.GetName().GetPtr() rangeOfString:@"Nvidia" options:NSCaseInsensitiveSearch].location != NSNotFound)
diff --git a/Engine/Source/Runtime/Apple/MetalRHI/Private/MetalCommandQueue.h b/Engine/Source/Runtime/Apple/MetalRHI/Private/MetalCommandQueue.h
--- a/Engine/Source/Runtime/Apple/MetalRHI/Private/MetalCommandQueue.h
+++ b/Engine/Source/Runtime/Apple/MetalRHI/Private/MetalCommandQueue.h
@@ -89,8 +89,10 @@ typedef NS_OPTIONS(uint64, EMetalFeatures)
 	EMetalFeaturesFunctionConstants = 1llu << 34llu,
 	/** Supports efficient buffer-blits */
 	EMetalFeaturesEfficientBufferBlits = 1llu << 35llu,
+	/** Supports any kind of buffer sub-allocation */
+	EMetalFeaturesBufferSubAllocation = 1llu << 36llu,
 	/** Supports private buffer sub-allocation */
-	EMetalFeaturesPrivateBufferSubAllocation = 1llu << 36llu,
+	EMetalFeaturesPrivateBufferSubAllocation = 1llu << 37llu,
 };
 
 /**
