In MetalRHI don't mark ring-buffer entries as single-use, for shader paramters they aren't per-draw but instead per-PSO. Fixes crashes on Nvidia Metal.

diff --git a/Engine/Source/Runtime/Apple/MetalRHI/Private/MetalBuffer.cpp b/Engine/Source/Runtime/Apple/MetalRHI/Private/MetalBuffer.cpp
--- a/Engine/Source/Runtime/Apple/MetalRHI/Private/MetalBuffer.cpp
+++ b/Engine/Source/Runtime/Apple/MetalRHI/Private/MetalBuffer.cpp
@@ -657,7 +657,7 @@ FMetalBuffer FMetalSubBufferRing::NewBuffer(NSUInteger Size, uint32 Alignment)
 #endif
 			
 			WriteHead += FullSize;
-			NewBuffer.MarkSingleUse();
+			// NewBuffer.MarkSingleUse();
 			return NewBuffer;
 		}
 #if PLATFORM_MAC
@@ -703,7 +703,7 @@ FMetalBuffer FMetalSubBufferRing::NewBuffer(NSUInteger Size, uint32 Alignment)
 #endif
 		
 		WriteHead += FullSize;
-		NewBuffer.MarkSingleUse();
+		// NewBuffer.MarkSingleUse();
 		return NewBuffer;
 	}
 }
