Mark Metal buffers that really are single-use so that they don't remain bound after they shouldn't.

diff --git a/Engine/Source/Runtime/Apple/MetalRHI/Private/MetalCommands.cpp b/Engine/Source/Runtime/Apple/MetalRHI/Private/MetalCommands.cpp
--- a/Engine/Source/Runtime/Apple/MetalRHI/Private/MetalCommands.cpp
+++ b/Engine/Source/Runtime/Apple/MetalRHI/Private/MetalCommands.cpp
@@ -790,6 +790,7 @@ void FMetalRHICommandContext::RHIBeginDrawPrimitiveUP( uint32 PrimitiveType_DEPR
 	
 	// allocate space
 	PendingVertexBuffer = Context->AllocateFromRingBuffer(VertexDataStride * NumVertices);
+	PendingVertexBuffer.MarkSingleUse();
 
 	// get the pointer to send back for writing
 	OutVertexData = (uint8*)PendingVertexBuffer.GetContents();
@@ -839,6 +840,8 @@ void FMetalRHICommandContext::RHIBeginDrawIndexedPrimitiveUP( uint32 PrimitiveTy
 	uint32 IndexSize = Align(IndexDataStride * NumIndices, BufferOffsetAlignment);
 	PendingVertexBuffer = Context->AllocateFromRingBuffer(VertexSize);
 	PendingIndexBuffer = Context->AllocateFromRingBuffer(IndexSize);
+	PendingVertexBuffer.MarkSingleUse();
+	PendingIndexBuffer.MarkSingleUse();
 	
 	// get the pointers to send back for writing
 	OutVertexData = (uint8*)PendingVertexBuffer.GetContents();
diff --git a/Engine/Source/Runtime/Apple/MetalRHI/Private/MetalUAV.cpp b/Engine/Source/Runtime/Apple/MetalRHI/Private/MetalUAV.cpp
--- a/Engine/Source/Runtime/Apple/MetalRHI/Private/MetalUAV.cpp
+++ b/Engine/Source/Runtime/Apple/MetalRHI/Private/MetalUAV.cpp
@@ -531,6 +531,7 @@ void FMetalRHICommandContext::RHIClearTinyUAV(FUnorderedAccessViewRHIParamRef Un
 			else
 			{
 				Temp = Context->AllocateFromRingBuffer(AlignedSize);
+				Temp.MarkSingleUse();
 			}
 			
 			// Construct a pattern that can be encoded into the temporary buffer (handles packing & 2-byte formats).
