Fix for MSAA-depth-resolve not being available on A7, A8 & A8X.

diff --git a/Engine/Source/Runtime/Apple/MetalRHI/Private/MetalStateCache.cpp b/Engine/Source/Runtime/Apple/MetalRHI/Private/MetalStateCache.cpp
--- a/Engine/Source/Runtime/Apple/MetalRHI/Private/MetalStateCache.cpp
+++ b/Engine/Source/Runtime/Apple/MetalRHI/Private/MetalStateCache.cpp
@@ -35,9 +35,18 @@ FORCEINLINE mtlpp::StoreAction GetMetalRTStoreAction(ERenderTargetStoreAction St
 		case ERenderTargetStoreAction::EStore: return mtlpp::StoreAction::Store;
 		//default store action in the desktop renderers needs to be mtlpp::StoreAction::StoreAndMultisampleResolve.  Trying to express the renderer by the requested maxrhishaderplatform
         //because we may render to the same MSAA target twice in two separate passes.  BasePass, then some stuff, then translucency for example and we need to not lose the prior MSAA contents to do this properly.
-		case ERenderTargetStoreAction::EMultisampleResolve: 
-			return	FMetalCommandQueue::SupportsFeature(EMetalFeaturesMSAAStoreAndResolve) && (GMaxRHIShaderPlatform == SP_METAL_MRT || GMaxRHIShaderPlatform == SP_METAL_SM5 || GMaxRHIShaderPlatform == SP_METAL_SM5_NOTESS || GMaxRHIShaderPlatform == SP_METAL_MRT_MAC) ?
-					mtlpp::StoreAction::StoreAndMultisampleResolve : mtlpp::StoreAction::MultisampleResolve;
+		case ERenderTargetStoreAction::EMultisampleResolve:
+		{
+			static bool bSupportsMSAAStoreResolve = FMetalCommandQueue::SupportsFeature(EMetalFeaturesMSAAStoreAndResolve) && (GMaxRHIFeatureLevel >= ERHIFeatureLevel::SM5);
+			if (bSupportsMSAAStoreResolve)
+			{
+				return mtlpp::StoreAction::StoreAndMultisampleResolve;
+			}
+			else
+			{
+				return mtlpp::StoreAction::MultisampleResolve;
+			}
+		}
 		default: return mtlpp::StoreAction::DontCare;
 	}
 }
@@ -680,12 +689,12 @@ bool FMetalStateCache::SetRenderTargetsInfo(FRHISetRenderTargetsInfo const& InRe
 				
                 //needed to quiet the metal validation that runs when you end renderpass. (it requires some kind of 'resolve' for an msaa target)
 				//But with deferredstore we don't set the real one until submit time.
-                NewDepthStore = GetMetalRTStoreAction(HighLevelStoreAction);
+				const bool bSupportsMSAADepthResolve = GetMetalDeviceContext().SupportsFeature(EMetalFeaturesMSAADepthResolve);
+				NewDepthStore = !Surface.MSAATexture || bSupportsMSAADepthResolve ? GetMetalRTStoreAction(HighLevelStoreAction) : mtlpp::StoreAction::DontCare;
 				DepthAttachment.SetStoreAction(bSupportsDeferredStore ? mtlpp::StoreAction::Unknown : NewDepthStore);
 				DepthAttachment.SetClearDepth(DepthClearValue);
 				check(SampleCount > 0);
 
-				const bool bSupportsMSAADepthResolve = GetMetalDeviceContext().SupportsFeature(EMetalFeaturesMSAADepthResolve);
 				if (Surface.MSAATexture && bSupportsMSAADepthResolve)
 				{
                     if (!bDepthStencilSampleCountMismatchFixup)
@@ -2086,4 +2095,4 @@ void FMetalStateCache::DiscardRenderTargets(bool Depth, bool Stencil, uint32 Col
 			ColorStore[Index] = mtlpp::StoreAction::DontCare;
 		}
 	}
-}
\ No newline at end of file
+}
