Duplicate 4190391:

Fix a long-standing error in MetalRHI that allowed for thread-unsafe operations to occur when binding an SRV or UAV - it would call GetLinearTexture which would potentially allocate. This should not have been happening - the allocation is meant to occur in the SRV/UAV constructor. This CL makes that separation explicit and will simply assert if the expected Linear Texture is not present.

diff --git a/Engine/Source/Runtime/Apple/MetalRHI/Private/MetalUAV.cpp b/Engine/Source/Runtime/Apple/MetalRHI/Private/MetalUAV.cpp
--- a/Engine/Source/Runtime/Apple/MetalRHI/Private/MetalUAV.cpp
+++ b/Engine/Source/Runtime/Apple/MetalRHI/Private/MetalUAV.cpp
@@ -153,6 +153,8 @@ FUnorderedAccessViewRHIRef FMetalDynamicRHI::RHICreateUnorderedAccessView(FVerte
 	if (FMetalCommandQueue::SupportsFeature(EMetalFeaturesLinearTextureUAVs))
 	{
 		check(VertexBuffer->GetUsage() & BUF_UnorderedAccess);
+		ns::AutoReleased<FMetalTexture> Tex = VertexBuffer->CreateLinearTexture((EPixelFormat)Format);
+		check(Tex);
 	}
 		
 	// create the UAV buffer to point to the structured buffer's memory
@@ -178,6 +180,8 @@ FUnorderedAccessViewRHIRef FMetalDynamicRHI::RHICreateUnorderedAccessView(FIndex
 		if (FMetalCommandQueue::SupportsFeature(EMetalFeaturesLinearTextureUAVs))
 		{
 			check(IndexBuffer->GetUsage() & BUF_UnorderedAccess);
+			ns::AutoReleased<FMetalTexture> Tex = IndexBuffer->CreateLinearTexture((EPixelFormat)Format);
+			check(Tex);
 		}
 		
 		// create the UAV buffer to point to the structured buffer's memory
@@ -320,7 +324,8 @@ FShaderResourceViewRHIRef FMetalDynamicRHI::RHICreateShaderResourceView(FVertexB
 		check(Stride == GPixelFormats[Format].BlockBytes);
 		check(VertexBuffer->GetUsage() & BUF_ShaderResource);
 		
-		check(VertexBuffer->GetLinearTexture((EPixelFormat)Format));
+		ns::AutoReleased<FMetalTexture> Tex = VertexBuffer->CreateLinearTexture((EPixelFormat)Format);
+		check(Tex);
 	}
 	
 	return SRV;
@@ -339,7 +344,11 @@ FShaderResourceViewRHIRef FMetalDynamicRHI::RHICreateShaderResourceView(FIndexBu
 	SRV->SourceStructuredBuffer = nullptr;
 	SRV->Format = (Buffer->IndexType == mtlpp::IndexType::UInt16) ? PF_R16_UINT : PF_R32_UINT;
 	
-	check(!FMetalCommandQueue::SupportsFeature(EMetalFeaturesLinearTextures) || Buffer->GetLinearTexture((EPixelFormat)SRV->Format));
+	if (FMetalCommandQueue::SupportsFeature(EMetalFeaturesLinearTextures))
+	{
+		ns::AutoReleased<FMetalTexture> Tex = Buffer->CreateLinearTexture((EPixelFormat)SRV->Format);
+		check(Tex);
+	}
 	
 	return SRV;
 	}
diff --git a/Engine/Source/Runtime/Apple/MetalRHI/Private/MetalVertexBuffer.cpp b/Engine/Source/Runtime/Apple/MetalRHI/Private/MetalVertexBuffer.cpp
--- a/Engine/Source/Runtime/Apple/MetalRHI/Private/MetalVertexBuffer.cpp
+++ b/Engine/Source/Runtime/Apple/MetalRHI/Private/MetalVertexBuffer.cpp
@@ -282,7 +282,7 @@ FMetalTexture FMetalRHIBuffer::AllocLinearTexture(EPixelFormat Format)
 	}
 }
 
-ns::AutoReleased<FMetalTexture> FMetalRHIBuffer::GetLinearTexture(EPixelFormat Format)
+ns::AutoReleased<FMetalTexture> FMetalRHIBuffer::CreateLinearTexture(EPixelFormat Format)
 {
 	ns::AutoReleased<FMetalTexture> Texture;
 	if (FMetalCommandQueue::SupportsFeature(EMetalFeaturesLinearTextures) && (Usage & (BUF_UnorderedAccess|BUF_ShaderResource)) && GMetalBufferFormats[Format].LinearTextureFormat != mtlpp::PixelFormat::Invalid)
@@ -304,6 +304,20 @@ ns::AutoReleased<FMetalTexture> FMetalRHIBuffer::GetLinearTexture(EPixelFormat F
 	return Texture;
 }
 
+ns::AutoReleased<FMetalTexture> FMetalRHIBuffer::GetLinearTexture(EPixelFormat Format)
+{
+	ns::AutoReleased<FMetalTexture> Texture;
+	if (FMetalCommandQueue::SupportsFeature(EMetalFeaturesLinearTextures) && (Usage & (BUF_UnorderedAccess|BUF_ShaderResource)) && GMetalBufferFormats[Format].LinearTextureFormat != mtlpp::PixelFormat::Invalid)
+	{
+		FMetalTexture* ExistingTexture = LinearTextures.Find(Format);
+		if (ExistingTexture)
+		{
+			Texture = *ExistingTexture;
+		}
+	}
+	return Texture;
+}
+
 void* FMetalRHIBuffer::Lock(EResourceLockMode LockMode, uint32 Offset, uint32 InSize)
 {
 	check(LockSize == 0 && LockOffset == 0);
diff --git a/Engine/Source/Runtime/Apple/MetalRHI/Public/MetalResources.h b/Engine/Source/Runtime/Apple/MetalRHI/Public/MetalResources.h
--- a/Engine/Source/Runtime/Apple/MetalRHI/Public/MetalResources.h
+++ b/Engine/Source/Runtime/Apple/MetalRHI/Public/MetalResources.h
@@ -750,6 +750,11 @@ public:
 	 */
 	FMetalTexture AllocLinearTexture(EPixelFormat Format);
 	
+	/**
+	 * Get a linear texture for given format.
+	 */
+	ns::AutoReleased<FMetalTexture> CreateLinearTexture(EPixelFormat Format);
+	
 	/**
 	 * Get a linear texture for given format.
 	 */
