Using Shared->Private blits for Uniform buffers on Mac Metal creates a threading violation when the RHI thread & parallel translation are active. Just use Managed memory (Shared on iOS) for Uniform buffers and avoid this chaos.

diff --git a/Engine/Source/Runtime/Apple/MetalRHI/Private/MetalUniformBuffer.cpp b/Engine/Source/Runtime/Apple/MetalRHI/Private/MetalUniformBuffer.cpp
--- a/Engine/Source/Runtime/Apple/MetalRHI/Private/MetalUniformBuffer.cpp
+++ b/Engine/Source/Runtime/Apple/MetalRHI/Private/MetalUniformBuffer.cpp
@@ -12,7 +12,7 @@
 
 FMetalUniformBuffer::FMetalUniformBuffer(const void* Contents, const FRHIUniformBufferLayout& Layout, EUniformBufferUsage InUsage)
 	: FRHIUniformBuffer(Layout)
-	, FMetalRHIBuffer(Layout.ConstantBufferSize, ((InUsage & UniformBuffer_MultiFrame) ? BUF_Static : BUF_Volatile), RRT_UniformBuffer)
+	, FMetalRHIBuffer(Layout.ConstantBufferSize, BUF_Volatile, RRT_UniformBuffer)
 {
 	if (Layout.ConstantBufferSize > 0)
 	{
