Copy of CL 4222794 from 4.20

Fixed a crash at exit in Mac editor caused by an attempt to use MetalProfiler after deleting it

diff --git a/Engine/Source/Runtime/Apple/MetalRHI/Private/MetalProfiler.cpp b/Engine/Source/Runtime/Apple/MetalRHI/Private/MetalProfiler.cpp
--- a/Engine/Source/Runtime/Apple/MetalRHI/Private/MetalProfiler.cpp
+++ b/Engine/Source/Runtime/Apple/MetalRHI/Private/MetalProfiler.cpp
@@ -1220,6 +1220,7 @@ IMetalStatistics* FMetalProfiler::GetStatistics()
 void FMetalProfiler::DestroyProfiler()
 {
 	delete Self;
+	Self = nullptr;
 }
 
 void FMetalProfiler::BeginCapture(int InNumFramesToCapture)
