Implement a new MetalRHI debug feature 'rhi.Metal.ResetOnPSOChange' which will reset all the resource bindings on each graphics PSO change to track down any missing resource bindings. Didn't help with the problem we were looking into, but could be useful in the future.

diff --git a/Engine/Source/Runtime/Apple/MetalRHI/Private/MetalCommandEncoder.cpp b/Engine/Source/Runtime/Apple/MetalRHI/Private/MetalCommandEncoder.cpp
--- a/Engine/Source/Runtime/Apple/MetalRHI/Private/MetalCommandEncoder.cpp
+++ b/Engine/Source/Runtime/Apple/MetalRHI/Private/MetalCommandEncoder.cpp
@@ -134,6 +134,48 @@ void FMetalCommandEncoder::Reset(void)
 	[DebugGroups removeAllObjects];
 }
 
+void FMetalCommandEncoder::ResetLive(void)
+{
+	for (uint32 Frequency = 0; Frequency < uint32(mtlpp::FunctionType::Kernel)+1; Frequency++)
+	{
+		for (uint32 i = 0; i < ML_MaxBuffers; i++)
+		{
+			ShaderBuffers[Frequency].Buffers[i] = nil;
+		}
+		FMemory::Memzero(ShaderBuffers[Frequency].Bytes);
+		FMemory::Memzero(ShaderBuffers[Frequency].Offsets);
+		FMemory::Memzero(ShaderBuffers[Frequency].Lengths);
+		ShaderBuffers[Frequency].Bound = 0;
+	}
+	
+	if (IsRenderCommandEncoderActive())
+	{
+		for (uint32 i = 0; i < ML_MaxBuffers; i++)
+		{
+			RenderCommandEncoder.SetVertexBuffer(nil, 0, i);
+			RenderCommandEncoder.SetFragmentBuffer(nil, 0, i);
+		}
+		
+		for (uint32 i = 0; i < ML_MaxTextures; i++)
+		{
+			RenderCommandEncoder.SetVertexTexture(nil, i);
+			RenderCommandEncoder.SetFragmentTexture(nil, i);
+		}
+	}
+	else if (IsComputeCommandEncoderActive())
+	{
+		for (uint32 i = 0; i < ML_MaxBuffers; i++)
+		{
+			ComputeCommandEncoder.SetBuffer(nil, 0, i);
+		}
+		
+		for (uint32 i = 0; i < ML_MaxTextures; i++)
+		{
+			ComputeCommandEncoder.SetTexture(nil, i);
+		}
+	}
+}
+
 #pragma mark - Public Command Buffer Mutators -
 
 void FMetalCommandEncoder::StartCommandBuffer(void)
diff --git a/Engine/Source/Runtime/Apple/MetalRHI/Private/MetalCommandEncoder.h b/Engine/Source/Runtime/Apple/MetalRHI/Private/MetalCommandEncoder.h
--- a/Engine/Source/Runtime/Apple/MetalRHI/Private/MetalCommandEncoder.h
+++ b/Engine/Source/Runtime/Apple/MetalRHI/Private/MetalCommandEncoder.h
@@ -61,6 +61,9 @@ public:
 	/** Reset cached state for reuse */
 	void Reset(void);
 	
+	/** Reset cached state for reuse while in rendering */
+	void ResetLive(void);
+	
 #pragma mark - Public Command Buffer Mutators -
 
 	/**
diff --git a/Engine/Source/Runtime/Apple/MetalRHI/Private/MetalCommands.cpp b/Engine/Source/Runtime/Apple/MetalRHI/Private/MetalCommands.cpp
--- a/Engine/Source/Runtime/Apple/MetalRHI/Private/MetalCommands.cpp
+++ b/Engine/Source/Runtime/Apple/MetalRHI/Private/MetalCommands.cpp
@@ -233,6 +233,12 @@ void FMetalRHICommandContext::RHISetGraphicsPipelineState(FGraphicsPipelineState
 {
 	@autoreleasepool {
 		FMetalGraphicsPipelineState* PipelineState = ResourceCast(GraphicsState);
+#if METAL_DEBUG_OPTIONS
+		if (GMetalResetOnPSOChange >= 4 && Context->GetCurrentState().GetGraphicsPSO() != PipelineState)
+		{
+			Context->GetCurrentRenderPass().GetCurrentCommandEncoder().ResetLive();
+		}
+#endif
 		Context->GetCurrentState().SetGraphicsPipelineState(PipelineState);
 		// Bad Arne!
 		RHISetStencilRef(0);
diff --git a/Engine/Source/Runtime/Apple/MetalRHI/Private/MetalContext.cpp b/Engine/Source/Runtime/Apple/MetalRHI/Private/MetalContext.cpp
--- a/Engine/Source/Runtime/Apple/MetalRHI/Private/MetalContext.cpp
+++ b/Engine/Source/Runtime/Apple/MetalRHI/Private/MetalContext.cpp
@@ -54,6 +54,12 @@ static FAutoConsoleVariableRef CVarMetalBufferZeroFill(
 	TEXT("Debug option: when enabled will fill the buffer contents with 0 when allocating buffer objects, or regions thereof. (Default: 0, Off)"));
 
 #if METAL_DEBUG_OPTIONS
+int32 GMetalResetOnPSOChange = 0; // Deliberately not static
+static FAutoConsoleVariableRef CVarMetalResetOnPSOChange(
+	TEXT("rhi.Metal.ResetOnPSOChange"),
+	GMetalResetOnPSOChange,
+	TEXT("Debug option: when enabled will reset all the resource bindings when the PSO changes to aid debugging unbound resources. (Default: 0, Off)"));
+
 int32 GMetalBufferScribble = 0; // Deliberately not static, see InitFrame_UniformBufferPoolCleanup
 static FAutoConsoleVariableRef CVarMetalBufferScribble(
 	TEXT("rhi.Metal.BufferScribble"),
@@ -517,7 +523,7 @@ void FMetalDeviceContext::EndFrame()
 	uint32 SubmitFlags = EMetalSubmitFlagsResetState | EMetalSubmitFlagsForce | EMetalSubmitFlagsLastCommandBuffer;
 #if METAL_DEBUG_OPTIONS
 	// Latched update of whether to use runtime debugging features
-	if (GMetalRuntimeDebugLevel != CommandQueue.GetRuntimeDebuggingLevel())
+	if (GMetalResetOnPSOChange != CommandQueue.GetRuntimeDebuggingLevel())
 	{
 		CommandQueue.SetRuntimeDebuggingLevel(GMetalRuntimeDebugLevel);
 		
diff --git a/Engine/Source/Runtime/Apple/MetalRHI/Private/MetalRHIPrivate.h b/Engine/Source/Runtime/Apple/MetalRHI/Private/MetalRHIPrivate.h
--- a/Engine/Source/Runtime/Apple/MetalRHI/Private/MetalRHIPrivate.h
+++ b/Engine/Source/Runtime/Apple/MetalRHI/Private/MetalRHIPrivate.h
@@ -95,6 +95,7 @@ extern FMetalBufferFormat GMetalBufferFormats[PF_MAX];
 #define METAL_DEBUG_OPTIONS !(UE_BUILD_SHIPPING || UE_BUILD_TEST)
 #if METAL_DEBUG_OPTIONS
 #define METAL_DEBUG_OPTION(Code) Code
+extern int32 GMetalResetOnPSOChange;
 #else
 #define METAL_DEBUG_OPTION(Code)
 #endif
diff --git a/Engine/Source/Runtime/Apple/MetalRHI/Private/MetalRenderPass.h b/Engine/Source/Runtime/Apple/MetalRHI/Private/MetalRenderPass.h
--- a/Engine/Source/Runtime/Apple/MetalRHI/Private/MetalRenderPass.h
+++ b/Engine/Source/Runtime/Apple/MetalRHI/Private/MetalRenderPass.h
@@ -125,6 +125,12 @@ public:
 	mtlpp::CommandBuffer const& GetCurrentCommandBuffer(void) const;
 	mtlpp::CommandBuffer& GetCurrentCommandBuffer(void);
 	
+    /*
+     * Get the internal current command-encoder.
+     * @returns The current command encoder.
+     */
+	inline FMetalCommandEncoder& GetCurrentCommandEncoder(void) { return CurrentEncoder; }
+	
 	/*
 	 * Get the internal ring-buffer used for temporary allocations.
 	 * @returns The temporary allocation buffer for the command-pass.
diff --git a/Engine/Source/Runtime/Apple/MetalRHI/Private/MetalStateCache.cpp b/Engine/Source/Runtime/Apple/MetalRHI/Private/MetalStateCache.cpp
--- a/Engine/Source/Runtime/Apple/MetalRHI/Private/MetalStateCache.cpp
+++ b/Engine/Source/Runtime/Apple/MetalRHI/Private/MetalStateCache.cpp
@@ -1144,6 +1144,25 @@ void FMetalStateCache::SetGraphicsPipelineState(FMetalGraphicsPipelineState* Sta
 		
 		PipelineBits |= EMetalPipelineFlagPipelineState;
 		
+#if METAL_DEBUG_OPTIONS
+        if (GMetalResetOnPSOChange >= 4)
+        {
+            for (uint32 i = 0; i < EMetalShaderStages::Num; i++)
+            {
+                ShaderBuffers[i].Bound = UINT32_MAX;
+#if PLATFORM_MAC
+#ifndef UINT128_MAX
+#define UINT128_MAX (((__uint128_t)1 << 127) - (__uint128_t)1 + ((__uint128_t)1 << 127))
+#endif
+                ShaderTextures[i].Bound = UINT128_MAX;
+#else
+                ShaderTextures[i].Bound = UINT32_MAX;
+#endif
+                ShaderSamplers[i].Bound = UINT16_MAX;
+            }
+        }
+#endif
+		
 		SetDepthStencilState(State->DepthStencilState);
 		SetRasterizerState(State->RasterizerState);
 
