Undo accidental change that won't work in some MetalRHI build types.

diff --git a/Engine/Source/Runtime/Apple/MetalRHI/Private/MetalContext.cpp b/Engine/Source/Runtime/Apple/MetalRHI/Private/MetalContext.cpp
--- a/Engine/Source/Runtime/Apple/MetalRHI/Private/MetalContext.cpp
+++ b/Engine/Source/Runtime/Apple/MetalRHI/Private/MetalContext.cpp
@@ -523,7 +523,7 @@ void FMetalDeviceContext::EndFrame()
 	uint32 SubmitFlags = EMetalSubmitFlagsResetState | EMetalSubmitFlagsForce | EMetalSubmitFlagsLastCommandBuffer;
 #if METAL_DEBUG_OPTIONS
 	// Latched update of whether to use runtime debugging features
-	if (GMetalResetOnPSOChange != CommandQueue.GetRuntimeDebuggingLevel())
+	if (GMetalRuntimeDebugLevel != CommandQueue.GetRuntimeDebuggingLevel())
 	{
 		CommandQueue.SetRuntimeDebuggingLevel(GMetalRuntimeDebugLevel);
 		
