Avoid redundant metal shader lib creates, cache metal shader libraries during RHICreateShaderLibrary.

diff --git a/Engine/Source/Runtime/Apple/MetalRHI/Private/MetalShaders.cpp b/Engine/Source/Runtime/Apple/MetalRHI/Private/MetalShaders.cpp
--- a/Engine/Source/Runtime/Apple/MetalRHI/Private/MetalShaders.cpp
+++ b/Engine/Source/Runtime/Apple/MetalRHI/Private/MetalShaders.cpp
@@ -1024,17 +1024,22 @@ FDomainShaderRHIRef FMetalDynamicRHI::CreateDomainShader_RenderThread(class FRHI
 	return RHICreateDomainShader(Library, Hash);
 }
 
-FMetalShaderLibrary::FMetalShaderLibrary(EShaderPlatform InPlatform, FString const& Name, TArray<mtlpp::Library> InLibrary, FMetalShaderMap const& InMap)
+static FCriticalSection LoadedShaderLibraryMutex;
+static TMap<FString, FRHIShaderLibraryParamRef> LoadedShaderLibraryMap;
+
+FMetalShaderLibrary::FMetalShaderLibrary(EShaderPlatform InPlatform, FString const& Name, TArray<mtlpp::Library> InLibrary, FMetalShaderMap const& InMap, const FString& InShaderLibraryFilename)
 : FRHIShaderLibrary(InPlatform, Name)
 , Library(InLibrary)
 , Map(InMap)
+, ShaderLibraryFilename(InShaderLibraryFilename)
 {
 	
 }
 
 FMetalShaderLibrary::~FMetalShaderLibrary()
 {
-	
+	FScopeLock Lock(&LoadedShaderLibraryMutex);
+	LoadedShaderLibraryMap.Remove(ShaderLibraryFilename);
 }
 
 bool FMetalShaderLibrary::ContainsEntry(const FSHAHash& Hash)
@@ -1207,6 +1212,13 @@ FRHIShaderLibraryRef FMetalDynamicRHI::RHICreateShaderLibrary(EShaderPlatform Pl
 		BinaryShaderFile = FPaths::ProjectContentDir() / LibName + METAL_MAP_EXTENSION;
 	}
 
+	FScopeLock Lock(&LoadedShaderLibraryMutex);
+	FRHIShaderLibraryParamRef* FoundShaderLibrary = LoadedShaderLibraryMap.Find(BinaryShaderFile);
+	if (FoundShaderLibrary)
+	{
+		return *FoundShaderLibrary;
+	}
+
 	FArchive* BinaryShaderAr = IFileManager::Get().CreateFileReader(*BinaryShaderFile);
 
 	if( BinaryShaderAr != NULL )
@@ -1241,7 +1253,8 @@ FRHIShaderLibraryRef FMetalDynamicRHI::RHICreateShaderLibrary(EShaderPlatform Pl
 				}
 			}
 
-			Result = new FMetalShaderLibrary(Platform, Name, Libraries, Map);
+			Result = new FMetalShaderLibrary(Platform, Name, Libraries, Map, BinaryShaderFile);
+			LoadedShaderLibraryMap.Add(BinaryShaderFile, Result.GetReference());
 		}
 		else
 		{
diff --git a/Engine/Source/Runtime/Apple/MetalRHI/Public/MetalResources.h b/Engine/Source/Runtime/Apple/MetalRHI/Public/MetalResources.h
--- a/Engine/Source/Runtime/Apple/MetalRHI/Public/MetalResources.h
+++ b/Engine/Source/Runtime/Apple/MetalRHI/Public/MetalResources.h
@@ -1054,7 +1054,7 @@ private:
 class FMetalShaderLibrary final : public FRHIShaderLibrary
 {	
 public:
-	FMetalShaderLibrary(EShaderPlatform Platform, FString const& Name, TArray<mtlpp::Library> Library, FMetalShaderMap const& Map);
+	FMetalShaderLibrary(EShaderPlatform Platform, FString const& Name, TArray<mtlpp::Library> Library, FMetalShaderMap const& Map, const FString& InShaderLibraryFilename);
 	virtual ~FMetalShaderLibrary();
 	
 	virtual bool IsNativeLibrary() const override final {return true;}
@@ -1105,6 +1105,7 @@ private:
 private:
 	TArray<mtlpp::Library> Library;
 	FMetalShaderMap Map;
+	FString ShaderLibraryFilename;
 };
 
 template<class T>
