Setting a compute shader in MetalRHI should not affect the graphics tessellation state or we'll fail to go down the proper DrawPatches rendering path and ultimately assert on a missing PrologueEncoder.

diff --git a/Engine/Source/Runtime/Apple/MetalRHI/Private/MetalStateCache.cpp b/Engine/Source/Runtime/Apple/MetalRHI/Private/MetalStateCache.cpp
--- a/Engine/Source/Runtime/Apple/MetalRHI/Private/MetalStateCache.cpp
+++ b/Engine/Source/Runtime/Apple/MetalRHI/Private/MetalStateCache.cpp
@@ -383,10 +383,6 @@ void FMetalStateCache::SetComputeShader(FMetalComputeShader* InComputeShader)
 		
 		PipelineBits |= EMetalPipelineFlagComputeShader;
 		
-#if PLATFORM_SUPPORTS_TESSELLATION_SHADERS
-		bUsingTessellation = false;
-#endif
-		
 		DirtyUniformBuffers[EMetalShaderStages::Compute] = 0xffffffff;
 
 		for (const auto& PackedGlobalArray : InComputeShader->Bindings.PackedGlobalArrays)
