Duplicate CLs 5312633 & 5312684:

Change the way we restore the macOS binary shader cache as this may help on Intel/Nvidia Metal drivers.

diff --git a/Engine/Source/Runtime/Apple/MetalRHI/Private/MetalRHI.cpp b/Engine/Source/Runtime/Apple/MetalRHI/Private/MetalRHI.cpp
--- a/Engine/Source/Runtime/Apple/MetalRHI/Private/MetalRHI.cpp
+++ b/Engine/Source/Runtime/Apple/MetalRHI/Private/MetalRHI.cpp
@@ -835,12 +835,28 @@ FMetalDynamicRHI::FMetalDynamicRHI(ERHIFeatureLevel::Type RequestedFeatureLevel)
 			NSString* TempDir = [NSString stringWithFormat:@"%@/../C/%@/com.apple.metal", NSTemporaryDirectory(), [NSBundle mainBundle].bundleIdentifier];
 
 			NSError* Err = nil;
-			BOOL bOK = [[NSFileManager defaultManager] removeItemAtPath:TempDir
-						error:&Err];
-
-			bOK = [[NSFileManager defaultManager] copyItemAtPath:DstPath
-						toPath:TempDir
-						error:&Err];
+			if(![[NSFileManager defaultManager] fileExistsAtPath:TempDir])
+			{
+				[[NSFileManager defaultManager] createDirectoryAtPath:TempDir
+										  withIntermediateDirectories:YES
+														   attributes:nil
+																error:&Err];
+			}
+			
+			NSDirectoryEnumerator<NSString *> * Enum = [[NSFileManager defaultManager] enumeratorAtPath:DstPath];
+			for (NSString* Path in Enum)
+			{
+				[Enum skipDescendents];
+				
+				NSString* Dest = [NSString stringWithFormat:@"%@/%@", TempDir, Path];
+				if(![[NSFileManager defaultManager] fileExistsAtPath:Dest])
+				{
+					NSString* Src = [NSString stringWithFormat:@"%@/%@", DstPath, Path];
+					[[NSFileManager defaultManager] copyItemAtPath:Src
+															toPath:Dest
+															 error:&Err];
+				}
+			}
 		}
 	}
 #endif
