Disable MTLHeap on macOS - the drivers from each vendor are impossibly unstable.

diff --git a/Engine/Source/Runtime/Apple/MetalRHI/Private/MetalCommandQueue.cpp b/Engine/Source/Runtime/Apple/MetalRHI/Private/MetalCommandQueue.cpp
--- a/Engine/Source/Runtime/Apple/MetalRHI/Private/MetalCommandQueue.cpp
+++ b/Engine/Source/Runtime/Apple/MetalRHI/Private/MetalCommandQueue.cpp
@@ -261,7 +261,8 @@ FMetalCommandQueue::FMetalCommandQueue(mtlpp::Device InDevice, uint32 const MaxN
 					Features |= EMetalFeaturesFences;
 				}
 				
-				if (!FParse::Param(FCommandLine::Get(),TEXT("nometalheap")) && ([Device.GetName().GetPtr() rangeOfString:@"Intel" options:NSCaseInsensitiveSearch].location == NSNotFound || FParse::Param(FCommandLine::Get(),TEXT("forcemetalheap"))))
+				// There are still too many driver bugs to use MTLHeap on macOS - nothing works without causing random, undebuggable GPU hangs that completely deadlock the Mac and don't generate any validation errors or command-buffer failures
+				if (FParse::Param(FCommandLine::Get(),TEXT("forcemetalheap")))
 				{
 					Features |= EMetalFeaturesHeaps;
 				}
