Don't allow memoryless targets to use store-actions, other than MSAA Resolve, as they don't work. Also don't allow non-MSAA targets to be memoryless - unfortunately the engine doesn't flag those that might be copied from and that won't work with a memoryless target.

diff --git a/Engine/Source/Runtime/Apple/MetalRHI/Private/MetalStateCache.cpp b/Engine/Source/Runtime/Apple/MetalRHI/Private/MetalStateCache.cpp
--- a/Engine/Source/Runtime/Apple/MetalRHI/Private/MetalStateCache.cpp
+++ b/Engine/Source/Runtime/Apple/MetalRHI/Private/MetalStateCache.cpp
@@ -536,26 +536,31 @@ bool FMetalStateCache::SetRenderTargetsInfo(FRHISetRenderTargetsInfo const& InRe
 				
 				if (Surface.MSAATexture)
 				{
-					// set up an MSAA attachment
-					ColorAttachment.SetTexture(Surface.MSAATexture);
-					NewColorStore[RenderTargetIndex] = GetMetalRTStoreAction(ERenderTargetStoreAction::EMultisampleResolve);
-					ColorAttachment.SetStoreAction(bSupportsDeferredStore && GRHIDeviceId > 2 ? mtlpp::StoreAction::Unknown : NewColorStore[RenderTargetIndex]);
-					ColorAttachment.SetResolveTexture(Surface.MSAAResolveTexture ? Surface.MSAAResolveTexture : Surface.Texture);
-					SampleCount = Surface.MSAATexture.GetSampleCount();
+					bool bMemoryless = false;
 #if PLATFORM_IOS
 					if (Surface.MSAATexture.GetStorageMode() == mtlpp::StorageMode::Memoryless)
 					{
+						bMemoryless = true;
 						HighLevelLoadAction = ERenderTargetLoadAction::EClear;
 					}
 #endif
+					// set up an MSAA attachment
+					ColorAttachment.SetTexture(Surface.MSAATexture);
+					NewColorStore[RenderTargetIndex] = GetMetalRTStoreAction(ERenderTargetStoreAction::EMultisampleResolve);
+					
+					ColorAttachment.SetStoreAction(bSupportsDeferredStore && !bMemoryless && GRHIDeviceId > 2 ? mtlpp::StoreAction::Unknown : NewColorStore[RenderTargetIndex]);
+					ColorAttachment.SetResolveTexture(Surface.MSAAResolveTexture ? Surface.MSAAResolveTexture : Surface.Texture);
+					SampleCount = Surface.MSAATexture.GetSampleCount();
 					// only allow one MRT with msaa
 					checkf(RenderTargetsInfo.NumColorRenderTargets == 1, TEXT("Only expected one MRT when using MSAA"));
 				}
 				else
 				{
+					bool bMemoryless = false;
 #if PLATFORM_IOS
 					if (Surface.Texture.GetStorageMode() == mtlpp::StorageMode::Memoryless)
 					{
+						bMemoryless = true;
 						HighLevelStoreAction = ERenderTargetStoreAction::ENoAction;
 						HighLevelLoadAction = ERenderTargetLoadAction::EClear;
 					}
@@ -563,7 +568,7 @@ bool FMetalStateCache::SetRenderTargetsInfo(FRHISetRenderTargetsInfo const& InRe
 					// set up non-MSAA attachment
 					ColorAttachment.SetTexture(Surface.Texture);
 					NewColorStore[RenderTargetIndex] = GetMetalRTStoreAction(HighLevelStoreAction);
-					ColorAttachment.SetStoreAction(bSupportsDeferredStore ? mtlpp::StoreAction::Unknown : NewColorStore[RenderTargetIndex]);
+					ColorAttachment.SetStoreAction(bSupportsDeferredStore && !bMemoryless ? mtlpp::StoreAction::Unknown : NewColorStore[RenderTargetIndex]);
                     SampleCount = 1;
 				}
 				
@@ -806,7 +811,7 @@ bool FMetalStateCache::SetRenderTargetsInfo(FRHISetRenderTargetsInfo const& InRe
                 //needed to quiet the metal validation that runs when you end renderpass. (it requires some kind of 'resolve' for an msaa target)
 				//But with deferredstore we don't set the real one until submit time.
 				NewDepthStore = !Surface.MSAATexture || bSupportsMSAADepthResolve ? GetMetalRTStoreAction(HighLevelStoreAction) : mtlpp::StoreAction::DontCare;
-				DepthAttachment.SetStoreAction(bSupportsDeferredStore && Surface.MSAATexture && GRHIDeviceId > 2 ? mtlpp::StoreAction::Unknown : NewDepthStore);
+				DepthAttachment.SetStoreAction(bSupportsDeferredStore && !bDepthTextureMemoryless && Surface.MSAATexture && GRHIDeviceId > 2 ? mtlpp::StoreAction::Unknown : NewDepthStore);
 				DepthAttachment.SetClearDepth(DepthClearValue);
 				check(SampleCount > 0);
 
@@ -862,9 +867,11 @@ bool FMetalStateCache::SetRenderTargetsInfo(FRHISetRenderTargetsInfo const& InRe
 						HighLevelStoreAction = ERenderTargetStoreAction::EStore;
 					}
 					
+					bool bStencilMemoryless = false;
 #if PLATFORM_IOS
 					if (StencilTexture.GetStorageMode() == mtlpp::StorageMode::Memoryless)
 					{
+						bStencilMemoryless = true;
 						HighLevelStoreAction = ERenderTargetStoreAction::ENoAction;
 						StencilAttachment.SetLoadAction(mtlpp::LoadAction::Clear);
 					}
@@ -873,7 +880,7 @@ bool FMetalStateCache::SetRenderTargetsInfo(FRHISetRenderTargetsInfo const& InRe
 					// For the case where Depth+Stencil is MSAA we can't Resolve depth and Store stencil - we can only Resolve + DontCare or StoreResolve + Store (on newer H/W and iOS).
 					// We only allow use of StoreResolve in the Desktop renderers as the mobile renderer does not and should not assume hardware support for it.
 					NewStencilStore = (StencilTexture.GetSampleCount() == 1  || GetMetalRTStoreAction(ERenderTargetStoreAction::EMultisampleResolve) == mtlpp::StoreAction::StoreAndMultisampleResolve) ? GetMetalRTStoreAction(HighLevelStoreAction) : mtlpp::StoreAction::DontCare;
-					StencilAttachment.SetStoreAction(bSupportsDeferredStore && StencilTexture.GetSampleCount() > 1 && GRHIDeviceId > 2 ? mtlpp::StoreAction::Unknown : NewStencilStore);
+					StencilAttachment.SetStoreAction(bSupportsDeferredStore && !bStencilMemoryless && StencilTexture.GetSampleCount() > 1 && GRHIDeviceId > 2 ? mtlpp::StoreAction::Unknown : NewStencilStore);
                     StencilAttachment.SetClearStencil(StencilClearValue);
 
                     if (SampleCount == 0)
diff --git a/Engine/Source/Runtime/Apple/MetalRHI/Private/MetalTexture.cpp b/Engine/Source/Runtime/Apple/MetalRHI/Private/MetalTexture.cpp
--- a/Engine/Source/Runtime/Apple/MetalRHI/Private/MetalTexture.cpp
+++ b/Engine/Source/Runtime/Apple/MetalRHI/Private/MetalTexture.cpp
@@ -876,19 +876,9 @@ FMetalSurface::FMetalSurface(ERHIResourceType ResourceType, EPixelFormat Format,
 		else if (Flags & (TexCreate_RenderTargetable|TexCreate_DepthStencilTargetable|TexCreate_ResolveTargetable|TexCreate_DepthStencilResolveTarget))
 		{
 			check(!(Flags & TexCreate_CPUReadback));
-#if PLATFORM_IOS
-			if (FMetalCommandQueue::SupportsFeature(EMetalFeaturesMemoryLessResources) && !(Flags & (TexCreate_ShaderResource|TexCreate_UAV)) && (GMaxRHIFeatureLevel < ERHIFeatureLevel::SM5))
-			{
-				Desc.SetStorageMode(mtlpp::StorageMode::Memoryless);
-				Desc.SetResourceOptions(mtlpp::ResourceOptions::StorageModeMemoryless);
-			}
-			else
-#endif
-			{
-				Desc.SetCpuCacheMode(mtlpp::CpuCacheMode::DefaultCache);
-				Desc.SetStorageMode(mtlpp::StorageMode::Private);
-				Desc.SetResourceOptions((mtlpp::ResourceOptions)(mtlpp::ResourceOptions::CpuCacheModeDefaultCache|mtlpp::ResourceOptions::StorageModePrivate));
-			}
+			Desc.SetCpuCacheMode(mtlpp::CpuCacheMode::DefaultCache);
+			Desc.SetStorageMode(mtlpp::StorageMode::Private);
+			Desc.SetResourceOptions((mtlpp::ResourceOptions)(mtlpp::ResourceOptions::CpuCacheModeDefaultCache|mtlpp::ResourceOptions::StorageModePrivate));
 		}
 		else
 		{
