Fixed: Game texture becomes distorted after joining any game mode in iOS

diff --git a/Engine/Source/Runtime/Apple/MetalRHI/Private/MetalStateCache.cpp b/Engine/Source/Runtime/Apple/MetalRHI/Private/MetalStateCache.cpp
--- a/Engine/Source/Runtime/Apple/MetalRHI/Private/MetalStateCache.cpp
+++ b/Engine/Source/Runtime/Apple/MetalRHI/Private/MetalStateCache.cpp
@@ -534,7 +534,15 @@ bool FMetalStateCache::SetRenderTargetsInfo(FRHISetRenderTargetsInfo const& InRe
 				ERenderTargetStoreAction HighLevelStoreAction = RenderTargetView.StoreAction;
 				ERenderTargetLoadAction HighLevelLoadAction = RenderTargetView.LoadAction;
 				
-				if (Surface.MSAATexture)
+				// on iOS with memory-less MSAA textures we can't load them
+                // in case high level code wants to load and render to MSAA target, set attachment to a resolved texture
+				const bool bUseResolvedTexture = (
+					PLATFORM_IOS && 
+					Surface.MSAATexture && 
+					Surface.MSAATexture.GetStorageMode() == mtlpp::StorageMode::Memoryless && 
+					HighLevelLoadAction == ERenderTargetLoadAction::ELoad);
+				
+				if (Surface.MSAATexture && !bUseResolvedTexture)
 				{
 					bool bMemoryless = false;
 #if PLATFORM_IOS
