Mac compile fix

diff --git a/Engine/Source/Runtime/Apple/MetalRHI/Private/MetalStateCache.cpp b/Engine/Source/Runtime/Apple/MetalRHI/Private/MetalStateCache.cpp
--- a/Engine/Source/Runtime/Apple/MetalRHI/Private/MetalStateCache.cpp
+++ b/Engine/Source/Runtime/Apple/MetalRHI/Private/MetalStateCache.cpp
@@ -533,14 +533,16 @@ bool FMetalStateCache::SetRenderTargetsInfo(FRHISetRenderTargetsInfo const& InRe
 	
 				ERenderTargetStoreAction HighLevelStoreAction = RenderTargetView.StoreAction;
 				ERenderTargetLoadAction HighLevelLoadAction = RenderTargetView.LoadAction;
-				
+
 				// on iOS with memory-less MSAA textures we can't load them
                 // in case high level code wants to load and render to MSAA target, set attachment to a resolved texture
-				const bool bUseResolvedTexture = (
-					PLATFORM_IOS && 
+				bool bUseResolvedTexture = false;
+#if PLATFORM_IOS
+				bUseResolvedTexture = (
 					Surface.MSAATexture && 
 					Surface.MSAATexture.GetStorageMode() == mtlpp::StorageMode::Memoryless && 
 					HighLevelLoadAction == ERenderTargetLoadAction::ELoad);
+#endif
 				
 				if (Surface.MSAATexture && !bUseResolvedTexture)
 				{
