[LLM] removing VertexBuffer, IndexBuffer and UniformBuffer from Metal's LLM tagging for parity with other platforms, and OpenGL uniform buffers are now tracked in the current tag scope by pointer. (NB There may still be an underlying issue of the uniform buffer being tracked twice if OpenGL.UseEmulatedUBs is set.)

diff --git a/Engine/Source/Runtime/Apple/MetalRHI/Private/MetalIndexBuffer.cpp b/Engine/Source/Runtime/Apple/MetalRHI/Private/MetalIndexBuffer.cpp
--- a/Engine/Source/Runtime/Apple/MetalRHI/Private/MetalIndexBuffer.cpp
+++ b/Engine/Source/Runtime/Apple/MetalRHI/Private/MetalIndexBuffer.cpp
@@ -65,7 +65,6 @@ FIndexBufferRHIRef FMetalDynamicRHI::RHICreateIndexBuffer(uint32 Stride,uint32 S
 	{
 		if (IndexBuffer->UsePrivateMemory())
 		{
-			LLM_SCOPE(ELLMTag::IndexBuffer);
 			SafeReleaseMetalBuffer(IndexBuffer->CPUBuffer);
 			IndexBuffer->CPUBuffer = nil;
 		}
@@ -124,7 +123,6 @@ struct FMetalRHICommandInitialiseIndexBuffer : public FRHICommand<FMetalRHIComma
 			GetMetalDeviceContext().AsyncCopyFromBufferToBuffer(Buffer->CPUBuffer, 0, Buffer->Buffer, 0, Buffer->Buffer.GetLength());
 			if (Buffer->UsePrivateMemory())
 			{
-				LLM_SCOPE(ELLMTag::IndexBuffer);
 				SafeReleaseMetalBuffer(Buffer->CPUBuffer);
 			}
 			else
@@ -188,7 +186,6 @@ FIndexBufferRHIRef FMetalDynamicRHI::CreateIndexBuffer_RenderThread(class FRHICo
 		{
 			if (IndexBuffer->UsePrivateMemory())
 			{
-				LLM_SCOPE(ELLMTag::IndexBuffer);
 				SafeReleaseMetalBuffer(IndexBuffer->CPUBuffer);
 				IndexBuffer->CPUBuffer = nil;
 			}
diff --git a/Engine/Source/Runtime/Apple/MetalRHI/Private/MetalStructuredBuffer.cpp b/Engine/Source/Runtime/Apple/MetalRHI/Private/MetalStructuredBuffer.cpp
--- a/Engine/Source/Runtime/Apple/MetalRHI/Private/MetalStructuredBuffer.cpp
+++ b/Engine/Source/Runtime/Apple/MetalRHI/Private/MetalStructuredBuffer.cpp
@@ -35,7 +35,6 @@ FStructuredBufferRHIRef FMetalDynamicRHI::RHICreateStructuredBuffer(uint32 Strid
 	{
 		if (Buffer->GetUsage() & (BUF_Dynamic|BUF_Static))
 		{
-			LLM_SCOPE(ELLMTag::VertexBuffer);
 			SafeReleaseMetalBuffer(Buffer->CPUBuffer);
 			Buffer->CPUBuffer = nil;
 		}
@@ -93,7 +92,6 @@ struct FMetalRHICommandInitialiseStructuredBuffer : public FRHICommand<FMetalRHI
 			GetMetalDeviceContext().AsyncCopyFromBufferToBuffer(Buffer->CPUBuffer, 0, Buffer->Buffer, 0, Buffer->Buffer.GetLength());
 			if (Buffer->GetUsage() & (BUF_Dynamic|BUF_Static))
 			{
-				LLM_SCOPE(ELLMTag::VertexBuffer);
 				SafeReleaseMetalBuffer(Buffer->CPUBuffer);
 			}
 			else
@@ -159,7 +157,6 @@ FStructuredBufferRHIRef FMetalDynamicRHI::CreateStructuredBuffer_RenderThread(cl
 		{
 			if (VertexBuffer->GetUsage() & (BUF_Dynamic|BUF_Static))
 			{
-				LLM_SCOPE(ELLMTag::VertexBuffer);
 				SafeReleaseMetalBuffer(VertexBuffer->CPUBuffer);
 				VertexBuffer->CPUBuffer = nil;
 			}
diff --git a/Engine/Source/Runtime/Apple/MetalRHI/Private/MetalVertexBuffer.cpp b/Engine/Source/Runtime/Apple/MetalRHI/Private/MetalVertexBuffer.cpp
--- a/Engine/Source/Runtime/Apple/MetalRHI/Private/MetalVertexBuffer.cpp
+++ b/Engine/Source/Runtime/Apple/MetalRHI/Private/MetalVertexBuffer.cpp
@@ -12,21 +12,6 @@
 #include "RenderUtils.h"
 #include "MetalLLM.h"
 #include <objc/runtime.h>
-#include "HAL/LowLevelMemTracker.h"
-
-#if ENABLE_LOW_LEVEL_MEM_TRACKER
-#define METAL_LLM_BUFFER_SCOPE(Type) \
-	ELLMTag Tag; \
-	switch(Type)	{ \
-		case RRT_UniformBuffer: Tag = ELLMTag::UniformBuffer; break; \
-		case RRT_IndexBuffer: Tag = ELLMTag::IndexBuffer; break; \
-		case RRT_VertexBuffer: \
-		default: Tag = ELLMTag::VertexBuffer; break; \
-	} \
-	LLM_SCOPE(Tag)
-#else
-#define METAL_LLM_BUFFER_SCOPE(Type)
-#endif
 
 #if STATS
 #define METAL_INC_DWORD_STAT_BY(Type, Name, Size) \
@@ -126,7 +111,6 @@ FMetalRHIBuffer::FMetalRHIBuffer(uint32 InSize, uint32 InUsage, ERHIResourceType
 	{
 		checkf(InSize <= 1024 * 1024 * 1024, TEXT("Metal doesn't support buffers > 1GB"));
 		
-		METAL_LLM_BUFFER_SCOPE(Type);
 		// Temporary buffers less than the buffer page size - currently 4Kb - is better off going through the set*Bytes API if available.
 		// These can't be used for shader resources or UAVs if we want to use the 'Linear Texture' code path
 		if (!(InUsage & (BUF_UnorderedAccess|BUF_ShaderResource|EMetalBufferUsage_GPUOnly)) && (InUsage & BUF_Volatile) && InSize < MetalBufferPageSize && (InSize < MetalBufferBytesSize))
@@ -184,7 +168,6 @@ FMetalRHIBuffer::FMetalRHIBuffer(uint32 InSize, uint32 InUsage, ERHIResourceType
 
 FMetalRHIBuffer::~FMetalRHIBuffer()
 {
-	METAL_LLM_BUFFER_SCOPE(Type);
 	for (TPair<EPixelFormat, FMetalTexture>& Pair : LinearTextures)
 	{
 		SafeReleaseMetalTexture(Pair.Value);
@@ -235,7 +218,6 @@ void FMetalRHIBuffer::Unalias()
 
 void FMetalRHIBuffer::Alloc(uint32 InSize, EResourceLockMode LockMode)
 {
-	METAL_LLM_BUFFER_SCOPE(Type);
 	bool const bUsePrivateMem = UsePrivateMemory();
 
 	if (!Buffer)
@@ -276,7 +258,6 @@ void FMetalRHIBuffer::Alloc(uint32 InSize, EResourceLockMode LockMode)
 
 FMetalTexture FMetalRHIBuffer::AllocLinearTexture(EPixelFormat Format)
 {
-	METAL_LLM_BUFFER_SCOPE(Type);
 	if (FMetalCommandQueue::SupportsFeature(EMetalFeaturesLinearTextures) && (Usage & (BUF_UnorderedAccess|BUF_ShaderResource)))
 	{
 		mtlpp::PixelFormat MTLFormat = (mtlpp::PixelFormat)GMetalBufferFormats[Format].LinearTextureFormat;
@@ -432,7 +413,6 @@ void* FMetalRHIBuffer::Lock(EResourceLockMode LockMode, uint32 Offset, uint32 In
         bool const bUsePrivateMem = UsePrivateMemory();
         if (bUsePrivateMem)
         {
-			METAL_LLM_BUFFER_SCOPE(Type);
 			if (CPUBuffer)
 			{
 				METAL_INC_DWORD_STAT_BY(Type, MemFreed, Len);
@@ -522,7 +502,6 @@ void FMetalRHIBuffer::Unlock()
 			GetMetalDeviceContext().AsyncCopyFromBufferToBuffer(CPUBuffer, 0, Buffer, 0, Buffer.GetLength());
 			if (UsePrivateMemory())
             {
-				METAL_LLM_BUFFER_SCOPE(Type);
 				SafeReleaseMetalBuffer(CPUBuffer);
 				CPUBuffer = nil;
 			}
@@ -570,7 +549,6 @@ FVertexBufferRHIRef FMetalDynamicRHI::RHICreateVertexBuffer(uint32 Size, uint32
 	{
 		if (VertexBuffer->UsePrivateMemory())
 		{
-			LLM_SCOPE(ELLMTag::VertexBuffer);
 			SafeReleaseMetalBuffer(VertexBuffer->CPUBuffer);
 			VertexBuffer->CPUBuffer = nil;
 		}
@@ -653,7 +631,6 @@ struct FMetalRHICommandInitialiseVertexBuffer : public FRHICommand<FMetalRHIComm
 
 			if (Buffer->UsePrivateMemory())
 			{
-				LLM_SCOPE(ELLMTag::VertexBuffer);
 				SafeReleaseMetalBuffer(Buffer->CPUBuffer);
 			}
 			else
@@ -717,7 +694,6 @@ FVertexBufferRHIRef FMetalDynamicRHI::CreateVertexBuffer_RenderThread(class FRHI
 		{
 			if (VertexBuffer->UsePrivateMemory())
 			{
-				LLM_SCOPE(ELLMTag::VertexBuffer);
 				SafeReleaseMetalBuffer(VertexBuffer->CPUBuffer);
 				VertexBuffer->CPUBuffer = nil;
 			}
diff --git a/Engine/Source/Runtime/Core/Private/HAL/LowLevelMemTracker.cpp b/Engine/Source/Runtime/Core/Private/HAL/LowLevelMemTracker.cpp
--- a/Engine/Source/Runtime/Core/Private/HAL/LowLevelMemTracker.cpp
+++ b/Engine/Source/Runtime/Core/Private/HAL/LowLevelMemTracker.cpp
@@ -91,9 +91,6 @@ DECLARE_LLM_MEMORY_STAT(TEXT("StreamingManager"), STAT_StreamingManagerLLM, STAT
 DECLARE_LLM_MEMORY_STAT(TEXT("Graphics"), STAT_GraphicsPlatformLLM, STATGROUP_LLMPlatform);
 DECLARE_LLM_MEMORY_STAT(TEXT("FileSystem"), STAT_FileSystemLLM, STATGROUP_LLMFULL);
 DECLARE_LLM_MEMORY_STAT(TEXT("Localization"), STAT_LocalizationLLM, STATGROUP_LLMFULL);
-DECLARE_LLM_MEMORY_STAT(TEXT("VertexBuffer"), STAT_VertexBufferLLM, STATGROUP_LLMFULL);
-DECLARE_LLM_MEMORY_STAT(TEXT("IndexBuffer"), STAT_IndexBufferLLM, STATGROUP_LLMFULL);
-DECLARE_LLM_MEMORY_STAT(TEXT("UniformBuffer"), STAT_UniformBufferLLM, STATGROUP_LLMFULL);
 DECLARE_LLM_MEMORY_STAT(TEXT("AssetRegistry"), STAT_AssetRegistryLLM, STATGROUP_LLMFULL);
 DECLARE_LLM_MEMORY_STAT(TEXT("ConfigSystem"), STAT_ConfigSystemLLM, STATGROUP_LLMFULL);
 DECLARE_LLM_MEMORY_STAT(TEXT("InitUObject"), STAT_InitUObjectLLM, STATGROUP_LLMFULL);
diff --git a/Engine/Source/Runtime/Core/Public/HAL/LowLevelMemTracker.h b/Engine/Source/Runtime/Core/Public/HAL/LowLevelMemTracker.h
--- a/Engine/Source/Runtime/Core/Public/HAL/LowLevelMemTracker.h
+++ b/Engine/Source/Runtime/Core/Public/HAL/LowLevelMemTracker.h
@@ -16,7 +16,7 @@
 	#define LLM_SUPPORTED_PLATFORM (PLATFORM_XBOXONE || PLATFORM_PS4 || PLATFORM_WINDOWS || PLATFORM_IOS || PLATFORM_MAC || PLATFORM_ANDROID || PLATFORM_SWITCH || PLATFORM_UNIX)
 
 #ifndef ALLOW_LOW_LEVEL_MEM_TRACKER_IN_TEST
-#define ALLOW_LOW_LEVEL_MEM_TRACKER_IN_TEST 0
+#define ALLOW_LOW_LEVEL_MEM_TRACKER_IN_TEST 1
 #endif
 
 	// *** enable/disable LLM here ***
@@ -208,9 +208,6 @@ enum class ELLMTagSet : uint8
 	macro(GraphicsPlatform,						"Graphics",						GET_STATFNAME(STAT_GraphicsPlatformLLM),					NAME_None,										-1)\
 	macro(FileSystem,							"FileSystem",					GET_STATFNAME(STAT_FileSystemLLM),							GET_STATFNAME(STAT_EngineSummaryLLM),			-1)\
 	macro(Localization,							"Localization",					GET_STATFNAME(STAT_LocalizationLLM),						GET_STATFNAME(STAT_EngineSummaryLLM),			-1)\
-	macro(VertexBuffer,							"VertexBuffer",					GET_STATFNAME(STAT_VertexBufferLLM),						GET_STATFNAME(STAT_EngineSummaryLLM),			-1)\
-	macro(IndexBuffer,							"IndexBuffer",					GET_STATFNAME(STAT_IndexBufferLLM),							GET_STATFNAME(STAT_EngineSummaryLLM),			-1)\
-	macro(UniformBuffer,						"UniformBuffer",				GET_STATFNAME(STAT_UniformBufferLLM),						GET_STATFNAME(STAT_EngineSummaryLLM),			-1)\
 	macro(AssetRegistry,						"AssetRegistry",				GET_STATFNAME(STAT_AssetRegistryLLM),						GET_STATFNAME(STAT_EngineSummaryLLM),			-1)\
 	macro(ConfigSystem,							"ConfigSystem",					GET_STATFNAME(STAT_ConfigSystemLLM),						GET_STATFNAME(STAT_EngineSummaryLLM),			-1)\
 	macro(InitUObject,							"InitUObject",					GET_STATFNAME(STAT_InitUObjectLLM),							GET_STATFNAME(STAT_EngineSummaryLLM),			-1)\
diff --git a/Engine/Source/Runtime/OpenGLDrv/Private/OpenGLUniformBuffer.cpp b/Engine/Source/Runtime/OpenGLDrv/Private/OpenGLUniformBuffer.cpp
--- a/Engine/Source/Runtime/OpenGLDrv/Private/OpenGLUniformBuffer.cpp
+++ b/Engine/Source/Runtime/OpenGLDrv/Private/OpenGLUniformBuffer.cpp
@@ -391,9 +391,7 @@ void FOpenGLUniformBuffer::SetGLUniformBufferParams(GLuint InResource, uint32 In
 	AllocatedSize = InAllocatedSize;
 	bStreamDraw = bInStreamDraw;
 
-#if ENABLE_LOW_LEVEL_MEM_TRACKER
-	LLM_SCOPED_PAUSE_TRACKING_WITH_ENUM_AND_AMOUNT(ELLMTag::UniformBuffer, InAllocatedSize, ELLMTracker::Default, ELLMAllocType::None);
-#endif
+	LLM(FLowLevelMemTracker::Get().OnLowLevelAlloc(ELLMTracker::Default, ((uint8*)this)+1, InAllocatedSize)); //+1 because ptr must be unique for LLM
 }
 
 FOpenGLUniformBuffer::~FOpenGLUniformBuffer()
@@ -431,9 +429,8 @@ FOpenGLUniformBuffer::~FOpenGLUniformBuffer()
 			ReleaseUniformBuffer(IsValidRef(EmulatedBufferData), Resource, AllocatedSize);
 			Resource = 0; 
 		}
-#if ENABLE_LOW_LEVEL_MEM_TRACKER
-		LLM_SCOPED_PAUSE_TRACKING_WITH_ENUM_AND_AMOUNT(ELLMTag::UniformBuffer, -(int64)AllocatedSize, ELLMTracker::Default, ELLMAllocType::None);
-#endif
+
+		LLM(FLowLevelMemTracker::Get().OnLowLevelFree(ELLMTracker::Default, ((uint8*)this)+1)); //+1 because ptr must be unique for LLM
 	}
 }
 
