Fix name conflict with GMetalDebugComputeShader .

diff --git a/Engine/Source/Runtime/Apple/MetalRHI/Private/MetalPipeline.cpp b/Engine/Source/Runtime/Apple/MetalRHI/Private/MetalPipeline.cpp
--- a/Engine/Source/Runtime/Apple/MetalRHI/Private/MetalPipeline.cpp
+++ b/Engine/Source/Runtime/Apple/MetalRHI/Private/MetalPipeline.cpp
@@ -92,7 +92,7 @@ static NSString* GMetalDebugShader = @"#include <metal_stdlib>\n"
 #endif
 
 // A compute debug shader for trying to emulate Aftermath style failure reporting
-static NSString* GMetalDebugComputeShader = @"#include <metal_stdlib>\n"
+static NSString* GMetalDebugMarkerComputeShader = @"#include <metal_stdlib>\n"
 "#include <metal_compute>\n"
 "\n"
 "using namespace metal;\n"
@@ -134,7 +134,7 @@ struct FMetalHelperFunctions
 			DebugShadersLib = GetMetalDeviceContext().GetDevice().NewLibrary(GMetalDebugShader, CompileOptions, &Error);
             DebugFunc = DebugShadersLib.NewFunction(@"Main_Debug");
 			
-			DebugComputeShadersLib = GetMetalDeviceContext().GetDevice().NewLibrary(GMetalDebugComputeShader, CompileOptions, &Error);
+			DebugComputeShadersLib = GetMetalDeviceContext().GetDevice().NewLibrary(GMetalDebugMarkerComputeShader, CompileOptions, &Error);
 			DebugComputeFunc = DebugShadersLib.NewFunction(@"Main_Debug");
 			
 			DebugComputeState = GetMetalDeviceContext().GetDevice().NewComputePipelineState(DebugComputeFunc, &Error);
