When using MetalRHI's internal 'fast' validation check the depth-stencil-state against the Metal render pass configuration to ensure depth & stencil writing will work and fail if not.

diff --git a/Engine/Source/Runtime/Apple/MetalRHI/Private/MetalStateCache.cpp b/Engine/Source/Runtime/Apple/MetalRHI/Private/MetalStateCache.cpp
--- a/Engine/Source/Runtime/Apple/MetalRHI/Private/MetalStateCache.cpp
+++ b/Engine/Source/Runtime/Apple/MetalRHI/Private/MetalStateCache.cpp
@@ -363,6 +363,12 @@ void FMetalStateCache::SetDepthStencilState(FMetalDepthStencilState* InDepthSten
 	{
 		DepthStencilState = InDepthStencilState;
 		RasterBits |= EMetalRenderFlagDepthStencilState;
+		
+		if (DepthStencilState && SafeGetRuntimeDebuggingLevel() >= EMetalDebugLevelFastValidation)
+		{
+			METAL_FATAL_ASSERT(RenderPassDesc.GetDepthAttachment().GetTexture() || DepthStencilState->bIsDepthWriteEnabled == false, TEXT("Attempting to set a depth-stencil state that writes depth but no depth texture is configured!\nState: %s\nRender Pass: %s"), *FString([DepthStencilState->State.GetPtr() description]), *FString([RenderPassDesc.GetPtr() description]));
+			METAL_FATAL_ASSERT(RenderPassDesc.GetStencilAttachment().GetTexture() || DepthStencilState->bIsStencilWriteEnabled == false, TEXT("Attempting to set a depth-stencil state that writes stencil but no stencil texture is configured!\nState: %s\nRender Pass: %s"), *FString([DepthStencilState->State.GetPtr() description]), *FString([RenderPassDesc.GetPtr() description]));
+		}
 	}
 }
 
