Build fix to allow using newer APIs than we build for - these APIs will only be called on the newer OSes that provide them.

diff --git a/Engine/Source/Runtime/Apple/MetalRHI/Private/MetalRenderPass.cpp b/Engine/Source/Runtime/Apple/MetalRHI/Private/MetalRenderPass.cpp
--- a/Engine/Source/Runtime/Apple/MetalRHI/Private/MetalRenderPass.cpp
+++ b/Engine/Source/Runtime/Apple/MetalRHI/Private/MetalRenderPass.cpp
@@ -1974,6 +1974,15 @@ uint32 FMetalRenderPass::GetCommandBufferIndex(void) const
 	}
 }
 
+#if PLATFORM_MAC
+@protocol IMTLRenderCommandEncoder
+- (void)memoryBarrierWithResources:(const id<MTLResource>[])resources count:(NSUInteger)count afterStages:(MTLRenderStages)after beforeStages:(MTLRenderStages)before;
+@end
+#endif
+@protocol IMTLComputeCommandEncoder
+- (void)memoryBarrierWithResources:(const id<MTLResource>[])resources count:(NSUInteger)count;
+@end
+
 void FMetalRenderPass::InsertDebugDraw(FMetalCommandData& Data)
 {
 	if (GMetalCommandBufferDebuggingEnabled)
@@ -2009,7 +2018,7 @@ void FMetalRenderPass::InsertDebugDraw(FMetalCommandData& Data)
 
 		#if PLATFORM_MAC
 			id<MTLBuffer> DebugBufferPtr = State.GetDebugBuffer().GetPtr();
-			[CurrentEncoder.GetRenderCommandEncoder().GetPtr() memoryBarrierWithResources:&DebugBufferPtr count:1 afterStages:MTLRenderStageFragment beforeStages:MTLRenderStageVertex];
+			[(id<IMTLRenderCommandEncoder>)CurrentEncoder.GetRenderCommandEncoder().GetPtr() memoryBarrierWithResources:&DebugBufferPtr count:1 afterStages:MTLRenderStageFragment beforeStages:MTLRenderStageVertex];
 			
 			CurrentEncoder.SetShaderBytes(mtlpp::FunctionType::Vertex, (uint8 const*)&DebugInfo, sizeof(DebugInfo), 0);
 			State.SetShaderBufferDirty(EMetalShaderStages::Vertex, 0);
@@ -2019,7 +2028,7 @@ void FMetalRenderPass::InsertDebugDraw(FMetalCommandData& Data)
 
 			CurrentEncoder.GetRenderCommandEncoder().Draw(mtlpp::PrimitiveType::Point, 0, 1);
 			
-			[CurrentEncoder.GetRenderCommandEncoder().GetPtr() memoryBarrierWithResources:&DebugBufferPtr count:1 afterStages:MTLRenderStageVertex beforeStages:MTLRenderStageVertex];
+			[(id<IMTLRenderCommandEncoder>)CurrentEncoder.GetRenderCommandEncoder().GetPtr() memoryBarrierWithResources:&DebugBufferPtr count:1 afterStages:MTLRenderStageVertex beforeStages:MTLRenderStageVertex];
 		#else
 			CurrentEncoder.GetRenderCommandEncoder().SetTileData((uint8 const*)&DebugInfo, sizeof(DebugInfo), 0);
 			CurrentEncoder.GetRenderCommandEncoder().SetTileBuffer(State.GetDebugBuffer(), 0, 1);
@@ -2056,7 +2065,7 @@ void FMetalRenderPass::InsertDebugDispatch(FMetalCommandData& Data)
 			CurrentEncoder.GetComputeCommandEncoder().SetComputePipelineState(GetMetalDebugComputeState());
 			
 			id<MTLBuffer> DebugBufferPtr = State.GetDebugBuffer().GetPtr();
-			[CurrentEncoder.GetComputeCommandEncoder().GetPtr() memoryBarrierWithResources:&DebugBufferPtr count:1];
+			[(id<IMTLComputeCommandEncoder>)CurrentEncoder.GetComputeCommandEncoder().GetPtr() memoryBarrierWithResources:&DebugBufferPtr count:1];
 			
 			CurrentEncoder.SetShaderBytes(mtlpp::FunctionType::Kernel, (uint8 const*)&DebugInfo, sizeof(DebugInfo), 0);
 			State.SetShaderBufferDirty(EMetalShaderStages::Compute, 0);
@@ -2067,7 +2076,7 @@ void FMetalRenderPass::InsertDebugDispatch(FMetalCommandData& Data)
 			mtlpp::Size ThreadsPerTile(1, 1, 1);
 			CurrentEncoder.GetComputeCommandEncoder().DispatchThreads(ThreadsPerTile, ThreadsPerTile);
 			
-			[CurrentEncoder.GetComputeCommandEncoder().GetPtr() memoryBarrierWithResources:&DebugBufferPtr count:1];
+			[(id<IMTLComputeCommandEncoder>)CurrentEncoder.GetComputeCommandEncoder().GetPtr() memoryBarrierWithResources:&DebugBufferPtr count:1];
 			
 			FMetalShaderPipeline* Pipeline = BoundShaderState->GetPipeline();
 			CurrentEncoder.GetComputeCommandEncoder().SetComputePipelineState(Pipeline->ComputePipelineState);
