The current tvOS hardware & API doesn't support tile-shaders so we must disable the GPU debugging code on PLATFORM_TVOS.

diff --git a/Engine/Source/Runtime/Apple/MetalRHI/Private/MetalPipeline.cpp b/Engine/Source/Runtime/Apple/MetalRHI/Private/MetalPipeline.cpp
--- a/Engine/Source/Runtime/Apple/MetalRHI/Private/MetalPipeline.cpp
+++ b/Engine/Source/Runtime/Apple/MetalRHI/Private/MetalPipeline.cpp
@@ -126,6 +126,7 @@ struct FMetalHelperFunctions
     
     FMetalHelperFunctions()
     {
+#if !PLATFORM_TVOS
         if (GMetalCommandBufferDebuggingEnabled)
         {
             mtlpp::CompileOptions CompileOptions;
@@ -139,6 +140,7 @@ struct FMetalHelperFunctions
 			
 			DebugComputeState = GetMetalDeviceContext().GetDevice().NewComputePipelineState(DebugComputeFunc, &Error);
         }
+#endif
     }
     
     static FMetalHelperFunctions& Get()
@@ -604,7 +606,7 @@ static FMetalShaderPipeline* CreateMTLRenderPipeline(bool const bSync, FMetalGra
 		mtlpp::ComputePipelineDescriptor ComputePipelineDesc(nil);
 #if PLATFORM_MAC
         mtlpp::RenderPipelineDescriptor DebugPipelineDesc;
-#else
+#elif !PLATFORM_TVOS
         mtlpp::TileRenderPipelineDescriptor DebugPipelineDesc;
 #endif
 		
@@ -664,7 +666,9 @@ static FMetalShaderPipeline* CreateMTLRenderPipeline(bool const bSync, FMetalGra
         FMetalBlendState* BlendState = (FMetalBlendState*)Init.BlendState;
 		
 		ns::Array<mtlpp::RenderPipelineColorAttachmentDescriptor> ColorAttachments = RenderPipelineDesc.GetColorAttachments();
+#if !PLATFORM_TVOS
 		auto DebugColorAttachements = DebugPipelineDesc.GetColorAttachments();
+#endif
 		
 		uint32 TargetWidth = 0;
         for (uint32 i = 0; i < NumActiveTargets; i++)
@@ -1189,7 +1193,8 @@ static FMetalShaderPipeline* CreateMTLRenderPipeline(bool const bSync, FMetalGra
 	#endif
         Pipeline->FragmentSource = PixelShader ? PixelShader->GetSourceCode() : nil;
     #endif
-        
+
+#if !PLATFORM_TVOS
 		if (GMetalCommandBufferDebuggingEnabled)
 		{
 #if PLATFORM_MAC
@@ -1208,6 +1213,7 @@ static FMetalShaderPipeline* CreateMTLRenderPipeline(bool const bSync, FMetalGra
 			METAL_GPUPROFILE(FScopedMetalCPUStats CPUStat(FString::Printf(TEXT("NewDebugPipeline: %s"), TEXT("")/**FString([RenderPipelineDesc.GetPtr() description])*/)));
 			Pipeline->DebugPipelineState = Device.NewRenderPipelineState(DebugPipelineDesc, mtlpp::PipelineOption::NoPipelineOption, Reflection, nullptr);
 		}
+#endif
         
 #if METAL_DEBUG_OPTIONS
         if (GFrameCounter > 3)
diff --git a/Engine/Source/Runtime/Apple/MetalRHI/Private/MetalRenderPass.cpp b/Engine/Source/Runtime/Apple/MetalRHI/Private/MetalRenderPass.cpp
--- a/Engine/Source/Runtime/Apple/MetalRHI/Private/MetalRenderPass.cpp
+++ b/Engine/Source/Runtime/Apple/MetalRHI/Private/MetalRenderPass.cpp
@@ -1985,6 +1985,7 @@ uint32 FMetalRenderPass::GetCommandBufferIndex(void) const
 
 void FMetalRenderPass::InsertDebugDraw(FMetalCommandData& Data)
 {
+#if !PLATFORM_TVOS
 	if (GMetalCommandBufferDebuggingEnabled)
 	{
 		FMetalGraphicsPipelineState* BoundShaderState = State.GetGraphicsPSO();
@@ -2039,10 +2040,12 @@ void FMetalRenderPass::InsertDebugDraw(FMetalCommandData& Data)
 			CurrentEncoder.GetRenderCommandEncoder().SetRenderPipelineState(PSO->RenderPipelineState);
 		}
 	}
+#endif
 }
 
 void FMetalRenderPass::InsertDebugDispatch(FMetalCommandData& Data)
 {
+#if !PLATFORM_TVOS
 	if (GMetalCommandBufferDebuggingEnabled)
 	{
 		FMetalComputeShader* BoundShaderState = State.GetComputeShader();
@@ -2082,4 +2085,5 @@ void FMetalRenderPass::InsertDebugDispatch(FMetalCommandData& Data)
 			CurrentEncoder.GetComputeCommandEncoder().SetComputePipelineState(Pipeline->ComputePipelineState);
 		}
 	}
+#endif
 }
