Missed one of the tile-shader references that aren't available on tvOS.

diff --git a/Engine/Source/Runtime/Apple/MetalRHI/Private/MetalPipeline.cpp b/Engine/Source/Runtime/Apple/MetalRHI/Private/MetalPipeline.cpp
--- a/Engine/Source/Runtime/Apple/MetalRHI/Private/MetalPipeline.cpp
+++ b/Engine/Source/Runtime/Apple/MetalRHI/Private/MetalPipeline.cpp
@@ -694,9 +694,11 @@ static FMetalShaderPipeline* CreateMTLRenderPipeline(bool const bSync, FMetalGra
             
             mtlpp::RenderPipelineColorAttachmentDescriptor Attachment = ColorAttachments[i];
             Attachment.SetPixelFormat(MetalFormat);
-            
+			
+#if !PLATFORM_TVOS
 			auto DebugAttachment = DebugColorAttachements[i];;
 			DebugAttachment.SetPixelFormat(MetalFormat);
+#endif
 			
             mtlpp::RenderPipelineColorAttachmentDescriptor Blend = BlendState->RenderTargetStates[i].BlendState;
             if(TargetFormat != PF_Unknown)
