Since switching the engine to using real PSOs the higher-level really does assume that if the stencil component is bound as a target that the depth component is bound as well. This means it is no longer valid to omit one or the other as we'll have compiled the PSOs incorrectly.

diff --git a/Engine/Source/Runtime/Apple/MetalRHI/Private/MetalStateCache.cpp b/Engine/Source/Runtime/Apple/MetalRHI/Private/MetalStateCache.cpp
--- a/Engine/Source/Runtime/Apple/MetalRHI/Private/MetalStateCache.cpp
+++ b/Engine/Source/Runtime/Apple/MetalRHI/Private/MetalStateCache.cpp
@@ -762,10 +762,8 @@ bool FMetalStateCache::SetRenderTargetsInfo(FRHISetRenderTargetsInfo const& InRe
 				DepthClearValue = 1.0f;
 			}
 
-            static bool const bUsingValidation = FMetalCommandQueue::SupportsFeature(EMetalFeaturesValidation) && !FApplePlatformMisc::IsOSAtLeastVersion((uint32[]){10, 14, 0}, (uint32[]){12, 0, 0}, (uint32[]){12, 0, 0});
-            
-            bool const bCombinedDepthStencilUsingStencil = (DepthTexture && (mtlpp::PixelFormat)DepthTexture.GetPixelFormat() != mtlpp::PixelFormat::Depth32Float && RenderTargetsInfo.DepthStencilRenderTarget.GetDepthStencilAccess().IsUsingStencil());			
-			bool const bUsingDepth = (RenderTargetsInfo.DepthStencilRenderTarget.GetDepthStencilAccess().IsUsingDepth() || (bUsingValidation && bCombinedDepthStencilUsingStencil));
+            bool const bCombinedDepthStencilUsingStencil = (DepthTexture && (mtlpp::PixelFormat)DepthTexture.GetPixelFormat() != mtlpp::PixelFormat::Depth32Float && RenderTargetsInfo.DepthStencilRenderTarget.GetDepthStencilAccess().IsUsingStencil());
+			bool const bUsingDepth = (RenderTargetsInfo.DepthStencilRenderTarget.GetDepthStencilAccess().IsUsingDepth() || (bCombinedDepthStencilUsingStencil));
 			if (DepthTexture && bUsingDepth)
 			{
 				mtlpp::RenderPassDepthAttachmentDescriptor DepthAttachment;
@@ -843,7 +841,7 @@ bool FMetalStateCache::SetRenderTargetsInfo(FRHISetRenderTargetsInfo const& InRe
             //doesn't have an autoresolve target to use.
 			
 			bool const bCombinedDepthStencilUsingDepth = (StencilTexture && StencilTexture.GetPixelFormat() != mtlpp::PixelFormat::Stencil8 && RenderTargetsInfo.DepthStencilRenderTarget.GetDepthStencilAccess().IsUsingDepth());
-			bool const bUsingStencil = RenderTargetsInfo.DepthStencilRenderTarget.GetDepthStencilAccess().IsUsingStencil() || (bUsingValidation && bCombinedDepthStencilUsingDepth);
+			bool const bUsingStencil = RenderTargetsInfo.DepthStencilRenderTarget.GetDepthStencilAccess().IsUsingStencil() || (bCombinedDepthStencilUsingDepth);
 			if (StencilTexture && bUsingStencil && (FMetalCommandQueue::SupportsFeature(EMetalFeaturesCombinedDepthStencil) || !bDepthStencilSampleCountMismatchFixup))
 			{
                 if (!FMetalCommandQueue::SupportsFeature(EMetalFeaturesCombinedDepthStencil) && bDepthStencilSampleCountMismatchFixup)
