Harden MetalUniformBuffer by caching all the data fetched from the UniformBufferLayout such that the RHI thread never needs to access it.

diff --git a/Engine/Source/Runtime/Apple/MetalRHI/Private/MetalUniformBuffer.cpp b/Engine/Source/Runtime/Apple/MetalRHI/Private/MetalUniformBuffer.cpp
--- a/Engine/Source/Runtime/Apple/MetalRHI/Private/MetalUniformBuffer.cpp
+++ b/Engine/Source/Runtime/Apple/MetalRHI/Private/MetalUniformBuffer.cpp
@@ -179,15 +179,38 @@ FMetalUniformBuffer::FMetalUniformBuffer(const void* Contents, const FRHIUniform
     , FMetalRHIBuffer(Layout.ConstantBufferSize, (FMetalCommandQueue::SupportsFeature(EMetalFeaturesIABs) && Layout.Resources.Num() ? (EMetalBufferUsage_GPUOnly|BUF_Volatile) : BUF_Volatile), RRT_UniformBuffer)
 	, UniformUsage(InUsage)
 	, IAB(nullptr)
+	, NumResources(0)
+	, ConstantSize(Layout.ConstantBufferSize)
 {
-	uint32 NumResources = Layout.Resources.Num();
+	NumResources = Layout.Resources.Num();
 	if (NumResources)
 	{
 		ResourceTable.Empty(NumResources);
 		ResourceTable.AddZeroed(NumResources);
+		
+		for (uint32 i = 0; i < NumResources; ++i)
+		{
+			FRHIResource* Resource = *(FRHIResource**)((uint8*)Contents + Layout.Resources[i].MemberOffset);
+			
+			// Allow null SRV's in uniform buffers for feature levels that don't support SRV's in shaders
+			if (Validation == EUniformBufferValidation::ValidateResources && !(GMaxRHIFeatureLevel <= ERHIFeatureLevel::ES3_1 && Layout.Resources[i].MemberType == UBMT_SRV))
+			{
+				check(Resource);
+			}
+			
+			ResourceTable[i] = Resource;
+		}
 	}
 	
-    Update(Contents, Validation);
+	if (FMetalCommandQueue::SupportsFeature(EMetalFeaturesIABs))
+	{
+		for (int32 i = 0; i < NumResources; ++i)
+		{
+			ResourceTypes[i] = Layout.Resources[i].MemberType;
+		}
+	}
+
+    Update(Contents, ResourceTable, Validation);
 	
     if (NumResources && FMetalCommandQueue::SupportsFeature(EMetalFeaturesIABs))
     {
@@ -232,7 +255,6 @@ FMetalUniformBuffer::FMetalIndirectArgumentBuffer& FMetalUniformBuffer::GetIAB()
 
 void FMetalUniformBuffer::InitIAB()
 {
-	int32 NumResources = ResourceTable.Num();
 	if (NumResources && FMetalCommandQueue::SupportsFeature(EMetalFeaturesIABs) && !IAB)
 	{
 		FMetalIndirectArgumentBuffer* NewIAB = new FMetalIndirectArgumentBuffer;
@@ -245,7 +267,7 @@ void FMetalUniformBuffer::InitIAB()
 		{
 			FRHIResource* Resource = ResourceTable[i].GetReference();
 			
-			switch(GetLayout().Resources[i].MemberType)
+			switch(ResourceTypes[i])
 			{
 				case UBMT_RDG_TEXTURE_SRV:
 				case UBMT_RDG_BUFFER_SRV:
@@ -434,7 +456,7 @@ void FMetalUniformBuffer::InitIAB()
 			NewIAB->IndirectArgumentResources.Add(Argument(NewIAB->IndirectArgumentBufferSideTable, mtlpp::ResourceUsage::Read));
 		}
 		
-		if (GetLayout().ConstantBufferSize > 0)
+		if (ConstantSize > 0)
 		{
 			FMetalArgumentDesc& Desc = Arguments.Emplace_GetRef();
 			Desc.SetIndex(Index++);
@@ -502,88 +524,77 @@ void const* FMetalUniformBuffer::GetData()
 	}
 }
 
-void FMetalUniformBuffer::Update(const void* Contents, EUniformBufferValidation Validation)
+void FMetalUniformBuffer::Update(const void* Contents, TArray<TRefCountPtr<FRHIResource>>& Resources, EUniformBufferValidation Validation)
 {
-    const FRHIUniformBufferLayout Layout = GetLayout();
-    if (Layout.ConstantBufferSize > 0)
+    if (ConstantSize > 0)
     {
-        UE_CLOG(Layout.ConstantBufferSize > 65536, LogMetal, Fatal, TEXT("Trying to allocated a uniform layout of size %d that is greater than the maximum permitted 64k."), Layout.ConstantBufferSize);
+        UE_CLOG(ConstantSize > 65536, LogMetal, Fatal, TEXT("Trying to allocated a uniform layout of size %d that is greater than the maximum permitted 64k."), ConstantSize);
         
         void* Data = Lock(RLM_WriteOnly, 0);
-        FMemory::Memcpy(Data, Contents, Layout.ConstantBufferSize);
+        FMemory::Memcpy(Data, Contents, ConstantSize);
         Unlock();
     }
-    
-    // set up an SRT-style uniform buffer
-    if (Layout.Resources.Num())
-    {
-        int32 NumResources = Layout.Resources.Num();
-        for (int32 i = 0; i < NumResources; ++i)
-        {
-            FRHIResource* Resource = *(FRHIResource**)((uint8*)Contents + Layout.Resources[i].MemberOffset);
-            
-            // Allow null SRV's in uniform buffers for feature levels that don't support SRV's in shaders
-            if (Validation == EUniformBufferValidation::ValidateResources && !(GMaxRHIFeatureLevel <= ERHIFeatureLevel::ES3_1 && Layout.Resources[i].MemberType == UBMT_SRV))
-            {
-                check(Resource);
-            }
-            
-            ResourceTable[i] = Resource;
-            
-            if (FMetalCommandQueue::SupportsFeature(EMetalFeaturesIABs) && Resource)
-            {
-                switch(Layout.Resources[i].MemberType)
-                {
-                    case UBMT_RDG_TEXTURE_SRV:
-                    case UBMT_RDG_BUFFER_SRV:
-                    case UBMT_SRV:
-                    {
-                        FMetalShaderResourceView* SRV = (FMetalShaderResourceView*)Resource;
-                        FRHITexture* Texture = SRV->SourceTexture.GetReference();
-                        if (Texture && Texture->GetTextureReference())
-                        {
-                            TextureReferences.Add(Texture->GetTextureReference());
-                        }
-                        break;
-                    }
-                    case UBMT_RDG_TEXTURE_UAV:
-                    case UBMT_RDG_BUFFER_UAV:
-                    {
-                        FMetalUnorderedAccessView* UAV = (FMetalUnorderedAccessView*)Resource;
-                        FRHITexture* Texture = UAV->SourceView->SourceTexture.GetReference();
-                        if (Texture && Texture->GetTextureReference())
-                        {
-                            TextureReferences.Add(Texture->GetTextureReference());
-                        }
-                        break;
-                    }
-                    case UBMT_RDG_TEXTURE:
-                    case UBMT_TEXTURE:
-                    {
-                        FRHITexture* Texture = (FRHITexture*)Resource;
-                        if (Texture && Texture->GetTextureReference())
-                        {
-                            TextureReferences.Add(Texture->GetTextureReference());
-                        }
-                        break;
-                    }
-                    default:
-                    {
-                        break;
-                    }
-                }
-            }
-        }
-        
-        if (FMetalCommandQueue::SupportsFeature(EMetalFeaturesIABs))
-        {
-            FRHICommandListImmediate& RHICmdList = FRHICommandListExecutor::GetImmediateCommandList();
-            if (!(UniformUsage & UniformBuffer_SingleDraw) && IsRunningRHIInSeparateThread() && !RHICmdList.Bypass() && IsInRenderingThread())
-            {
-                new (RHICmdList.AllocCommand<FMetalRHICommandInitialiseUniformBufferIAB>()) FMetalRHICommandInitialiseUniformBufferIAB(this);
-            }
-        }
-    }
+	
+	ResourceTable = Resources;
+	
+	if (NumResources && FMetalCommandQueue::SupportsFeature(EMetalFeaturesIABs))
+	{
+		// set up an SRT-style uniform buffer
+		for (uint32 i = 0; i < NumResources; ++i)
+		{
+			FRHIResource* Resource = ResourceTable[i];
+			
+			if (FMetalCommandQueue::SupportsFeature(EMetalFeaturesIABs) && Resource)
+			{
+				switch(ResourceTypes[i])
+				{
+					case UBMT_RDG_TEXTURE_SRV:
+					case UBMT_RDG_BUFFER_SRV:
+					case UBMT_SRV:
+					{
+						FMetalShaderResourceView* SRV = (FMetalShaderResourceView*)Resource;
+						FRHITexture* Texture = SRV->SourceTexture.GetReference();
+						if (Texture && Texture->GetTextureReference())
+						{
+							TextureReferences.Add(Texture->GetTextureReference());
+						}
+						break;
+					}
+					case UBMT_RDG_TEXTURE_UAV:
+					case UBMT_RDG_BUFFER_UAV:
+					{
+						FMetalUnorderedAccessView* UAV = (FMetalUnorderedAccessView*)Resource;
+						FRHITexture* Texture = UAV->SourceView->SourceTexture.GetReference();
+						if (Texture && Texture->GetTextureReference())
+						{
+							TextureReferences.Add(Texture->GetTextureReference());
+						}
+						break;
+					}
+					case UBMT_RDG_TEXTURE:
+					case UBMT_TEXTURE:
+					{
+						FRHITexture* Texture = (FRHITexture*)Resource;
+						if (Texture && Texture->GetTextureReference())
+						{
+							TextureReferences.Add(Texture->GetTextureReference());
+						}
+						break;
+					}
+					default:
+					{
+						break;
+					}
+				}
+			}
+		}
+		
+		FRHICommandListImmediate& RHICmdList = FRHICommandListExecutor::GetImmediateCommandList();
+		if (!(UniformUsage & UniformBuffer_SingleDraw) && IsRunningRHIInSeparateThread() && !RHICmdList.Bypass() && IsInRenderingThread())
+		{
+			new (RHICmdList.AllocCommand<FMetalRHICommandInitialiseUniformBufferIAB>()) FMetalRHICommandInitialiseUniformBufferIAB(this);
+		}
+	}
 }
 
 FUniformBufferRHIRef FMetalDynamicRHI::RHICreateUniformBuffer(const void* Contents, const FRHIUniformBufferLayout& Layout, EUniformBufferUsage Usage, EUniformBufferValidation Validation)
@@ -597,18 +608,15 @@ FUniformBufferRHIRef FMetalDynamicRHI::RHICreateUniformBuffer(const void* Conten
 struct FMetalRHICommandUpateUniformBuffer : public FRHICommand<FMetalRHICommandUpateUniformBuffer>
 {
 	TRefCountPtr<FMetalUniformBuffer> Buffer;
+	TArray<TRefCountPtr<FRHIResource> > ResourceTable;
 	char* Contents;
 	
-	FORCEINLINE_DEBUGGABLE FMetalRHICommandUpateUniformBuffer(FMetalUniformBuffer* InBuffer, void const* Data)
+	FORCEINLINE_DEBUGGABLE FMetalRHICommandUpateUniformBuffer(FMetalUniformBuffer* InBuffer, void const* Data, TArray<TRefCountPtr<FRHIResource>>& Resources)
 	: Buffer(InBuffer)
+	, ResourceTable(Resources)
 	, Contents(nullptr)
 	{
-		const FRHIUniformBufferLayout Layout = Buffer->GetLayout();
-		uint32 MaxLayoutSize = Layout.ConstantBufferSize;
-		for (int32 i = 0; i < Layout.Resources.Num(); ++i)
-		{
-			MaxLayoutSize = FMath::Max(MaxLayoutSize, (uint32)(Layout.Resources[i].MemberOffset + sizeof(FRHIResource*)));
-		}
+		uint32 MaxLayoutSize = InBuffer->ConstantSize;
 		Contents = new char[MaxLayoutSize];
 		FMemory::Memcpy(Contents, Data, MaxLayoutSize);
 	}
@@ -620,7 +628,7 @@ struct FMetalRHICommandUpateUniformBuffer : public FRHICommand<FMetalRHICommandU
 	
 	void Execute(FRHICommandListBase& CmdList)
 	{
-		Buffer->Update(Contents, EUniformBufferValidation::None);
+		Buffer->Update(Contents, ResourceTable, EUniformBufferValidation::None);
 	}
 };
 
@@ -631,13 +639,25 @@ void FMetalDynamicRHI::RHIUpdateUniformBuffer(FUniformBufferRHIParamRef UniformB
 
 	FMetalUniformBuffer* UniformBuffer = ResourceCast(UniformBufferRHI);
 	FRHICommandListImmediate& RHICmdList = FRHICommandListExecutor::GetImmediateCommandList();
+		
+	TArray<TRefCountPtr<FRHIResource> > ResourceTable;
+	ResourceTable.AddZeroed(UniformBuffer->NumResources);
+		
+	const FRHIUniformBufferLayout& Layout = UniformBuffer->GetLayout();
+		
+	for (uint32 i = 0; i < UniformBuffer->NumResources; ++i)
+	{
+		FRHIResource* Resource = *(FRHIResource**)((uint8*)Contents + Layout.Resources[i].MemberOffset);
+		ResourceTable[i] = Resource;
+	}	
+		
 	if (RHICmdList.Bypass() || !IsRunningRHIInSeparateThread())
 	{
-		UniformBuffer->Update(Contents, EUniformBufferValidation::None);
+		UniformBuffer->Update(Contents, ResourceTable, EUniformBufferValidation::None);
 	}
 	else
 	{
-		new (RHICmdList.AllocCommand<FMetalRHICommandUpateUniformBuffer>()) FMetalRHICommandUpateUniformBuffer(UniformBuffer, Contents);
+		new (RHICmdList.AllocCommand<FMetalRHICommandUpateUniformBuffer>()) FMetalRHICommandUpateUniformBuffer(UniformBuffer, Contents, ResourceTable);
 		RHICmdList.RHIThreadFence(true);
 	}
 	}
diff --git a/Engine/Source/Runtime/Apple/MetalRHI/Public/MetalResources.h b/Engine/Source/Runtime/Apple/MetalRHI/Public/MetalResources.h
--- a/Engine/Source/Runtime/Apple/MetalRHI/Public/MetalResources.h
+++ b/Engine/Source/Runtime/Apple/MetalRHI/Public/MetalResources.h
@@ -852,7 +852,7 @@ public:
 
 	void InitIAB();
 
-	void Update(const void* Contents, EUniformBufferValidation Validation);
+	void Update(const void* Contents, TArray<TRefCountPtr<FRHIResource>>& Resources, EUniformBufferValidation Validation);
 	
 	/** Resource table containing RHI references. */
 	TArray<TRefCountPtr<FRHIResource> > ResourceTable;
@@ -885,6 +885,9 @@ public:
 	EUniformBufferUsage UniformUsage;
 	FMetalIndirectArgumentBuffer& GetIAB();
 	FMetalIndirectArgumentBuffer* IAB;
+	TArray<EUniformBufferBaseType> ResourceTypes;
+	uint32 NumResources;
+	uint32 ConstantSize;
 };
 
 
