fix for metal validation layer

diff --git a/Engine/Source/Runtime/Apple/MetalRHI/Private/MetalStateCache.cpp b/Engine/Source/Runtime/Apple/MetalRHI/Private/MetalStateCache.cpp
--- a/Engine/Source/Runtime/Apple/MetalRHI/Private/MetalStateCache.cpp
+++ b/Engine/Source/Runtime/Apple/MetalRHI/Private/MetalStateCache.cpp
@@ -363,12 +363,6 @@ void FMetalStateCache::SetDepthStencilState(FMetalDepthStencilState* InDepthSten
 	{
 		DepthStencilState = InDepthStencilState;
 		RasterBits |= EMetalRenderFlagDepthStencilState;
-		
-		if (DepthStencilState && SafeGetRuntimeDebuggingLevel() >= EMetalDebugLevelFastValidation)
-		{
-			METAL_FATAL_ASSERT(RenderPassDesc.GetDepthAttachment().GetTexture() || DepthStencilState->bIsDepthWriteEnabled == false, TEXT("Attempting to set a depth-stencil state that writes depth but no depth texture is configured!\nState: %s\nRender Pass: %s"), *FString([DepthStencilState->State.GetPtr() description]), *FString([RenderPassDesc.GetPtr() description]));
-			METAL_FATAL_ASSERT(RenderPassDesc.GetStencilAttachment().GetTexture() || DepthStencilState->bIsStencilWriteEnabled == false, TEXT("Attempting to set a depth-stencil state that writes stencil but no stencil texture is configured!\nState: %s\nRender Pass: %s"), *FString([DepthStencilState->State.GetPtr() description]), *FString([RenderPassDesc.GetPtr() description]));
-		}
 	}
 }
 
@@ -2081,6 +2075,13 @@ void FMetalStateCache::SetRenderState(FMetalCommandEncoder& CommandEncoder, FMet
 		if (RasterBits & EMetalRenderFlagDepthStencilState)
 		{
 			check(IsValidRef(DepthStencilState));
+            
+            if (DepthStencilState && RenderPassDesc && SafeGetRuntimeDebuggingLevel() >= EMetalDebugLevelFastValidation)
+            {
+                METAL_FATAL_ASSERT(DepthStencilState->bIsDepthWriteEnabled == false || (RenderPassDesc.GetDepthAttachment() && RenderPassDesc.GetDepthAttachment().GetTexture()) , TEXT("Attempting to set a depth-stencil state that writes depth but no depth texture is configured!\nState: %s\nRender Pass: %s"), *FString([DepthStencilState->State.GetPtr() description]), *FString([RenderPassDesc.GetPtr() description]));
+                METAL_FATAL_ASSERT(DepthStencilState->bIsStencilWriteEnabled == false || (RenderPassDesc.GetStencilAttachment() && RenderPassDesc.GetStencilAttachment().GetTexture()), TEXT("Attempting to set a depth-stencil state that writes stencil but no stencil texture is configured!\nState: %s\nRender Pass: %s"), *FString([DepthStencilState->State.GetPtr() description]), *FString([RenderPassDesc.GetPtr() description]));
+            }
+            
 			CommandEncoder.SetDepthStencilState(DepthStencilState ? DepthStencilState->State : nil);
 		}
 		if (RasterBits & EMetalRenderFlagStencilReferenceValue)
