Simplify the way ranges are allocated and freed in FMetalSubBufferMagazine to avoid having to allocate/deallocate memory - this wouldn't have been leaking (it will always be referenced) but it will consume more memory than necessary as the number of magazine's increases.

diff --git a/Engine/Source/Runtime/Apple/MetalRHI/Private/MetalBuffer.cpp b/Engine/Source/Runtime/Apple/MetalRHI/Private/MetalBuffer.cpp
--- a/Engine/Source/Runtime/Apple/MetalRHI/Private/MetalBuffer.cpp
+++ b/Engine/Source/Runtime/Apple/MetalRHI/Private/MetalBuffer.cpp
@@ -705,11 +705,8 @@ FMetalSubBufferMagazine::FMetalSubBufferMagazine(NSUInteger Size, NSUInteger Chu
 		
 		INC_MEMORY_STAT_BY(STAT_MetalBufferUnusedMemory, FullSize);
 		INC_MEMORY_STAT_BY(STAT_MetalMagazineBufferUnusedMemory, FullSize);
-		while(FullSize)
-		{
-			FullSize -= ChunkSize;
-			FreeRanges.Push(new ns::Range(FullSize, ChunkSize));
-		}
+        uint32 BlockCount = FullSize / ChunkSize;
+        Blocks.AddZeroed(BlockCount);
 	}
 }
 
@@ -747,7 +744,8 @@ void FMetalSubBufferMagazine::FreeRange(ns::Range const& Range)
 		}
 #endif
 	
-		FreeRanges.Push(new ns::Range(Range.Location, Range.Length));
+		uint32 BlockIndex = Range.Location / Range.Length;
+		FPlatformAtomics::AtomicStore(&Blocks[BlockIndex], 0);
 		FPlatformAtomics::InterlockedAdd(&UsedSize, -((int64)Range.Length));
 		
 		INC_MEMORY_STAT_BY(STAT_MetalBufferUnusedMemory, Range.Length);
@@ -878,15 +876,18 @@ FMetalBuffer FMetalSubBufferMagazine::NewBuffer()
 	{
 		check(ParentBuffer && ParentBuffer.GetPtr());
 		
-		ns::Range* Range = FreeRanges.Pop();
-		if (Range)
+		for (uint32 i = 0; i < Blocks.Num(); i++)
 		{
-			FPlatformAtomics::InterlockedAdd(&UsedSize, ((int64)Range->Length));
-			DEC_MEMORY_STAT_BY(STAT_MetalBufferUnusedMemory, Range->Length);
-			DEC_MEMORY_STAT_BY(STAT_MetalMagazineBufferUnusedMemory, Range->Length);
-            INC_MEMORY_STAT_BY(STAT_MetalMagazineBufferMemory, Range->Length);
-			Result = FMetalBuffer(MTLPP_VALIDATE(mtlpp::Buffer, ParentBuffer, SafeGetRuntimeDebuggingLevel() >= EMetalDebugLevelValidation, NewBuffer(*Range)), this);
-			delete Range;
+			if (FPlatformAtomics::InterlockedCompareExchange(&Blocks[i], 1, 0) == 0)
+			{
+				ns::Range Range(i * BlockSize, BlockSize);
+				FPlatformAtomics::InterlockedAdd(&UsedSize, ((int64)Range.Length));
+				DEC_MEMORY_STAT_BY(STAT_MetalBufferUnusedMemory, Range.Length);
+				DEC_MEMORY_STAT_BY(STAT_MetalMagazineBufferUnusedMemory, Range.Length);
+				INC_MEMORY_STAT_BY(STAT_MetalMagazineBufferMemory, Range.Length);
+				Result = FMetalBuffer(MTLPP_VALIDATE(mtlpp::Buffer, ParentBuffer, SafeGetRuntimeDebuggingLevel() >= EMetalDebugLevelValidation, NewBuffer(Range)), this);
+				break;
+			}
 		}
 	}
 
diff --git a/Engine/Source/Runtime/Apple/MetalRHI/Private/MetalBuffer.h b/Engine/Source/Runtime/Apple/MetalRHI/Private/MetalBuffer.h
--- a/Engine/Source/Runtime/Apple/MetalRHI/Private/MetalBuffer.h
+++ b/Engine/Source/Runtime/Apple/MetalRHI/Private/MetalBuffer.h
@@ -129,7 +129,7 @@ private:
 	int64 volatile UsedSize;
 	mtlpp::Buffer ParentBuffer;
 	mutable mtlpp::Heap ParentHeap;
-	TLockFreePointerListLIFO<ns::Range> FreeRanges;
+	TArray<int8> Blocks;
 };
 
 struct FMetalRingBufferRef
