Really disable MTLFence and MTLHeap by default on iOS & tvOS to avoid timeouts/internal errors.
- Fix -nomsaa which requires disabling MSAA in pipeline and render-pass descriptors.
- Reorder some of the MetalRuntimeDebugLevel options so that splitting command-buffers and waiting on them can be used in Shipping builds!

diff --git a/Engine/Source/Runtime/Apple/MetalRHI/Private/MetalCommandBuffer.h b/Engine/Source/Runtime/Apple/MetalRHI/Private/MetalCommandBuffer.h
--- a/Engine/Source/Runtime/Apple/MetalRHI/Private/MetalCommandBuffer.h
+++ b/Engine/Source/Runtime/Apple/MetalRHI/Private/MetalCommandBuffer.h
@@ -35,10 +35,10 @@ enum EMetalDebugLevel
 	EMetalDebugLevelFastValidation,
 	EMetalDebugLevelTrackResources,
 	EMetalDebugLevelResetOnBind,
-	EMetalDebugLevelValidation,
-	EMetalDebugLevelLogOperations,
 	EMetalDebugLevelConditionalSubmit,
-	EMetalDebugLevelWaitForComplete
+	EMetalDebugLevelWaitForComplete,
+	EMetalDebugLevelValidation,
+	EMetalDebugLevelLogOperations
 };
 
 NS_ASSUME_NONNULL_BEGIN
diff --git a/Engine/Source/Runtime/Apple/MetalRHI/Private/MetalCommandEncoder.cpp b/Engine/Source/Runtime/Apple/MetalRHI/Private/MetalCommandEncoder.cpp
--- a/Engine/Source/Runtime/Apple/MetalRHI/Private/MetalCommandEncoder.cpp
+++ b/Engine/Source/Runtime/Apple/MetalRHI/Private/MetalCommandEncoder.cpp
@@ -75,10 +75,20 @@ struct FMetalCommandContextDebug
 {
 	@public
 	TArray<FMetalCommandContextDebug> Contexts;
+	uint32 Index;
 }
 @end
 @implementation FMetalCommandBufferDebug
 APPLE_PLATFORM_OBJECT_ALLOC_OVERRIDES(FMetalCommandBufferDebug)
+- (instancetype)init
+{
+	id Self = [super init];
+	if (Self)
+	{
+		Index = ~0u;
+	}
+	return Self;
+}
 - (void)dealloc
 {
 	Contexts.Empty();
@@ -121,6 +131,11 @@ uint32 FMetalCommandBufferMarkers::AddCommand(uint32 CmdBufIndex, uint32 Encoder
 	uint32 Num = 0;
 	if (m_ptr)
 	{
+		if (m_ptr->Index == ~0u)
+		{
+			m_ptr->Index = CmdBufIndex;
+		}
+		
 		FMetalCommandContextDebug& Context = m_ptr->Contexts[ContextIndex];
 		if (Context.DebugBuffer != DebugBuffer)
 		{
@@ -177,6 +192,16 @@ uint32 FMetalCommandBufferMarkers::NumContexts() const
 	return Num;
 }
 
+uint32 FMetalCommandBufferMarkers::GetIndex() const
+{
+	uint32 Num = 0;
+	if (m_ptr)
+	{
+		Num = m_ptr->Index;
+	}
+	return Num;
+}
+
 FMetalCommandBufferMarkers FMetalCommandBufferMarkers::Get(mtlpp::CommandBuffer const& CmdBuf)
 {
 	return CmdBuf.GetAssociatedObject<FMetalCommandBufferMarkers>(FMetalCommandBufferMarkers::kTableAssociationKey);
diff --git a/Engine/Source/Runtime/Apple/MetalRHI/Private/MetalCommandEncoder.h b/Engine/Source/Runtime/Apple/MetalRHI/Private/MetalCommandEncoder.h
--- a/Engine/Source/Runtime/Apple/MetalRHI/Private/MetalCommandEncoder.h
+++ b/Engine/Source/Runtime/Apple/MetalRHI/Private/MetalCommandEncoder.h
@@ -94,6 +94,7 @@ public:
 	TArray<FMetalCommandDebug>* GetCommands(uint32 Context);
 	ns::AutoReleased<FMetalBuffer> GetDebugBuffer(uint32 ContextIndex);
 	uint32 NumContexts() const;
+	uint32 GetIndex() const;
 	
 	static FMetalCommandBufferMarkers Get(mtlpp::CommandBuffer const& CmdBuf);
 	
diff --git a/Engine/Source/Runtime/Apple/MetalRHI/Private/MetalCommandList.cpp b/Engine/Source/Runtime/Apple/MetalRHI/Private/MetalCommandList.cpp
--- a/Engine/Source/Runtime/Apple/MetalRHI/Private/MetalCommandList.cpp
+++ b/Engine/Source/Runtime/Apple/MetalRHI/Private/MetalCommandList.cpp
@@ -74,6 +74,7 @@ static void ReportMetalCommandBufferFailure(mtlpp::CommandBuffer const& Complete
 		FMetalCommandBufferMarkers Markers = FMetalCommandBufferMarkers::Get(CompletedBuffer);
 		if (Markers)
 		{
+			uint32 CommandBufferIndex = Markers.GetIndex();
 			FMetalDebugInfo BrokenDraw;
 			BrokenDraw.EncoderIndex = UINT32_MAX;
 			BrokenDraw.CommandIndex = UINT32_MAX;
@@ -82,7 +83,7 @@ static void ReportMetalCommandBufferFailure(mtlpp::CommandBuffer const& Complete
             BrokenDraw.CommandBuffer = 0;
 			uint32 BrokenContext = 0;
 			bool bFoundBrokenDraw = false;
-			for (uint32 i = 0; !bFoundBrokenDraw && i < Markers.NumContexts(); i++)
+			for (uint32 i = 0; (CommandBufferIndex != ~0u) && !bFoundBrokenDraw && i < Markers.NumContexts(); i++)
 			{
 				ns::AutoReleased<FMetalBuffer> DebugBuffer = Markers.GetDebugBuffer(i);
 				TArray<FMetalCommandDebug>* Commands = Markers.GetCommands(i);
@@ -96,9 +97,11 @@ static void ReportMetalCommandBufferFailure(mtlpp::CommandBuffer const& Complete
 					FMetalDebugInfo CurrentDraw = DebugArray[0];
 					
 					// On TBDR we find the disjoint where one tile has progressed further than another
+					// We find the earliest failure in the probably vain hope that this is the tile that actually failed
+					// There's actually no guarantee of that, it might be one of the later tiles that exploded
 					for (uint32 j = 0; j < DebugCount; j++)
 					{
-						if (DebugArray[j].CommandIndex < CurrentDraw.CommandIndex)
+						if (CommandBufferIndex == DebugArray[j].CmdBuffIndex && DebugArray[j].EncoderIndex < CurrentDraw.EncoderIndex && DebugArray[j].CommandIndex < CurrentDraw.CommandIndex)
 						{
 							CurrentDraw = DebugArray[j];
 							break;
@@ -106,13 +109,7 @@ static void ReportMetalCommandBufferFailure(mtlpp::CommandBuffer const& Complete
 					}
 
 					// Find the first command to fail - which depends on the order of encoders - parallel contexts make this more complicated
-					if (CurrentDraw.CmdBuffIndex < BrokenDraw.CmdBuffIndex)
-					{
-						BrokenDraw = CurrentDraw;
-                        BrokenContext = CurrentDraw.ContextIndex;
-						bFoundBrokenDraw = true;
-					}
-					else if (CurrentDraw.CommandIndex < Commands->Num() && CurrentDraw.EncoderIndex < BrokenDraw.EncoderIndex)
+					if (CurrentDraw.CmdBuffIndex == BrokenDraw.CmdBuffIndex && CurrentDraw.CommandIndex < Commands->Num() && CurrentDraw.EncoderIndex < BrokenDraw.EncoderIndex)
 					{
 						BrokenDraw = CurrentDraw;
                         BrokenContext = i;
@@ -120,6 +117,12 @@ static void ReportMetalCommandBufferFailure(mtlpp::CommandBuffer const& Complete
 					}
                     else if(Commands->Num() > 1)
                     {
+						BrokenDraw.EncoderIndex = 0;
+						BrokenDraw.CommandIndex = 0;
+						BrokenDraw.CmdBuffIndex = CommandBufferIndex;
+						BrokenDraw.ContextIndex = 0;
+                        BrokenContext = i;
+						BrokenDraw.CommandBuffer = (uintptr_t)CompletedBuffer.GetPtr();
                         BrokenDraw = DebugArray[1];
                         bFoundBrokenDraw = true;
                     }
@@ -129,41 +132,7 @@ static void ReportMetalCommandBufferFailure(mtlpp::CommandBuffer const& Complete
 			{
 				id<MTLCommandBuffer> Ptr = reinterpret_cast<id<MTLCommandBuffer>>(BrokenDraw.CommandBuffer);
 				UE_LOG(LogMetal, Error, TEXT("GPU last wrote Command Buffer: %u (%llx) Encoder Index: %d Context Index: %d Draw Index: %d PSO: VS: %u_%u, PS: %u_%u."), BrokenDraw.CmdBuffIndex, BrokenDraw.CommandBuffer, BrokenDraw.EncoderIndex, BrokenDraw.ContextIndex, BrokenDraw.CommandIndex, BrokenDraw.PSOSignature[0], BrokenDraw.PSOSignature[1], BrokenDraw.PSOSignature[2], BrokenDraw.PSOSignature[3]);
-
-                if (Ptr && Ptr != CompletedBuffer.GetPtr())
-                {
-                    // Have we failed between commands-buffers? Or have our command-buffers been subject to coalescing?
-                    ns::AutoReleased<mtlpp::CommandBuffer> PrevCommandBuffer(Ptr);
-                    FMetalCommandBufferMarkers OtherMarkers = FMetalCommandBufferMarkers::Get(PrevCommandBuffer);
-                    if (OtherMarkers)
-                    {
-                        ns::AutoReleased<FMetalBuffer> DebugBuffer = OtherMarkers.GetDebugBuffer(BrokenContext);
-                        TArray<FMetalCommandDebug>* Commands = OtherMarkers.GetCommands(BrokenContext);
-
-                        UE_LOG(LogMetal, Error, TEXT("Failed executing following commands:"));
-                        uint32 StartIdx = BrokenDraw.CommandIndex;
-                        for (uint32 CmdIdx = StartIdx; CmdIdx < Commands->Num(); CmdIdx++)
-                        {
-                            FMetalCommandDebug& Command = (*Commands)[CmdIdx];
-                            FString VSHash = Command.PSO->VertexShader->GetHash().ToString();
-                            uint32 VSSig[2] = { Command.PSO->VertexShader->SourceLen, Command.PSO->VertexShader->SourceCRC };
-
-                            FString PSHash;
-                            uint32 PSSig[2] = { 0, 0 };
-                            if (Command.PSO->PixelShader.IsValid())
-                            {
-                                PSHash = Command.PSO->PixelShader->GetHash().ToString();
-                                PSSig[0] = Command.PSO->PixelShader->SourceLen;
-                                PSSig[1] = Command.PSO->PixelShader->SourceCRC;
-                            }
-
-                            UE_LOG(LogMetal, Error, TEXT("Command Buffer: %d (%p) Encoder: %d Command: %d: %s PSO: VS: %s (%u_%u), PS: %s (%u_%u)"), Command.CmdBufIndex, Ptr, Command.Encoder, Command.Index, *Command.Data.ToString(), *VSHash, VSSig[0], VSSig[1], *PSHash, PSSig[0], PSSig[1]);
-                        }
-                    }
-
-                    BrokenDraw.CommandIndex = 0;
-                }
-				
+			
 				ns::AutoReleased<FMetalBuffer> DebugBuffer = Markers.GetDebugBuffer(BrokenContext);
 				TArray<FMetalCommandDebug>* Commands = Markers.GetCommands(BrokenContext);
 				if (Commands->Num())
@@ -198,6 +167,13 @@ static void ReportMetalCommandBufferFailure(mtlpp::CommandBuffer const& Complete
                     for (uint32 i = 0; i < DebugCount; i++)
                     {
                         FMetalDebugInfo& Command = DebugArray[i];
+						
+						// Stop when nothing has been written into the buffer
+						if (Command.CmdBuffIndex == 0 && Command.CommandBuffer == 0 && Command.EncoderIndex == 0 && Command.CommandIndex == 0 && Command.PSOSignature[0] == 0 && Command.PSOSignature[1] == 0 && Command.PSOSignature[2] == 0 && Command.PSOSignature[3] == 0)
+						{
+							break;
+						}
+						
                         UE_LOG(LogMetal, Error, TEXT("Command Buffer: %d (%p) Debug Buffer: %p Tile: %u Context: %d Encoder: %d Command: %d PSO: VS: %u_%u, PS: %u_%u"), Command.CmdBuffIndex, Command.CommandBuffer, DebugBuffer.GetPtr(), i, Command.ContextIndex, Command.EncoderIndex, Command.CommandIndex, Command.PSOSignature[0], Command.PSOSignature[1], Command.PSOSignature[2], Command.PSOSignature[3]);
                     }
                 }
diff --git a/Engine/Source/Runtime/Apple/MetalRHI/Private/MetalCommandQueue.cpp b/Engine/Source/Runtime/Apple/MetalRHI/Private/MetalCommandQueue.cpp
--- a/Engine/Source/Runtime/Apple/MetalRHI/Private/MetalCommandQueue.cpp
+++ b/Engine/Source/Runtime/Apple/MetalRHI/Private/MetalCommandQueue.cpp
@@ -94,8 +94,16 @@ FMetalCommandQueue::FMetalCommandQueue(mtlpp::Device InDevice, uint32 const MaxN
 				if (Vers.majorVersion >= 12)
 				{
 					Features |= EMetalFeaturesMaxThreadsPerThreadgroup;
-					Features |= EMetalFeaturesFences;
-					Features |= EMetalFeaturesHeaps;
+					
+					if (FParse::Param(FCommandLine::Get(),TEXT("metalfence")))
+					{
+						Features |= EMetalFeaturesFences;
+					}
+					
+					if (FParse::Param(FCommandLine::Get(),TEXT("metalheap")))
+					{
+						Features |= EMetalFeaturesHeaps;
+					}
 					
 					if (MaxShaderVersion >= 4)
 					{
@@ -113,16 +121,6 @@ FMetalCommandQueue::FMetalCommandQueue(mtlpp::Device InDevice, uint32 const MaxN
 		if(Device.SupportsFeatureSet(mtlpp::FeatureSet::iOS_GPUFamily3_v2) || Device.SupportsFeatureSet(mtlpp::FeatureSet::iOS_GPUFamily2_v3) || Device.SupportsFeatureSet(mtlpp::FeatureSet::iOS_GPUFamily1_v3))
 		{
 			Features |= EMetalFeaturesStencilView | EMetalFeaturesFunctionConstants | EMetalFeaturesGraphicsUAVs | EMetalFeaturesMemoryLessResources;
-			
-			if (FParse::Param(FCommandLine::Get(),TEXT("metalfence")))
-			{
-				Features |= EMetalFeaturesFences;
-			}
-			
-			if (FParse::Param(FCommandLine::Get(),TEXT("metalheap")))
-			{
-				Features |= EMetalFeaturesHeaps;
-			}
 		}
 		
 		if(Device.SupportsFeatureSet(mtlpp::FeatureSet::iOS_GPUFamily3_v2))
@@ -158,8 +156,16 @@ FMetalCommandQueue::FMetalCommandQueue(mtlpp::Device InDevice, uint32 const MaxN
 				if (Vers.majorVersion >= 12)
 				{
 					Features |= EMetalFeaturesMaxThreadsPerThreadgroup;
-					Features |= EMetalFeaturesFences;
-					Features |= EMetalFeaturesHeaps;
+					
+                    if (FParse::Param(FCommandLine::Get(),TEXT("metalfence")))
+                    {
+                        Features |= EMetalFeaturesFences;
+                    }
+                    
+                    if (FParse::Param(FCommandLine::Get(),TEXT("metalheap")))
+                    {
+                        Features |= EMetalFeaturesHeaps;
+                    }
 					
 					if (MaxShaderVersion >= 4)
 					{
@@ -416,12 +422,10 @@ void FMetalCommandQueue::CommitCommandBuffer(mtlpp::CommandBuffer& CommandBuffer
 	MTLPP_VALIDATE(mtlpp::CommandBuffer, CommandBuffer, SafeGetRuntimeDebuggingLevel() >= EMetalDebugLevelValidation, Commit());
 	
 	// Wait for completion when debugging command-buffers.
-#if METAL_DEBUG_OPTIONS
 	if (RuntimeDebuggingLevel >= EMetalDebugLevelWaitForComplete)
 	{
 		CommandBuffer.WaitUntilCompleted();
 	}
-#endif
 }
 
 void FMetalCommandQueue::SubmitCommandBuffers(TArray<mtlpp::CommandBuffer> BufferList, uint32 Index, uint32 Count)
diff --git a/Engine/Source/Runtime/Apple/MetalRHI/Private/MetalContext.cpp b/Engine/Source/Runtime/Apple/MetalRHI/Private/MetalContext.cpp
--- a/Engine/Source/Runtime/Apple/MetalRHI/Private/MetalContext.cpp
+++ b/Engine/Source/Runtime/Apple/MetalRHI/Private/MetalContext.cpp
@@ -89,10 +89,10 @@ static FAutoConsoleVariableRef CVarMetalRuntimeDebugLevel(
 	TEXT("\t2: Enable light-weight validation of resource bindings & API usage,\n")
 	TEXT("\t3: Track resources and validate lifetime on command-buffer failure,\n")
 	TEXT("\t4: Reset resource bindings when binding a PSO/Compute-Shader to simplify GPU debugging,\n")
-	TEXT("\t5: Enable slower, more extensive validation checks for resource types & encoder usage,\n")
-	TEXT("\t6: Record the draw, blit & dispatch commands issued into a command-buffer and report them on failure,\n")
-	TEXT("\t7: Allow rhi.Metal.CommandBufferCommitThreshold to break command-encoders (except when MSAA is enabled),\n")
-	TEXT("\t8: Wait for each command-buffer to complete immediately after submission."));
+	TEXT("\t5: Allow rhi.Metal.CommandBufferCommitThreshold to break command-encoders (except when MSAA is enabled),\n")
+	TEXT("\t6: Wait for each command-buffer to complete immediately after submission.")
+	TEXT("\t7: Enable slower, more extensive validation checks for resource types & encoder usage,\n")
+	TEXT("\t8: Record the draw, blit & dispatch commands issued into a command-buffer and report them on failure,\n"));
 
 float GMetalPresentFramePacing = 0.0f;
 #if !PLATFORM_MAC
diff --git a/Engine/Source/Runtime/Apple/MetalRHI/Private/MetalPipeline.cpp b/Engine/Source/Runtime/Apple/MetalRHI/Private/MetalPipeline.cpp
--- a/Engine/Source/Runtime/Apple/MetalRHI/Private/MetalPipeline.cpp
+++ b/Engine/Source/Runtime/Apple/MetalRHI/Private/MetalPipeline.cpp
@@ -804,10 +804,11 @@ static FMetalShaderPipeline* CreateMTLRenderPipeline(bool const bSync, FMetalGra
 #endif
         }
         
-        RenderPipelineDesc.SetSampleCount(FMath::Max(Init.NumSamples, (uint16)1u));
+        static bool bNoMSAA = FParse::Param(FCommandLine::Get(), TEXT("nomsaa"));
+        RenderPipelineDesc.SetSampleCount(!bNoMSAA ? FMath::Max(Init.NumSamples, (uint16)1u) : (uint16)1u);
     #if PLATFORM_MAC
         RenderPipelineDesc.SetInputPrimitiveTopology(TranslatePrimitiveTopology(Init.PrimitiveType));
-		DebugPipelineDesc.SetSampleCount(FMath::Max(Init.NumSamples, (uint16)1u));
+		DebugPipelineDesc.SetSampleCount(!bNoMSAA ? FMath::Max(Init.NumSamples, (uint16)1u) : (uint16)1u);
 		DebugPipelineDesc.SetInputPrimitiveTopology(TranslatePrimitiveTopology(Init.PrimitiveType));
     #endif
         
diff --git a/Engine/Source/Runtime/Apple/MetalRHI/Private/MetalRenderPass.cpp b/Engine/Source/Runtime/Apple/MetalRHI/Private/MetalRenderPass.cpp
--- a/Engine/Source/Runtime/Apple/MetalRHI/Private/MetalRenderPass.cpp
+++ b/Engine/Source/Runtime/Apple/MetalRHI/Private/MetalRenderPass.cpp
@@ -1873,7 +1873,6 @@ void FMetalRenderPass::ConditionalSubmit()
 	
 	bool bCanForceSubmit = State.CanRestartRenderPass();
 
-#if METAL_DEBUG_OPTIONS
 	FRHISetRenderTargetsInfo CurrentRenderTargets = State.GetRenderTargetsInfo();
 	
 	// Force a command-encoder when GMetalRuntimeDebugLevel is enabled to help track down intermittent command-buffer failures.
@@ -1916,7 +1915,6 @@ void FMetalRenderPass::ConditionalSubmit()
 		
 		bCanForceSubmit = bCanChangeRT;
 	}
-#endif
 	
 	if (GMetalCommandBufferCommitThreshold > 0 && NumOutstandingOps > 0 && NumOutstandingOps >= GMetalCommandBufferCommitThreshold && bCanForceSubmit && !CurrentEncoder.IsParallel())
 	{
@@ -1926,7 +1924,6 @@ void FMetalRenderPass::ConditionalSubmit()
 			NumOutstandingOps = 0;
 		}
 		
-#if METAL_DEBUG_OPTIONS
 		// Force a command-encoder when GMetalRuntimeDebugLevel is enabled to help track down intermittent command-buffer failures.
 		if (bWithinRenderPass && CmdList.GetCommandQueue().GetRuntimeDebuggingLevel() >= EMetalDebugLevelConditionalSubmit && State.GetHasValidRenderTarget())
 		{
@@ -1946,7 +1943,6 @@ void FMetalRenderPass::ConditionalSubmit()
 				RestartRenderPass(State.GetRenderPassDescriptor());
 			}
 		}
-#endif
 	}
 }
 
diff --git a/Engine/Source/Runtime/Apple/MetalRHI/Private/MetalStateCache.cpp b/Engine/Source/Runtime/Apple/MetalRHI/Private/MetalStateCache.cpp
--- a/Engine/Source/Runtime/Apple/MetalRHI/Private/MetalStateCache.cpp
+++ b/Engine/Source/Runtime/Apple/MetalRHI/Private/MetalStateCache.cpp
@@ -36,8 +36,13 @@ FORCEINLINE mtlpp::StoreAction GetMetalRTStoreAction(ERenderTargetStoreAction St
         //because we may render to the same MSAA target twice in two separate passes.  BasePass, then some stuff, then translucency for example and we need to not lose the prior MSAA contents to do this properly.
 		case ERenderTargetStoreAction::EMultisampleResolve:
 		{
+            static bool bNoMSAA = FParse::Param(FCommandLine::Get(), TEXT("nomsaa"));
 			static bool bSupportsMSAAStoreResolve = FMetalCommandQueue::SupportsFeature(EMetalFeaturesMSAAStoreAndResolve) && (GMaxRHIFeatureLevel >= ERHIFeatureLevel::SM5);
-			if (bSupportsMSAAStoreResolve)
+            if (bNoMSAA)
+            {
+                return mtlpp::StoreAction::Store;
+            }
+			else if (bSupportsMSAAStoreResolve)
 			{
 				return mtlpp::StoreAction::StoreAndMultisampleResolve;
 			}
