Fix command-buffer failures when using MTLFence & using the buffer heap defragmentation code at the same time - the defragging was not being properly fenced. This change does *not* re-enable MTLFence for iOS - that will wait for a later release.

diff --git a/Engine/Source/Runtime/Apple/MetalRHI/Private/MetalBuffer.cpp b/Engine/Source/Runtime/Apple/MetalRHI/Private/MetalBuffer.cpp
--- a/Engine/Source/Runtime/Apple/MetalRHI/Private/MetalBuffer.cpp
+++ b/Engine/Source/Runtime/Apple/MetalRHI/Private/MetalBuffer.cpp
@@ -7,6 +7,7 @@
 #include <objc/runtime.h>
 #include "MetalCommandBuffer.h"
 #include "MetalProfiler.h"
+#include "MetalRenderPass.h"
 
 DECLARE_MEMORY_STAT(TEXT("Used Device Buffer Memory"), STAT_MetalDeviceBufferMemory, STATGROUP_MetalRHI);
 DECLARE_MEMORY_STAT(TEXT("Used Pooled Buffer Memory"), STAT_MetalPooledBufferMemory, STATGROUP_MetalRHI);
@@ -690,6 +691,7 @@ FMetalSubBufferMagazine::FMetalSubBufferMagazine(NSUInteger Size, NSUInteger Chu
 		Desc.SetStorageMode(Storage);
 		ParentHeap = GetMetalDeviceContext().GetDevice().NewHeap(Desc);
 		check(ParentHeap.GetPtr());
+		METAL_FATAL_ASSERT(ParentHeap, TEXT("Failed to create heap of size %u and resource options %u"), Size, (uint32)Options);
 #if STATS || ENABLE_LOW_LEVEL_MEM_TRACKER
 		MetalLLM::LogAllocHeap(GetMetalDeviceContext().GetDevice(), ParentHeap);
 #endif
@@ -699,6 +701,7 @@ FMetalSubBufferMagazine::FMetalSubBufferMagazine(NSUInteger Size, NSUInteger Chu
 		ParentBuffer = MTLPP_VALIDATE(mtlpp::Device, GetMetalDeviceContext().GetDevice(), SafeGetRuntimeDebuggingLevel() >= EMetalDebugLevelValidation, NewBuffer(FullSize, Options));
 		check(ParentBuffer.GetPtr());
 		check(ParentBuffer.GetLength() >= FullSize);
+		METAL_FATAL_ASSERT(ParentBuffer, TEXT("Failed to create heap of size %u and resource options %u"), Size, (uint32)Options);
 #if STATS || ENABLE_LOW_LEVEL_MEM_TRACKER
 		MetalLLM::LogAllocBuffer(GetMetalDeviceContext().GetDevice(), ParentBuffer);
 #endif
@@ -826,9 +829,16 @@ NSUInteger FMetalSubBufferMagazine::GetUsedSize() const
 	}
 }
 
-NSUInteger FMetalSubBufferMagazine::GetFreeSize() const 
+NSUInteger FMetalSubBufferMagazine::GetFreeSize() const
 {
-	return GetSize() - GetUsedSize();
+	if (ParentHeap)
+	{
+		return ParentHeap.MaxAvailableSizeWithAlignment(MinAlign);
+	}
+	else
+	{
+		return GetSize() - GetUsedSize();
+	}
 }
 
 int64 FMetalSubBufferMagazine::NumCurrentAllocations() const
@@ -1321,7 +1331,7 @@ FMetalResourceHeap::FMetalResourceHeap(void)
 
 FMetalResourceHeap::~FMetalResourceHeap()
 {
-	Compact(true);
+	Compact(nullptr, true);
 }
 
 void FMetalResourceHeap::Init(FMetalCommandQueue& InQueue)
@@ -1612,7 +1622,7 @@ FMetalBuffer FMetalResourceHeap::CreateBuffer(uint32 Size, uint32 Alignment, uin
 	}
 	
     METAL_DEBUG_OPTION(GetMetalDeviceContext().ValidateIsInactiveBuffer(Buffer));
-	check(Buffer && Buffer.GetPtr());
+	METAL_FATAL_ASSERT(Buffer, TEXT("Failed to create buffer of size %u and resource options %u"), Size, (uint32)Options);
 	return Buffer;
 }
 
@@ -1706,7 +1716,7 @@ void FMetalResourceHeap::ReleaseTexture(FMetalSurface* Surface, FMetalTexture& T
 	}
 }
 
-void FMetalResourceHeap::Compact(bool const bForce)
+void FMetalResourceHeap::Compact(FMetalRenderPass* Pass, bool const bForce)
 {
 	FScopeLock Lock(&Mutex);
     for (uint32 u = 0; u < NumUsageTypes; u++)
@@ -1726,15 +1736,13 @@ void FMetalResourceHeap::Compact(bool const bForce)
                 }
             }
             
-            mtlpp::CommandBuffer Buffer;
-            mtlpp::BlitCommandEncoder Encoder;
-            
             uint32 BytesCompacted = 0;
             uint32 const BytesToCompact = GMetalHeapBufferBytesToCompact;
             
             for (uint32 i = 0; i < NumHeapSizes; i++)
             {
-                for (uint32 j = 1; j < BufferHeaps[u][t][i].Num(); j++)
+				// When not forcing the disposal of all resources we can compact them using the render-pass
+                for (uint32 j = 1; !bForce && j < BufferHeaps[u][t][i].Num(); j++)
                 {
                     FMetalSubBufferHeap* Data = BufferHeaps[u][t][i][j];
                     if (Data->AllocRanges.Num() > 0 && BytesCompacted < BytesToCompact)
@@ -1748,18 +1756,12 @@ void FMetalResourceHeap::Compact(bool const bForce)
                                     FMetalSubBufferHeap* Prev = BufferHeaps[u][t][i][AllocIndex];
                                     if (Prev->MaxAvailableSize() >= Alloc.Range.Length)
                                     {
-                                        if (!Encoder)
-                                        {
-                                            Buffer = Queue->CreateCommandBuffer();
-                                            Encoder = Buffer.BlitCommandEncoder();
-                                        }
-                                        
                                         FMetalBuffer NewBuffer = Prev->NewBuffer(Alloc.Range.Length);
                                         check(NewBuffer && NewBuffer.GetPtr());
                                         
                                         FMetalBuffer SourceResource = FMetalBuffer((id<MTLBuffer>)Alloc.Resource);
-                                        Encoder.Copy(SourceResource, Alloc.Range.Location, NewBuffer, 0, Alloc.Range.Length);
-                                        
+										Pass->AsyncCopyFromBufferToBuffer(SourceResource, Alloc.Range.Location, NewBuffer, 0, Alloc.Range.Length);
+										
                                         FMetalBuffer PrevBuffer;
                                         FMetalBuffer PrevCPUBuffer;
                                         if (Alloc.Owner->Buffer.GetPtr() == Alloc.Resource && Alloc.Owner->Buffer.GetOffset() == Alloc.Range.Location)
@@ -1796,12 +1798,6 @@ void FMetalResourceHeap::Compact(bool const bForce)
                     }
                 }
             }
-            
-            if (Encoder)
-            {
-                Encoder.EndEncoding();
-                Queue->CommitCommandBuffer(Buffer);
-            }
         }
     }
 
diff --git a/Engine/Source/Runtime/Apple/MetalRHI/Private/MetalBuffer.h b/Engine/Source/Runtime/Apple/MetalRHI/Private/MetalBuffer.h
--- a/Engine/Source/Runtime/Apple/MetalRHI/Private/MetalBuffer.h
+++ b/Engine/Source/Runtime/Apple/MetalRHI/Private/MetalBuffer.h
@@ -424,7 +424,7 @@ public:
 	void ReleaseBuffer(FMetalBuffer& Buffer);
 	void ReleaseTexture(FMetalSurface* Surface, FMetalTexture& Texture);
 	
-	void Compact(bool const bForce);
+	void Compact(class FMetalRenderPass* Pass, bool const bForce);
 	
 private:
 	uint32 GetMagazineIndex(uint32 Size);
diff --git a/Engine/Source/Runtime/Apple/MetalRHI/Private/MetalCommandBuffer.h b/Engine/Source/Runtime/Apple/MetalRHI/Private/MetalCommandBuffer.h
--- a/Engine/Source/Runtime/Apple/MetalRHI/Private/MetalCommandBuffer.h
+++ b/Engine/Source/Runtime/Apple/MetalRHI/Private/MetalCommandBuffer.h
@@ -36,9 +36,9 @@ enum EMetalDebugLevel
 	EMetalDebugLevelTrackResources,
 	EMetalDebugLevelResetOnBind,
 	EMetalDebugLevelConditionalSubmit,
-	EMetalDebugLevelWaitForComplete,
 	EMetalDebugLevelValidation,
-	EMetalDebugLevelLogOperations
+	EMetalDebugLevelLogOperations,
+	EMetalDebugLevelWaitForComplete,
 };
 
 NS_ASSUME_NONNULL_BEGIN
diff --git a/Engine/Source/Runtime/Apple/MetalRHI/Private/MetalContext.cpp b/Engine/Source/Runtime/Apple/MetalRHI/Private/MetalContext.cpp
--- a/Engine/Source/Runtime/Apple/MetalRHI/Private/MetalContext.cpp
+++ b/Engine/Source/Runtime/Apple/MetalRHI/Private/MetalContext.cpp
@@ -90,9 +90,9 @@ static FAutoConsoleVariableRef CVarMetalRuntimeDebugLevel(
 	TEXT("\t3: Track resources and validate lifetime on command-buffer failure,\n")
 	TEXT("\t4: Reset resource bindings when binding a PSO/Compute-Shader to simplify GPU debugging,\n")
 	TEXT("\t5: Allow rhi.Metal.CommandBufferCommitThreshold to break command-encoders (except when MSAA is enabled),\n")
-	TEXT("\t6: Wait for each command-buffer to complete immediately after submission.")
-	TEXT("\t7: Enable slower, more extensive validation checks for resource types & encoder usage,\n")
-	TEXT("\t8: Record the draw, blit & dispatch commands issued into a command-buffer and report them on failure,\n"));
+	TEXT("\t6: Enable slower, more extensive validation checks for resource types & encoder usage,\n")
+    TEXT("\t7: Record the draw, blit & dispatch commands issued into a command-buffer and report them on failure,\n")
+    TEXT("\t8: Wait for each command-buffer to complete immediately after submission."));
 
 float GMetalPresentFramePacing = 0.0f;
 #if !PLATFORM_MAC
@@ -487,7 +487,7 @@ void FMetalDeviceContext::ClearFreeList()
 
 void FMetalDeviceContext::DrainHeap()
 {
-	Heap.Compact(false);
+	Heap.Compact(&RenderPass, false);
 }
 
 void FMetalDeviceContext::EndFrame()
@@ -531,7 +531,7 @@ void FMetalDeviceContext::EndFrame()
     
     ClearFreeList();
     
-    Heap.Compact(false);
+	DrainHeap();
     
 	InitFrame(true, 0, 0);
 }
