For non-Shipping Mac builds enable my internal GPU/Metal debugging tools by default to provide more information about problems - this means that bugs that aren't 100% reproducible should still let QA report more useful information. Also fixes the location of where we enabled the GPU debugging for macOS so that it is actually useful.

diff --git a/Engine/Source/Runtime/Apple/MetalRHI/Private/MetalCommandQueue.cpp b/Engine/Source/Runtime/Apple/MetalRHI/Private/MetalCommandQueue.cpp
--- a/Engine/Source/Runtime/Apple/MetalRHI/Private/MetalCommandQueue.cpp
+++ b/Engine/Source/Runtime/Apple/MetalRHI/Private/MetalCommandQueue.cpp
@@ -251,9 +251,6 @@ FMetalCommandQueue::FMetalCommandQueue(mtlpp::Device InDevice, uint32 const MaxN
 			if (MaxShaderVersion >= 4)
 			{
 				Features |= EMetalFeaturesTextureBuffers;
-            
-                IConsoleVariable* GPUCrashDebuggingCVar = IConsoleManager::Get().FindConsoleVariable(TEXT("r.GPUCrashDebugging"));
-                GMetalCommandBufferDebuggingEnabled = (GPUCrashDebuggingCVar && GPUCrashDebuggingCVar->GetInt() != 0) || FParse::Param(FCommandLine::Get(),TEXT("metalgpudebug"));
             }
             if (MaxShaderVersion >= 5)
             {
@@ -277,6 +274,14 @@ FMetalCommandQueue::FMetalCommandQueue(mtlpp::Device InDevice, uint32 const MaxN
 					Features |= EMetalFeaturesHeaps;
 				}
 			}
+            
+#if UE_BUILD_SHIPPING
+			IConsoleVariable* GPUCrashDebuggingCVar = IConsoleManager::Get().FindConsoleVariable(TEXT("r.GPUCrashDebugging"));
+			GMetalCommandBufferDebuggingEnabled = (GPUCrashDebuggingCVar && GPUCrashDebuggingCVar->GetInt() != 0) || FParse::Param(FCommandLine::Get(),TEXT("metalgpudebug"));
+#else
+			// Force the command-buffer debugging on for the moment to see if we can't work out what the heck is going on on macOS
+			GMetalCommandBufferDebuggingEnabled = true;
+#endif
 		}
     }
     else if ([Device.GetName().GetPtr() rangeOfString:@"Nvidia" options:NSCaseInsensitiveSearch].location != NSNotFound)
diff --git a/Engine/Source/Runtime/Apple/MetalRHI/Private/MetalContext.cpp b/Engine/Source/Runtime/Apple/MetalRHI/Private/MetalContext.cpp
--- a/Engine/Source/Runtime/Apple/MetalRHI/Private/MetalContext.cpp
+++ b/Engine/Source/Runtime/Apple/MetalRHI/Private/MetalContext.cpp
@@ -73,11 +73,15 @@ static FAutoConsoleVariableRef CVarMetalResourceDeferDeleteNumFrames(
 	TEXT("Debug option: set to the number of frames that must have passed before resource free-lists are processed and resources disposed of. (Default: 0, Off)"));
 #endif
 
-#if UE_BUILD_SHIPPING || UE_BUILD_TEST
+#if UE_BUILD_SHIPPING
 int32 GMetalRuntimeDebugLevel = 0;
 #else
+#if PLATFORM_MAC
+int32 GMetalRuntimeDebugLevel = 3;
+#else
 int32 GMetalRuntimeDebugLevel = 1;
 #endif
+#endif
 static FAutoConsoleVariableRef CVarMetalRuntimeDebugLevel(
 	TEXT("rhi.Metal.RuntimeDebugLevel"),
 	GMetalRuntimeDebugLevel,
