Add texture type validation to MetalRHI's internal fast-validation runtime - you don't need to use Apple's tools to check that the texture types are what the shader expects anymore.

diff --git a/Engine/Source/Runtime/Apple/MetalRHI/Private/MetalPipeline.cpp b/Engine/Source/Runtime/Apple/MetalRHI/Private/MetalPipeline.cpp
--- a/Engine/Source/Runtime/Apple/MetalRHI/Private/MetalPipeline.cpp
+++ b/Engine/Source/Runtime/Apple/MetalRHI/Private/MetalPipeline.cpp
@@ -514,6 +514,7 @@ APPLE_PLATFORM_OBJECT_ALLOC_OVERRIDES(FMetalShaderPipeline)
 			{
 				checkf(Arg.index < ML_MaxTextures, TEXT("Metal texture index exceeded!"));
 				ResourceMask[Frequency].TextureMask |= (1 << Arg.index);
+				TextureTypes[Frequency].Add(Arg.index, (uint8)Arg.textureType);
 				break;
 			}
 			case MTLArgumentTypeSampler:
diff --git a/Engine/Source/Runtime/Apple/MetalRHI/Private/MetalPipeline.h b/Engine/Source/Runtime/Apple/MetalRHI/Private/MetalPipeline.h
--- a/Engine/Source/Runtime/Apple/MetalRHI/Private/MetalPipeline.h
+++ b/Engine/Source/Runtime/Apple/MetalRHI/Private/MetalPipeline.h
@@ -76,6 +76,7 @@ struct FMetalTessellationPipelineDesc
 	mtlpp::RenderPipelineState DebugPipelineState;
 	FMetalTessellationPipelineDesc TessellationPipelineDesc;
 	TArray<uint32> BufferDataSizes[EMetalShaderStagesNum];
+	TMap<uint8, uint8> TextureTypes[EMetalShaderStagesNum];
 	FMetalDebugShaderResourceMask ResourceMask[EMetalShaderStagesNum];
 	mtlpp::RenderPipelineReflection RenderPipelineReflection;
 	mtlpp::ComputePipelineReflection ComputePipelineReflection;
diff --git a/Engine/Source/Runtime/Apple/MetalRHI/Private/MetalStateCache.cpp b/Engine/Source/Runtime/Apple/MetalRHI/Private/MetalStateCache.cpp
--- a/Engine/Source/Runtime/Apple/MetalRHI/Private/MetalStateCache.cpp
+++ b/Engine/Source/Runtime/Apple/MetalRHI/Private/MetalStateCache.cpp
@@ -2136,6 +2136,7 @@ void FMetalStateCache::SetRenderPipelineState(FMetalCommandEncoder& CommandEncod
 			
 			FMetalDebugShaderResourceMask ComputeMask = Pipeline->ResourceMask[EMetalShaderCompute];
 			TArray<uint32>& MinComputeBufferSizes = Pipeline->BufferDataSizes[EMetalShaderCompute];
+			TMap<uint8, uint8> ComputeTexTypes = Pipeline->TextureTypes[EMetalShaderCompute];
 			while(ComputeMask.BufferMask)
 			{
 				uint32 Index = __builtin_ctz(ComputeMask.BufferMask);
@@ -2155,6 +2156,8 @@ void FMetalStateCache::SetRenderPipelineState(FMetalCommandEncoder& CommandEncod
 					uint32 Index = __builtin_ctzll(LoTextures);
 					LoTextures &= ~(uint64(1) << uint64(Index));
 					ensure(ShaderTextures[EMetalShaderStages::Vertex].Textures[Index] || ShaderTextures[EMetalShaderStages::Hull].Textures[Index + 64]);
+					ensure(!ShaderTextures[EMetalShaderStages::Vertex].Textures[Index] || ShaderTextures[EMetalShaderStages::Vertex].Textures[Index].GetTextureType() == (mtlpp::TextureType)ComputeTexTypes.FindRef(Index));
+					ensure(!ShaderTextures[EMetalShaderStages::Hull].Textures[Index] || ShaderTextures[EMetalShaderStages::Hull].Textures[Index].GetTextureType() == (mtlpp::TextureType)ComputeTexTypes.FindRef(Index));
 				}
 				
 				uint64 HiTextures = (uint64)(ComputeMask.TextureMask >> FMetalTextureMask(64));
@@ -2163,6 +2166,8 @@ void FMetalStateCache::SetRenderPipelineState(FMetalCommandEncoder& CommandEncod
 					uint32 Index = __builtin_ctzll(HiTextures);
 					HiTextures &= ~(uint64(1) << uint64(Index));
 					ensure(ShaderTextures[EMetalShaderStages::Vertex].Textures[Index + 64] || ShaderTextures[EMetalShaderStages::Hull].Textures[Index + 64]);
+					ensure(!ShaderTextures[EMetalShaderStages::Vertex].Textures[Index + 64] || ShaderTextures[EMetalShaderStages::Vertex].Textures[Index + 64].GetTextureType() == (mtlpp::TextureType)ComputeTexTypes.FindRef(Index + 64));
+					ensure(!ShaderTextures[EMetalShaderStages::Hull].Textures[Index + 64] || ShaderTextures[EMetalShaderStages::Hull].Textures[Index + 64].GetTextureType() == (mtlpp::TextureType)ComputeTexTypes.FindRef(Index + 64));
 				}
 			}
 #else
@@ -2172,6 +2177,8 @@ void FMetalStateCache::SetRenderPipelineState(FMetalCommandEncoder& CommandEncod
 				ComputeMask.TextureMask &= ~(1 << Index);
 				
 				ensure(ShaderTextures[EMetalShaderStages::Vertex].Textures[Index] || ShaderTextures[EMetalShaderStages::Hull].Textures[Index]);
+				ensure(!ShaderTextures[EMetalShaderStages::Vertex].Textures[Index] || ShaderTextures[EMetalShaderStages::Vertex].Textures[Index].GetTextureType() == (mtlpp::TextureType)ComputeTexTypes.FindRef(Index));
+				ensure(!ShaderTextures[EMetalShaderStages::Hull].Textures[Index] || ShaderTextures[EMetalShaderStages::Hull].Textures[Index].GetTextureType() == (mtlpp::TextureType)ComputeTexTypes.FindRef(Index));
 			}
 #endif
 			while(ComputeMask.SamplerMask)
@@ -2185,6 +2192,7 @@ void FMetalStateCache::SetRenderPipelineState(FMetalCommandEncoder& CommandEncod
 		
 		FMetalDebugShaderResourceMask VertexMask = Pipeline->ResourceMask[EMetalShaderVertex];
 		TArray<uint32>& MinVertexBufferSizes = Pipeline->BufferDataSizes[EMetalShaderVertex];
+		TMap<uint8, uint8> VertexTexTypes = Pipeline->TextureTypes[EMetalShaderVertex];
 		while(VertexMask.BufferMask)
 		{
 			uint32 Index = __builtin_ctz(VertexMask.BufferMask);
@@ -2203,6 +2211,7 @@ void FMetalStateCache::SetRenderPipelineState(FMetalCommandEncoder& CommandEncod
 				uint32 Index = __builtin_ctzll(LoTextures);
 				LoTextures &= ~(uint64(1) << uint64(Index));
 				ensure(ShaderTextures[VertexStage].Textures[Index]);
+				ensure(ShaderTextures[VertexStage].Textures[Index].GetTextureType() == (mtlpp::TextureType)VertexTexTypes.FindRef(Index));
 			}
 			
 			uint64 HiTextures = (uint64)(VertexMask.TextureMask >> FMetalTextureMask(64));
@@ -2211,6 +2220,7 @@ void FMetalStateCache::SetRenderPipelineState(FMetalCommandEncoder& CommandEncod
 				uint32 Index = __builtin_ctzll(HiTextures);
 				HiTextures &= ~(uint64(1) << uint64(Index));
 				ensure(ShaderTextures[VertexStage].Textures[Index + 64]);
+				ensure(ShaderTextures[VertexStage].Textures[Index + 64].GetTextureType() == (mtlpp::TextureType)VertexTexTypes.FindRef(Index + 64));
 			}
 		}
 #else
@@ -2220,6 +2230,7 @@ void FMetalStateCache::SetRenderPipelineState(FMetalCommandEncoder& CommandEncod
 			VertexMask.TextureMask &= ~(1 << Index);
 			
 			ensure(ShaderTextures[VertexStage].Textures[Index]);
+			ensure(ShaderTextures[VertexStage].Textures[Index].GetTextureType() == (mtlpp::TextureType)VertexTexTypes.FindRef(Index));
 		}
 #endif
 		while(VertexMask.SamplerMask)
@@ -2231,6 +2242,7 @@ void FMetalStateCache::SetRenderPipelineState(FMetalCommandEncoder& CommandEncod
 		
 		FMetalDebugShaderResourceMask FragmentMask = Pipeline->ResourceMask[EMetalShaderFragment];
 		TArray<uint32>& MinFragmentBufferSizes = Pipeline->BufferDataSizes[EMetalShaderFragment];
+		TMap<uint8, uint8> FragmentTexTypes = Pipeline->TextureTypes[EMetalShaderFragment];
 		while(FragmentMask.BufferMask)
 		{
 			uint32 Index = __builtin_ctz(FragmentMask.BufferMask);
@@ -2249,6 +2261,7 @@ void FMetalStateCache::SetRenderPipelineState(FMetalCommandEncoder& CommandEncod
 				uint32 Index = __builtin_ctzll(LoTextures);
 				LoTextures &= ~(uint64(1) << uint64(Index));
 				ensure(ShaderTextures[EMetalShaderStages::Pixel].Textures[Index]);
+				ensure(ShaderTextures[EMetalShaderStages::Pixel].Textures[Index].GetTextureType() == (mtlpp::TextureType)FragmentTexTypes.FindRef(Index));
 			}
 			
 			uint64 HiTextures = (uint64)(FragmentMask.TextureMask >> FMetalTextureMask(64));
@@ -2257,6 +2270,7 @@ void FMetalStateCache::SetRenderPipelineState(FMetalCommandEncoder& CommandEncod
 				uint32 Index = __builtin_ctzll(HiTextures);
 				HiTextures &= ~(uint64(1) << uint64(Index));
 				ensure(ShaderTextures[EMetalShaderStages::Pixel].Textures[Index + 64]);
+				ensure(ShaderTextures[EMetalShaderStages::Pixel].Textures[Index + 64].GetTextureType() == (mtlpp::TextureType)FragmentTexTypes.FindRef(Index + 64));
 			}
 		}
 #else
@@ -2266,6 +2280,7 @@ void FMetalStateCache::SetRenderPipelineState(FMetalCommandEncoder& CommandEncod
 			FragmentMask.TextureMask &= ~(1 << Index);
 			
 			ensure(ShaderTextures[EMetalShaderStages::Pixel].Textures[Index]);
+			ensure(ShaderTextures[EMetalShaderStages::Pixel].Textures[Index].GetTextureType() == (mtlpp::TextureType)FragmentTexTypes.FindRef(Index));
 		}
 #endif
 		while(FragmentMask.SamplerMask)
@@ -2295,6 +2310,7 @@ void FMetalStateCache::SetComputePipelineState(FMetalCommandEncoder& CommandEnco
 		
 		FMetalDebugShaderResourceMask ComputeMask = Pipeline->ResourceMask[EMetalShaderCompute];
 		TArray<uint32>& MinComputeBufferSizes = Pipeline->BufferDataSizes[EMetalShaderCompute];
+		TMap<uint8, uint8> ComputeTexTypes = Pipeline->TextureTypes[EMetalShaderCompute];
 		while(ComputeMask.BufferMask)
 		{
 			uint32 Index = __builtin_ctz(ComputeMask.BufferMask);
@@ -2313,6 +2329,7 @@ void FMetalStateCache::SetComputePipelineState(FMetalCommandEncoder& CommandEnco
 				uint32 Index = __builtin_ctzll(LoTextures);
 				LoTextures &= ~(uint64(1) << uint64(Index));
 				ensure(ShaderTextures[EMetalShaderStages::Compute].Textures[Index]);
+				ensure(ShaderTextures[EMetalShaderStages::Compute].Textures[Index].GetTextureType() == (mtlpp::TextureType)ComputeTexTypes.FindRef(Index));
 			}
 			
 			uint64 HiTextures = (uint64)(ComputeMask.TextureMask >> FMetalTextureMask(64));
@@ -2321,6 +2338,7 @@ void FMetalStateCache::SetComputePipelineState(FMetalCommandEncoder& CommandEnco
 				uint32 Index = __builtin_ctzll(HiTextures);
 				HiTextures &= ~(uint64(1) << uint64(Index));
 				ensure(ShaderTextures[EMetalShaderStages::Compute].Textures[Index + 64]);
+				ensure(ShaderTextures[EMetalShaderStages::Compute].Textures[Index + 64].GetTextureType() == (mtlpp::TextureType)ComputeTexTypes.FindRef(Index + 64));
 			}
 		}
 #else
@@ -2330,6 +2348,7 @@ void FMetalStateCache::SetComputePipelineState(FMetalCommandEncoder& CommandEnco
 			ComputeMask.TextureMask &= ~(1 << Index);
 			
 			ensure(ShaderTextures[EMetalShaderStages::Compute].Textures[Index]);
+			ensure(ShaderTextures[EMetalShaderStages::Compute].Textures[Index].GetTextureType() == (mtlpp::TextureType)ComputeTexTypes.FindRef(Index));
 		}
 #endif
 		while(ComputeMask.SamplerMask)
