Properly disambiguate Intel command-buffer failures.

diff --git a/Engine/Source/Runtime/Apple/MetalRHI/Private/MetalCommandList.cpp b/Engine/Source/Runtime/Apple/MetalRHI/Private/MetalCommandList.cpp
--- a/Engine/Source/Runtime/Apple/MetalRHI/Private/MetalCommandList.cpp
+++ b/Engine/Source/Runtime/Apple/MetalRHI/Private/MetalCommandList.cpp
@@ -272,7 +272,7 @@ static __attribute__ ((optnone)) void HandleNVIDIAMetalCommandBufferError(mtlpp:
 	HandleMetalCommandBufferError(CompletedBuffer);
 }
 
-static void HandleIntelMetalCommandBufferError(mtlpp::CommandBuffer const& CompletedBuffer)
+static __attribute__ ((optnone)) void HandleIntelMetalCommandBufferError(mtlpp::CommandBuffer const& CompletedBuffer)
 {
 	HandleMetalCommandBufferError(CompletedBuffer);
 }
