[NO NON-PII COMMIT MESSAGE]

diff --git a/Engine/Source/Runtime/Apple/AudioMixerAudioUnit/Private/AudioMixerPlatformAudioUnit.cpp b/Engine/Source/Runtime/Apple/AudioMixerAudioUnit/Private/AudioMixerPlatformAudioUnit.cpp
--- a/Engine/Source/Runtime/Apple/AudioMixerAudioUnit/Private/AudioMixerPlatformAudioUnit.cpp
+++ b/Engine/Source/Runtime/Apple/AudioMixerAudioUnit/Private/AudioMixerPlatformAudioUnit.cpp
@@ -327,6 +327,8 @@ namespace Audio
 		AudioStreamInfo.OutputDeviceIndex = Params.OutputDeviceIndex;
 		AudioStreamInfo.AudioMixer = Params.AudioMixer;
 		
+		OpenStreamParams = Params;
+		
 		// Initialize the audio unit graph
 		OSStatus Status = AUGraphInitialize(AudioUnitGraph);
 		if (Status != noErr)
@@ -374,16 +376,23 @@ namespace Audio
 
 	bool FMixerPlatformAudioUnit::StopAudioStream()
 	{
-		if(!bInitialized || AudioStreamInfo.StreamState != EAudioOutputStreamState::Running)
+		if(!bInitialized)
 		{
 			return false;
 		}
 		
-		AudioStreamInfo.StreamState = EAudioOutputStreamState::Stopping;
-		
-		AUGraphStop(AudioUnitGraph);
+		if (AudioStreamInfo.StreamState != EAudioOutputStreamState::Stopped && AudioStreamInfo.StreamState != EAudioOutputStreamState::Closed)
+		{
+			AUGraphStop(AudioUnitGraph);
+			
+			if (AudioStreamInfo.StreamState == EAudioOutputStreamState::Running)
+			{
+				StopGeneratingAudio();
+			}
+
+			check(AudioStreamInfo.StreamState == EAudioOutputStreamState::Stopped);
+		}
 
-		AudioStreamInfo.StreamState = EAudioOutputStreamState::Stopped;
 		
 		return true;
 	}
