The full fix for crashes on launch for macOS with the validation & GPU debugging enabled. The compute shader is in a different library and the Mac render command debug PSO needs the Point topology.

diff --git a/Engine/Source/Runtime/Apple/MetalRHI/Private/MetalCommandQueue.cpp b/Engine/Source/Runtime/Apple/MetalRHI/Private/MetalCommandQueue.cpp
--- a/Engine/Source/Runtime/Apple/MetalRHI/Private/MetalCommandQueue.cpp
+++ b/Engine/Source/Runtime/Apple/MetalRHI/Private/MetalCommandQueue.cpp
@@ -275,13 +275,8 @@ FMetalCommandQueue::FMetalCommandQueue(mtlpp::Device InDevice, uint32 const MaxN
 				}
 			}
             
-#if UE_BUILD_SHIPPING
 			IConsoleVariable* GPUCrashDebuggingCVar = IConsoleManager::Get().FindConsoleVariable(TEXT("r.GPUCrashDebugging"));
 			GMetalCommandBufferDebuggingEnabled = (GPUCrashDebuggingCVar && GPUCrashDebuggingCVar->GetInt() != 0) || FParse::Param(FCommandLine::Get(),TEXT("metalgpudebug"));
-#else
-			// Force the command-buffer debugging on for the moment to see if we can't work out what the heck is going on on macOS
-			GMetalCommandBufferDebuggingEnabled = true;
-#endif
 		}
     }
     else if ([Device.GetName().GetPtr() rangeOfString:@"Nvidia" options:NSCaseInsensitiveSearch].location != NSNotFound)
diff --git a/Engine/Source/Runtime/Apple/MetalRHI/Private/MetalPipeline.cpp b/Engine/Source/Runtime/Apple/MetalRHI/Private/MetalPipeline.cpp
--- a/Engine/Source/Runtime/Apple/MetalRHI/Private/MetalPipeline.cpp
+++ b/Engine/Source/Runtime/Apple/MetalRHI/Private/MetalPipeline.cpp
@@ -136,7 +136,7 @@ struct FMetalHelperFunctions
             DebugFunc = DebugShadersLib.NewFunction(@"Main_Debug");
 			
 			DebugComputeShadersLib = GetMetalDeviceContext().GetDevice().NewLibrary(GMetalDebugMarkerComputeShader, CompileOptions, &Error);
-			DebugComputeFunc = DebugShadersLib.NewFunction(@"Main_Debug");
+			DebugComputeFunc = DebugComputeShadersLib.NewFunction(@"Main_Debug");
 			
 			DebugComputeState = GetMetalDeviceContext().GetDevice().NewComputePipelineState(DebugComputeFunc, &Error);
         }
@@ -810,7 +810,7 @@ static FMetalShaderPipeline* CreateMTLRenderPipeline(bool const bSync, FMetalGra
     #if PLATFORM_MAC
         RenderPipelineDesc.SetInputPrimitiveTopology(TranslatePrimitiveTopology(Init.PrimitiveType));
 		DebugPipelineDesc.SetSampleCount(!bNoMSAA ? FMath::Max(Init.NumSamples, (uint16)1u) : (uint16)1u);
-		DebugPipelineDesc.SetInputPrimitiveTopology(TranslatePrimitiveTopology(Init.PrimitiveType));
+		DebugPipelineDesc.SetInputPrimitiveTopology(mtlpp::PrimitiveTopologyClass::Point);
     #endif
         
         FMetalVertexDeclaration* VertexDecl = (FMetalVertexDeclaration*)Init.BoundShaderState.VertexDeclarationRHI;
