Change to FMetalShaderPipelineCache::GetRenderPipeline to allow it to create multiple pipelines simultaneously

diff --git a/Engine/Source/Runtime/Apple/MetalRHI/Private/MetalPipeline.cpp b/Engine/Source/Runtime/Apple/MetalRHI/Private/MetalPipeline.cpp
--- a/Engine/Source/Runtime/Apple/MetalRHI/Private/MetalPipeline.cpp
+++ b/Engine/Source/Runtime/Apple/MetalRHI/Private/MetalPipeline.cpp
@@ -13,6 +13,7 @@
 #include "MetalCommandBuffer.h"
 #include "RenderUtils.h"
 #include "Misc/ScopeRWLock.h"
+#include "HAL/PThreadEvent.h"
 #include <objc/runtime.h>
 
 static int32 GMetalCacheShaderPipelines = 1;
@@ -334,37 +335,80 @@ public:
 		InitMetalGraphicsPipelineKey(Key, Init, IndexType);
 		
 		// By default there'll be more threads trying to read this than to write it.
-		FRWScopeLock Lock(PipelineMutex, SLT_ReadOnly);
-		
+		PipelineMutex.ReadLock();
+
 		// Try to find the entry in the cache.
 		FMetalShaderPipeline* Desc = Pipelines.FindRef(Key);
+
+		PipelineMutex.ReadUnlock();
+
 		if (Desc == nil)
 		{
-			Desc = CreateMTLRenderPipeline(bSync, Key, Init, IndexType);
-			
-			// Bail cleanly if compilation fails.
-			if(!Desc)
+
+			// By default there'll be more threads trying to read this than to write it.
+			EventsMutex.ReadLock();
+
+			// Try to find a pipeline creation event for this key. If it's found, we already have a thread creating this pipeline and we just have to wait.
+			TSharedPtr<FPThreadEvent, ESPMode::ThreadSafe> Event = PipelineEvents.FindRef(Key);
+
+			EventsMutex.ReadUnlock();
+
+			bool bCompile = false;
+			if (!Event.IsValid())
 			{
-				return nil;
+				// Create an event other threads can use to wait if they request the same pipeline this thread is creating
+				EventsMutex.WriteLock();
+
+				Event = PipelineEvents.FindRef(Key);
+				if (!Event.IsValid())
+				{
+					Event = PipelineEvents.Add(Key, MakeShareable(new FPThreadEvent()));
+					Event->Create(true);
+					bCompile = true;
+				}
+				check(Event.IsValid());
+
+				EventsMutex.WriteUnlock();
 			}
-			
-			// Now we are a writer as we want to create & add the new pipeline
-			Lock.ReleaseReadOnlyLockAndAcquireWriteLock_USE_WITH_CAUTION();
-			
-			// Retest to ensure no-one beat us here!
-			if (Pipelines.FindRef(Key) == nil)
+
+			if (bCompile)
 			{
-				Pipelines.Add(Key, Desc);
-				ReverseLookup.Add(Desc, Key);
-				
-				if (GMetalCacheShaderPipelines == 0)
+				Desc = CreateMTLRenderPipeline(bSync, Key, Init, IndexType);
+
+				if (Desc != nil)
 				{
-					// When we aren't caching for program lifetime we autorelease so that the PSO is released to the OS once all RHI references are released.
-					[Desc autorelease];
+					PipelineMutex.WriteLock();
+
+					Pipelines.Add(Key, Desc);
+					ReverseLookup.Add(Desc, Key);
+
+					PipelineMutex.WriteUnlock();
+
+					if (GMetalCacheShaderPipelines == 0)
+					{
+						// When we aren't caching for program lifetime we autorelease so that the PSO is released to the OS once all RHI references are released.
+						[Desc autorelease];
+					}
 				}
+
+				EventsMutex.WriteLock();
+
+				Event->Trigger();
+				PipelineEvents.Remove(Key);
+
+				EventsMutex.WriteUnlock();
+			}
+			else
+			{
+				check(Event.IsValid());
+				Event->Wait();
+
+				PipelineMutex.ReadLock();
+				Desc = Pipelines.FindRef(Key);
+				PipelineMutex.ReadUnlock();
+				check(Desc);
 			}
 		}
-		check(Desc);
 		
 		return Desc;
 	}
@@ -399,8 +443,10 @@ public:
 	
 private:
 	FRWLock PipelineMutex;
+	FRWLock EventsMutex;
 	TMap<FMetalGraphicsPipelineKey, FMetalShaderPipeline*> Pipelines;
 	TMap<FMetalShaderPipeline*, FMetalGraphicsPipelineKey> ReverseLookup;
+	TMap<FMetalGraphicsPipelineKey, TSharedPtr<FPThreadEvent, ESPMode::ThreadSafe>> PipelineEvents;
 };
 
 @implementation FMetalShaderPipeline
