Fix handling of visibility queries in the GPU debugger - under validation we can't actually track occlusion draw calls because we might cause it to bind the same vis-offset more than once which is technically incorrect.

diff --git a/Engine/Source/Runtime/Apple/MetalRHI/Private/MetalRenderPass.cpp b/Engine/Source/Runtime/Apple/MetalRHI/Private/MetalRenderPass.cpp
--- a/Engine/Source/Runtime/Apple/MetalRHI/Private/MetalRenderPass.cpp
+++ b/Engine/Source/Runtime/Apple/MetalRHI/Private/MetalRenderPass.cpp
@@ -1982,7 +1982,7 @@ uint32 FMetalRenderPass::GetCommandBufferIndex(void) const
 void FMetalRenderPass::InsertDebugDraw(FMetalCommandData& Data)
 {
 #if !PLATFORM_TVOS
-	if (GMetalCommandBufferDebuggingEnabled)
+	if (GMetalCommandBufferDebuggingEnabled && (!FMetalCommandQueue::SupportsFeature(EMetalFeaturesValidation) || State.GetVisibilityResultMode() == mtlpp::VisibilityResultMode::Disabled))
 	{
 		FMetalGraphicsPipelineState* BoundShaderState = State.GetGraphicsPSO();
 		
@@ -2013,6 +2013,13 @@ void FMetalRenderPass::InsertDebugDraw(FMetalCommandData& Data)
 			
 			CurrentEncoder.GetRenderCommandEncoder().SetRenderPipelineState(PSO->DebugPipelineState);
 
+			mtlpp::VisibilityResultMode VisMode = State.GetVisibilityResultMode();
+			uint32 VisibilityOffset = State.GetVisibilityResultOffset();
+			if (VisMode != mtlpp::VisibilityResultMode::Disabled)
+			{
+				CurrentEncoder.GetRenderCommandEncoder().SetVisibilityResultMode(mtlpp::VisibilityResultMode::Disabled, 0);
+			}
+
 		#if PLATFORM_MAC
 			id<MTLBuffer> DebugBufferPtr = State.GetDebugBuffer().GetPtr();
 			[(id<IMTLRenderCommandEncoder>)CurrentEncoder.GetRenderCommandEncoder().GetPtr() memoryBarrierWithResources:&DebugBufferPtr count:1 afterStages:MTLRenderStageFragment beforeStages:MTLRenderStageVertex];
@@ -2034,6 +2041,11 @@ void FMetalRenderPass::InsertDebugDraw(FMetalCommandData& Data)
 		#endif
 
 			CurrentEncoder.GetRenderCommandEncoder().SetRenderPipelineState(PSO->RenderPipelineState);
+
+			if (VisMode != mtlpp::VisibilityResultMode::Disabled)
+			{
+				CurrentEncoder.GetRenderCommandEncoder().SetVisibilityResultMode(VisMode, VisibilityOffset);
+			}
 		}
 	}
 #endif
diff --git a/Engine/Source/Runtime/Apple/MetalRHI/Private/MetalStateCache.h b/Engine/Source/Runtime/Apple/MetalRHI/Private/MetalStateCache.h
--- a/Engine/Source/Runtime/Apple/MetalRHI/Private/MetalStateCache.h
+++ b/Engine/Source/Runtime/Apple/MetalRHI/Private/MetalStateCache.h
@@ -148,6 +148,8 @@ public:
     bool IsLinearBuffer(EMetalShaderStages ShaderStage, uint32 BindIndex);
     FMetalShaderPipeline* GetPipelineState(void) const { return GraphicsPSO->GetPipeline(GetIndexType()); }
 	EPrimitiveType GetPrimitiveType() { check(IsValidRef(GraphicsPSO)); return GraphicsPSO->GetPrimitiveType(); }
+	mtlpp::VisibilityResultMode GetVisibilityResultMode() { return VisibilityMode; }
+	uint32 GetVisibilityResultOffset() { return VisibilityOffset; }
 	
 	FTexture2DRHIRef CreateFallbackDepthStencilSurface(uint32 Width, uint32 Height);
 	bool GetFallbackDepthStencilBound(void) const { return bFallbackDepthStencilBound; }
