Make sure that Metal allocation resource options are validated with calls to GetCompatibleResourceOptions to ensure that they match.

diff --git a/Engine/Source/Runtime/Apple/MetalRHI/Private/MetalBuffer.cpp b/Engine/Source/Runtime/Apple/MetalRHI/Private/MetalBuffer.cpp
--- a/Engine/Source/Runtime/Apple/MetalRHI/Private/MetalBuffer.cpp
+++ b/Engine/Source/Runtime/Apple/MetalRHI/Private/MetalBuffer.cpp
@@ -170,6 +170,7 @@ FMetalSubBufferHeap::FMetalSubBufferHeap(NSUInteger Size, NSUInteger Alignment,
 , MinAlign(Alignment)
 , UsedSize(0)
 {
+	Options = (mtlpp::ResourceOptions)GetMetalDeviceContext().GetCommandQueue().GetCompatibleResourceOptions(Options);
 	static bool bSupportsHeaps = GetMetalDeviceContext().SupportsFeature(EMetalFeaturesHeaps);
 	NSUInteger FullSize = Align(Size, Alignment);
 	METAL_GPUPROFILE(FScopedMetalCPUStats CPUStat(FString::Printf(TEXT("AllocBuffer: %llu, %llu"), FullSize, Options)));
@@ -540,6 +541,7 @@ FMetalSubBufferLinear::FMetalSubBufferLinear(NSUInteger Size, NSUInteger Alignme
 , UsedSize(0)
 , FreedSize(0)
 {
+	Options = (mtlpp::ResourceOptions)GetMetalDeviceContext().GetCommandQueue().GetCompatibleResourceOptions(Options);
 	NSUInteger FullSize = Align(Size, Alignment);
 	METAL_GPUPROFILE(FScopedMetalCPUStats CPUStat(FString::Printf(TEXT("AllocBuffer: %llu, %llu"), FullSize, Options)));
 	
@@ -670,6 +672,7 @@ FMetalSubBufferMagazine::FMetalSubBufferMagazine(NSUInteger Size, NSUInteger Chu
 , OutstandingAllocs(0)
 , UsedSize(0)
 {
+	Options = (mtlpp::ResourceOptions)GetMetalDeviceContext().GetCommandQueue().GetCompatibleResourceOptions(Options);
     static bool bSupportsHeaps = GetMetalDeviceContext().SupportsFeature(EMetalFeaturesHeaps);
     mtlpp::StorageMode Storage = (mtlpp::StorageMode)((Options & mtlpp::ResourceStorageModeMask) >> mtlpp::ResourceStorageModeShift);
     if (bSupportsHeaps && Storage == mtlpp::StorageMode::Private)
@@ -941,6 +944,7 @@ FMetalSubBufferRing::FMetalSubBufferRing(NSUInteger Size, NSUInteger Alignment,
 , Options(InOptions)
 , Storage((mtlpp::StorageMode)((Options & mtlpp::ResourceStorageModeMask) >> mtlpp::ResourceStorageModeShift))
 {
+	Options = (mtlpp::ResourceOptions)GetMetalDeviceContext().GetCommandQueue().GetCompatibleResourceOptions(Options);
 	check(Storage != mtlpp::StorageMode::Private /* Private memory requires command-buffers and encoders to properly marshal! */);
 	FMemory::Memzero(FrameSize);
 }
