To avoid initialisation order issues and to make it faster in general make GetCompatibleResourceOptions static so that it can be accessed directly.

diff --git a/Engine/Source/Runtime/Apple/MetalRHI/Private/MetalBuffer.cpp b/Engine/Source/Runtime/Apple/MetalRHI/Private/MetalBuffer.cpp
--- a/Engine/Source/Runtime/Apple/MetalRHI/Private/MetalBuffer.cpp
+++ b/Engine/Source/Runtime/Apple/MetalRHI/Private/MetalBuffer.cpp
@@ -170,7 +170,7 @@ FMetalSubBufferHeap::FMetalSubBufferHeap(NSUInteger Size, NSUInteger Alignment,
 , MinAlign(Alignment)
 , UsedSize(0)
 {
-	Options = (mtlpp::ResourceOptions)GetMetalDeviceContext().GetCommandQueue().GetCompatibleResourceOptions(Options);
+	Options = (mtlpp::ResourceOptions)FMetalCommandQueue::GetCompatibleResourceOptions(Options);
 	static bool bSupportsHeaps = GetMetalDeviceContext().SupportsFeature(EMetalFeaturesHeaps);
 	NSUInteger FullSize = Align(Size, Alignment);
 	METAL_GPUPROFILE(FScopedMetalCPUStats CPUStat(FString::Printf(TEXT("AllocBuffer: %llu, %llu"), FullSize, Options)));
@@ -541,7 +541,7 @@ FMetalSubBufferLinear::FMetalSubBufferLinear(NSUInteger Size, NSUInteger Alignme
 , UsedSize(0)
 , FreedSize(0)
 {
-	Options = (mtlpp::ResourceOptions)GetMetalDeviceContext().GetCommandQueue().GetCompatibleResourceOptions(Options);
+	Options = (mtlpp::ResourceOptions)FMetalCommandQueue::GetCompatibleResourceOptions(Options);
 	NSUInteger FullSize = Align(Size, Alignment);
 	METAL_GPUPROFILE(FScopedMetalCPUStats CPUStat(FString::Printf(TEXT("AllocBuffer: %llu, %llu"), FullSize, Options)));
 	
@@ -672,7 +672,7 @@ FMetalSubBufferMagazine::FMetalSubBufferMagazine(NSUInteger Size, NSUInteger Chu
 , OutstandingAllocs(0)
 , UsedSize(0)
 {
-	Options = (mtlpp::ResourceOptions)GetMetalDeviceContext().GetCommandQueue().GetCompatibleResourceOptions(Options);
+	Options = (mtlpp::ResourceOptions)FMetalCommandQueue::GetCompatibleResourceOptions(Options);
     static bool bSupportsHeaps = GetMetalDeviceContext().SupportsFeature(EMetalFeaturesHeaps);
     mtlpp::StorageMode Storage = (mtlpp::StorageMode)((Options & mtlpp::ResourceStorageModeMask) >> mtlpp::ResourceStorageModeShift);
     if (bSupportsHeaps && Storage == mtlpp::StorageMode::Private)
@@ -944,7 +944,7 @@ FMetalSubBufferRing::FMetalSubBufferRing(NSUInteger Size, NSUInteger Alignment,
 , Options(InOptions)
 , Storage((mtlpp::StorageMode)((Options & mtlpp::ResourceStorageModeMask) >> mtlpp::ResourceStorageModeShift))
 {
-	Options = (mtlpp::ResourceOptions)GetMetalDeviceContext().GetCommandQueue().GetCompatibleResourceOptions(Options);
+	Options = (mtlpp::ResourceOptions)FMetalCommandQueue::GetCompatibleResourceOptions(Options);
 	check(Storage != mtlpp::StorageMode::Private /* Private memory requires command-buffers and encoders to properly marshal! */);
 	FMemory::Memzero(FrameSize);
 }
@@ -1195,7 +1195,7 @@ FMetalBuffer FMetalBufferPoolPolicyData::CreateResource(CreationArguments Args)
 	check(Args.Device);	
 	uint32 BufferSize = GetPoolBucketSize(GetPoolBucketIndex(Args));
 	METAL_GPUPROFILE(FScopedMetalCPUStats CPUStat(FString::Printf(TEXT("AllocBuffer: %llu, %llu"), BufferSize, mtlpp::ResourceOptions(BUFFER_CACHE_MODE | ((NSUInteger)Args.Storage << mtlpp::ResourceStorageModeShift)))));
-	FMetalBuffer NewBuf(MTLPP_VALIDATE(mtlpp::Device, Args.Device, SafeGetRuntimeDebuggingLevel() >= EMetalDebugLevelValidation, NewBuffer(BufferSize, GetMetalDeviceContext().GetCommandQueue().GetCompatibleResourceOptions(mtlpp::ResourceOptions(BUFFER_CACHE_MODE | mtlpp::ResourceOptions::HazardTrackingModeUntracked | ((NSUInteger)Args.Storage << mtlpp::ResourceStorageModeShift))))), true);
+	FMetalBuffer NewBuf(MTLPP_VALIDATE(mtlpp::Device, Args.Device, SafeGetRuntimeDebuggingLevel() >= EMetalDebugLevelValidation, NewBuffer(BufferSize, FMetalCommandQueue::GetCompatibleResourceOptions(mtlpp::ResourceOptions(BUFFER_CACHE_MODE | mtlpp::ResourceOptions::HazardTrackingModeUntracked | ((NSUInteger)Args.Storage << mtlpp::ResourceStorageModeShift))))), true);
 #if STATS || ENABLE_LOW_LEVEL_MEM_TRACKER
 	MetalLLM::LogAllocBuffer(Args.Device, NewBuf);
 #endif
diff --git a/Engine/Source/Runtime/Apple/MetalRHI/Private/MetalCommandEncoder.cpp b/Engine/Source/Runtime/Apple/MetalRHI/Private/MetalCommandEncoder.cpp
--- a/Engine/Source/Runtime/Apple/MetalRHI/Private/MetalCommandEncoder.cpp
+++ b/Engine/Source/Runtime/Apple/MetalRHI/Private/MetalCommandEncoder.cpp
@@ -212,7 +212,7 @@ FMetalCommandBufferMarkers FMetalCommandBufferMarkers::Get(mtlpp::CommandBuffer
 FMetalCommandEncoder::FMetalCommandEncoder(FMetalCommandList& CmdList)
 : CommandList(CmdList)
 , bSupportsMetalFeaturesSetBytes(CmdList.GetCommandQueue().SupportsFeature(EMetalFeaturesSetBytes))
-, RingBuffer(EncoderRingBufferSize, BufferOffsetAlignment, CmdList.GetCommandQueue().GetCompatibleResourceOptions((mtlpp::ResourceOptions)(mtlpp::ResourceOptions::HazardTrackingModeUntracked | BUFFER_RESOURCE_STORAGE_MANAGED)))
+, RingBuffer(EncoderRingBufferSize, BufferOffsetAlignment, FMetalCommandQueue::GetCompatibleResourceOptions((mtlpp::ResourceOptions)(mtlpp::ResourceOptions::HazardTrackingModeUntracked | BUFFER_RESOURCE_STORAGE_MANAGED)))
 , RenderPassDesc(nil)
 , EncoderFence(nil)
 #if ENABLE_METAL_GPUPROFILE
diff --git a/Engine/Source/Runtime/Apple/MetalRHI/Private/MetalCommandQueue.h b/Engine/Source/Runtime/Apple/MetalRHI/Private/MetalCommandQueue.h
--- a/Engine/Source/Runtime/Apple/MetalRHI/Private/MetalCommandQueue.h
+++ b/Engine/Source/Runtime/Apple/MetalRHI/Private/MetalCommandQueue.h
@@ -160,7 +160,7 @@ public:
 	mtlpp::Device& GetDevice(void);
 	
 	/** Converts a Metal v1.1+ resource option to something valid on the current version. */
-	mtlpp::ResourceOptions GetCompatibleResourceOptions(mtlpp::ResourceOptions Options) const;
+	static mtlpp::ResourceOptions GetCompatibleResourceOptions(mtlpp::ResourceOptions Options);
 	
 	/**
 	 * @param InFeature A specific Metal feature to check for.
@@ -204,6 +204,6 @@ private:
 	TLockFreePointerListLIFO<mtlpp::CommandBufferFence> CommandBufferFences;
 	uint64 ParallelCommandLists;
 	int32 RuntimeDebuggingLevel;
-	NSUInteger PermittedOptions;
+	static NSUInteger PermittedOptions;
 	static uint64 Features;
 };
diff --git a/Engine/Source/Runtime/Apple/MetalRHI/Private/MetalContext.cpp b/Engine/Source/Runtime/Apple/MetalRHI/Private/MetalContext.cpp
--- a/Engine/Source/Runtime/Apple/MetalRHI/Private/MetalContext.cpp
+++ b/Engine/Source/Runtime/Apple/MetalRHI/Private/MetalContext.cpp
@@ -778,7 +778,7 @@ FMetalTexture FMetalDeviceContext::CreateTexture(FMetalSurface* Surface, mtlpp::
 
 FMetalBuffer FMetalDeviceContext::CreatePooledBuffer(FMetalPooledBufferArgs const& Args)
 {
-    FMetalBuffer Buffer = Heap.CreateBuffer(Args.Size, BufferOffsetAlignment, Args.Flags, GetCommandQueue().GetCompatibleResourceOptions((mtlpp::ResourceOptions)(BUFFER_CACHE_MODE | mtlpp::ResourceOptions::HazardTrackingModeUntracked | ((NSUInteger)Args.Storage << mtlpp::ResourceStorageModeShift))));
+    FMetalBuffer Buffer = Heap.CreateBuffer(Args.Size, BufferOffsetAlignment, Args.Flags, FMetalCommandQueue::GetCompatibleResourceOptions((mtlpp::ResourceOptions)(BUFFER_CACHE_MODE | mtlpp::ResourceOptions::HazardTrackingModeUntracked | ((NSUInteger)Args.Storage << mtlpp::ResourceStorageModeShift))));
 	check(Buffer && Buffer.GetPtr());
 #if METAL_DEBUG_OPTIONS
 	if (GMetalResourcePurgeOnDelete && !Buffer.GetHeap())
diff --git a/Engine/Source/Runtime/Apple/MetalRHI/Private/MetalQuery.cpp b/Engine/Source/Runtime/Apple/MetalRHI/Private/MetalQuery.cpp
--- a/Engine/Source/Runtime/Apple/MetalRHI/Private/MetalQuery.cpp
+++ b/Engine/Source/Runtime/Apple/MetalRHI/Private/MetalQuery.cpp
@@ -57,7 +57,7 @@ FMetalQueryBuffer* FMetalQueryBufferPool::GetCurrentQueryBuffer()
 			LLM_PLATFORM_SCOPE_METAL(ELLMTagMetal::Buffers);
 
 			METAL_GPUPROFILE(FScopedMetalCPUStats CPUStat(FString::Printf(TEXT("AllocBuffer: %llu, %llu"), EQueryBufferMaxSize, mtlpp::ResourceOptions::StorageModeShared)));
-            Buffer = GetMetalDeviceContext().GetResourceHeap().CreateBuffer(EQueryBufferMaxSize, 16, BUF_Dynamic, Context->GetCommandQueue().GetCompatibleResourceOptions((mtlpp::ResourceOptions)(BUFFER_CACHE_MODE | mtlpp::ResourceOptions::HazardTrackingModeUntracked | mtlpp::ResourceOptions::StorageModeShared)), true);
+            Buffer = GetMetalDeviceContext().GetResourceHeap().CreateBuffer(EQueryBufferMaxSize, 16, BUF_Dynamic, FMetalCommandQueue::GetCompatibleResourceOptions((mtlpp::ResourceOptions)(BUFFER_CACHE_MODE | mtlpp::ResourceOptions::HazardTrackingModeUntracked | mtlpp::ResourceOptions::StorageModeShared)), true);
 			FMemory::Memzero((((uint8*)Buffer.GetContents())), EQueryBufferMaxSize);
 
 #if STATS || ENABLE_LOW_LEVEL_MEM_TRACKER
diff --git a/Engine/Source/Runtime/Apple/MetalRHI/Private/MetalRenderTarget.cpp b/Engine/Source/Runtime/Apple/MetalRHI/Private/MetalRenderTarget.cpp
--- a/Engine/Source/Runtime/Apple/MetalRHI/Private/MetalRenderTarget.cpp
+++ b/Engine/Source/Runtime/Apple/MetalRHI/Private/MetalRenderTarget.cpp
@@ -633,7 +633,7 @@ void FMetalDynamicRHI::RHIReadSurfaceData(FTextureRHIParamRef TextureRHI, FIntRe
 			Desc.SetSampleCount(Texture.GetSampleCount());
 			Desc.SetArrayLength(Texture.GetArrayLength());
 			
-			mtlpp::ResourceOptions GeneralResourceOption = (mtlpp::ResourceOptions)GetMetalDeviceContext().GetCommandQueue().GetCompatibleResourceOptions(mtlpp::ResourceOptions(((NSUInteger)Texture.GetCpuCacheMode() << mtlpp::ResourceCpuCacheModeShift) | ((NSUInteger)StorageMode << mtlpp::ResourceStorageModeShift) | mtlpp::ResourceOptions::HazardTrackingModeUntracked));
+			mtlpp::ResourceOptions GeneralResourceOption = (mtlpp::ResourceOptions)FMetalCommandQueue::GetCompatibleResourceOptions(mtlpp::ResourceOptions(((NSUInteger)Texture.GetCpuCacheMode() << mtlpp::ResourceCpuCacheModeShift) | ((NSUInteger)StorageMode << mtlpp::ResourceStorageModeShift) | mtlpp::ResourceOptions::HazardTrackingModeUntracked));
 			Desc.SetResourceOptions(GeneralResourceOption);
 			
 			Desc.SetCpuCacheMode(Texture.GetCpuCacheMode());
diff --git a/Engine/Source/Runtime/Apple/MetalRHI/Private/MetalTexture.cpp b/Engine/Source/Runtime/Apple/MetalRHI/Private/MetalTexture.cpp
--- a/Engine/Source/Runtime/Apple/MetalRHI/Private/MetalTexture.cpp
+++ b/Engine/Source/Runtime/Apple/MetalRHI/Private/MetalTexture.cpp
@@ -264,7 +264,7 @@ FMetalTexture FMetalSurface::Reallocate(FMetalTexture InTexture, mtlpp::TextureU
 	Desc.SetSampleCount(InTexture.GetSampleCount());
 	Desc.SetArrayLength(InTexture.GetArrayLength());
 	
-	static mtlpp::ResourceOptions GeneralResourceOption = (mtlpp::ResourceOptions)GetMetalDeviceContext().GetCommandQueue().GetCompatibleResourceOptions(mtlpp::ResourceOptions::HazardTrackingModeUntracked);
+	static mtlpp::ResourceOptions GeneralResourceOption = (mtlpp::ResourceOptions)FMetalCommandQueue::GetCompatibleResourceOptions(mtlpp::ResourceOptions::HazardTrackingModeUntracked);
 	
 	Desc.SetResourceOptions(mtlpp::ResourceOptions(((NSUInteger)InTexture.GetCpuCacheMode() << mtlpp::ResourceCpuCacheModeShift) | ((NSUInteger)Texture.GetStorageMode() << mtlpp::ResourceStorageModeShift) | GeneralResourceOption));
 	Desc.SetCpuCacheMode(InTexture.GetCpuCacheMode());
@@ -904,7 +904,7 @@ FMetalSurface::FMetalSurface(ERHIResourceType ResourceType, EPixelFormat Format,
 #endif
 		}
 		
-		static mtlpp::ResourceOptions GeneralResourceOption = GetMetalDeviceContext().GetCommandQueue().GetCompatibleResourceOptions(mtlpp::ResourceOptions::HazardTrackingModeUntracked);
+		static mtlpp::ResourceOptions GeneralResourceOption = FMetalCommandQueue::GetCompatibleResourceOptions(mtlpp::ResourceOptions::HazardTrackingModeUntracked);
 		Desc.SetResourceOptions((mtlpp::ResourceOptions)(Desc.GetResourceOptions() | GeneralResourceOption));
 	}
 	
@@ -1251,7 +1251,7 @@ FMetalBuffer FMetalSurface::AllocSurface(uint32 MipIndex, uint32 ArrayIndex, ERe
 	uint32 MipBytes = GetMipSize(MipIndex, &DestStride, false);
 	
 	// allocate some temporary memory
-	mtlpp::ResourceOptions ResMode = GetMetalDeviceContext().GetCommandQueue().GetCompatibleResourceOptions((mtlpp::ResourceOptions)(mtlpp::ResourceOptions::StorageModeShared | (bSupportsResourceOptions && !(PLATFORM_MAC && PixelFormat == PF_G8 && (Flags & TexCreate_SRGB)) ? mtlpp::ResourceOptions::CpuCacheModeWriteCombined : 0)));
+	mtlpp::ResourceOptions ResMode = FMetalCommandQueue::GetCompatibleResourceOptions((mtlpp::ResourceOptions)(mtlpp::ResourceOptions::StorageModeShared | (bSupportsResourceOptions && !(PLATFORM_MAC && PixelFormat == PF_G8 && (Flags & TexCreate_SRGB)) ? mtlpp::ResourceOptions::CpuCacheModeWriteCombined : 0)));
 	
     FMetalBuffer Buffer = GetMetalDeviceContext().GetResourceHeap().CreateBuffer(MipBytes, BufferOffsetAlignment, BUF_Dynamic, ResMode);
 	
diff --git a/Engine/Source/Runtime/Apple/MetalRHI/Private/MetalVertexBuffer.cpp b/Engine/Source/Runtime/Apple/MetalRHI/Private/MetalVertexBuffer.cpp
--- a/Engine/Source/Runtime/Apple/MetalRHI/Private/MetalVertexBuffer.cpp
+++ b/Engine/Source/Runtime/Apple/MetalRHI/Private/MetalVertexBuffer.cpp
@@ -264,7 +264,7 @@ FMetalTexture FMetalRHIBuffer::AllocLinearTexture(EPixelFormat Format)
 		
 		mtlpp::TextureDescriptor Desc;
 		NSUInteger Mode = ((NSUInteger)Buffer.GetStorageMode() << mtlpp::ResourceStorageModeShift) | ((NSUInteger)Buffer.GetCpuCacheMode() << mtlpp::ResourceCpuCacheModeShift);
-		Mode = GetMetalDeviceContext().GetCommandQueue().GetCompatibleResourceOptions(mtlpp::ResourceOptions(Mode | mtlpp::ResourceOptions::HazardTrackingModeUntracked));
+		Mode = FMetalCommandQueue::GetCompatibleResourceOptions(mtlpp::ResourceOptions(Mode | mtlpp::ResourceOptions::HazardTrackingModeUntracked));
 		NSUInteger TexUsage = mtlpp::TextureUsage::Unknown;
 		if (Usage & BUF_ShaderResource)
 		{
