#UE4 Mac CIS fix attempt

diff --git a/Engine/Source/Runtime/Apple/MetalRHI/Private/MetalCommandQueue.cpp b/Engine/Source/Runtime/Apple/MetalRHI/Private/MetalCommandQueue.cpp
--- a/Engine/Source/Runtime/Apple/MetalRHI/Private/MetalCommandQueue.cpp
+++ b/Engine/Source/Runtime/Apple/MetalRHI/Private/MetalCommandQueue.cpp
@@ -502,7 +502,7 @@ mtlpp::Device& FMetalCommandQueue::GetDevice(void)
 	return Device;
 }
 
-mtlpp::ResourceOptions FMetalCommandQueue::GetCompatibleResourceOptions(mtlpp::ResourceOptions Options) const
+mtlpp::ResourceOptions FMetalCommandQueue::GetCompatibleResourceOptions(mtlpp::ResourceOptions Options)
 {
 	NSUInteger NewOptions = (Options & PermittedOptions);
 #if PLATFORM_IOS // Swizzle Managed to Shared for iOS - we can do this as they are equivalent, unlike Shared -> Managed on Mac.
