Fix missing static member variable declaration.

diff --git a/Engine/Source/Runtime/Apple/MetalRHI/Private/MetalCommandQueue.cpp b/Engine/Source/Runtime/Apple/MetalRHI/Private/MetalCommandQueue.cpp
--- a/Engine/Source/Runtime/Apple/MetalRHI/Private/MetalCommandQueue.cpp
+++ b/Engine/Source/Runtime/Apple/MetalRHI/Private/MetalCommandQueue.cpp
@@ -17,6 +17,7 @@
 #include "command_buffer.hpp"
 
 #pragma mark - Private C++ Statics -
+NSUInteger FMetalCommandQueue::PermittedOptions = 0;
 uint64 FMetalCommandQueue::Features = 0;
 extern mtlpp::VertexFormat GMetalFColorVertexFormat;
 bool GMetalCommandBufferDebuggingEnabled = 0;
