Duplicate 5993626 //UE4/Dev-Rendering to Main (//UE4/Main)

Don't record the BufferSizes into the list of buffers to validate - we will always bind it but it happens too late to validate with everything else.

diff --git a/Engine/Source/Runtime/Apple/MetalRHI/Private/MetalPipeline.cpp b/Engine/Source/Runtime/Apple/MetalRHI/Private/MetalPipeline.cpp
--- a/Engine/Source/Runtime/Apple/MetalRHI/Private/MetalPipeline.cpp
+++ b/Engine/Source/Runtime/Apple/MetalRHI/Private/MetalPipeline.cpp
@@ -543,13 +543,15 @@ APPLE_PLATFORM_OBJECT_ALLOC_OVERRIDES(FMetalShaderPipeline)
 			case MTLArgumentTypeBuffer:
 			{
 				checkf(Arg.index < ML_MaxBuffers, TEXT("Metal buffer index exceeded!"));
-				ResourceMask[Frequency].BufferMask |= (1 << Arg.index);
-				
-				if(BufferDataSizes[Frequency].Num() < 31)
-					BufferDataSizes[Frequency].SetNumZeroed(31);
+				if (FString(Arg.name) != TEXT("BufferSizes"))
+				{
+					ResourceMask[Frequency].BufferMask |= (1 << Arg.index);
 				
-				BufferDataSizes[Frequency][Arg.index] = Arg.bufferDataSize;
+					if(BufferDataSizes[Frequency].Num() < 31)
+						BufferDataSizes[Frequency].SetNumZeroed(31);
 				
+					BufferDataSizes[Frequency][Arg.index] = Arg.bufferDataSize;
+				}
 				break;
 			}
 			case MTLArgumentTypeThreadgroupMemory:
