Duplicate 5993726 //UE4/Dev-Rendering to Main (//UE4/Main)

Don't validate the various additional buffers used for Metal tessellation either - they are also bound after other resources.

diff --git a/Engine/Source/Runtime/Apple/MetalRHI/Private/MetalStateCache.cpp b/Engine/Source/Runtime/Apple/MetalRHI/Private/MetalStateCache.cpp
--- a/Engine/Source/Runtime/Apple/MetalRHI/Private/MetalStateCache.cpp
+++ b/Engine/Source/Runtime/Apple/MetalRHI/Private/MetalStateCache.cpp
@@ -2142,11 +2142,19 @@ void FMetalStateCache::SetRenderPipelineState(FMetalCommandEncoder& CommandEncod
 				uint32 Index = __builtin_ctz(ComputeMask.BufferMask);
 				ComputeMask.BufferMask &= ~(1 << Index);
 				
-				FMetalBufferBinding const& Binding = ShaderBuffers[EMetalShaderStages::Vertex].Buffers[Index];
-				FMetalBufferBinding const& HullBinding = ShaderBuffers[EMetalShaderStages::Hull].Buffers[Index];
-				ensure(Binding.Buffer || Binding.Bytes || HullBinding.Buffer || HullBinding.Bytes);
-				ensure(MinComputeBufferSizes.Num() > Index);
-				ensure(Binding.Length >= MinComputeBufferSizes[Index] || HullBinding.Length >= MinComputeBufferSizes[Index]);
+				if (Pipeline->TessellationPipelineDesc.TessellationControlPointIndexBufferIndex != Index
+					&& Pipeline->TessellationPipelineDesc.TessellationIndexBufferIndex != Index
+					&& Pipeline->TessellationPipelineDesc.TessellationOutputControlPointBufferIndex != Index
+					&& Pipeline->TessellationPipelineDesc.TessellationPatchConstBufferIndex != Index
+					&& Pipeline->TessellationPipelineDesc.TessellationFactorBufferIndex != Index
+					&& Pipeline->TessellationPipelineDesc.TessellationPatchCountBufferIndex != Index)
+				{
+					FMetalBufferBinding const& Binding = ShaderBuffers[EMetalShaderStages::Vertex].Buffers[Index];
+					FMetalBufferBinding const& HullBinding = ShaderBuffers[EMetalShaderStages::Hull].Buffers[Index];
+					ensure(Binding.Buffer || Binding.Bytes || HullBinding.Buffer || HullBinding.Bytes);
+					ensure(MinComputeBufferSizes.Num() > Index);
+					ensure(Binding.Length >= MinComputeBufferSizes[Index] || HullBinding.Length >= MinComputeBufferSizes[Index]);
+				}
 			}
 #if PLATFORM_MAC
 			{
@@ -2198,10 +2206,14 @@ void FMetalStateCache::SetRenderPipelineState(FMetalCommandEncoder& CommandEncod
 			uint32 Index = __builtin_ctz(VertexMask.BufferMask);
 			VertexMask.BufferMask &= ~(1 << Index);
 			
-			FMetalBufferBinding const& Binding = ShaderBuffers[VertexStage].Buffers[Index];
-			ensure(Binding.Buffer || Binding.Bytes);
-			ensure(MinVertexBufferSizes.Num() > Index);
-			ensure(Binding.Length >= MinVertexBufferSizes[Index]);
+			if ((VertexStage == EMetalShaderStages::Vertex) || (Pipeline->TessellationPipelineDesc.TessellationInputPatchConstBufferIndex != Index
+				&& Pipeline->TessellationPipelineDesc.TessellationInputControlPointBufferIndex != Index))
+			{
+				FMetalBufferBinding const& Binding = ShaderBuffers[VertexStage].Buffers[Index];
+				ensure(Binding.Buffer || Binding.Bytes);
+				ensure(MinVertexBufferSizes.Num() > Index);
+				ensure(Binding.Length >= MinVertexBufferSizes[Index]);
+			}
 		}
 #if PLATFORM_MAC
 		{
