allow non-SM5 to use RHIThread for Metal
can be turned on with r.Metal.RHIThread=1 (off by default)

diff --git a/Engine/Source/Runtime/Apple/MetalRHI/Private/MetalRHI.cpp b/Engine/Source/Runtime/Apple/MetalRHI/Private/MetalRHI.cpp
--- a/Engine/Source/Runtime/Apple/MetalRHI/Private/MetalRHI.cpp
+++ b/Engine/Source/Runtime/Apple/MetalRHI/Private/MetalRHI.cpp
@@ -32,6 +32,16 @@ bool GMetalCommandBufferHasStartEndTimeAPI = false;
 
 FMetalBufferFormat GMetalBufferFormats[PF_MAX];
 
+static TAutoConsoleVariable<int32> CVarUseRHIThread(
+													TEXT("r.Metal.IOSRHIThread"),
+													0,
+													TEXT("Controls RHIThread usage for IOS:\n")
+													TEXT("\t0: No RHIThread.\n")
+													TEXT("\t1: Use RHIThread.\n")
+													TEXT("Default is 0."),
+													ECVF_Default | ECVF_RenderThreadSafe
+													);
+
 static void ValidateTargetedRHIFeatureLevelExists(EShaderPlatform Platform)
 {
 	bool bSupportsShaderPlatform = false;
@@ -150,9 +160,10 @@ FMetalDynamicRHI::FMetalDynamicRHI(ERHIFeatureLevel::Type RequestedFeatureLevel)
 
     bool bProjectSupportsMRTs = false;
     GConfig->GetBool(TEXT("/Script/IOSRuntimeSettings.IOSRuntimeSettings"), TEXT("bSupportsMetalMRT"), bProjectSupportsMRTs, GEngineIni);
-	
+
 	bool const bRequestedMetalMRT = ((RequestedFeatureLevel >= ERHIFeatureLevel::SM4) || (!bRequestedFeatureLevel && FParse::Param(FCommandLine::Get(),TEXT("metalmrt"))));
-	
+	bSupportsRHIThread = FParse::Param(FCommandLine::Get(),TEXT("rhithread"));
+
     // only allow GBuffers, etc on A8s (A7s are just not going to cut it)
     if (bProjectSupportsMRTs && bCanUseWideMRTs && bRequestedMetalMRT)
     {
@@ -164,8 +175,6 @@ FMetalDynamicRHI::FMetalDynamicRHI(ERHIFeatureLevel::Type RequestedFeatureLevel)
         GMaxRHIShaderPlatform = SP_METAL_MRT;
 #endif
 		GMaxRHIFeatureLevel = ERHIFeatureLevel::SM5;
-		
-		bSupportsRHIThread = FParse::Param(FCommandLine::Get(),TEXT("rhithread"));
     }
     else
 	{
@@ -422,6 +431,7 @@ FMetalDynamicRHI::FMetalDynamicRHI(ERHIFeatureLevel::Type RequestedFeatureLevel)
 	}
 	else
 	{
+		GRHISupportsRHIThread = bSupportsRHIThread || (CVarUseRHIThread.GetValueOnAnyThread() > 0);
 		GRHISupportsParallelRHIExecute = false;
 		GSupportsEfficientAsyncCompute = false;
 		GSupportsParallelOcclusionQueries = false;
