changes to get mesh streaming working on IOS, checked in disabled

diff --git a/Engine/Source/Runtime/Apple/MetalRHI/Private/MetalBuffer.cpp b/Engine/Source/Runtime/Apple/MetalRHI/Private/MetalBuffer.cpp
--- a/Engine/Source/Runtime/Apple/MetalRHI/Private/MetalBuffer.cpp
+++ b/Engine/Source/Runtime/Apple/MetalRHI/Private/MetalBuffer.cpp
@@ -1697,6 +1697,7 @@ FMetalTexture FMetalResourceHeap::CreateTexture(mtlpp::TextureDescriptor Desc, F
 	}
 	else if (Desc.GetUsage() & mtlpp::TextureUsage::RenderTarget)
 	{
+		LLM_PLATFORM_SCOPE_METAL(ELLMTagMetal::RenderTargets);
 		return TargetPool.CreateTexture(Queue->GetDevice(), Desc);
 	}
 	else
diff --git a/Engine/Source/Runtime/Apple/MetalRHI/Private/MetalIndexBuffer.cpp b/Engine/Source/Runtime/Apple/MetalRHI/Private/MetalIndexBuffer.cpp
--- a/Engine/Source/Runtime/Apple/MetalRHI/Private/MetalIndexBuffer.cpp
+++ b/Engine/Source/Runtime/Apple/MetalRHI/Private/MetalIndexBuffer.cpp
@@ -40,9 +40,23 @@ FMetalIndexBuffer::~FMetalIndexBuffer()
 {
 }
 
-FIndexBufferRHIRef FMetalDynamicRHI::RHICreateIndexBuffer(uint32 Stride,uint32 Size, uint32 InUsage, FRHIResourceCreateInfo& CreateInfo)
+void FMetalIndexBuffer::Swap(FMetalIndexBuffer& Other)
 {
 	@autoreleasepool {
+	FRHIIndexBuffer::Swap(Other);
+	FMetalRHIBuffer::Swap(Other);
+	::Swap(IndexType, Other.IndexType);
+	}
+}
+
+FIndexBufferRHIRef FMetalDynamicRHI::RHICreateIndexBuffer(uint32 Stride, uint32 Size, uint32 InUsage, FRHIResourceCreateInfo& CreateInfo)
+{
+	@autoreleasepool {
+	if (CreateInfo.bWithoutNativeResource)
+	{
+		return new FMetalIndexBuffer(2, 0, 0);
+	}
+		
 	// make the RHI object, which will allocate memory
 	FMetalIndexBuffer* IndexBuffer = new FMetalIndexBuffer(Stride, Size, InUsage);
 	
@@ -85,6 +99,25 @@ FIndexBufferRHIRef FMetalDynamicRHI::RHICreateIndexBuffer(uint32 Stride,uint32 S
 	}
 }
 
+void FMetalDynamicRHI::RHITransferIndexBufferUnderlyingResource(FIndexBufferRHIParamRef DestIndexBuffer, FIndexBufferRHIParamRef SrcIndexBuffer)
+{
+	@autoreleasepool {
+	check(DestIndexBuffer);
+	FMetalIndexBuffer* Dest = ResourceCast(DestIndexBuffer);
+	if (!SrcIndexBuffer)
+	{
+		FRHIResourceCreateInfo CreateInfo;
+		TRefCountPtr<FMetalIndexBuffer> DeletionProxy = new FMetalIndexBuffer(2, 0, 0);
+		Dest->Swap(*DeletionProxy);
+	}
+	else
+	{
+		FMetalIndexBuffer* Src = ResourceCast(SrcIndexBuffer);
+		Dest->Swap(*Src);
+	}
+	}
+}
+
 void* FMetalDynamicRHI::RHILockIndexBuffer(FIndexBufferRHIParamRef IndexBufferRHI, uint32 Offset, uint32 Size, EResourceLockMode LockMode)
 {
 	@autoreleasepool {
@@ -140,6 +173,11 @@ struct FMetalRHICommandInitialiseIndexBuffer : public FRHICommand<FMetalRHIComma
 FIndexBufferRHIRef FMetalDynamicRHI::CreateIndexBuffer_RenderThread(class FRHICommandListImmediate& RHICmdList, uint32 Stride, uint32 Size, uint32 InUsage, FRHIResourceCreateInfo& CreateInfo)
 {
 	@autoreleasepool {
+		if (CreateInfo.bWithoutNativeResource)
+		{
+			return new FMetalIndexBuffer(2, 0, 0);
+		}
+		
 		// make the RHI object, which will allocate memory
 		TRefCountPtr<FMetalIndexBuffer> IndexBuffer = new FMetalIndexBuffer(Stride, Size, InUsage);
 		
diff --git a/Engine/Source/Runtime/Apple/MetalRHI/Private/MetalUAV.cpp b/Engine/Source/Runtime/Apple/MetalRHI/Private/MetalUAV.cpp
--- a/Engine/Source/Runtime/Apple/MetalRHI/Private/MetalUAV.cpp
+++ b/Engine/Source/Runtime/Apple/MetalRHI/Private/MetalUAV.cpp
@@ -319,6 +319,17 @@ FShaderResourceViewRHIRef FMetalDynamicRHI::RHICreateShaderResourceView(FStructu
 FShaderResourceViewRHIRef FMetalDynamicRHI::RHICreateShaderResourceView(FVertexBufferRHIParamRef VertexBufferRHI, uint32 Stride, uint8 Format)
 {
 	@autoreleasepool {
+	if (!VertexBufferRHI)
+	{
+		FMetalShaderResourceView* SRV = new FMetalShaderResourceView;
+		SRV->SourceVertexBuffer = nullptr;
+		SRV->TextureView = nullptr;
+		SRV->SourceIndexBuffer = nullptr;
+		SRV->SourceStructuredBuffer = nullptr;
+		SRV->Format = Format;
+		SRV->Stride = 0;
+		return SRV;
+	}
 	FMetalVertexBuffer* VertexBuffer = ResourceCast(VertexBufferRHI);
 		
 	FMetalShaderResourceView* SRV = new FMetalShaderResourceView;
@@ -344,6 +355,18 @@ FShaderResourceViewRHIRef FMetalDynamicRHI::RHICreateShaderResourceView(FVertexB
 FShaderResourceViewRHIRef FMetalDynamicRHI::RHICreateShaderResourceView(FIndexBufferRHIParamRef BufferRHI)
 {
 	@autoreleasepool {
+	if (!BufferRHI)
+	{
+		FMetalShaderResourceView* SRV = new FMetalShaderResourceView;
+		SRV->SourceVertexBuffer = nullptr;
+		SRV->TextureView = nullptr;
+		SRV->SourceIndexBuffer = nullptr;
+		SRV->SourceStructuredBuffer = nullptr;
+		SRV->Format = PF_R16_UINT;
+		SRV->Stride = 0;
+		return SRV;
+	}
+
 	FMetalIndexBuffer* Buffer = ResourceCast(BufferRHI);
 		
 	FMetalShaderResourceView* SRV = new FMetalShaderResourceView;
@@ -492,6 +515,56 @@ FShaderResourceViewRHIRef FMetalDynamicRHI::RHICreateShaderResourceView(FTexture
 	}
 }
 
+void FMetalDynamicRHI::RHIUpdateShaderResourceView(FShaderResourceViewRHIParamRef SRVRHI, FVertexBufferRHIParamRef VertexBufferRHI, uint32 Stride, uint8 Format)
+{
+	check(SRVRHI);
+	FMetalShaderResourceView* SRV = ResourceCast(SRVRHI);
+	if (!VertexBufferRHI)
+	{
+		SRV->SourceVertexBuffer = nullptr;
+		SRV->TextureView = nullptr;
+		SRV->SourceIndexBuffer = nullptr;
+		SRV->SourceStructuredBuffer = nullptr;
+		SRV->Format = Format;
+		SRV->Stride = Stride;
+	}
+	else if (SRV->SourceVertexBuffer != VertexBufferRHI)
+	{
+		FMetalVertexBuffer* VertexBuffer = ResourceCast(VertexBufferRHI);
+		SRV->SourceVertexBuffer = VertexBuffer;
+		SRV->TextureView = nullptr;
+		SRV->SourceIndexBuffer = nullptr;
+		SRV->SourceStructuredBuffer = nullptr;
+		SRV->Format = Format;
+		SRV->Stride = Stride;
+	}
+}
+
+void FMetalDynamicRHI::RHIUpdateShaderResourceView(FShaderResourceViewRHIParamRef SRVRHI, FIndexBufferRHIParamRef IndexBufferRHI)
+{
+	check(SRVRHI);
+	FMetalShaderResourceView* SRV = ResourceCast(SRVRHI);
+	if (!IndexBufferRHI)
+	{
+		SRV->SourceVertexBuffer = nullptr;
+		SRV->TextureView = nullptr;
+		SRV->SourceIndexBuffer = nullptr;
+		SRV->SourceStructuredBuffer = nullptr;
+		SRV->Format = PF_R16_UINT;
+		SRV->Stride = 0;
+	}
+	else if (SRV->SourceIndexBuffer != IndexBufferRHI)
+	{
+		FMetalIndexBuffer* IndexBuffer = ResourceCast(IndexBufferRHI);
+		SRV->SourceVertexBuffer = nullptr;
+		SRV->TextureView = nullptr;
+		SRV->SourceIndexBuffer = IndexBuffer;
+		SRV->SourceStructuredBuffer = nullptr;
+		SRV->Format = (IndexBuffer->IndexType == mtlpp::IndexType::UInt16) ? PF_R16_UINT : PF_R32_UINT;
+		SRV->Stride = 0;
+	}
+}
+
 void FMetalRHICommandContext::RHIClearTinyUAV(FUnorderedAccessViewRHIParamRef UnorderedAccessViewRHI, const uint32* Values)
 {
 	@autoreleasepool {
diff --git a/Engine/Source/Runtime/Apple/MetalRHI/Private/MetalVertexBuffer.cpp b/Engine/Source/Runtime/Apple/MetalRHI/Private/MetalVertexBuffer.cpp
--- a/Engine/Source/Runtime/Apple/MetalRHI/Private/MetalVertexBuffer.cpp
+++ b/Engine/Source/Runtime/Apple/MetalRHI/Private/MetalVertexBuffer.cpp
@@ -86,6 +86,17 @@ FMetalVertexBuffer::~FMetalVertexBuffer()
 {
 }
 
+void FMetalVertexBuffer::Swap(FMetalVertexBuffer& Other)
+{
+	FRHIVertexBuffer::Swap(Other);
+	FMetalRHIBuffer::Swap(Other);
+}
+
+void FMetalRHIBuffer::Swap(FMetalRHIBuffer& Other)
+{
+	::Swap(*this, Other);
+}
+
 bool FMetalRHIBuffer::UsePrivateMemory() const
 {
 	return (FMetalCommandQueue::SupportsFeature(EMetalFeaturesEfficientBufferBlits) && (Usage & (BUF_Dynamic|BUF_Static)))
@@ -527,6 +538,11 @@ void FMetalRHIBuffer::Unlock()
 FVertexBufferRHIRef FMetalDynamicRHI::RHICreateVertexBuffer(uint32 Size, uint32 InUsage, FRHIResourceCreateInfo& CreateInfo)
 {
 	@autoreleasepool {
+	if (CreateInfo.bWithoutNativeResource)
+	{
+		return new FMetalVertexBuffer(0, 0);
+	}
+	
 	// make the RHI object, which will allocate memory
 	FMetalVertexBuffer* VertexBuffer = new FMetalVertexBuffer(Size, InUsage);
 
@@ -648,6 +664,11 @@ struct FMetalRHICommandInitialiseVertexBuffer : public FRHICommand<FMetalRHIComm
 FVertexBufferRHIRef FMetalDynamicRHI::CreateVertexBuffer_RenderThread(class FRHICommandListImmediate& RHICmdList, uint32 Size, uint32 InUsage, FRHIResourceCreateInfo& CreateInfo)
 {
 	@autoreleasepool {
+		if (CreateInfo.bWithoutNativeResource)
+		{
+			return new FMetalVertexBuffer(0, 0);
+		}
+		
 		// make the RHI object, which will allocate memory
 		TRefCountPtr<FMetalVertexBuffer> VertexBuffer = new FMetalVertexBuffer(Size, InUsage);
 		
@@ -722,6 +743,22 @@ FVertexBufferRHIRef FMetalDynamicRHI::CreateVertexBuffer_RenderThread(class FRHI
 	}
 }
 
+void FMetalDynamicRHI::RHITransferVertexBufferUnderlyingResource(FVertexBufferRHIParamRef DestVertexBuffer, FVertexBufferRHIParamRef SrcVertexBuffer)
+{
+	check(DestVertexBuffer);
+	FMetalVertexBuffer* Dest = ResourceCast(DestVertexBuffer);
+	if (!SrcVertexBuffer)
+	{
+		TRefCountPtr<FMetalVertexBuffer> DeletionProxy = new FMetalVertexBuffer(0, 0);
+		Dest->Swap(*DeletionProxy);
+	}
+	else
+	{
+		FMetalVertexBuffer* Src = ResourceCast(SrcVertexBuffer);
+		Dest->Swap(*Src);
+	}
+}
+
 void* FMetalDynamicRHI::RHILockStagingBuffer(FStagingBufferRHIParamRef StagingBuffer, uint32 Offset, uint32 SizeRHI)
 {
 	FMetalStagingBuffer* Buffer = ResourceCast(StagingBuffer);
diff --git a/Engine/Source/Runtime/Apple/MetalRHI/Public/MetalRHI.h b/Engine/Source/Runtime/Apple/MetalRHI/Public/MetalRHI.h
--- a/Engine/Source/Runtime/Apple/MetalRHI/Public/MetalRHI.h
+++ b/Engine/Source/Runtime/Apple/MetalRHI/Public/MetalRHI.h
@@ -66,10 +66,12 @@ public:
 	virtual FIndexBufferRHIRef RHICreateIndexBuffer(uint32 Stride, uint32 Size, uint32 InUsage, FRHIResourceCreateInfo& CreateInfo) final override;
 	virtual void* RHILockIndexBuffer(FIndexBufferRHIParamRef IndexBuffer, uint32 Offset, uint32 Size, EResourceLockMode LockMode) final override;
 	virtual void RHIUnlockIndexBuffer(FIndexBufferRHIParamRef IndexBuffer) final override;
+	virtual void RHITransferIndexBufferUnderlyingResource(FIndexBufferRHIParamRef DestIndexBuffer, FIndexBufferRHIParamRef SrcIndexBuffer) final override;
 	virtual FVertexBufferRHIRef RHICreateVertexBuffer(uint32 Size, uint32 InUsage, FRHIResourceCreateInfo& CreateInfo) final override;
 	virtual void* RHILockVertexBuffer(FVertexBufferRHIParamRef VertexBuffer, uint32 Offset, uint32 SizeRHI, EResourceLockMode LockMode) final override;
 	virtual void RHIUnlockVertexBuffer(FVertexBufferRHIParamRef VertexBuffer) final override;
 	virtual void RHICopyVertexBuffer(FVertexBufferRHIParamRef SourceBuffer,FVertexBufferRHIParamRef DestBuffer) final override;
+	virtual void RHITransferVertexBufferUnderlyingResource(FVertexBufferRHIParamRef DestVertexBuffer, FVertexBufferRHIParamRef SrcVertexBuffer) final override;
 	virtual FStructuredBufferRHIRef RHICreateStructuredBuffer(uint32 Stride, uint32 Size, uint32 InUsage, FRHIResourceCreateInfo& CreateInfo) final override;
 	virtual void* RHILockStructuredBuffer(FStructuredBufferRHIParamRef StructuredBuffer, uint32 Offset, uint32 SizeRHI, EResourceLockMode LockMode) final override;
 	virtual void RHIUnlockStructuredBuffer(FStructuredBufferRHIParamRef StructuredBuffer) final override;
@@ -80,6 +82,8 @@ public:
 	virtual FShaderResourceViewRHIRef RHICreateShaderResourceView(FStructuredBufferRHIParamRef StructuredBuffer) final override;
 	virtual FShaderResourceViewRHIRef RHICreateShaderResourceView(FVertexBufferRHIParamRef VertexBuffer, uint32 Stride, uint8 Format) final override;
 	virtual FShaderResourceViewRHIRef RHICreateShaderResourceView(FIndexBufferRHIParamRef Buffer) final override;
+	virtual void RHIUpdateShaderResourceView(FShaderResourceViewRHIParamRef SRV, FVertexBufferRHIParamRef VertexBuffer, uint32 Stride, uint8 Format) final override;
+	virtual void RHIUpdateShaderResourceView(FShaderResourceViewRHIParamRef SRV, FIndexBufferRHIParamRef IndexBuffer) final override;
 	virtual uint64 RHICalcTexture2DPlatformSize(uint32 SizeX, uint32 SizeY, uint8 Format, uint32 NumMips, uint32 NumSamples, uint32 Flags, uint32& OutAlign) final override;
 	virtual uint64 RHICalcTexture3DPlatformSize(uint32 SizeX, uint32 SizeY, uint32 SizeZ, uint8 Format, uint32 NumMips, uint32 Flags, uint32& OutAlign) final override;
 	virtual uint64 RHICalcTextureCubePlatformSize(uint32 Size, uint8 Format, uint32 NumMips, uint32 Flags, uint32& OutAlign) final override;
diff --git a/Engine/Source/Runtime/Apple/MetalRHI/Public/MetalResources.h b/Engine/Source/Runtime/Apple/MetalRHI/Public/MetalResources.h
--- a/Engine/Source/Runtime/Apple/MetalRHI/Public/MetalResources.h
+++ b/Engine/Source/Runtime/Apple/MetalRHI/Public/MetalResources.h
@@ -779,6 +779,8 @@ public:
 	 */
 	void Unlock();
 	
+	void Swap(FMetalRHIBuffer& Other);
+	
 	/**
 	 * Whether to allocate the resource rom private memory or not.
 	 */
@@ -824,6 +826,8 @@ public:
 	FMetalIndexBuffer(uint32 InStride, uint32 InSize, uint32 InUsage);
 	virtual ~FMetalIndexBuffer();
 	
+	void Swap(FMetalIndexBuffer& Other);
+	
 	// 16- or 32-bit
 	mtlpp::IndexType IndexType;
 };
@@ -836,6 +840,8 @@ public:
 	/** Constructor */
 	FMetalVertexBuffer(uint32 InSize, uint32 InUsage);
 	virtual ~FMetalVertexBuffer();
+
+	void Swap(FMetalVertexBuffer& Other);
 };
 
 class FMetalUniformBuffer : public FRHIUniformBuffer, public FMetalRHIBuffer
