break out render targets from textures for Metal platform LLM

diff --git a/Engine/Source/Runtime/Apple/MetalRHI/Private/MetalLLM.cpp b/Engine/Source/Runtime/Apple/MetalRHI/Private/MetalLLM.cpp
--- a/Engine/Source/Runtime/Apple/MetalRHI/Private/MetalLLM.cpp
+++ b/Engine/Source/Runtime/Apple/MetalRHI/Private/MetalLLM.cpp
@@ -19,6 +19,7 @@ struct FLLMTagInfoMetal
 DECLARE_LLM_MEMORY_STAT(TEXT("Metal Buffers"), STAT_MetalBuffersLLM, STATGROUP_LLMPlatform);
 DECLARE_LLM_MEMORY_STAT(TEXT("Metal Textures"), STAT_MetalTexturesLLM, STATGROUP_LLMPlatform);
 DECLARE_LLM_MEMORY_STAT(TEXT("Metal Heaps"), STAT_MetalHeapsLLM, STATGROUP_LLMPlatform);
+DECLARE_LLM_MEMORY_STAT(TEXT("Metal RenderTargets"), STAT_MetalRenderTargetsLLM, STATGROUP_LLMPlatform);
 
 // *** order must match ELLMTagMetal enum ***
 const FLLMTagInfoMetal ELLMTagNamesMetal[] =
@@ -27,6 +28,7 @@ const FLLMTagInfoMetal ELLMTagNamesMetal[] =
 	{ TEXT("Metal Buffers"),		GET_STATFNAME(STAT_MetalBuffersLLM),		GET_STATFNAME(STAT_EngineSummaryLLM) },		// ELLMTagMetal::Buffers
 	{ TEXT("Metal Textures"),		GET_STATFNAME(STAT_MetalTexturesLLM),		GET_STATFNAME(STAT_EngineSummaryLLM) },		// ELLMTagMetal::Textures
 	{ TEXT("Metal Heaps"),			GET_STATFNAME(STAT_MetalHeapsLLM),			GET_STATFNAME(STAT_EngineSummaryLLM) },		// ELLMTagMetal::Heaps
+	{ TEXT("Metal Render Targets"),	GET_STATFNAME(STAT_MetalRenderTargetsLLM),	GET_STATFNAME(STAT_EngineSummaryLLM) },		// ELLMTagMetal::RenderTargets
 };
 
 /*
@@ -214,21 +216,42 @@ void MetalLLM::LogAllocTexture(mtlpp::Device& Device, mtlpp::TextureDescriptor c
 	{
 		LLM_SCOPED_PAUSE_TRACKING(ELLMAllocType::System);
 		
-		objc_setAssociatedObject(Texture.GetPtr(), (void*)&MetalLLM::LogAllocTexture,
-		[[[FMetalDeallocHandler alloc] initWithBlock:^{
-			LLM_PLATFORM_SCOPE_METAL(ELLMTagMetal::Textures);
+		if (Desc.GetUsage() & mtlpp::TextureUsage::RenderTarget)
+		{
+			objc_setAssociatedObject(Texture.GetPtr(), (void*)&MetalLLM::LogAllocTexture,
+			[[[FMetalDeallocHandler alloc] initWithBlock:^{
+				LLM_PLATFORM_SCOPE_METAL(ELLMTagMetal::RenderTargets);
+				
+				LLM(FLowLevelMemTracker::Get().OnLowLevelFree(ELLMTracker::Platform, Ptr, ELLMAllocType::System));
+				
+#if PLATFORM_IOS
+				if (!bMemoryless)
+#endif
+				{
+					DEC_MEMORY_STAT_BY(STAT_MetalTextureMemory, Size);
+				}
+				DEC_DWORD_STAT(STAT_MetalTextureCount);
+			}] autorelease],
+			OBJC_ASSOCIATION_RETAIN);
+		}
+		else
+		{
+			objc_setAssociatedObject(Texture.GetPtr(), (void*)&MetalLLM::LogAllocTexture,
+			[[[FMetalDeallocHandler alloc] initWithBlock:^{
+				LLM_PLATFORM_SCOPE_METAL(ELLMTagMetal::Textures);
 			
-			LLM(FLowLevelMemTracker::Get().OnLowLevelFree(ELLMTracker::Platform, Ptr, ELLMAllocType::System));
+				LLM(FLowLevelMemTracker::Get().OnLowLevelFree(ELLMTracker::Platform, Ptr, ELLMAllocType::System));
 			
 #if PLATFORM_IOS
-			if (!bMemoryless)
+				if (!bMemoryless)
 #endif
-			{
-				DEC_MEMORY_STAT_BY(STAT_MetalTextureMemory, Size);
-			}
-			DEC_DWORD_STAT(STAT_MetalTextureCount);
-		}] autorelease],
-		OBJC_ASSOCIATION_RETAIN);
+				{
+					DEC_MEMORY_STAT_BY(STAT_MetalTextureMemory, Size);
+				}
+				DEC_DWORD_STAT(STAT_MetalTextureCount);
+			}] autorelease],
+			OBJC_ASSOCIATION_RETAIN);
+		}
 	}
 }
 
diff --git a/Engine/Source/Runtime/Apple/MetalRHI/Private/MetalLLM.h b/Engine/Source/Runtime/Apple/MetalRHI/Private/MetalLLM.h
--- a/Engine/Source/Runtime/Apple/MetalRHI/Private/MetalLLM.h
+++ b/Engine/Source/Runtime/Apple/MetalRHI/Private/MetalLLM.h
@@ -23,6 +23,7 @@ enum class ELLMTagMetal : LLM_TAG_TYPE
 	Buffers = (LLM_TAG_TYPE)ELLMTag::PlatformTagStart,
 	Textures,
 	Heaps,
+	RenderTargets,
 	
 	Count
 };
