Experimentally disable HDR support on macOS older than 10.14.4 to see if that improves stability on old AMD drivers

diff --git a/Engine/Source/Runtime/Apple/MetalRHI/Private/MetalRHI.cpp b/Engine/Source/Runtime/Apple/MetalRHI/Private/MetalRHI.cpp
--- a/Engine/Source/Runtime/Apple/MetalRHI/Private/MetalRHI.cpp
+++ b/Engine/Source/Runtime/Apple/MetalRHI/Private/MetalRHI.cpp
@@ -511,7 +511,7 @@ FMetalDynamicRHI::FMetalDynamicRHI(ERHIFeatureLevel::Type RequestedFeatureLevel)
 	GMaxShadowDepthBufferSizeX = 16384;
 	GMaxShadowDepthBufferSizeY = 16384;
     bSupportsD16 = !FParse::Param(FCommandLine::Get(),TEXT("nometalv2")) && Device.SupportsFeatureSet(mtlpp::FeatureSet::macOS_GPUFamily1_v2);
-    GRHISupportsHDROutput = FPlatformMisc::MacOSXVersionCompare(10,13,0) >= 0 && Device.SupportsFeatureSet(mtlpp::FeatureSet::macOS_GPUFamily1_v2);
+    GRHISupportsHDROutput = FPlatformMisc::MacOSXVersionCompare(10,14,4) >= 0 && Device.SupportsFeatureSet(mtlpp::FeatureSet::macOS_GPUFamily1_v2);
 	GRHIHDRDisplayOutputFormat = (GRHISupportsHDROutput) ? PF_PLATFORM_HDR_0 : PF_B8G8R8A8;
 #else
 #if PLATFORM_TVOS
