Fixes case where multiple waits on the async context cause an assert to trip.

This isn't a full fix - need to support multiple waits.

diff --git a/Engine/Source/Runtime/Apple/MetalRHI/Private/MetalRenderPass.cpp b/Engine/Source/Runtime/Apple/MetalRHI/Private/MetalRenderPass.cpp
--- a/Engine/Source/Runtime/Apple/MetalRHI/Private/MetalRenderPass.cpp
+++ b/Engine/Source/Runtime/Apple/MetalRHI/Private/MetalRenderPass.cpp
@@ -1705,6 +1705,18 @@ void FMetalRenderPass::ConditionalSwitchToAsyncCompute(void)
 			}
 			PrologueEncoderFence = nullptr;
 		}
+		if (PassStartFence)
+		{
+			if(PassStartFence->NeedsWait(mtlpp::RenderStages::Vertex))
+			{
+				PrologueEncoder.WaitForFence(PassStartFence);
+			}
+			else
+			{
+				PrologueEncoder.WaitAndUpdateFence(PassStartFence);
+			}
+			PassStartFence = nullptr;
+		}
 #if METAL_DEBUG_OPTIONS
 //		if (GetEmitDrawEvents() && PrologueEncoder.GetEncoderFence())
 //		{
