Add a 1ms sleep to a busy wait on iOS when we can't get a backbuffer. Solves overheating on iPhone X and other 30 fps devices.

diff --git a/Engine/Source/Runtime/Apple/MetalRHI/Private/MetalViewport.cpp b/Engine/Source/Runtime/Apple/MetalRHI/Private/MetalViewport.cpp
--- a/Engine/Source/Runtime/Apple/MetalRHI/Private/MetalViewport.cpp
+++ b/Engine/Source/Runtime/Apple/MetalRHI/Private/MetalViewport.cpp
@@ -315,10 +315,17 @@ mtlpp::Drawable FMetalViewport::GetDrawable(EMetalViewportAccessFlag Accessor)
 			do
 			{
 				Drawable = [AppDelegate.IOSView MakeDrawable];
-				Size.width = ((id<CAMetalDrawable>)Drawable).texture.width;
-				Size.height = ((id<CAMetalDrawable>)Drawable).texture.height;
+				if (Drawable != nil)
+				{
+					Size.width = ((id<CAMetalDrawable>)Drawable).texture.width;
+					Size.height = ((id<CAMetalDrawable>)Drawable).texture.height;
+				}
+				else
+				{
+					FPlatformProcess::Sleep(0.001f);
+				}
 			}
-			while (Size.width != BackBuffer[GetViewportIndex(Accessor)]->GetSizeX() || Size.height != BackBuffer[GetViewportIndex(Accessor)]->GetSizeY());
+			while (Drawable == nil || Size.width != BackBuffer[GetViewportIndex(Accessor)]->GetSizeX() || Size.height != BackBuffer[GetViewportIndex(Accessor)]->GetSizeY());
 			
 	#endif
 			
