Disable Metal fences on Mac. They can be enabled with -metalfence, like on iOS

diff --git a/Engine/Source/Runtime/Apple/MetalRHI/Private/MetalCommandQueue.cpp b/Engine/Source/Runtime/Apple/MetalRHI/Private/MetalCommandQueue.cpp
--- a/Engine/Source/Runtime/Apple/MetalRHI/Private/MetalCommandQueue.cpp
+++ b/Engine/Source/Runtime/Apple/MetalRHI/Private/MetalCommandQueue.cpp
@@ -264,7 +264,7 @@ FMetalCommandQueue::FMetalCommandQueue(mtlpp::Device InDevice, uint32 const MaxN
 			if (!GIsEditor)
 #endif
 			{
-				if (!FParse::Param(FCommandLine::Get(),TEXT("nometalfence")))
+				if (FParse::Param(FCommandLine::Get(),TEXT("metalfence")))
 				{
 					Features |= EMetalFeaturesFences;
 				}
