Decoupled device and audio mixer callbacks for iOS with a circular buffer.
iOS audio mixer init now observes project settings
Added optional command line arguments to override callback buffer size and num buffers  "-ForceIOSAudioMixerBufferSize=" and "-ForceIOSAudioMixerNumBuffers="

diff --git a/Engine/Source/Runtime/Apple/AudioMixerAudioUnit/Private/AudioMixerPlatformAudioUnit.h b/Engine/Source/Runtime/Apple/AudioMixerAudioUnit/Private/AudioMixerPlatformAudioUnit.h
--- a/Engine/Source/Runtime/Apple/AudioMixerAudioUnit/Private/AudioMixerPlatformAudioUnit.h
+++ b/Engine/Source/Runtime/Apple/AudioMixerAudioUnit/Private/AudioMixerPlatformAudioUnit.h
@@ -3,7 +3,7 @@
 #pragma once
 
 #include "AudioMixer.h"
-
+#include "DSP/Dsp.h"
 #include <AudioToolbox/AudioToolbox.h>
 #include <AudioUnit/AudioUnit.h>
 
@@ -66,6 +66,20 @@ namespace Audio
         int32       BytesPerSubmittedBuffer;
         
         double GraphSampleRate;
+
+		// We may have to grow the circular buffer capacity since Audio Unit callback size is not guaranteed to be constant
+		// Currently, this just zero's-out and reallocates, so it will pop. (We always keep largest capacity)
+		void GrowCircularBufferIfNeeded(const int32 InNumSamplesPerRenderCallback, const int32 InNumSamplesPerDeviceCallback);
+
+		// This buffer is pushed to and popped from in the SubmitBuffer callback. 
+		// This is required for devices that require frame counts per callback that are not powers of two.
+		Audio::TCircularAudioBuffer<int8> CircularOutputBuffer;
+
+		// This is the buffer we pop CircularOutputBuffer into in SubmitBuffer.
+		TArray<int8> DeviceBuffer;
+
+		int32 NumSamplesPerRenderCallback;
+		int32 NumSamplesPerDeviceCallback;
         
 		bool PerformCallback(AudioBufferList* OutputBufferData);
 		void HandleError(const TCHAR* InLogOutput, bool bTeardown = true);
