Use scheduled handler for present in Metal on Mac instead of command buffer's present method to work around timeouts related to space switching

diff --git a/Engine/Source/Runtime/Apple/MetalRHI/Private/MetalViewport.cpp b/Engine/Source/Runtime/Apple/MetalRHI/Private/MetalViewport.cpp
--- a/Engine/Source/Runtime/Apple/MetalRHI/Private/MetalViewport.cpp
+++ b/Engine/Source/Runtime/Apple/MetalRHI/Private/MetalViewport.cpp
@@ -508,46 +508,28 @@ void FMetalViewport::Present(FMetalCommandQueue& CommandQueue, bool bLockToVsync
 #endif
 						};
 						
-#if WITH_EDITOR			// The Editor needs the older way to present otherwise we end up with bad behaviour of the completion handlers that causes GPU timeouts.
-						if (GIsEditor)
+#if PLATFORM_MAC		// Mac needs the older way to present otherwise we end up with bad behaviour of the completion handlers that causes GPU timeouts.
+						mtlpp::CommandBufferHandler H = [LocalDrawable](mtlpp::CommandBuffer const&)
 						{
-#if !PLATFORM_IOS
-							mtlpp::CommandBufferHandler H = [LocalDrawable](mtlpp::CommandBuffer const&)
-							{
-#else
-							mtlpp::CommandBufferHandler H = [LocalDrawable, MinPresentDuration, FramePace](mtlpp::CommandBuffer const&)
-							{
-								if (MinPresentDuration && GEnablePresentPacing)
-								{
-									[LocalDrawable presentAfterMinimumDuration:1.0f/(float)FramePace];
-								}
-								else
-#endif
-								{
-									[LocalDrawable present];
-								};
-							};
+							[LocalDrawable present];
+						};
 								
-							CurrentCommandBuffer.AddCompletedHandler(C);
-							CurrentCommandBuffer.AddScheduledHandler(H);
+						CurrentCommandBuffer.AddCompletedHandler(C);
+						CurrentCommandBuffer.AddScheduledHandler(H);
+
+#else // PLATFORM_MAC
+						CurrentCommandBuffer.AddCompletedHandler(C);
+
+						if (MinPresentDuration && GEnablePresentPacing)
+						{
+							CurrentCommandBuffer.PresentAfterMinimumDuration(LocalDrawable, 1.0f/(float)FramePace);
 						}
 						else
-#endif
 						{
-							CurrentCommandBuffer.AddCompletedHandler(C);
-#if PLATFORM_IOS
-							if (MinPresentDuration && GEnablePresentPacing)
-							{
-								CurrentCommandBuffer.PresentAfterMinimumDuration(LocalDrawable, 1.0f/(float)FramePace);
-							}
-							else
-#endif
-							{
-								CurrentCommandBuffer.Present(LocalDrawable);
-							}
+							CurrentCommandBuffer.Present(LocalDrawable);
 						}
-						
-						
+#endif // PLATFORM_MAC
+
 						METAL_GPUPROFILE(Stats->End(CurrentCommandBuffer));
 						CommandQueue.CommitCommandBuffer(CurrentCommandBuffer);
 					}
