Audio mixer on iOS now uses a circular buffer to let device callback size change.
iOS audio mixer Project settings are now observed.
Added optional comand line arguments to override project settings without recooking

diff --git a/Engine/Source/Runtime/Apple/AudioMixerAudioUnit/Private/AudioMixerPlatformAudioUnit.h b/Engine/Source/Runtime/Apple/AudioMixerAudioUnit/Private/AudioMixerPlatformAudioUnit.h
--- a/Engine/Source/Runtime/Apple/AudioMixerAudioUnit/Private/AudioMixerPlatformAudioUnit.h
+++ b/Engine/Source/Runtime/Apple/AudioMixerAudioUnit/Private/AudioMixerPlatformAudioUnit.h
@@ -75,11 +75,10 @@ namespace Audio
 		// This is required for devices that require frame counts per callback that are not powers of two.
 		Audio::TCircularAudioBuffer<int8> CircularOutputBuffer;
 
-		// This is the buffer we pop CircularOutputBuffer into in SubmitBuffer.
-		TArray<int8> DeviceBuffer;
-
 		int32 NumSamplesPerRenderCallback;
 		int32 NumSamplesPerDeviceCallback;
+        mutable bool bInternalPlatformSettingsInitialized{ false };
+        mutable FAudioPlatformSettings InternalPlatformSettings;
         
 		bool PerformCallback(AudioBufferList* OutputBufferData);
 		void HandleError(const TCHAR* InLogOutput, bool bTeardown = true);
