UE4.23 - Deprecate FUniformBufferRHIParamRef, FIndexBufferRHIParamRef, FVertexBufferRHIParamRef, FStructuredBufferRHIParamRef, FTextureRHIParamRef, FTexture2DRHIParamRef, FTexture2DArrayRHIParamRef, FTexture3DRHIParamRef, FTextureCubeRHIParamRef

diff --git a/Engine/Plugins/Experimental/GeometryCache/Source/GeometryCache/Private/GeometryCacheSceneProxy.cpp b/Engine/Plugins/Experimental/GeometryCache/Source/GeometryCache/Private/GeometryCacheSceneProxy.cpp
--- a/Engine/Plugins/Experimental/GeometryCache/Source/GeometryCache/Private/GeometryCacheSceneProxy.cpp
+++ b/Engine/Plugins/Experimental/GeometryCache/Source/GeometryCache/Private/GeometryCacheSceneProxy.cpp
@@ -214,12 +214,12 @@ struct FRHICommandUpdateGeometryCacheBuffer : public FRHICommand<FRHICommandUpda
 {
 	FGraphEventRef BufferGenerationCompleteFence;
 
-	FVertexBufferRHIParamRef VertexBuffer;
+	FRHIVertexBuffer* VertexBuffer;
 	//void *VertexData;
 	//uint32 VertexSize;
 	TArray<uint8> VertexData;
 
-	FIndexBufferRHIParamRef IndexBuffer;
+	FRHIIndexBuffer* IndexBuffer;
 	//void *IndexData;
 	//uint32 IndexSize;
 	TArray<uint8> IndexData;
@@ -228,10 +228,10 @@ struct FRHICommandUpdateGeometryCacheBuffer : public FRHICommand<FRHICommandUpda
 	
 	FORCEINLINE_DEBUGGABLE FRHICommandUpdateGeometryCacheBuffer(
 		FGraphEventRef& InBufferGenerationCompleteFence,
-		FVertexBufferRHIParamRef InVertexBuffer,
+		FRHIVertexBuffer* InVertexBuffer,
 		void *InVertexData,
 		uint32 InVertexSize,
-		FIndexBufferRHIParamRef InIndexBuffer,
+		FRHIIndexBuffer* InIndexBuffer,
 		void *InIndexData,
 		uint32 InIndexSize)
 	:
diff --git a/Engine/Plugins/FX/Niagara/Source/Niagara/Classes/NiagaraEmitterInstanceBatcher.h b/Engine/Plugins/FX/Niagara/Source/Niagara/Classes/NiagaraEmitterInstanceBatcher.h
--- a/Engine/Plugins/FX/Niagara/Source/Niagara/Classes/NiagaraEmitterInstanceBatcher.h
+++ b/Engine/Plugins/FX/Niagara/Source/Niagara/Classes/NiagaraEmitterInstanceBatcher.h
@@ -99,9 +99,9 @@ public:
 
 	virtual void PostRenderOpaque(
 		FRHICommandListImmediate& RHICmdList,
-		const FUniformBufferRHIParamRef ViewUniformBuffer,
+		FRHIUniformBuffer* ViewUniformBuffer,
 		const class FShaderParametersMetadata* SceneTexturesUniformBufferStruct,
-		FUniformBufferRHIParamRef SceneTexturesUniformBuffer) override;
+		FRHIUniformBuffer* SceneTexturesUniformBuffer) override;
 
 	int32 AddSortedGPUSimulation(const FNiagaraGPUSortInfo& SortInfo);
 	void SortGPUParticles(FRHICommandListImmediate& RHICmdList);
@@ -122,7 +122,7 @@ public:
 				const uint32 TotalNumInstances, 
 				FNiagaraShader* Shader,
 				FRHICommandList &RHICmdList, 
-				FUniformBufferRHIParamRef ViewUniformBuffer, 
+				FRHIUniformBuffer* ViewUniformBuffer, 
 				bool bCopyBeforeStart = false
 			) const;
 
@@ -132,10 +132,10 @@ public:
 
 private:
 
-	void ExecuteAll(FRHICommandList &RHICmdList, FUniformBufferRHIParamRef ViewUniformBuffer, bool bSetReadback);
-	void TickSingle(const FNiagaraGPUSystemTick& Tick, FNiagaraComputeInstanceData *Instance, FRHICommandList &RHICmdList, FUniformBufferRHIParamRef ViewUniformBuffer, bool bSetReadback) const;
+	void ExecuteAll(FRHICommandList &RHICmdList, FRHIUniformBuffer* ViewUniformBuffer, bool bSetReadback);
+	void TickSingle(const FNiagaraGPUSystemTick& Tick, FNiagaraComputeInstanceData *Instance, FRHICommandList &RHICmdList, FRHIUniformBuffer* ViewUniformBuffer, bool bSetReadback) const;
 	void ResizeBuffersAndGatherResources(FOverlappableTicks& OverlappableTick, FRHICommandList& RHICmdList, FNiagaraBufferArray& DestDataBuffers, FNiagaraBufferArray& CurrDataBuffers, FNiagaraBufferArray& DestBufferIntFloat, FNiagaraBufferArray& CurrBufferIntFloat);
-	void DispatchAllOnCompute(FOverlappableTicks& OverlappableTick, FRHICommandList& RHICmdList, FUniformBufferRHIParamRef ViewUniformBuffer, FNiagaraBufferArray& DestDataBuffers, FNiagaraBufferArray& CurrDataBuffers, FNiagaraBufferArray& DestBufferIntFloat, FNiagaraBufferArray& CurrBufferIntFloat, bool bSetReadback);
+	void DispatchAllOnCompute(FOverlappableTicks& OverlappableTick, FRHICommandList& RHICmdList, FRHIUniformBuffer* ViewUniformBuffer, FNiagaraBufferArray& DestDataBuffers, FNiagaraBufferArray& CurrDataBuffers, FNiagaraBufferArray& DestBufferIntFloat, FNiagaraBufferArray& CurrBufferIntFloat, bool bSetReadback);
 
 //	void SimStepReadback(const FNiagaraComputeInstanceData& Instance, FRHICommandList& RHICmdList) const;
 
diff --git a/Engine/Plugins/FX/Niagara/Source/Niagara/Private/NiagaraDataSet.cpp b/Engine/Plugins/FX/Niagara/Source/Niagara/Private/NiagaraDataSet.cpp
--- a/Engine/Plugins/FX/Niagara/Source/Niagara/Private/NiagaraDataSet.cpp
+++ b/Engine/Plugins/FX/Niagara/Source/Niagara/Private/NiagaraDataSet.cpp
@@ -1047,7 +1047,7 @@ void FScopedNiagaraDataSetGPUReadback::ReadbackData(NiagaraEmitterInstanceBatche
 			const uint32 BufferOffset = DataBuffer->GetGPUInstanceCountBufferOffset();
 			if (Batcher && BufferOffset != INDEX_NONE)
 			{
-				FVertexBufferRHIParamRef InstanceCountBuffer = Batcher->GetGPUInstanceCounterManager().GetInstanceCountBuffer().Buffer;
+				FRHIVertexBuffer* InstanceCountBuffer = Batcher->GetGPUInstanceCounterManager().GetInstanceCountBuffer().Buffer;
 
 				void* Data = RHICmdList.LockVertexBuffer(InstanceCountBuffer, 0, (BufferOffset + 1) * sizeof(int32), RLM_ReadOnly);
 				NumInstances = reinterpret_cast<int32*>(Data)[BufferOffset];
diff --git a/Engine/Plugins/FX/Niagara/Source/Niagara/Private/NiagaraEmitterInstanceBatcher.cpp b/Engine/Plugins/FX/Niagara/Source/Niagara/Private/NiagaraEmitterInstanceBatcher.cpp
--- a/Engine/Plugins/FX/Niagara/Source/Niagara/Private/NiagaraEmitterInstanceBatcher.cpp
+++ b/Engine/Plugins/FX/Niagara/Source/Niagara/Private/NiagaraEmitterInstanceBatcher.cpp
@@ -225,7 +225,7 @@ void NiagaraEmitterInstanceBatcher::ResizeBuffersAndGatherResources(FOverlappabl
 	}
 }
 
-void NiagaraEmitterInstanceBatcher::DispatchAllOnCompute(FOverlappableTicks& OverlappableTick, FRHICommandList& RHICmdList, FUniformBufferRHIParamRef ViewUniformBuffer, FNiagaraBufferArray& DestDataBuffers, FNiagaraBufferArray& CurrDataBuffers, FNiagaraBufferArray& DestBufferIntFloat, FNiagaraBufferArray& CurrBufferIntFloat, bool bSetReadback)
+void NiagaraEmitterInstanceBatcher::DispatchAllOnCompute(FOverlappableTicks& OverlappableTick, FRHICommandList& RHICmdList, FRHIUniformBuffer* ViewUniformBuffer, FNiagaraBufferArray& DestDataBuffers, FNiagaraBufferArray& CurrDataBuffers, FNiagaraBufferArray& DestBufferIntFloat, FNiagaraBufferArray& CurrBufferIntFloat, bool bSetReadback)
 {
 	FRHICommandListImmediate& RHICmdListImmediate = FRHICommandListExecutor::GetImmediateCommandList();
 
@@ -328,7 +328,7 @@ void NiagaraEmitterInstanceBatcher::DispatchAllOnCompute(FOverlappableTicks& Ove
 	RHICmdList.SubmitCommandsHint();
 }
 
-void NiagaraEmitterInstanceBatcher::PostRenderOpaque(FRHICommandListImmediate& RHICmdList, const FUniformBufferRHIParamRef ViewUniformBuffer, const class FShaderParametersMetadata* SceneTexturesUniformBufferStruct, FUniformBufferRHIParamRef SceneTexturesUniformBuffer)
+void NiagaraEmitterInstanceBatcher::PostRenderOpaque(FRHICommandListImmediate& RHICmdList, FRHIUniformBuffer* ViewUniformBuffer, const class FShaderParametersMetadata* SceneTexturesUniformBufferStruct, FRHIUniformBuffer* SceneTexturesUniformBuffer)
 {
 	// Setup new readback since if there is no pending request, there is no risk of having invalid data read (offset being allocated after the readback was sent).
 	ExecuteAll(RHICmdList, ViewUniformBuffer, !GPUInstanceCounterManager.HasPendingGPUReadback());
@@ -339,7 +339,7 @@ void NiagaraEmitterInstanceBatcher::PostRenderOpaque(FRHICommandListImmediate& R
 	}
 }
 
-void NiagaraEmitterInstanceBatcher::ExecuteAll(FRHICommandList &RHICmdList, FUniformBufferRHIParamRef ViewUniformBuffer, bool bSetReadback)
+void NiagaraEmitterInstanceBatcher::ExecuteAll(FRHICommandList &RHICmdList, FRHIUniformBuffer* ViewUniformBuffer, bool bSetReadback)
 {
 	// @todo REMOVE THIS HACK
 	LastFrameThatDrainedData = GFrameNumberRenderThread;
@@ -419,7 +419,7 @@ void NiagaraEmitterInstanceBatcher::ExecuteAll(FRHICommandList &RHICmdList, FUni
 	FinishDispatches();
 }
 
-void NiagaraEmitterInstanceBatcher::TickSingle(const FNiagaraGPUSystemTick& Tick, FNiagaraComputeInstanceData *Instance, FRHICommandList &RHICmdList, FUniformBufferRHIParamRef ViewUniformBuffer, bool bSetReadback) const
+void NiagaraEmitterInstanceBatcher::TickSingle(const FNiagaraGPUSystemTick& Tick, FNiagaraComputeInstanceData *Instance, FRHICommandList &RHICmdList, FRHIUniformBuffer* ViewUniformBuffer, bool bSetReadback) const
 {
 	SCOPE_CYCLE_COUNTER(STAT_NiagaraGPUSimTick_RT);
 
@@ -915,7 +915,7 @@ void NiagaraEmitterInstanceBatcher::UnsetDataInterfaceParameters(const TArray<FN
  */
  template<bool bDoResourceTransitions>
 void NiagaraEmitterInstanceBatcher::Run(const FNiagaraGPUSystemTick& Tick, const FNiagaraComputeInstanceData* Instance, uint32 UpdateStartInstance, const uint32 TotalNumInstances, FNiagaraShader* Shader,
-	FRHICommandList &RHICmdList, FUniformBufferRHIParamRef ViewUniformBuffer, bool bCopyBeforeStart) const
+	FRHICommandList &RHICmdList, FRHIUniformBuffer* ViewUniformBuffer, bool bCopyBeforeStart) const
 {
 	FNiagaraComputeExecutionContext* Context = Instance->Context;
 	if (TotalNumInstances == 0 && !bDoResourceTransitions)
diff --git a/Engine/Plugins/FX/Niagara/Source/Niagara/Public/NiagaraWorldManager.h b/Engine/Plugins/FX/Niagara/Source/Niagara/Public/NiagaraWorldManager.h
--- a/Engine/Plugins/FX/Niagara/Source/Niagara/Public/NiagaraWorldManager.h
+++ b/Engine/Plugins/FX/Niagara/Source/Niagara/Public/NiagaraWorldManager.h
@@ -35,9 +35,9 @@ public:
 		SceneTexturesUniformParams = Params.SceneTexturesUniformParams;
 	}
 
-	FTexture2DRHIParamRef GetSceneDepthTexture() { return SceneDepthTexture; }
-	FTexture2DRHIParamRef GetSceneNormalTexture() { return SceneNormalTexture; }
-	FUniformBufferRHIParamRef GetViewUniformBuffer() { return ViewUniformBuffer; }
+	FRHITexture2D* GetSceneDepthTexture() { return SceneDepthTexture; }
+	FRHITexture2D* GetSceneNormalTexture() { return SceneNormalTexture; }
+	FRHIUniformBuffer* GetViewUniformBuffer() { return ViewUniformBuffer; }
 	TUniformBufferRef<FSceneTexturesUniformParameters> GetSceneTextureUniformParameters() { return SceneTexturesUniformParams; }
 
 	virtual void InitDynamicRHI() override;
@@ -45,9 +45,9 @@ public:
 	virtual void ReleaseDynamicRHI() override;
 
 private:
-	FTexture2DRHIParamRef SceneDepthTexture;
-	FTexture2DRHIParamRef SceneNormalTexture;
-	FUniformBufferRHIParamRef ViewUniformBuffer;
+	FRHITexture2D* SceneDepthTexture;
+	FRHITexture2D* SceneNormalTexture;
+	FRHIUniformBuffer* ViewUniformBuffer;
 
 	TUniformBufferRef<FSceneTexturesUniformParameters> SceneTexturesUniformParams;
 	FPostOpaqueRenderDelegate PostOpaqueDelegate;
diff --git a/Engine/Plugins/FX/Niagara/Source/NiagaraShader/Public/NiagaraShader.h b/Engine/Plugins/FX/Niagara/Source/NiagaraShader/Public/NiagaraShader.h
--- a/Engine/Plugins/FX/Niagara/Source/NiagaraShader/Public/NiagaraShader.h
+++ b/Engine/Plugins/FX/Niagara/Source/NiagaraShader/Public/NiagaraShader.h
@@ -68,7 +68,7 @@ public:
 	
 	void SetDataInterfaceParameterInfo(const TArray<FNiagaraDataInterfaceGPUParamInfo>& InDIParamInfo);
 
-//	FUniformBufferRHIParamRef GetParameterCollectionBuffer(const FGuid& Id, const FSceneInterface* SceneInterface) const;
+//	FRHIUniformBuffer* GetParameterCollectionBuffer(const FGuid& Id, const FSceneInterface* SceneInterface) const;
 	/*
 	template<typename ShaderRHIParamRef>
 	FORCEINLINE_DEBUGGABLE void SetViewParameters(FRHICommandList& RHICmdList, const ShaderRHIParamRef ShaderRHI, const FSceneView& View, const TUniformBufferRef<FViewUniformShaderParameters>& ViewUniformBuffer)
diff --git a/Engine/Plugins/FX/Niagara/Source/NiagaraVertexFactories/Public/NiagaraMeshVertexFactory.h b/Engine/Plugins/FX/Niagara/Source/NiagaraVertexFactories/Public/NiagaraMeshVertexFactory.h
--- a/Engine/Plugins/FX/Niagara/Source/NiagaraVertexFactories/Public/NiagaraMeshVertexFactory.h
+++ b/Engine/Plugins/FX/Niagara/Source/NiagaraVertexFactories/Public/NiagaraMeshVertexFactory.h
@@ -156,7 +156,7 @@ public:
 	/**
 	* Retrieve the uniform buffer for this vertex factory.
 	*/
-	FORCEINLINE FUniformBufferRHIParamRef GetUniformBuffer()
+	FORCEINLINE FRHIUniformBuffer* GetUniformBuffer()
 	{
 		return MeshParticleUniformBuffer;
 	}
@@ -193,7 +193,7 @@ protected:
 	uint32 MeshFacingMode;
 
 	/** Uniform buffer with mesh particle parameters. */
-	FUniformBufferRHIParamRef MeshParticleUniformBuffer;
+	FRHIUniformBuffer* MeshParticleUniformBuffer;
 	
 	/** Used to remember this in the case that we reuse the same vertex factory for multiple renders . */
 	FNiagaraMeshInstanceVertices* InstanceVerticesCPU;
diff --git a/Engine/Plugins/FX/Niagara/Source/NiagaraVertexFactories/Public/NiagaraSpriteVertexFactory.h b/Engine/Plugins/FX/Niagara/Source/NiagaraVertexFactories/Public/NiagaraSpriteVertexFactory.h
--- a/Engine/Plugins/FX/Niagara/Source/NiagaraVertexFactories/Public/NiagaraSpriteVertexFactory.h
+++ b/Engine/Plugins/FX/Niagara/Source/NiagaraVertexFactories/Public/NiagaraSpriteVertexFactory.h
@@ -136,7 +136,7 @@ public:
 	/**
 	 * Retrieve the uniform buffer for this vertex factory.
 	 */
-	FORCEINLINE FUniformBufferRHIParamRef GetSpriteUniformBuffer()
+	FORCEINLINE FRHIUniformBuffer* GetSpriteUniformBuffer()
 	{
 		return SpriteUniformBuffer;
 	}
diff --git a/Engine/Plugins/Lumin/MagicLeap/Source/MagicLeapHelperVulkan/Private/MagicLeapHelperVulkan.cpp b/Engine/Plugins/Lumin/MagicLeap/Source/MagicLeapHelperVulkan/Private/MagicLeapHelperVulkan.cpp
--- a/Engine/Plugins/Lumin/MagicLeap/Source/MagicLeapHelperVulkan/Private/MagicLeapHelperVulkan.cpp
+++ b/Engine/Plugins/Lumin/MagicLeap/Source/MagicLeapHelperVulkan/Private/MagicLeapHelperVulkan.cpp
@@ -290,7 +290,7 @@ bool FMagicLeapHelperVulkan::GetMediaTexture(FTextureRHIRef& ResultTexture, FSam
 	return false;
 }
 
-void FMagicLeapHelperVulkan::AliasMediaTexture(FTextureRHIParamRef DestTexture, FTextureRHIParamRef SrcTexture)
+void FMagicLeapHelperVulkan::AliasMediaTexture(FRHITexture* DestTexture, FRHITexture* SrcTexture)
 {
 #if PLATFORM_LUMIN
 	GDynamicRHI->RHIAliasTextureResources(DestTexture, SrcTexture);
diff --git a/Engine/Plugins/Lumin/MagicLeap/Source/MagicLeapHelperVulkan/Public/MagicLeapHelperVulkan.h b/Engine/Plugins/Lumin/MagicLeap/Source/MagicLeapHelperVulkan/Public/MagicLeapHelperVulkan.h
--- a/Engine/Plugins/Lumin/MagicLeap/Source/MagicLeapHelperVulkan/Public/MagicLeapHelperVulkan.h
+++ b/Engine/Plugins/Lumin/MagicLeap/Source/MagicLeapHelperVulkan/Public/MagicLeapHelperVulkan.h
@@ -1,4 +1,4 @@
-﻿// Copyright 1998-2019 Epic Games, Inc. All Rights Reserved.
+// Copyright 1998-2019 Epic Games, Inc. All Rights Reserved.
 
 #pragma once
 
@@ -14,5 +14,5 @@ public:
 	static uint64 AliasImageSRGB(const uint64 Allocation, const uint64 AllocationOffset, const uint32 Width, const uint32 Height);
 	static bool GetVulkanDeviceExtensionsRequired(struct VkPhysicalDevice_T* pPhysicalDevice, TArray<const ANSICHAR*>& Out);
 	static bool GetMediaTexture(FTextureRHIRef& Result, FSamplerStateRHIRef& SamplerResult, const uint64 MediaTextureHandle);
-	static void AliasMediaTexture(FTextureRHIParamRef DestTexture, FTextureRHIParamRef SrcTexture);
+	static void AliasMediaTexture(FRHITexture* DestTexture, FRHITexture* SrcTexture);
 };
diff --git a/Engine/Plugins/Media/WmfMedia/Source/WmfMedia/Private/Wmf/WmfMediaHardwareVideoDecodingRendering.cpp b/Engine/Plugins/Media/WmfMedia/Source/WmfMedia/Private/Wmf/WmfMediaHardwareVideoDecodingRendering.cpp
--- a/Engine/Plugins/Media/WmfMedia/Source/WmfMedia/Private/Wmf/WmfMediaHardwareVideoDecodingRendering.cpp
+++ b/Engine/Plugins/Media/WmfMedia/Source/WmfMedia/Private/Wmf/WmfMediaHardwareVideoDecodingRendering.cpp
@@ -34,7 +34,7 @@ struct FRHICommandCopyResource final : public FRHICommand<FRHICommandCopyResourc
 	TComPtr<ID3D11Texture2D> SampleTexture;
 	FTexture2DRHIRef SampleDestinationTexture;
 
-	FRHICommandCopyResource(ID3D11Texture2D* InSampleTexture, FTexture2DRHIParamRef InSampleDestinationTexture)
+	FRHICommandCopyResource(ID3D11Texture2D* InSampleTexture, FRHITexture2D* InSampleDestinationTexture)
 		: SampleTexture(InSampleTexture)
 		, SampleDestinationTexture(InSampleDestinationTexture)
 	{
diff --git a/Engine/Plugins/Runtime/AR/Apple/AppleARKit/Source/AppleARKit/Private/AppleARKitTextures.cpp b/Engine/Plugins/Runtime/AR/Apple/AppleARKit/Source/AppleARKit/Private/AppleARKitTextures.cpp
--- a/Engine/Plugins/Runtime/AR/Apple/AppleARKit/Source/AppleARKit/Private/AppleARKitTextures.cpp
+++ b/Engine/Plugins/Runtime/AR/Apple/AppleARKit/Source/AppleARKit/Private/AppleARKitTextures.cpp
@@ -58,7 +58,7 @@ public:
 				Size.Y = ImageExtent.size.height;
 				
 				// Let go of the last texture
-				RHIUpdateTextureReference(Owner->TextureReference.TextureReferenceRHI, FTextureRHIParamRef());
+				RHIUpdateTextureReference(Owner->TextureReference.TextureReferenceRHI, nullptr);
 				DecodedTextureRef.SafeRelease();
 				
 				// Create the target texture that we'll update into
@@ -95,7 +95,7 @@ public:
 
 	virtual void ReleaseRHI() override
 	{
-		RHIUpdateTextureReference(Owner->TextureReference.TextureReferenceRHI, FTextureRHIParamRef());
+		RHIUpdateTextureReference(Owner->TextureReference.TextureReferenceRHI, nullptr);
 #if PLATFORM_MAC || PLATFORM_IOS
 		if (CameraImage != nullptr)
 		{
@@ -498,7 +498,7 @@ public:
 	
 	virtual void ReleaseRHI() override
 	{
-		RHIUpdateTextureReference(Owner->TextureReference.TextureReferenceRHI, FTextureRHIParamRef());
+		RHIUpdateTextureReference(Owner->TextureReference.TextureReferenceRHI, nullptr);
 		EnvCubemapTextureRHIRef.SafeRelease();
 		FTextureResource::ReleaseRHI();
 		FExternalTextureRegistry::Get().UnregisterExternalTexture(Owner->ExternalTextureGuid);
diff --git a/Engine/Plugins/Runtime/AR/Microsoft/HoloLensAR/Source/HoloLensAR/Private/HoloLensCameraImageTexture.cpp b/Engine/Plugins/Runtime/AR/Microsoft/HoloLensAR/Source/HoloLensAR/Private/HoloLensCameraImageTexture.cpp
--- a/Engine/Plugins/Runtime/AR/Microsoft/HoloLensAR/Source/HoloLensAR/Private/HoloLensCameraImageTexture.cpp
+++ b/Engine/Plugins/Runtime/AR/Microsoft/HoloLensAR/Source/HoloLensAR/Private/HoloLensCameraImageTexture.cpp
@@ -221,7 +221,7 @@ public:
 
 	virtual void ReleaseRHI() override
 	{
-		RHIUpdateTextureReference(Owner->TextureReference.TextureReferenceRHI, FTextureRHIParamRef());
+		RHIUpdateTextureReference(Owner->TextureReference.TextureReferenceRHI, nullptr);
 		CameraImage = nullptr;
 		CopyTextureRef.SafeRelease();
 		DecodedTextureRef.SafeRelease();
diff --git a/Engine/Plugins/Runtime/GoogleVR/GoogleVRHMD/Source/GoogleVRHMD/Private/GoogleVRHMDRender.cpp b/Engine/Plugins/Runtime/GoogleVR/GoogleVRHMD/Source/GoogleVRHMD/Private/GoogleVRHMDRender.cpp
--- a/Engine/Plugins/Runtime/GoogleVR/GoogleVRHMD/Source/GoogleVRHMD/Private/GoogleVRHMDRender.cpp
+++ b/Engine/Plugins/Runtime/GoogleVR/GoogleVRHMD/Source/GoogleVRHMD/Private/GoogleVRHMDRender.cpp
@@ -223,7 +223,7 @@ static void ResolvePendingRenderTarget(FRHICommandListImmediate& RHICmdList, FGr
 #endif
 }
 
-void FGoogleVRHMD::RenderTexture_RenderThread(FRHICommandListImmediate& RHICmdList, FTexture2DRHIParamRef BackBuffer, FTexture2DRHIParamRef SrcTexture, FVector2D WindowSize) const
+void FGoogleVRHMD::RenderTexture_RenderThread(FRHICommandListImmediate& RHICmdList, FRHITexture2D* BackBuffer, FRHITexture2D* SrcTexture, FVector2D WindowSize) const
 {
 	check(IsInRenderingThread());
 
diff --git a/Engine/Plugins/Runtime/GoogleVR/GoogleVRHMD/Source/GoogleVRHMD/Private/GoogleVRSplash.cpp b/Engine/Plugins/Runtime/GoogleVR/GoogleVRHMD/Source/GoogleVRHMD/Private/GoogleVRSplash.cpp
--- a/Engine/Plugins/Runtime/GoogleVR/GoogleVRHMD/Source/GoogleVRHMD/Private/GoogleVRSplash.cpp
+++ b/Engine/Plugins/Runtime/GoogleVR/GoogleVRHMD/Source/GoogleVRHMD/Private/GoogleVRSplash.cpp
@@ -188,7 +188,7 @@ bool FGoogleVRSplash::IsTickable() const
 	return bIsShown;
 }
 
-void FGoogleVRSplash::RenderStereoSplashScreen(FRHICommandListImmediate& RHICmdList, FTexture2DRHIParamRef DstTexture)
+void FGoogleVRSplash::RenderStereoSplashScreen(FRHICommandListImmediate& RHICmdList, FRHITexture2D* DstTexture)
 {
 	UpdateSplashScreenEyeOffset();
 
@@ -301,7 +301,7 @@ void FGoogleVRSplash::SubmitBlackFrame()
 	check(IsInRenderingThread());
 
 	FRHICommandListImmediate& RHICmdList = FRHICommandListExecutor::GetImmediateCommandList();
-	FTexture2DRHIParamRef DstTexture = GVRCustomPresent->TextureSet->GetTexture2D();
+	FRHITexture2D* DstTexture = GVRCustomPresent->TextureSet->GetTexture2D();
 
 	GVRCustomPresent->BeginRendering(GVRHMD->CachedHeadPose);
 
diff --git a/Engine/Plugins/Runtime/GoogleVR/GoogleVRHMD/Source/GoogleVRHMD/Private/GoogleVRSplash.h b/Engine/Plugins/Runtime/GoogleVR/GoogleVRHMD/Source/GoogleVRHMD/Private/GoogleVRSplash.h
--- a/Engine/Plugins/Runtime/GoogleVR/GoogleVRHMD/Source/GoogleVRHMD/Private/GoogleVRSplash.h
+++ b/Engine/Plugins/Runtime/GoogleVR/GoogleVRHMD/Source/GoogleVRHMD/Private/GoogleVRSplash.h
@@ -29,7 +29,7 @@ public:
 
 	void Tick(float DeltaTime);
 	bool IsTickable() const;
-	void RenderStereoSplashScreen(FRHICommandListImmediate& RHICmdList, FTexture2DRHIParamRef DstTexture);
+	void RenderStereoSplashScreen(FRHICommandListImmediate& RHICmdList, FRHITexture2D* DstTexture);
 	void ForceRerenderSplashScreen();
 
 private:
diff --git a/Engine/Plugins/Runtime/Oculus/OculusVR/Source/OculusHMD/Private/OculusHMD.cpp b/Engine/Plugins/Runtime/Oculus/OculusVR/Source/OculusHMD/Private/OculusHMD.cpp
--- a/Engine/Plugins/Runtime/Oculus/OculusVR/Source/OculusHMD/Private/OculusHMD.cpp
+++ b/Engine/Plugins/Runtime/Oculus/OculusVR/Source/OculusHMD/Private/OculusHMD.cpp
@@ -1326,7 +1326,7 @@ namespace OculusHMD
 	}
 
 
-	void FOculusHMD::CopyTexture_RenderThread(FRHICommandListImmediate& RHICmdList, FTexture2DRHIParamRef SrcTexture, FIntRect SrcRect, FTexture2DRHIParamRef DstTexture, FIntRect DstRect, bool bClearBlack, bool bNoAlpha) const
+	void FOculusHMD::CopyTexture_RenderThread(FRHICommandListImmediate& RHICmdList, FRHITexture2D* SrcTexture, FIntRect SrcRect, FRHITexture2D* DstTexture, FIntRect DstRect, bool bClearBlack, bool bNoAlpha) const
 	{
 		if (bClearBlack)
 		{
diff --git a/Engine/Plugins/Runtime/Oculus/OculusVR/Source/OculusHMD/Private/OculusHMD.h b/Engine/Plugins/Runtime/Oculus/OculusVR/Source/OculusHMD/Private/OculusHMD.h
--- a/Engine/Plugins/Runtime/Oculus/OculusVR/Source/OculusHMD/Private/OculusHMD.h
+++ b/Engine/Plugins/Runtime/Oculus/OculusVR/Source/OculusHMD/Private/OculusHMD.h
@@ -174,7 +174,7 @@ public:
 	// FHeadMoundedDisplayBase interface
 	virtual FVector2D GetEyeCenterPoint_RenderThread(EStereoscopicPass StereoPassType) const override;
 	virtual FIntRect GetFullFlatEyeRect_RenderThread(FTexture2DRHIRef EyeTexture) const override;
-	virtual void CopyTexture_RenderThread(FRHICommandListImmediate& RHICmdList, FTexture2DRHIParamRef SrcTexture, FIntRect SrcRect, FTexture2DRHIParamRef DstTexture, FIntRect DstRect, bool bClearBlack, bool bNoAlpha) const override;
+	virtual void CopyTexture_RenderThread(FRHICommandListImmediate& RHICmdList, FRHITexture2D* SrcTexture, FIntRect SrcRect, FRHITexture2D* DstTexture, FIntRect DstRect, bool bClearBlack, bool bNoAlpha) const override;
 	virtual bool PopulateAnalyticsAttributes(TArray<struct FAnalyticsEventAttribute>& EventAttributes) override;
 
 	// FXRRenderTargetManager interface
diff --git a/Engine/Plugins/Runtime/Oculus/OculusVR/Source/OculusHMD/Private/OculusHMD_CustomPresent.cpp b/Engine/Plugins/Runtime/Oculus/OculusVR/Source/OculusHMD/Private/OculusHMD_CustomPresent.cpp
--- a/Engine/Plugins/Runtime/Oculus/OculusVR/Source/OculusHMD/Private/OculusHMD_CustomPresent.cpp
+++ b/Engine/Plugins/Runtime/Oculus/OculusVR/Source/OculusHMD/Private/OculusHMD_CustomPresent.cpp
@@ -297,15 +297,15 @@ FXRSwapChainPtr FCustomPresent::CreateSwapChain_RenderThread(uint32 InSizeX, uin
 }
 
 
-void FCustomPresent::CopyTexture_RenderThread(FRHICommandListImmediate& RHICmdList, FTextureRHIParamRef DstTexture, FTextureRHIParamRef SrcTexture,
+void FCustomPresent::CopyTexture_RenderThread(FRHICommandListImmediate& RHICmdList, FRHITexture* DstTexture, FRHITexture* SrcTexture,
 	FIntRect DstRect, FIntRect SrcRect, bool bAlphaPremultiply, bool bNoAlphaWrite, bool bInvertY, bool sRGBSource) const
 {
 	CheckInRenderThread();
 
-	FTexture2DRHIParamRef DstTexture2D = DstTexture->GetTexture2D();
-	FTextureCubeRHIParamRef DstTextureCube = DstTexture->GetTextureCube();
-	FTexture2DRHIParamRef SrcTexture2D = SrcTexture->GetTexture2D();
-	FTextureCubeRHIParamRef SrcTextureCube = SrcTexture->GetTextureCube();
+	FRHITexture2D* DstTexture2D = DstTexture->GetTexture2D();
+	FRHITextureCube* DstTextureCube = DstTexture->GetTextureCube();
+	FRHITexture2D* SrcTexture2D = SrcTexture->GetTexture2D();
+	FRHITextureCube* SrcTextureCube = SrcTexture->GetTextureCube();
 
 	FIntPoint DstSize;
 	FIntPoint SrcSize;
diff --git a/Engine/Plugins/Runtime/Oculus/OculusVR/Source/OculusHMD/Private/OculusHMD_CustomPresent.h b/Engine/Plugins/Runtime/Oculus/OculusVR/Source/OculusHMD/Private/OculusHMD_CustomPresent.h
--- a/Engine/Plugins/Runtime/Oculus/OculusVR/Source/OculusHMD/Private/OculusHMD_CustomPresent.h
+++ b/Engine/Plugins/Runtime/Oculus/OculusVR/Source/OculusHMD/Private/OculusHMD_CustomPresent.h
@@ -67,7 +67,7 @@ public:
 
 	virtual FTextureRHIRef CreateTexture_RenderThread(uint32 InSizeX, uint32 InSizeY, EPixelFormat InFormat, FClearValueBinding InBinding, uint32 InNumMips, uint32 InNumSamples, uint32 InNumSamplesTileMem, ERHIResourceType InResourceType, ovrpTextureHandle InTexture, uint32 TexCreateFlags) = 0;
 	FXRSwapChainPtr CreateSwapChain_RenderThread(uint32 InSizeX, uint32 InSizeY, EPixelFormat InFormat, FClearValueBinding InBinding, uint32 InNumMips, uint32 InNumSamples, uint32 InNumSamplesTileMem, ERHIResourceType InResourceType, const TArray<ovrpTextureHandle>& InTextures, uint32 InTexCreateFlags);
-	void CopyTexture_RenderThread(FRHICommandListImmediate& RHICmdList, FTextureRHIParamRef DstTexture, FTextureRHIParamRef SrcTexture, FIntRect DstRect = FIntRect(), FIntRect SrcRect = FIntRect(), bool bAlphaPremultiply = false, bool bNoAlphaWrite = false, bool bInvertY = true, bool sRGBSource = false) const;
+	void CopyTexture_RenderThread(FRHICommandListImmediate& RHICmdList, FRHITexture* DstTexture, FRHITexture* SrcTexture, FIntRect DstRect = FIntRect(), FIntRect SrcRect = FIntRect(), bool bAlphaPremultiply = false, bool bNoAlphaWrite = false, bool bInvertY = true, bool sRGBSource = false) const;
 	virtual void SubmitGPUFrameTime(float GPUFrameTime) { }
 
 	bool supportsSRGB() { return bSupportsSRGB; }
diff --git a/Engine/Plugins/Runtime/OpenXR/Source/OpenXRHMD/Private/OpenXRHMD.cpp b/Engine/Plugins/Runtime/OpenXR/Source/OpenXRHMD/Private/OpenXRHMD.cpp
--- a/Engine/Plugins/Runtime/OpenXR/Source/OpenXRHMD/Private/OpenXRHMD.cpp
+++ b/Engine/Plugins/Runtime/OpenXR/Source/OpenXRHMD/Private/OpenXRHMD.cpp
@@ -212,7 +212,7 @@ bool FOpenXRHMDPlugin::PreInit()
 	return true;
 }
 
-FOpenXRHMD::FOpenXRSwapchain::FOpenXRSwapchain(XrSwapchain InSwapchain, FTexture2DRHIParamRef InRHITexture, const TArray<FTexture2DRHIRef>& InRHITextureSwapChain) :
+FOpenXRHMD::FOpenXRSwapchain::FOpenXRSwapchain(XrSwapchain InSwapchain, FRHITexture2D* InRHITexture, const TArray<FTexture2DRHIRef>& InRHITextureSwapChain) :
 	Handle(InSwapchain), RHITexture(InRHITexture), RHITextureSwapChain(InRHITextureSwapChain), SwapChainIndex_RenderThread(0), IsAcquired(false)
 {
 	IncrementSwapChainIndex_RenderThread(XR_NO_DURATION);
@@ -979,7 +979,7 @@ FIntRect FOpenXRHMD::GetFullFlatEyeRect_RenderThread(FTexture2DRHIRef EyeTexture
 	return FIntRect(EyeTexture->GetSizeX() * SrcNormRectMin.X, EyeTexture->GetSizeY() * SrcNormRectMin.Y, EyeTexture->GetSizeX() * SrcNormRectMax.X, EyeTexture->GetSizeY() * SrcNormRectMax.Y);
 }
 
-void FOpenXRHMD::CopyTexture_RenderThread(FRHICommandListImmediate& RHICmdList, FTexture2DRHIParamRef SrcTexture, FIntRect SrcRect, FTexture2DRHIParamRef DstTexture, FIntRect DstRect, bool bClearBlack, bool bNoAlpha) const
+void FOpenXRHMD::CopyTexture_RenderThread(FRHICommandListImmediate& RHICmdList, FRHITexture2D* SrcTexture, FIntRect SrcRect, FRHITexture2D* DstTexture, FIntRect DstRect, bool bClearBlack, bool bNoAlpha) const
 {
 	check(IsInRenderingThread());
 
diff --git a/Engine/Plugins/Runtime/OpenXR/Source/OpenXRHMD/Private/OpenXRHMD.h b/Engine/Plugins/Runtime/OpenXR/Source/OpenXRHMD/Private/OpenXRHMD.h
--- a/Engine/Plugins/Runtime/OpenXR/Source/OpenXRHMD/Private/OpenXRHMD.h
+++ b/Engine/Plugins/Runtime/OpenXR/Source/OpenXRHMD/Private/OpenXRHMD.h
@@ -26,7 +26,7 @@ public:
 	class FOpenXRSwapchain : public TSharedFromThis<FOpenXRSwapchain, ESPMode::ThreadSafe>
 	{
 	public:
-		FOpenXRSwapchain(XrSwapchain InSwapchain, FTexture2DRHIParamRef InRHITexture, const TArray<FTexture2DRHIRef>& InRHITextureSwapChain);
+		FOpenXRSwapchain(XrSwapchain InSwapchain, FRHITexture2D* InRHITexture, const TArray<FTexture2DRHIRef>& InRHITextureSwapChain);
 		virtual ~FOpenXRSwapchain();
 
 		FRHITexture* GetTexture() const { return RHITexture.GetReference(); }
@@ -121,7 +121,7 @@ public:
 	virtual FIntPoint GetIdealRenderTargetSize() const override;
 	virtual bool GetHMDDistortionEnabled(EShadingPath ShadingPath) const override { return false; }
 	virtual FIntRect GetFullFlatEyeRect_RenderThread(FTexture2DRHIRef EyeTexture) const override;
-	virtual void CopyTexture_RenderThread(FRHICommandListImmediate& RHICmdList, FTexture2DRHIParamRef SrcTexture, FIntRect SrcRect, FTexture2DRHIParamRef DstTexture, FIntRect DstRect, bool bClearBlack, bool bNoAlpha) const override;
+	virtual void CopyTexture_RenderThread(FRHICommandListImmediate& RHICmdList, FRHITexture2D* SrcTexture, FIntRect SrcRect, FRHITexture2D* DstTexture, FIntRect DstRect, bool bClearBlack, bool bNoAlpha) const override;
 	virtual bool HasHiddenAreaMesh() const override { return false; };
 	virtual void DrawHiddenAreaMesh_RenderThread(class FRHICommandList& RHICmdList, EStereoscopicPass StereoPass) const override;
 	virtual void OnBeginRendering_RenderThread(FRHICommandListImmediate& RHICmdList, FSceneViewFamily& ViewFamily) override;
diff --git a/Engine/Plugins/Runtime/Steam/SteamVR/Source/SteamVR/Private/SteamVRHMD.h b/Engine/Plugins/Runtime/Steam/SteamVR/Source/SteamVR/Private/SteamVRHMD.h
--- a/Engine/Plugins/Runtime/Steam/SteamVR/Source/SteamVR/Private/SteamVRHMD.h
+++ b/Engine/Plugins/Runtime/Steam/SteamVR/Source/SteamVR/Private/SteamVRHMD.h
@@ -208,7 +208,7 @@ public:
 	virtual void AdjustViewRect(EStereoscopicPass StereoPass, int32& X, int32& Y, uint32& SizeX, uint32& SizeY) const override;
 	virtual void CalculateStereoViewOffset(const EStereoscopicPass StereoPassType, FRotator& ViewRotation, const float MetersToWorld, FVector& ViewLocation) override;
 	virtual FMatrix GetStereoProjectionMatrix(const enum EStereoscopicPass StereoPassType) const override;
-	virtual void RenderTexture_RenderThread(FRHICommandListImmediate& RHICmdList, FTexture2DRHIParamRef BackBuffer, FTexture2DRHIParamRef SrcTexture, FVector2D WindowSize) const override;
+	virtual void RenderTexture_RenderThread(FRHICommandListImmediate& RHICmdList, FRHITexture2D* BackBuffer, FRHITexture2D* SrcTexture, FVector2D WindowSize) const override;
 	virtual void GetEyeRenderParams_RenderThread(const FRenderingCompositePassContext& Context, FVector2D& EyeToSrcUVScaleValue, FVector2D& EyeToSrcUVOffsetValue) const override;
 	virtual IStereoRenderTargetManager* GetRenderTargetManager() override { return this; }
 	virtual IStereoLayers* GetStereoLayers() override;
@@ -243,7 +243,7 @@ private:
 	void CreateSpectatorScreenController();
 public:
 	virtual FIntRect GetFullFlatEyeRect_RenderThread(FTexture2DRHIRef EyeTexture) const override;
-	virtual void CopyTexture_RenderThread(FRHICommandListImmediate& RHICmdList, FTexture2DRHIParamRef SrcTexture, FIntRect SrcRect, FTexture2DRHIParamRef DstTexture, FIntRect DstRect, bool bClearBlack, bool bNoAlpha) const override;
+	virtual void CopyTexture_RenderThread(FRHICommandListImmediate& RHICmdList, FRHITexture2D* SrcTexture, FIntRect SrcRect, FRHITexture2D* DstTexture, FIntRect DstRect, bool bClearBlack, bool bNoAlpha) const override;
 
 	class BridgeBaseImpl : public FXRRenderBridge
 	{
diff --git a/Engine/Plugins/Runtime/Steam/SteamVR/Source/SteamVR/Private/SteamVRRender.cpp b/Engine/Plugins/Runtime/Steam/SteamVR/Source/SteamVR/Private/SteamVRRender.cpp
--- a/Engine/Plugins/Runtime/Steam/SteamVR/Source/SteamVR/Private/SteamVRRender.cpp
+++ b/Engine/Plugins/Runtime/Steam/SteamVR/Source/SteamVR/Private/SteamVRRender.cpp
@@ -30,7 +30,7 @@ void FSteamVRHMD::DrawDistortionMesh_RenderThread(struct FRenderingCompositePass
 	check(0);
 }
 
-void FSteamVRHMD::RenderTexture_RenderThread(FRHICommandListImmediate& RHICmdList, FTexture2DRHIParamRef BackBuffer, FTexture2DRHIParamRef SrcTexture, FVector2D WindowSize) const
+void FSteamVRHMD::RenderTexture_RenderThread(FRHICommandListImmediate& RHICmdList, FRHITexture2D* BackBuffer, FRHITexture2D* SrcTexture, FVector2D WindowSize) const
 {
 	check(IsInRenderingThread());
 
diff --git a/Engine/Plugins/Runtime/Steam/SteamVR/Source/SteamVR/Private/SteamVRSpectatorScreen.cpp b/Engine/Plugins/Runtime/Steam/SteamVR/Source/SteamVR/Private/SteamVRSpectatorScreen.cpp
--- a/Engine/Plugins/Runtime/Steam/SteamVR/Source/SteamVR/Private/SteamVRSpectatorScreen.cpp
+++ b/Engine/Plugins/Runtime/Steam/SteamVR/Source/SteamVR/Private/SteamVRSpectatorScreen.cpp
@@ -34,7 +34,7 @@ FIntRect FSteamVRHMD::GetFullFlatEyeRect_RenderThread(FTexture2DRHIRef EyeTextur
 	return FIntRect(EyeTexture->GetSizeX() * SrcNormRectMin.X, EyeTexture->GetSizeY() * SrcNormRectMin.Y, EyeTexture->GetSizeX() * SrcNormRectMax.X, EyeTexture->GetSizeY() * SrcNormRectMax.Y);
 }
 
-void FSteamVRHMD::CopyTexture_RenderThread(FRHICommandListImmediate& RHICmdList, FTexture2DRHIParamRef SrcTexture, FIntRect SrcRect, FTexture2DRHIParamRef DstTexture, FIntRect DstRect, bool bClearBlack, bool bNoAlpha) const
+void FSteamVRHMD::CopyTexture_RenderThread(FRHICommandListImmediate& RHICmdList, FRHITexture2D* SrcTexture, FIntRect SrcRect, FRHITexture2D* DstTexture, FIntRect DstRect, bool bClearBlack, bool bNoAlpha) const
 {
 	check(IsInRenderingThread());
 
diff --git a/Engine/Plugins/Runtime/WindowsMixedReality/Source/WindowsMixedRealityHMD/Private/WindowsMixedRealityHMD.cpp b/Engine/Plugins/Runtime/WindowsMixedReality/Source/WindowsMixedRealityHMD/Private/WindowsMixedRealityHMD.cpp
--- a/Engine/Plugins/Runtime/WindowsMixedReality/Source/WindowsMixedRealityHMD/Private/WindowsMixedRealityHMD.cpp
+++ b/Engine/Plugins/Runtime/WindowsMixedReality/Source/WindowsMixedRealityHMD/Private/WindowsMixedRealityHMD.cpp
@@ -73,7 +73,7 @@ public:
 		WorldToMeters.Bind(Initializer.ParameterMap, TEXT("WorldToMeters"));
 	}
 
-	void SetParameters(FRHICommandList& RHICmdList, float nearPlaneM, float farPlaneM, float worldToMetersValue, FTextureRHIParamRef DepthTexture)
+	void SetParameters(FRHICommandList& RHICmdList, float nearPlaneM, float farPlaneM, float worldToMetersValue, FRHITexture* DepthTexture)
 	{
 		FRHIPixelShader* PixelShaderRHI = GetPixelShader();
 
diff --git a/Engine/Plugins/Runtime/WindowsMoviePlayer/Source/WindowsMoviePlayer/Private/WindowsMovieStreamer.cpp b/Engine/Plugins/Runtime/WindowsMoviePlayer/Source/WindowsMoviePlayer/Private/WindowsMovieStreamer.cpp
--- a/Engine/Plugins/Runtime/WindowsMoviePlayer/Source/WindowsMoviePlayer/Private/WindowsMovieStreamer.cpp
+++ b/Engine/Plugins/Runtime/WindowsMoviePlayer/Source/WindowsMoviePlayer/Private/WindowsMovieStreamer.cpp
@@ -113,7 +113,7 @@ void FMediaFoundationMovieStreamer::ConvertSample()
 	const FIntPoint OutputDim = SourceFormat.OutputDim;
 
 	FSlateTexture2DRHIRef* CurrentTexture = Texture.Get();
-	FTextureRHIParamRef RenderTarget = CurrentTexture->GetRHIRef();
+	FRHITexture* RenderTarget = CurrentTexture->GetRHIRef();
 
 	// perform the conversion
 	FRHICommandListImmediate& CommandList = FRHICommandListExecutor::GetImmediateCommandList();
diff --git a/Engine/Plugins/Runtime/nDisplay/Source/DisplayCluster/Private/Render/Device/DisplayClusterDeviceBase.cpp b/Engine/Plugins/Runtime/nDisplay/Source/DisplayCluster/Private/Render/Device/DisplayClusterDeviceBase.cpp
--- a/Engine/Plugins/Runtime/nDisplay/Source/DisplayCluster/Private/Render/Device/DisplayClusterDeviceBase.cpp
+++ b/Engine/Plugins/Runtime/nDisplay/Source/DisplayCluster/Private/Render/Device/DisplayClusterDeviceBase.cpp
@@ -340,7 +340,7 @@ FMatrix FDisplayClusterDeviceBase::GetStereoProjectionMatrix(const enum EStereos
 	return PrjMatrix;
 }
 
-void FDisplayClusterDeviceBase::RenderTexture_RenderThread(FRHICommandListImmediate& RHICmdList, FTexture2DRHIParamRef BackBuffer, FTexture2DRHIParamRef SrcTexture, FVector2D WindowSize) const
+void FDisplayClusterDeviceBase::RenderTexture_RenderThread(FRHICommandListImmediate& RHICmdList, FRHITexture2D* BackBuffer, FRHITexture2D* SrcTexture, FVector2D WindowSize) const
 {
 	// Get registered PP operations map
 	const TMap<FString, IDisplayClusterRenderManager::FDisplayClusterPPInfo> PPOperationsMap = GDisplayCluster->GetRenderMgr()->GetRegisteredPostprocessOperations();
@@ -657,7 +657,7 @@ void FDisplayClusterDeviceBase::AddViewport(const FString& InViewportId, const F
 	}
 }
 
-void FDisplayClusterDeviceBase::CopyTextureToBackBuffer_RenderThread(FRHICommandListImmediate& RHICmdList, FTexture2DRHIParamRef BackBuffer, FTexture2DRHIParamRef SrcTexture, FVector2D WindowSize) const
+void FDisplayClusterDeviceBase::CopyTextureToBackBuffer_RenderThread(FRHICommandListImmediate& RHICmdList, FRHITexture2D* BackBuffer, FRHITexture2D* SrcTexture, FVector2D WindowSize) const
 {
 	DISPLAY_CLUSTER_FUNC_TRACE(LogDisplayClusterRender);
 
diff --git a/Engine/Plugins/Runtime/nDisplay/Source/DisplayCluster/Private/Render/Device/DisplayClusterDeviceBase.h b/Engine/Plugins/Runtime/nDisplay/Source/DisplayCluster/Private/Render/Device/DisplayClusterDeviceBase.h
--- a/Engine/Plugins/Runtime/nDisplay/Source/DisplayCluster/Private/Render/Device/DisplayClusterDeviceBase.h
+++ b/Engine/Plugins/Runtime/nDisplay/Source/DisplayCluster/Private/Render/Device/DisplayClusterDeviceBase.h
@@ -50,7 +50,7 @@ public:
 	virtual void AdjustViewRect(enum EStereoscopicPass StereoPassType, int32& X, int32& Y, uint32& SizeX, uint32& SizeY) const;
 	virtual void CalculateStereoViewOffset(const enum EStereoscopicPass StereoPassType, FRotator& ViewRotation, const float WorldToMeters, FVector& ViewLocation) override;
 	virtual FMatrix GetStereoProjectionMatrix(const enum EStereoscopicPass StereoPassType) const override;
-	virtual void RenderTexture_RenderThread(FRHICommandListImmediate& RHICmdList, FTexture2DRHIParamRef BackBuffer, FTexture2DRHIParamRef SrcTexture, FVector2D WindowSize) const override final;
+	virtual void RenderTexture_RenderThread(FRHICommandListImmediate& RHICmdList, FRHITexture2D* BackBuffer, FRHITexture2D* SrcTexture, FVector2D WindowSize) const override final;
 	virtual int32 GetDesiredNumberOfViews(bool bStereoRequested) const override;
 	virtual EStereoscopicPass GetViewPassForIndex(bool bStereoRequested, uint32 ViewIndex) const override;
 	virtual uint32 GetViewIndexForPass(EStereoscopicPass StereoPassType) const override;
@@ -122,7 +122,7 @@ protected:
 	// Adds a new viewport with specified parameters and projection policy object
 	void AddViewport(const FString& InViewportId, const FIntPoint& InViewportLocation, const FIntPoint& InViewportSize, TSharedPtr<IDisplayClusterProjectionPolicy> InProjPolicy, const FString& InCameraId);
 	// Performs copying of render target data to the back buffer
-	virtual void CopyTextureToBackBuffer_RenderThread(FRHICommandListImmediate& RHICmdList, FTexture2DRHIParamRef BackBuffer, FTexture2DRHIParamRef SrcTexture, FVector2D WindowSize) const;
+	virtual void CopyTextureToBackBuffer_RenderThread(FRHICommandListImmediate& RHICmdList, FRHITexture2D* BackBuffer, FRHITexture2D* SrcTexture, FVector2D WindowSize) const;
 
 protected:
 	// Viewports
diff --git a/Engine/Plugins/Runtime/nDisplay/Source/DisplayCluster/Private/Render/Device/DisplayClusterDeviceBase_PostProcess.cpp b/Engine/Plugins/Runtime/nDisplay/Source/DisplayCluster/Private/Render/Device/DisplayClusterDeviceBase_PostProcess.cpp
--- a/Engine/Plugins/Runtime/nDisplay/Source/DisplayCluster/Private/Render/Device/DisplayClusterDeviceBase_PostProcess.cpp
+++ b/Engine/Plugins/Runtime/nDisplay/Source/DisplayCluster/Private/Render/Device/DisplayClusterDeviceBase_PostProcess.cpp
@@ -19,7 +19,7 @@ static FAutoConsoleVariableRef CVarPostprocessViewBeforeWarpBlend(
 	TEXT("Enable PP per view before warp&blend (0 = disabled)\n")
 );
 
-void FDisplayClusterDeviceBase_PostProcess::PerformPostProcessViewBeforeWarpBlend_RenderThread(FRHICommandListImmediate& RHICmdList, FTexture2DRHIParamRef SrcTexture, const FIntRect& ViewRect) const
+void FDisplayClusterDeviceBase_PostProcess::PerformPostProcessViewBeforeWarpBlend_RenderThread(FRHICommandListImmediate& RHICmdList, FRHITexture2D* SrcTexture, const FIntRect& ViewRect) const
 {
 	DISPLAY_CLUSTER_FUNC_TRACE(LogDisplayClusterRender);
 
@@ -57,7 +57,7 @@ static FAutoConsoleVariableRef CVarPostprocessFrameBeforeWarpBlend(
 	TEXT("Enable PP per eye frame before warp&blend (0 = disabled)\n")
 );
 
-void FDisplayClusterDeviceBase_PostProcess::PerformPostProcessFrameBeforeWarpBlend_RenderThread(FRHICommandListImmediate& RHICmdList, FTexture2DRHIParamRef SrcTexture, const FIntRect& ViewRect) const
+void FDisplayClusterDeviceBase_PostProcess::PerformPostProcessFrameBeforeWarpBlend_RenderThread(FRHICommandListImmediate& RHICmdList, FRHITexture2D* SrcTexture, const FIntRect& ViewRect) const
 {
 	DISPLAY_CLUSTER_FUNC_TRACE(LogDisplayClusterRender);
 
@@ -95,7 +95,7 @@ static FAutoConsoleVariableRef CVarPostprocessTargetBeforeWarpBlend(
 	TEXT("Enable PP for the whole render target before warp&blend (0 = disabled)\n")
 );
 
-void FDisplayClusterDeviceBase_PostProcess::PerformPostProcessRenderTargetBeforeWarpBlend_RenderThread(FRHICommandListImmediate& RHICmdList, FTexture2DRHIParamRef SrcTexture) const
+void FDisplayClusterDeviceBase_PostProcess::PerformPostProcessRenderTargetBeforeWarpBlend_RenderThread(FRHICommandListImmediate& RHICmdList, FRHITexture2D* SrcTexture) const
 {
 	DISPLAY_CLUSTER_FUNC_TRACE(LogDisplayClusterRender);
 
@@ -133,7 +133,7 @@ static FAutoConsoleVariableRef CVarPostprocessViewAfterWarpBlend(
 	TEXT("Enable PP per view after warp&blend (0 = disabled)\n")
 );
 
-void FDisplayClusterDeviceBase_PostProcess::PerformPostProcessViewAfterWarpBlend_RenderThread(FRHICommandListImmediate& RHICmdList, FTexture2DRHIParamRef SrcTexture, const FIntRect& FrameRect) const
+void FDisplayClusterDeviceBase_PostProcess::PerformPostProcessViewAfterWarpBlend_RenderThread(FRHICommandListImmediate& RHICmdList, FRHITexture2D* SrcTexture, const FIntRect& FrameRect) const
 {
 	DISPLAY_CLUSTER_FUNC_TRACE(LogDisplayClusterRender);
 
@@ -171,7 +171,7 @@ static FAutoConsoleVariableRef CVarPostprocessFrameAfterWarpBlend(
 	TEXT("Enable PP per eye frame after warp&blend (0 = disabled)\n")
 );
 
-void FDisplayClusterDeviceBase_PostProcess::PerformPostProcessFrameAfterWarpBlend_RenderThread(FRHICommandListImmediate& RHICmdList, FTexture2DRHIParamRef SrcTexture, const FIntRect& FrameRect) const
+void FDisplayClusterDeviceBase_PostProcess::PerformPostProcessFrameAfterWarpBlend_RenderThread(FRHICommandListImmediate& RHICmdList, FRHITexture2D* SrcTexture, const FIntRect& FrameRect) const
 {
 	DISPLAY_CLUSTER_FUNC_TRACE(LogDisplayClusterRender);
 
@@ -209,7 +209,7 @@ static FAutoConsoleVariableRef CVarPostprocessTargetAfterWarpBlend(
 	TEXT("Enable PP for the whole render target after warp&blend (0 = disabled)\n")
 );
 
-void FDisplayClusterDeviceBase_PostProcess::PerformPostProcessRenderTargetAfterWarpBlend_RenderThread(FRHICommandListImmediate& RHICmdList, FTexture2DRHIParamRef SrcTexture) const
+void FDisplayClusterDeviceBase_PostProcess::PerformPostProcessRenderTargetAfterWarpBlend_RenderThread(FRHICommandListImmediate& RHICmdList, FRHITexture2D* SrcTexture) const
 {
 	DISPLAY_CLUSTER_FUNC_TRACE(LogDisplayClusterRender);
 
diff --git a/Engine/Plugins/Runtime/nDisplay/Source/DisplayCluster/Private/Render/Device/DisplayClusterDeviceBase_PostProcess.h b/Engine/Plugins/Runtime/nDisplay/Source/DisplayCluster/Private/Render/Device/DisplayClusterDeviceBase_PostProcess.h
--- a/Engine/Plugins/Runtime/nDisplay/Source/DisplayCluster/Private/Render/Device/DisplayClusterDeviceBase_PostProcess.h
+++ b/Engine/Plugins/Runtime/nDisplay/Source/DisplayCluster/Private/Render/Device/DisplayClusterDeviceBase_PostProcess.h
@@ -25,12 +25,12 @@ public:
 
 
 protected:
-	virtual void PerformPostProcessViewBeforeWarpBlend_RenderThread(FRHICommandListImmediate& RHICmdList, FTexture2DRHIParamRef SrcTexture, const FIntRect& ViewRect) const override final;
-	virtual void PerformPostProcessFrameBeforeWarpBlend_RenderThread(FRHICommandListImmediate& RHICmdList, FTexture2DRHIParamRef SrcTexture, const FIntRect& FrameRect) const override final;
-	virtual void PerformPostProcessRenderTargetBeforeWarpBlend_RenderThread(FRHICommandListImmediate& RHICmdList, FTexture2DRHIParamRef SrcTexture) const override final;
-	virtual void PerformPostProcessViewAfterWarpBlend_RenderThread(FRHICommandListImmediate& RHICmdList, FTexture2DRHIParamRef SrcTexture, const FIntRect& ViewRect) const override final;
-	virtual void PerformPostProcessFrameAfterWarpBlend_RenderThread(FRHICommandListImmediate& RHICmdList, FTexture2DRHIParamRef SrcTexture, const FIntRect& FrameRect) const override final;
-	virtual void PerformPostProcessRenderTargetAfterWarpBlend_RenderThread(FRHICommandListImmediate& RHICmdList, FTexture2DRHIParamRef SrcTexture) const override final;
+	virtual void PerformPostProcessViewBeforeWarpBlend_RenderThread(FRHICommandListImmediate& RHICmdList, FRHITexture2D* SrcTexture, const FIntRect& ViewRect) const override final;
+	virtual void PerformPostProcessFrameBeforeWarpBlend_RenderThread(FRHICommandListImmediate& RHICmdList, FRHITexture2D* SrcTexture, const FIntRect& FrameRect) const override final;
+	virtual void PerformPostProcessRenderTargetBeforeWarpBlend_RenderThread(FRHICommandListImmediate& RHICmdList, FRHITexture2D* SrcTexture) const override final;
+	virtual void PerformPostProcessViewAfterWarpBlend_RenderThread(FRHICommandListImmediate& RHICmdList, FRHITexture2D* SrcTexture, const FIntRect& ViewRect) const override final;
+	virtual void PerformPostProcessFrameAfterWarpBlend_RenderThread(FRHICommandListImmediate& RHICmdList, FRHITexture2D* SrcTexture, const FIntRect& FrameRect) const override final;
+	virtual void PerformPostProcessRenderTargetAfterWarpBlend_RenderThread(FRHICommandListImmediate& RHICmdList, FRHITexture2D* SrcTexture) const override final;
 
 	mutable TArray<IDisplayClusterRenderManager::FDisplayClusterPPInfo> PPOperations;
 
diff --git a/Engine/Plugins/Runtime/nDisplay/Source/DisplayCluster/Private/Render/Device/Monoscopic/DisplayClusterDeviceMonoscopicOpenGL.h b/Engine/Plugins/Runtime/nDisplay/Source/DisplayCluster/Private/Render/Device/Monoscopic/DisplayClusterDeviceMonoscopicOpenGL.h
--- a/Engine/Plugins/Runtime/nDisplay/Source/DisplayCluster/Private/Render/Device/Monoscopic/DisplayClusterDeviceMonoscopicOpenGL.h
+++ b/Engine/Plugins/Runtime/nDisplay/Source/DisplayCluster/Private/Render/Device/Monoscopic/DisplayClusterDeviceMonoscopicOpenGL.h
@@ -32,7 +32,7 @@ protected:
 	//////////////////////////////////////////////////////////////////////////////////////////////
 	// FDisplayClusterDeviceBase
 	//////////////////////////////////////////////////////////////////////////////////////////////
-	virtual void CopyTextureToBackBuffer_RenderThread(FRHICommandListImmediate& RHICmdList, FTexture2DRHIParamRef BackBuffer, FTexture2DRHIParamRef SrcTexture, FVector2D WindowSize) const override
+	virtual void CopyTextureToBackBuffer_RenderThread(FRHICommandListImmediate& RHICmdList, FRHITexture2D* BackBuffer, FRHITexture2D* SrcTexture, FVector2D WindowSize) const override
 	{
 		// Forbid the default FDisplayClusterDeviceBase::CopyTextureToBackBuffer_RenderThread implementation. The OpenGL copies data in the Present method.
 	}
diff --git a/Engine/Plugins/Runtime/nDisplay/Source/DisplayCluster/Private/Render/Device/QuadBufferStereo/DisplayClusterDeviceQuadBufferStereoBase.cpp b/Engine/Plugins/Runtime/nDisplay/Source/DisplayCluster/Private/Render/Device/QuadBufferStereo/DisplayClusterDeviceQuadBufferStereoBase.cpp
--- a/Engine/Plugins/Runtime/nDisplay/Source/DisplayCluster/Private/Render/Device/QuadBufferStereo/DisplayClusterDeviceQuadBufferStereoBase.cpp
+++ b/Engine/Plugins/Runtime/nDisplay/Source/DisplayCluster/Private/Render/Device/QuadBufferStereo/DisplayClusterDeviceQuadBufferStereoBase.cpp
@@ -65,7 +65,7 @@ void FDisplayClusterDeviceQuadBufferStereoBase::AdjustViewRect(enum EStereoscopi
 	UE_LOG(LogDisplayClusterRender, Verbose, TEXT("Adjusted view rect: ViewportIdx=%d, EyeType=%d, [%d,%d - %d,%d]"), ViewportIndex, int(EyeType), r.Min.X, r.Min.Y, r.Max.X, r.Max.Y);
 }
 
-void FDisplayClusterDeviceQuadBufferStereoBase::CopyTextureToBackBuffer_RenderThread(FRHICommandListImmediate& RHICmdList, FTexture2DRHIParamRef BackBuffer, FTexture2DRHIParamRef SrcTexture, FVector2D WindowSize) const
+void FDisplayClusterDeviceQuadBufferStereoBase::CopyTextureToBackBuffer_RenderThread(FRHICommandListImmediate& RHICmdList, FRHITexture2D* BackBuffer, FRHITexture2D* SrcTexture, FVector2D WindowSize) const
 {
 	DISPLAY_CLUSTER_FUNC_TRACE(LogDisplayClusterRender);
 
diff --git a/Engine/Plugins/Runtime/nDisplay/Source/DisplayCluster/Private/Render/Device/QuadBufferStereo/DisplayClusterDeviceQuadBufferStereoBase.h b/Engine/Plugins/Runtime/nDisplay/Source/DisplayCluster/Private/Render/Device/QuadBufferStereo/DisplayClusterDeviceQuadBufferStereoBase.h
--- a/Engine/Plugins/Runtime/nDisplay/Source/DisplayCluster/Private/Render/Device/QuadBufferStereo/DisplayClusterDeviceQuadBufferStereoBase.h
+++ b/Engine/Plugins/Runtime/nDisplay/Source/DisplayCluster/Private/Render/Device/QuadBufferStereo/DisplayClusterDeviceQuadBufferStereoBase.h
@@ -32,5 +32,5 @@ protected:
 	//////////////////////////////////////////////////////////////////////////////////////////////
 	// FDisplayClusterDeviceBase
 	//////////////////////////////////////////////////////////////////////////////////////////////
-	virtual void CopyTextureToBackBuffer_RenderThread(FRHICommandListImmediate& RHICmdList, FTexture2DRHIParamRef BackBuffer, FTexture2DRHIParamRef SrcTexture, FVector2D WindowSize) const override;
+	virtual void CopyTextureToBackBuffer_RenderThread(FRHICommandListImmediate& RHICmdList, FRHITexture2D* BackBuffer, FRHITexture2D* SrcTexture, FVector2D WindowSize) const override;
 };
diff --git a/Engine/Plugins/Runtime/nDisplay/Source/DisplayCluster/Private/Render/Device/QuadBufferStereo/DisplayClusterDeviceQuadBufferStereoOpenGL.h b/Engine/Plugins/Runtime/nDisplay/Source/DisplayCluster/Private/Render/Device/QuadBufferStereo/DisplayClusterDeviceQuadBufferStereoOpenGL.h
--- a/Engine/Plugins/Runtime/nDisplay/Source/DisplayCluster/Private/Render/Device/QuadBufferStereo/DisplayClusterDeviceQuadBufferStereoOpenGL.h
+++ b/Engine/Plugins/Runtime/nDisplay/Source/DisplayCluster/Private/Render/Device/QuadBufferStereo/DisplayClusterDeviceQuadBufferStereoOpenGL.h
@@ -32,7 +32,7 @@ protected:
 	//////////////////////////////////////////////////////////////////////////////////////////////
 	// FDisplayClusterDeviceBase
 	//////////////////////////////////////////////////////////////////////////////////////////////
-	virtual void CopyTextureToBackBuffer_RenderThread(FRHICommandListImmediate& RHICmdList, FTexture2DRHIParamRef BackBuffer, FTexture2DRHIParamRef SrcTexture, FVector2D WindowSize) const override
+	virtual void CopyTextureToBackBuffer_RenderThread(FRHICommandListImmediate& RHICmdList, FRHITexture2D* BackBuffer, FRHITexture2D* SrcTexture, FVector2D WindowSize) const override
 	{
 		// Forbid the default FDisplayClusterDeviceBase::CopyTextureToBackBuffer_RenderThread implementation. The OpenGL copies data in the Present method.
 	}
diff --git a/Engine/Plugins/Runtime/nDisplay/Source/DisplayCluster/Public/Render/PostProcess/IDisplayClusterPostProcess.h b/Engine/Plugins/Runtime/nDisplay/Source/DisplayCluster/Public/Render/PostProcess/IDisplayClusterPostProcess.h
--- a/Engine/Plugins/Runtime/nDisplay/Source/DisplayCluster/Public/Render/PostProcess/IDisplayClusterPostProcess.h
+++ b/Engine/Plugins/Runtime/nDisplay/Source/DisplayCluster/Public/Render/PostProcess/IDisplayClusterPostProcess.h
@@ -37,7 +37,7 @@ public:
 	* @param SrcTexture - Source texture
 	* @param ViewRect   - View region of the texture
 	*/
-	virtual void PerformPostProcessViewBeforeWarpBlend_RenderThread(FRHICommandListImmediate& RHICmdList, FTexture2DRHIParamRef SrcTexture, const FIntRect& ViewRect) const
+	virtual void PerformPostProcessViewBeforeWarpBlend_RenderThread(FRHICommandListImmediate& RHICmdList, FRHITexture2D* SrcTexture, const FIntRect& ViewRect) const
 	{ }
 
 	/**
@@ -58,7 +58,7 @@ public:
 	* @param SrcTexture - Source texture
 	* @param ViewRect   - View region of the texture
 	*/
-	virtual void PerformPostProcessViewAfterWarpBlend_RenderThread(FRHICommandListImmediate& RHICmdList, FTexture2DRHIParamRef SrcTexture, const FIntRect& ViewRect) const
+	virtual void PerformPostProcessViewAfterWarpBlend_RenderThread(FRHICommandListImmediate& RHICmdList, FRHITexture2D* SrcTexture, const FIntRect& ViewRect) const
 	{ }
 
 	/**
@@ -81,7 +81,7 @@ public:
 	* @param SrcTexture - Source texture
 	* @param ViewRect   - Frame region of the texture
 	*/
-	virtual void PerformPostProcessFrameBeforeWarpBlend_RenderThread(FRHICommandListImmediate& RHICmdList, FTexture2DRHIParamRef SrcTexture, const FIntRect& FrameRect) const
+	virtual void PerformPostProcessFrameBeforeWarpBlend_RenderThread(FRHICommandListImmediate& RHICmdList, FRHITexture2D* SrcTexture, const FIntRect& FrameRect) const
 	{ }
 
 	/**
@@ -104,7 +104,7 @@ public:
 	* @param SrcTexture - Source texture
 	* @param ViewRect   - Frame region of the texture
 	*/
-	virtual void PerformPostProcessFrameAfterWarpBlend_RenderThread(FRHICommandListImmediate& RHICmdList, FTexture2DRHIParamRef SrcTexture, const FIntRect& FrameRect) const
+	virtual void PerformPostProcessFrameAfterWarpBlend_RenderThread(FRHICommandListImmediate& RHICmdList, FRHITexture2D* SrcTexture, const FIntRect& FrameRect) const
 	{ }
 
 	/**
@@ -123,7 +123,7 @@ public:
 	* @param RHICmdList - RHI command list
 	* @param SrcTexture - Source texture
 	*/
-	virtual void PerformPostProcessRenderTargetBeforeWarpBlend_RenderThread(FRHICommandListImmediate& RHICmdList, FTexture2DRHIParamRef SrcTexture) const
+	virtual void PerformPostProcessRenderTargetBeforeWarpBlend_RenderThread(FRHICommandListImmediate& RHICmdList, FRHITexture2D* SrcTexture) const
 	{ }
 
 	/**
@@ -142,6 +142,6 @@ public:
 	* @param RHICmdList - RHI command list
 	* @param SrcTexture - Source texture
 	*/
-	virtual void PerformPostProcessRenderTargetAfterWarpBlend_RenderThread(FRHICommandListImmediate& RHICmdList, FTexture2DRHIParamRef SrcTexture) const
+	virtual void PerformPostProcessRenderTargetAfterWarpBlend_RenderThread(FRHICommandListImmediate& RHICmdList, FRHITexture2D* SrcTexture) const
 	{ }
 };
diff --git a/Engine/Plugins/Runtime/nDisplay/Source/DisplayCluster/Public/Render/Projection/IDisplayClusterProjectionPolicy.h b/Engine/Plugins/Runtime/nDisplay/Source/DisplayCluster/Public/Render/Projection/IDisplayClusterProjectionPolicy.h
--- a/Engine/Plugins/Runtime/nDisplay/Source/DisplayCluster/Public/Render/Projection/IDisplayClusterProjectionPolicy.h
+++ b/Engine/Plugins/Runtime/nDisplay/Source/DisplayCluster/Public/Render/Projection/IDisplayClusterProjectionPolicy.h
@@ -89,6 +89,6 @@ public:
 	* @param ViewportRect - Region of the SrcTexture to perform warp&blend operations
 	*
 	*/
-	virtual void ApplyWarpBlend_RenderThread(const uint32 ViewIdx, FRHICommandListImmediate& RHICmdList, FTexture2DRHIParamRef SrcTexture, const FIntRect& ViewportRect)
+	virtual void ApplyWarpBlend_RenderThread(const uint32 ViewIdx, FRHICommandListImmediate& RHICmdList, FRHITexture2D* SrcTexture, const FIntRect& ViewportRect)
 	{ }
 };
diff --git a/Engine/Plugins/Runtime/nDisplay/Source/DisplayClusterProjection/Private/Policy/EasyBlend/DX11/DisplayClusterProjectionEasyBlendViewAdapterDX11.cpp b/Engine/Plugins/Runtime/nDisplay/Source/DisplayClusterProjection/Private/Policy/EasyBlend/DX11/DisplayClusterProjectionEasyBlendViewAdapterDX11.cpp
--- a/Engine/Plugins/Runtime/nDisplay/Source/DisplayClusterProjection/Private/Policy/EasyBlend/DX11/DisplayClusterProjectionEasyBlendViewAdapterDX11.cpp
+++ b/Engine/Plugins/Runtime/nDisplay/Source/DisplayClusterProjection/Private/Policy/EasyBlend/DX11/DisplayClusterProjectionEasyBlendViewAdapterDX11.cpp
@@ -157,7 +157,7 @@ bool FDisplayClusterProjectionEasyBlendViewAdapterDX11::GetProjectionMatrix(cons
 	return true;
 }
 
-bool FDisplayClusterProjectionEasyBlendViewAdapterDX11::ApplyWarpBlend_RenderThread(const uint32 ViewIdx, FRHICommandListImmediate& RHICmdList, FTexture2DRHIParamRef SrcTexture, const FIntRect& ViewportRect)
+bool FDisplayClusterProjectionEasyBlendViewAdapterDX11::ApplyWarpBlend_RenderThread(const uint32 ViewIdx, FRHICommandListImmediate& RHICmdList, FRHITexture2D* SrcTexture, const FIntRect& ViewportRect)
 {
 	check(IsInRenderingThread());
 	check(Views.Num() > (int)ViewIdx);
@@ -300,7 +300,7 @@ bool FDisplayClusterProjectionEasyBlendViewAdapterDX11::InitializeResources_Rend
 	return bIsRenderResourcesInitialized;
 }
 
-void FDisplayClusterProjectionEasyBlendViewAdapterDX11::LoadViewportTexture_RenderThread(const uint32 ViewIdx, FRHICommandListImmediate& RHICmdList, FTexture2DRHIParamRef SrcTexture, const FIntRect& ViewportRect)
+void FDisplayClusterProjectionEasyBlendViewAdapterDX11::LoadViewportTexture_RenderThread(const uint32 ViewIdx, FRHICommandListImmediate& RHICmdList, FRHITexture2D* SrcTexture, const FIntRect& ViewportRect)
 {
 	check(IsInRenderingThread());
 
@@ -325,7 +325,7 @@ void FDisplayClusterProjectionEasyBlendViewAdapterDX11::LoadViewportTexture_Rend
 	RHICmdList.ImmediateFlush(EImmediateFlushType::FlushRHIThreadFlushResources);
 }
 
-void FDisplayClusterProjectionEasyBlendViewAdapterDX11::SaveViewportTexture_RenderThread(const uint32 ViewIdx, FRHICommandListImmediate& RHICmdList, FTexture2DRHIParamRef DstTexture, const FIntRect& ViewportRect)
+void FDisplayClusterProjectionEasyBlendViewAdapterDX11::SaveViewportTexture_RenderThread(const uint32 ViewIdx, FRHICommandListImmediate& RHICmdList, FRHITexture2D* DstTexture, const FIntRect& ViewportRect)
 {
 	check(IsInRenderingThread());
 
diff --git a/Engine/Plugins/Runtime/nDisplay/Source/DisplayClusterProjection/Private/Policy/EasyBlend/DX11/DisplayClusterProjectionEasyBlendViewAdapterDX11.h b/Engine/Plugins/Runtime/nDisplay/Source/DisplayClusterProjection/Private/Policy/EasyBlend/DX11/DisplayClusterProjectionEasyBlendViewAdapterDX11.h
--- a/Engine/Plugins/Runtime/nDisplay/Source/DisplayClusterProjection/Private/Policy/EasyBlend/DX11/DisplayClusterProjectionEasyBlendViewAdapterDX11.h
+++ b/Engine/Plugins/Runtime/nDisplay/Source/DisplayClusterProjection/Private/Policy/EasyBlend/DX11/DisplayClusterProjectionEasyBlendViewAdapterDX11.h
@@ -21,12 +21,12 @@ public:
 public:
 	virtual bool CalculateView(const uint32 ViewIdx, FVector& InOutViewLocation, FRotator& InOutViewRotation, const FVector& ViewOffset, const float WorldToMeters, const float NCP, const float FCP) override;
 	virtual bool GetProjectionMatrix(const uint32 ViewIdx, FMatrix& OutPrjMatrix) override;
-	virtual bool ApplyWarpBlend_RenderThread(const uint32 ViewIdx, FRHICommandListImmediate& RHICmdList, FTexture2DRHIParamRef SrcTexture, const FIntRect& ViewportRect) override;
+	virtual bool ApplyWarpBlend_RenderThread(const uint32 ViewIdx, FRHICommandListImmediate& RHICmdList, FRHITexture2D* SrcTexture, const FIntRect& ViewportRect) override;
 
 private:
 	bool InitializeResources_RenderThread();
-	void LoadViewportTexture_RenderThread(const uint32 ViewIdx, FRHICommandListImmediate& RHICmdList, FTexture2DRHIParamRef SrcTexture, const FIntRect& ViewportRect);
-	void SaveViewportTexture_RenderThread(const uint32 ViewIdx, FRHICommandListImmediate& RHICmdList, FTexture2DRHIParamRef DstTexture, const FIntRect& ViewportRect);
+	void LoadViewportTexture_RenderThread(const uint32 ViewIdx, FRHICommandListImmediate& RHICmdList, FRHITexture2D* SrcTexture, const FIntRect& ViewportRect);
+	void SaveViewportTexture_RenderThread(const uint32 ViewIdx, FRHICommandListImmediate& RHICmdList, FRHITexture2D* DstTexture, const FIntRect& ViewportRect);
 
 private:
 	float ZNear;
diff --git a/Engine/Plugins/Runtime/nDisplay/Source/DisplayClusterProjection/Private/Policy/EasyBlend/DisplayClusterProjectionEasyBlendPolicyBase.cpp b/Engine/Plugins/Runtime/nDisplay/Source/DisplayClusterProjection/Private/Policy/EasyBlend/DisplayClusterProjectionEasyBlendPolicyBase.cpp
--- a/Engine/Plugins/Runtime/nDisplay/Source/DisplayClusterProjection/Private/Policy/EasyBlend/DisplayClusterProjectionEasyBlendPolicyBase.cpp
+++ b/Engine/Plugins/Runtime/nDisplay/Source/DisplayClusterProjection/Private/Policy/EasyBlend/DisplayClusterProjectionEasyBlendPolicyBase.cpp
@@ -131,7 +131,7 @@ bool FDisplayClusterProjectionEasyBlendPolicyBase::IsWarpBlendSupported()
 	return true;
 }
 
-void FDisplayClusterProjectionEasyBlendPolicyBase::ApplyWarpBlend_RenderThread(const uint32 ViewIdx, FRHICommandListImmediate& RHICmdList, FTexture2DRHIParamRef SrcTexture, const FIntRect& ViewportRect)
+void FDisplayClusterProjectionEasyBlendPolicyBase::ApplyWarpBlend_RenderThread(const uint32 ViewIdx, FRHICommandListImmediate& RHICmdList, FRHITexture2D* SrcTexture, const FIntRect& ViewportRect)
 {
 	check(IsInRenderingThread());
 
diff --git a/Engine/Plugins/Runtime/nDisplay/Source/DisplayClusterProjection/Private/Policy/EasyBlend/DisplayClusterProjectionEasyBlendPolicyBase.h b/Engine/Plugins/Runtime/nDisplay/Source/DisplayClusterProjection/Private/Policy/EasyBlend/DisplayClusterProjectionEasyBlendPolicyBase.h
--- a/Engine/Plugins/Runtime/nDisplay/Source/DisplayClusterProjection/Private/Policy/EasyBlend/DisplayClusterProjectionEasyBlendPolicyBase.h
+++ b/Engine/Plugins/Runtime/nDisplay/Source/DisplayClusterProjection/Private/Policy/EasyBlend/DisplayClusterProjectionEasyBlendPolicyBase.h
@@ -32,7 +32,7 @@ public:
 	virtual bool GetProjectionMatrix(const uint32 ViewIdx, FMatrix& OutPrjMatrix) override;
 
 	virtual bool IsWarpBlendSupported() override;
-	virtual void ApplyWarpBlend_RenderThread(const uint32 ViewIdx, FRHICommandListImmediate& RHICmdList, FTexture2DRHIParamRef SrcTexture, const FIntRect& ViewportRect) override;
+	virtual void ApplyWarpBlend_RenderThread(const uint32 ViewIdx, FRHICommandListImmediate& RHICmdList, FRHITexture2D* SrcTexture, const FIntRect& ViewportRect) override;
 
 protected:
 	// Delegate view adapter instantiation to the RHI specific children
diff --git a/Engine/Plugins/Runtime/nDisplay/Source/DisplayClusterProjection/Private/Policy/EasyBlend/DisplayClusterProjectionEasyBlendViewAdapterBase.h b/Engine/Plugins/Runtime/nDisplay/Source/DisplayClusterProjection/Private/Policy/EasyBlend/DisplayClusterProjectionEasyBlendViewAdapterBase.h
--- a/Engine/Plugins/Runtime/nDisplay/Source/DisplayClusterProjection/Private/Policy/EasyBlend/DisplayClusterProjectionEasyBlendViewAdapterBase.h
+++ b/Engine/Plugins/Runtime/nDisplay/Source/DisplayClusterProjection/Private/Policy/EasyBlend/DisplayClusterProjectionEasyBlendViewAdapterBase.h
@@ -46,7 +46,7 @@ public:
 public:
 	virtual bool CalculateView(const uint32 ViewIdx, FVector& InOutViewLocation, FRotator& InOutViewRotation, const FVector& ViewOffset, const float WorldToMeters, const float NCP, const float FCP) = 0;
 	virtual bool GetProjectionMatrix(const uint32 ViewIdx, FMatrix& OutPrjMatrix) = 0;
-	virtual bool ApplyWarpBlend_RenderThread(const uint32 ViewIdx, FRHICommandListImmediate& RHICmdList, FTexture2DRHIParamRef SrcTexture, const FIntRect& ViewportRect) = 0;
+	virtual bool ApplyWarpBlend_RenderThread(const uint32 ViewIdx, FRHICommandListImmediate& RHICmdList, FRHITexture2D* SrcTexture, const FIntRect& ViewportRect) = 0;
 
 private:
 	const FInitParams InitParams;
diff --git a/Engine/Plugins/Runtime/nDisplay/Source/DisplayClusterProjection/Private/Policy/MPCDI/DisplayClusterProjectionMPCDIPolicy.cpp b/Engine/Plugins/Runtime/nDisplay/Source/DisplayClusterProjection/Private/Policy/MPCDI/DisplayClusterProjectionMPCDIPolicy.cpp
--- a/Engine/Plugins/Runtime/nDisplay/Source/DisplayClusterProjection/Private/Policy/MPCDI/DisplayClusterProjectionMPCDIPolicy.cpp
+++ b/Engine/Plugins/Runtime/nDisplay/Source/DisplayClusterProjection/Private/Policy/MPCDI/DisplayClusterProjectionMPCDIPolicy.cpp
@@ -156,7 +156,7 @@ bool FDisplayClusterProjectionMPCDIPolicy::IsWarpBlendSupported()
 	return true;
 }
 
-void FDisplayClusterProjectionMPCDIPolicy::ApplyWarpBlend_RenderThread(const uint32 ViewIdx, FRHICommandListImmediate& RHICmdList, FTexture2DRHIParamRef SrcTexture, const FIntRect& ViewportRect)
+void FDisplayClusterProjectionMPCDIPolicy::ApplyWarpBlend_RenderThread(const uint32 ViewIdx, FRHICommandListImmediate& RHICmdList, FRHITexture2D* SrcTexture, const FIntRect& ViewportRect)
 {
 	check(IsInRenderingThread());
 
diff --git a/Engine/Plugins/Runtime/nDisplay/Source/DisplayClusterProjection/Private/Policy/MPCDI/DisplayClusterProjectionMPCDIPolicy.h b/Engine/Plugins/Runtime/nDisplay/Source/DisplayClusterProjection/Private/Policy/MPCDI/DisplayClusterProjectionMPCDIPolicy.h
--- a/Engine/Plugins/Runtime/nDisplay/Source/DisplayClusterProjection/Private/Policy/MPCDI/DisplayClusterProjectionMPCDIPolicy.h
+++ b/Engine/Plugins/Runtime/nDisplay/Source/DisplayClusterProjection/Private/Policy/MPCDI/DisplayClusterProjectionMPCDIPolicy.h
@@ -37,7 +37,7 @@ public:
 	virtual bool GetProjectionMatrix(const uint32 ViewIdx, FMatrix& OutPrjMatrix) override;
 
 	virtual bool IsWarpBlendSupported() override;
-	virtual void ApplyWarpBlend_RenderThread(const uint32 ViewIdx, FRHICommandListImmediate& RHICmdList, FTexture2DRHIParamRef SrcTexture, const FIntRect& ViewportRect) override;
+	virtual void ApplyWarpBlend_RenderThread(const uint32 ViewIdx, FRHICommandListImmediate& RHICmdList, FRHITexture2D* SrcTexture, const FIntRect& ViewportRect) override;
 
 protected:
 	bool ReadConfigData(const FString& ViewportId, FString& OutFile, FString& OutBuffer, FString& OutRegion, FString& OutOrigin);
diff --git a/Engine/Plugins/Runtime/nDisplay/Source/MPCDI/Private/MPCDIShader.h b/Engine/Plugins/Runtime/nDisplay/Source/MPCDI/Private/MPCDIShader.h
--- a/Engine/Plugins/Runtime/nDisplay/Source/MPCDI/Private/MPCDIShader.h
+++ b/Engine/Plugins/Runtime/nDisplay/Source/MPCDI/Private/MPCDIShader.h
@@ -190,7 +190,7 @@ public:
 		}
 
 	template<typename TShaderRHIParamRef>
-	void SetParameters(FRHICommandListImmediate& RHICmdList, const TShaderRHIParamRef ShaderRHI, FTexture2DRHIParamRef SourceTexture)
+	void SetParameters(FRHICommandListImmediate& RHICmdList, const TShaderRHIParamRef ShaderRHI, FRHITexture2D* SourceTexture)
 	{
 #if 0
 		//Cubemap support
@@ -222,7 +222,7 @@ public:
 	}
 
 	template<typename TShaderRHIParamRef>
-	void SetParameters(FRHICommandListImmediate& RHICmdList, const TShaderRHIParamRef ShaderRHI, FTexture2DRHIParamRef SourceTexture, const IMPCDI::FShaderInputData& ShaderInputData, MPCDI::FMPCDIRegion& MPCDIRegionData)
+	void SetParameters(FRHICommandListImmediate& RHICmdList, const TShaderRHIParamRef ShaderRHI, FRHITexture2D* SourceTexture, const IMPCDI::FShaderInputData& ShaderInputData, MPCDI::FMPCDIRegion& MPCDIRegionData)
 	{
 		SetParameters(RHICmdList, ShaderRHI, SourceTexture);
 		SetParameters(RHICmdList, ShaderRHI, ShaderInputData);
diff --git a/Engine/Plugins/Runtime/nDisplay/Source/MPCDI/Public/IMPCDI.h b/Engine/Plugins/Runtime/nDisplay/Source/MPCDI/Public/IMPCDI.h
--- a/Engine/Plugins/Runtime/nDisplay/Source/MPCDI/Public/IMPCDI.h
+++ b/Engine/Plugins/Runtime/nDisplay/Source/MPCDI/Public/IMPCDI.h
@@ -77,8 +77,8 @@ public:
 
 	struct FTextureWarpData
 	{
-		FTexture2DRHIParamRef SrcTexture; // Source texture (shader render)
-		FTexture2DRHIParamRef DstTexture; // Render target destination
+		FRHITexture2D* SrcTexture; // Source texture (shader render)
+		FRHITexture2D* DstTexture; // Render target destination
 
 		FIntRect SrcRect; // Source texture region
 		FIntRect DstRect; // Destination texture region
diff --git a/Engine/Plugins/Runtime/nDisplay/Source/PicpMPCDI/Private/PicpBlurPostProcess.cpp b/Engine/Plugins/Runtime/nDisplay/Source/PicpMPCDI/Private/PicpBlurPostProcess.cpp
--- a/Engine/Plugins/Runtime/nDisplay/Source/PicpMPCDI/Private/PicpBlurPostProcess.cpp
+++ b/Engine/Plugins/Runtime/nDisplay/Source/PicpMPCDI/Private/PicpBlurPostProcess.cpp
@@ -32,9 +32,9 @@ IMPLEMENT_SHADER_TYPE(, FDirectComposePS, PostProcessShaderFileName, TEXT("Direc
 
 void FPicpBlurPostProcess::ApplyCompose_RenderThread(
 	FRHICommandListImmediate& RHICmdList, 
-	FTextureRHIParamRef OverlayTexture,
-	FTexture2DRHIParamRef DstRenderTarget, 
-	FTexture2DRHIParamRef CopyTexture
+	FRHITexture* OverlayTexture,
+	FRHITexture2D* DstRenderTarget,
+	FRHITexture2D* CopyTexture
 )
 {
 	check(IsInRenderingThread());
@@ -79,8 +79,8 @@ void FPicpBlurPostProcess::ApplyCompose_RenderThread(
 template<uint32 ShaderType>
 static void PicpBlurPostProcess_RenderThread(
 	FRHICommandListImmediate& RHICmdList,
-	FTexture2DRHIParamRef InOutTexture,
-	FTexture2DRHIParamRef TempTexture,
+	FRHITexture2D* InOutTexture,
+	FRHITexture2D* TempTexture,
 	int KernelRadius,
 	float KernelScale
 )
@@ -127,7 +127,7 @@ static void PicpBlurPostProcess_RenderThread(
 	}	
 }
 
-void FPicpBlurPostProcess::ApplyBlur_RenderThread(FRHICommandListImmediate& RHICmdList, FTexture2DRHIParamRef InOutRT, FTexture2DRHIParamRef TempRT, int KernelRadius, float KernelScale, EPicpBlurPostProcessShaderType BlurType)
+void FPicpBlurPostProcess::ApplyBlur_RenderThread(FRHICommandListImmediate& RHICmdList, FRHITexture2D* InOutRT, FRHITexture2D* TempRT, int KernelRadius, float KernelScale, EPicpBlurPostProcessShaderType BlurType)
 {
 	switch (BlurType)
 	{
@@ -156,8 +156,8 @@ void FPicpBlurPostProcess::ApplyBlur(UTextureRenderTarget2D* InOutRenderTarget,
 	ENQUEUE_RENDER_COMMAND(CaptureCommand)(
 		[InOutRT, TempRT, KernelRadius, KernelScale, BlurType](FRHICommandListImmediate& RHICmdList)
 		{
-			FTexture2DRHIParamRef RenderTarget = InOutRT->GetRenderTargetTexture()->GetTexture2D();
-			FTexture2DRHIParamRef TempTexture = TempRT->GetRenderTargetTexture()->GetTexture2D();
+		FRHITexture2D* RenderTarget = InOutRT->GetRenderTargetTexture()->GetTexture2D();
+		FRHITexture2D* TempTexture = TempRT->GetRenderTargetTexture()->GetTexture2D();
 			if (nullptr==RenderTarget || nullptr==TempTexture)
 			{
 				//@todo handle error
@@ -184,9 +184,9 @@ void FPicpBlurPostProcess::ApplyCompose(UTexture* InputTexture, UTextureRenderTa
 	ENQUEUE_RENDER_COMMAND(CaptureCommand)(
 		[OverlayTextureResource, TextureRenderTarget, TextureRenderResult](FRHICommandListImmediate& RHICmdList)
 	{
-		FTexture2DRHIParamRef RTTexture = TextureRenderTarget ? TextureRenderTarget->GetRenderTargetTexture()->GetTexture2D(): nullptr;
-		FTexture2DRHIParamRef ResultTexture = TextureRenderResult ? TextureRenderResult->GetRenderTargetTexture()->GetTexture2D() : nullptr;
-		FTextureRHIParamRef   OverlayTexture = OverlayTextureResource ? OverlayTextureResource->TextureRHI : nullptr;;
+		FRHITexture2D* RTTexture = TextureRenderTarget ? TextureRenderTarget->GetRenderTargetTexture()->GetTexture2D(): nullptr;
+		FRHITexture2D* ResultTexture = TextureRenderResult ? TextureRenderResult->GetRenderTargetTexture()->GetTexture2D() : nullptr;
+		FRHITexture*   OverlayTexture = OverlayTextureResource ? OverlayTextureResource->TextureRHI : nullptr;;
 		
 		if (nullptr == RTTexture)
 		{
diff --git a/Engine/Plugins/Runtime/nDisplay/Source/PicpMPCDI/Private/PicpBlurPostProcess.h b/Engine/Plugins/Runtime/nDisplay/Source/PicpMPCDI/Private/PicpBlurPostProcess.h
--- a/Engine/Plugins/Runtime/nDisplay/Source/PicpMPCDI/Private/PicpBlurPostProcess.h
+++ b/Engine/Plugins/Runtime/nDisplay/Source/PicpMPCDI/Private/PicpBlurPostProcess.h
@@ -24,8 +24,8 @@ public:
 	static void ApplyCompose(UTexture* InputTexture, UTextureRenderTarget2D* OutputRenderTarget, UTextureRenderTarget2D* Result);
 
 	//onRenderThread
-	static void ApplyBlur_RenderThread(FRHICommandListImmediate& RHICmdList, FTexture2DRHIParamRef InOutRT, FTexture2DRHIParamRef TempRT, int KernelRadius, float KernelScale, EPicpBlurPostProcessShaderType BlurType);
-	static void ApplyCompose_RenderThread(FRHICommandListImmediate& RHICmdList, FTextureRHIParamRef OverlayTexture, FTexture2DRHIParamRef DstRenderTarget, FTexture2DRHIParamRef DstTexture);
+	static void ApplyBlur_RenderThread(FRHICommandListImmediate& RHICmdList, FRHITexture2D* InOutRT, FRHITexture2D* TempRT, int KernelRadius, float KernelScale, EPicpBlurPostProcessShaderType BlurType);
+	static void ApplyCompose_RenderThread(FRHICommandListImmediate& RHICmdList, FRHITexture* OverlayTexture, FRHITexture2D* DstRenderTarget, FRHITexture2D* DstTexture);
 };
 
 class FDirectProjectionVS : public FGlobalShader
@@ -101,7 +101,7 @@ public:
 		}
 	}
 
-	void SetParameters(FRHICommandListImmediate& RHICmdList, FTexture2DRHIParamRef SourceTexture, FVector2D SampleOffset, int KernelRadius)
+	void SetParameters(FRHICommandListImmediate& RHICmdList, FRHITexture2D* SourceTexture, FVector2D SampleOffset, int KernelRadius)
 	{
 		FRHIPixelShader* ShaderRHI = GetPixelShader();
 
diff --git a/Engine/Plugins/Runtime/nDisplay/Source/PicpMPCDI/Private/PicpMPCDIShader.h b/Engine/Plugins/Runtime/nDisplay/Source/PicpMPCDI/Private/PicpMPCDIShader.h
--- a/Engine/Plugins/Runtime/nDisplay/Source/PicpMPCDI/Private/PicpMPCDIShader.h
+++ b/Engine/Plugins/Runtime/nDisplay/Source/PicpMPCDI/Private/PicpMPCDIShader.h
@@ -225,7 +225,7 @@ public:
 
 
 	template<typename TShaderRHIParamRef>
-	void SetParameters(FRHICommandListImmediate& RHICmdList, const TShaderRHIParamRef ShaderRHI, FTexture2DRHIParamRef SourceTexture)
+	void SetParameters(FRHICommandListImmediate& RHICmdList, const TShaderRHIParamRef ShaderRHI, FRHITexture2D* SourceTexture)
 	{
 		{
 			SetTextureParameter(RHICmdList, ShaderRHI, PostprocessInputParameter0, SourceTexture);
@@ -282,7 +282,7 @@ public:
 	}
 
 	template<typename TShaderRHIParamRef>
-	void SetParameters(FRHICommandListImmediate& RHICmdList, const TShaderRHIParamRef ShaderRHI, FTexture2DRHIParamRef SourceTexture, const IMPCDI::FShaderInputData& ShaderInputData, MPCDI::FMPCDIRegion& MPCDIRegionData)
+	void SetParameters(FRHICommandListImmediate& RHICmdList, const TShaderRHIParamRef ShaderRHI, FRHITexture2D* SourceTexture, const IMPCDI::FShaderInputData& ShaderInputData, MPCDI::FMPCDIRegion& MPCDIRegionData)
 	{
 		SetParameters(RHICmdList, ShaderRHI, SourceTexture);
 		SetParameters(RHICmdList, ShaderRHI, ShaderInputData);
diff --git a/Engine/Plugins/Runtime/nDisplay/Source/PicpProjection/Private/Blueprints/PicpProjectionBlueprintAPIImpl.cpp b/Engine/Plugins/Runtime/nDisplay/Source/PicpProjection/Private/Blueprints/PicpProjectionBlueprintAPIImpl.cpp
--- a/Engine/Plugins/Runtime/nDisplay/Source/PicpProjection/Private/Blueprints/PicpProjectionBlueprintAPIImpl.cpp
+++ b/Engine/Plugins/Runtime/nDisplay/Source/PicpProjection/Private/Blueprints/PicpProjectionBlueprintAPIImpl.cpp
@@ -60,7 +60,7 @@ void UPicpProjectionAPIImpl::SetupOverlayCaptures(const TArray<struct FPicpOverl
 		FTextureRenderTargetResource* overlayRTTRes = capture.CameraOverlayFrame->GameThread_GetRenderTargetResource();
 		FTextureRenderTarget2DResource* overlayRTTRes2D = (FTextureRenderTarget2DResource*)overlayRTTRes;
 
-		FTexture2DRHIParamRef CameraTextureRef = overlayRTTRes2D->GetTextureRHI();
+		FRHITexture2D* CameraTextureRef = overlayRTTRes2D->GetTextureRHI();
 		FPicpProjectionOverlayCamera* overlayCamera = new FPicpProjectionOverlayCamera(CameraRotation, CameraLocation, Prj, CameraTextureRef);
 
 		// add inner camera
@@ -73,7 +73,7 @@ void UPicpProjectionAPIImpl::SetupOverlayCaptures(const TArray<struct FPicpOverl
 			FTextureRenderTargetResource* res = layer.SourceFrameCapture->GameThread_GetRenderTargetResource();
 
 			FTextureRenderTarget2DResource* rtt2d = (FTextureRenderTarget2DResource*)res;
-			FTexture2DRHIParamRef paramRef = rtt2d->GetTextureRHI();
+			FRHITexture2D* paramRef = rtt2d->GetTextureRHI();
 			FPicpProjectionOverlayViewport* overlayViewport = new FPicpProjectionOverlayViewport(paramRef);
 
 			if (layer.OverlayBlendMode == ECameraOverlayRenderMode::Over)
diff --git a/Engine/Plugins/Runtime/nDisplay/Source/PicpProjection/Private/PicpProjectionModule.cpp b/Engine/Plugins/Runtime/nDisplay/Source/PicpProjection/Private/PicpProjectionModule.cpp
--- a/Engine/Plugins/Runtime/nDisplay/Source/PicpProjection/Private/PicpProjectionModule.cpp
+++ b/Engine/Plugins/Runtime/nDisplay/Source/PicpProjection/Private/PicpProjectionModule.cpp
@@ -172,7 +172,7 @@ void FPicpProjectionModule::CaptureWarpTexture(UTextureRenderTarget2D* dst, cons
 							if (bCaptureNow)
 							{
 								FTextureRenderTarget2DResource* dstResource2D = (FTextureRenderTarget2DResource*)(dst->GameThread_GetRenderTargetResource());
-								FTexture2DRHIParamRef dstTextureRHI = dstResource2D->GetTextureRHI();
+								FRHITexture2D* dstTextureRHI = dstResource2D->GetTextureRHI();
 
 								MPCDIPolicy->SetWarpTextureCapture(ViewIdx, dstTextureRHI);
 							}
diff --git a/Engine/Plugins/Runtime/nDisplay/Source/PicpProjection/Private/Policy/MPCDI/PicpProjectionMPCDIPolicy.cpp b/Engine/Plugins/Runtime/nDisplay/Source/PicpProjection/Private/Policy/MPCDI/PicpProjectionMPCDIPolicy.cpp
--- a/Engine/Plugins/Runtime/nDisplay/Source/PicpProjection/Private/Policy/MPCDI/PicpProjectionMPCDIPolicy.cpp
+++ b/Engine/Plugins/Runtime/nDisplay/Source/PicpProjection/Private/Policy/MPCDI/PicpProjectionMPCDIPolicy.cpp
@@ -180,7 +180,7 @@ bool FPicpProjectionMPCDIPolicy::IsWarpBlendSupported()
 	return true;
 }
 
-void FPicpProjectionMPCDIPolicy::ApplyWarpBlend_RenderThread(const uint32 ViewIdx, FRHICommandListImmediate& RHICmdList, FTexture2DRHIParamRef SrcTexture, const FIntRect& ViewportRect)
+void FPicpProjectionMPCDIPolicy::ApplyWarpBlend_RenderThread(const uint32 ViewIdx, FRHICommandListImmediate& RHICmdList, FRHITexture2D* SrcTexture, const FIntRect& ViewportRect)
 {
 	check(IsInRenderingThread());
 
@@ -319,7 +319,7 @@ bool FPicpProjectionMPCDIPolicy::InitializeResources_RenderThread()
 
 
 //Picp extands api:
-void FPicpProjectionMPCDIPolicy::SetWarpTextureCapture(const uint32 ViewIdx, FTexture2DRHIParamRef target)
+void FPicpProjectionMPCDIPolicy::SetWarpTextureCapture(const uint32 ViewIdx, FRHITexture2D* target)
 {
 	if ((int)ViewIdx < Views.Num())
 	{
diff --git a/Engine/Plugins/Runtime/nDisplay/Source/PicpProjection/Private/Policy/MPCDI/PicpProjectionMPCDIPolicy.h b/Engine/Plugins/Runtime/nDisplay/Source/PicpProjection/Private/Policy/MPCDI/PicpProjectionMPCDIPolicy.h
--- a/Engine/Plugins/Runtime/nDisplay/Source/PicpProjection/Private/Policy/MPCDI/PicpProjectionMPCDIPolicy.h
+++ b/Engine/Plugins/Runtime/nDisplay/Source/PicpProjection/Private/Policy/MPCDI/PicpProjectionMPCDIPolicy.h
@@ -41,12 +41,12 @@ public:
 	virtual bool GetProjectionMatrix(const uint32 ViewIdx, FMatrix& OutPrjMatrix) override;
 
 	virtual bool IsWarpBlendSupported() override;
-	virtual void ApplyWarpBlend_RenderThread(const uint32 ViewIdx, FRHICommandListImmediate& RHICmdList, FTexture2DRHIParamRef SrcTexture, const FIntRect& ViewportRect) override;
+	virtual void ApplyWarpBlend_RenderThread(const uint32 ViewIdx, FRHICommandListImmediate& RHICmdList, FRHITexture2D* SrcTexture, const FIntRect& ViewportRect) override;
 
 	void UpdateOverlayViewportData(FPicpProjectionOverlayFrameData& OverlayFrameData);
 	void SetOverlayData_RenderThread(const FPicpProjectionOverlayViewportData* Source);
 
-	void SetWarpTextureCapture(const uint32 ViewIdx, FTexture2DRHIParamRef target);
+	void SetWarpTextureCapture(const uint32 ViewIdx, FRHITexture2D* target);
 	IMPCDI::FFrustum GetWarpFrustum(const uint32 ViewIdx, bool bIsCaptureWarpTextureFrustum);
 
 protected:
diff --git a/Engine/Plugins/Runtime/nDisplay/Source/PicpProjection/Public/Overlay/PicpProjectionOverlayCamera.h b/Engine/Plugins/Runtime/nDisplay/Source/PicpProjection/Public/Overlay/PicpProjectionOverlayCamera.h
--- a/Engine/Plugins/Runtime/nDisplay/Source/PicpProjection/Public/Overlay/PicpProjectionOverlayCamera.h
+++ b/Engine/Plugins/Runtime/nDisplay/Source/PicpProjection/Public/Overlay/PicpProjectionOverlayCamera.h
@@ -17,14 +17,14 @@ public:
 	//@ Add more render options here
 	
 	//Camera setup
-	FTextureRHIParamRef CameraTexture; // Texture to render
+	FRHITexture* CameraTexture; // Texture to render
 	FMatrix Prj; // Projection matrix
 
 	FRotator ViewRot;
 	FVector  ViewLoc;
 
 public:
-	FPicpProjectionOverlayCamera(const FRotator& CameraRotation, const FVector& CameraLocation, const FMatrix& CameraPrj, FTextureRHIParamRef CameraTextureRef)
+	FPicpProjectionOverlayCamera(const FRotator& CameraRotation, const FVector& CameraLocation, const FMatrix& CameraPrj, FRHITexture* CameraTextureRef)
 		: FPicpProjectionOverlayBase()
 		, SoftEdge(0.1f, 0.1f, 0.1f)
 		, CameraTexture(CameraTextureRef)
diff --git a/Engine/Plugins/Runtime/nDisplay/Source/PicpProjection/Public/Overlay/PicpProjectionOverlayLUT.h b/Engine/Plugins/Runtime/nDisplay/Source/PicpProjection/Public/Overlay/PicpProjectionOverlayLUT.h
--- a/Engine/Plugins/Runtime/nDisplay/Source/PicpProjection/Public/Overlay/PicpProjectionOverlayLUT.h
+++ b/Engine/Plugins/Runtime/nDisplay/Source/PicpProjection/Public/Overlay/PicpProjectionOverlayLUT.h
@@ -10,7 +10,7 @@ class FPicpProjectionOverlayLUT
 	: public FPicpProjectionOverlayBase
 {
 public:
-	FTexture2DRHIParamRef LUTTexture;           // if texture defined, applyed brightness correction from view angle
+	FRHITexture2D* LUTTexture;           // if texture defined, applyed brightness correction from view angle
 	FVector               ViewCorrectionCoeff;
 	FVector               EyeOrigin;            // Eye position in world space
 	//@ Add more render options here
@@ -22,7 +22,7 @@ public:
 		, EyeOrigin(0,0,0)
 	{ }
 
-	void Initialize(FTexture2DRHIParamRef& TextureRef, const FVector& ViewCorrection, const FVector& EyeLocation)
+	void Initialize(FRHITexture2D*& TextureRef, const FVector& ViewCorrection, const FVector& EyeLocation)
 	{
 		LUTTexture = TextureRef;
 		ViewCorrectionCoeff = ViewCorrection;
diff --git a/Engine/Plugins/Runtime/nDisplay/Source/PicpProjection/Public/Overlay/PicpProjectionOverlayViewport.h b/Engine/Plugins/Runtime/nDisplay/Source/PicpProjection/Public/Overlay/PicpProjectionOverlayViewport.h
--- a/Engine/Plugins/Runtime/nDisplay/Source/PicpProjection/Public/Overlay/PicpProjectionOverlayViewport.h
+++ b/Engine/Plugins/Runtime/nDisplay/Source/PicpProjection/Public/Overlay/PicpProjectionOverlayViewport.h
@@ -11,11 +11,11 @@ class FPicpProjectionOverlayViewport
 	: public FPicpProjectionOverlayBase
 {
 public:
-	FTexture2DRHIParamRef ViewportTexture; // Texture to render
+	FRHITexture2D* ViewportTexture; // Texture to render
 	//@ Add more render options here
 
 public:
-	FPicpProjectionOverlayViewport(FTexture2DRHIParamRef TextureRef)
+	FPicpProjectionOverlayViewport(FRHITexture2D* TextureRef)
 		: FPicpProjectionOverlayBase()
 		, ViewportTexture(TextureRef)
 	{ 
diff --git a/Engine/Source/Runtime/Apple/MetalRHI/Private/MetalCommands.cpp b/Engine/Source/Runtime/Apple/MetalRHI/Private/MetalCommands.cpp
--- a/Engine/Source/Runtime/Apple/MetalRHI/Private/MetalCommands.cpp
+++ b/Engine/Source/Runtime/Apple/MetalRHI/Private/MetalCommands.cpp
@@ -94,7 +94,7 @@ mtlpp::PrimitiveType TranslatePrimitiveType(uint32 PrimitiveType)
 	}
 }
 
-void FMetalRHICommandContext::RHISetStreamSource(uint32 StreamIndex,FVertexBufferRHIParamRef VertexBufferRHI,uint32 Offset)
+void FMetalRHICommandContext::RHISetStreamSource(uint32 StreamIndex, FRHIVertexBuffer* VertexBufferRHI,uint32 Offset)
 {
 	@autoreleasepool {
 		FMetalVertexBuffer* VertexBuffer = ResourceCast(VertexBufferRHI);
@@ -102,7 +102,7 @@ void FMetalRHICommandContext::RHISetStreamSource(uint32 StreamIndex,FVertexBuffe
 	}
 }
 
-void FMetalDynamicRHI::RHISetStreamOutTargets(uint32 NumTargets, const FVertexBufferRHIParamRef* VertexBuffers, const uint32* Offsets)
+void FMetalDynamicRHI::RHISetStreamOutTargets(uint32 NumTargets, const FRHIVertexBuffer** VertexBuffers, const uint32* Offsets)
 {
 	NOT_SUPPORTED("RHISetStreamOutTargets");
 
@@ -141,7 +141,7 @@ void FMetalRHICommandContext::RHIDispatchComputeShader(uint32 ThreadGroupCountX,
 	}
 }
 
-void FMetalRHICommandContext::RHIDispatchIndirectComputeShader(FVertexBufferRHIParamRef ArgumentBufferRHI, uint32 ArgumentOffset)
+void FMetalRHICommandContext::RHIDispatchIndirectComputeShader(FRHIVertexBuffer* ArgumentBufferRHI, uint32 ArgumentOffset)
 {
 	@autoreleasepool {
 	if (GetMetalDeviceContext().SupportsFeature(EMetalFeaturesIndirectBuffer))
@@ -263,7 +263,7 @@ void FMetalRHICommandContext::RHISetUAVParameter(FRHIComputeShader* ComputeShade
 }
 
 
-void FMetalRHICommandContext::RHISetShaderTexture(FRHIVertexShader* VertexShaderRHI, uint32 TextureIndex, FTextureRHIParamRef NewTextureRHI)
+void FMetalRHICommandContext::RHISetShaderTexture(FRHIVertexShader* VertexShaderRHI, uint32 TextureIndex, FRHITexture* NewTextureRHI)
 {
 	@autoreleasepool {
 	FMetalSurface* Surface = GetMetalSurfaceFromRHITexture(NewTextureRHI);
@@ -286,7 +286,7 @@ void FMetalRHICommandContext::RHISetShaderTexture(FRHIVertexShader* VertexShader
 	}
 }
 
-void FMetalRHICommandContext::RHISetShaderTexture(FRHIHullShader* HullShader, uint32 TextureIndex, FTextureRHIParamRef NewTextureRHI)
+void FMetalRHICommandContext::RHISetShaderTexture(FRHIHullShader* HullShader, uint32 TextureIndex, FRHITexture* NewTextureRHI)
 {
 #if PLATFORM_SUPPORTS_TESSELLATION_SHADERS
 	@autoreleasepool {
@@ -311,7 +311,7 @@ void FMetalRHICommandContext::RHISetShaderTexture(FRHIHullShader* HullShader, ui
 #endif
 }
 
-void FMetalRHICommandContext::RHISetShaderTexture(FRHIDomainShader* DomainShader, uint32 TextureIndex, FTextureRHIParamRef NewTextureRHI)
+void FMetalRHICommandContext::RHISetShaderTexture(FRHIDomainShader* DomainShader, uint32 TextureIndex, FRHITexture* NewTextureRHI)
 {
 #if PLATFORM_SUPPORTS_TESSELLATION_SHADERS
 	@autoreleasepool {
@@ -336,13 +336,13 @@ void FMetalRHICommandContext::RHISetShaderTexture(FRHIDomainShader* DomainShader
 #endif
 }
 
-void FMetalRHICommandContext::RHISetShaderTexture(FRHIGeometryShader* GeometryShader, uint32 TextureIndex, FTextureRHIParamRef NewTextureRHI)
+void FMetalRHICommandContext::RHISetShaderTexture(FRHIGeometryShader* GeometryShader, uint32 TextureIndex, FRHITexture* NewTextureRHI)
 {
 	NOT_SUPPORTED("RHISetShaderTexture-Geometry");
 
 }
 
-void FMetalRHICommandContext::RHISetShaderTexture(FRHIPixelShader* PixelShader, uint32 TextureIndex, FTextureRHIParamRef NewTextureRHI)
+void FMetalRHICommandContext::RHISetShaderTexture(FRHIPixelShader* PixelShader, uint32 TextureIndex, FRHITexture* NewTextureRHI)
 {
 	@autoreleasepool {
 	FMetalSurface* Surface = GetMetalSurfaceFromRHITexture(NewTextureRHI);
@@ -365,7 +365,7 @@ void FMetalRHICommandContext::RHISetShaderTexture(FRHIPixelShader* PixelShader,
 	}
 }
 
-void FMetalRHICommandContext::RHISetShaderTexture(FRHIComputeShader* ComputeShader, uint32 TextureIndex, FTextureRHIParamRef NewTextureRHI)
+void FMetalRHICommandContext::RHISetShaderTexture(FRHIComputeShader* ComputeShader, uint32 TextureIndex, FRHITexture* NewTextureRHI)
 {
 	@autoreleasepool {
 	FMetalSurface* Surface = GetMetalSurfaceFromRHITexture(NewTextureRHI);
@@ -534,7 +534,7 @@ void FMetalRHICommandContext::RHISetShaderParameter(FRHIComputeShader* ComputeSh
 	}
 }
 
-void FMetalRHICommandContext::RHISetShaderUniformBuffer(FRHIVertexShader* VertexShaderRHI, uint32 BufferIndex, FUniformBufferRHIParamRef BufferRHI)
+void FMetalRHICommandContext::RHISetShaderUniformBuffer(FRHIVertexShader* VertexShaderRHI, uint32 BufferIndex, FRHIUniformBuffer* BufferRHI)
 {
 	@autoreleasepool {
 	FMetalVertexShader* VertexShader = ResourceCast(VertexShaderRHI);
@@ -550,7 +550,7 @@ void FMetalRHICommandContext::RHISetShaderUniformBuffer(FRHIVertexShader* Vertex
 	}
 }
 
-void FMetalRHICommandContext::RHISetShaderUniformBuffer(FRHIHullShader* HullShaderRHI, uint32 BufferIndex, FUniformBufferRHIParamRef BufferRHI)
+void FMetalRHICommandContext::RHISetShaderUniformBuffer(FRHIHullShader* HullShaderRHI, uint32 BufferIndex, FRHIUniformBuffer* BufferRHI)
 {
 #if PLATFORM_SUPPORTS_TESSELLATION_SHADERS
 	@autoreleasepool {
@@ -568,7 +568,7 @@ void FMetalRHICommandContext::RHISetShaderUniformBuffer(FRHIHullShader* HullShad
 #endif
 }
 
-void FMetalRHICommandContext::RHISetShaderUniformBuffer(FRHIDomainShader* DomainShaderRHI, uint32 BufferIndex, FUniformBufferRHIParamRef BufferRHI)
+void FMetalRHICommandContext::RHISetShaderUniformBuffer(FRHIDomainShader* DomainShaderRHI, uint32 BufferIndex, FRHIUniformBuffer* BufferRHI)
 {
 #if PLATFORM_SUPPORTS_TESSELLATION_SHADERS
 	@autoreleasepool {
@@ -586,12 +586,12 @@ void FMetalRHICommandContext::RHISetShaderUniformBuffer(FRHIDomainShader* Domain
 #endif
 }
 
-void FMetalRHICommandContext::RHISetShaderUniformBuffer(FRHIGeometryShader* GeometryShader, uint32 BufferIndex, FUniformBufferRHIParamRef BufferRHI)
+void FMetalRHICommandContext::RHISetShaderUniformBuffer(FRHIGeometryShader* GeometryShader, uint32 BufferIndex, FRHIUniformBuffer* BufferRHI)
 {
 	NOT_SUPPORTED("RHISetShaderUniformBuffer-Geometry");
 }
 
-void FMetalRHICommandContext::RHISetShaderUniformBuffer(FRHIPixelShader* PixelShaderRHI, uint32 BufferIndex, FUniformBufferRHIParamRef BufferRHI)
+void FMetalRHICommandContext::RHISetShaderUniformBuffer(FRHIPixelShader* PixelShaderRHI, uint32 BufferIndex, FRHIUniformBuffer* BufferRHI)
 {
 	@autoreleasepool {
 	FMetalPixelShader* PixelShader = ResourceCast(PixelShaderRHI);
@@ -607,7 +607,7 @@ void FMetalRHICommandContext::RHISetShaderUniformBuffer(FRHIPixelShader* PixelSh
 	}
 }
 
-void FMetalRHICommandContext::RHISetShaderUniformBuffer(FRHIComputeShader* ComputeShaderRHI, uint32 BufferIndex, FUniformBufferRHIParamRef BufferRHI)
+void FMetalRHICommandContext::RHISetShaderUniformBuffer(FRHIComputeShader* ComputeShaderRHI, uint32 BufferIndex, FRHIUniformBuffer* BufferRHI)
 {
 	@autoreleasepool {
 	FMetalComputeShader* ComputeShader = ResourceCast(ComputeShaderRHI);
@@ -646,7 +646,7 @@ void FMetalRHICommandContext::RHISetRenderTargets(uint32 NumSimultaneousRenderTa
 	}
 	else
 	{
-		DepthView = FRHIDepthRenderTargetView(FTextureRHIParamRef(), ERenderTargetLoadAction::EClear, ERenderTargetStoreAction::ENoAction);
+		DepthView = FRHIDepthRenderTargetView(nullptr, ERenderTargetLoadAction::EClear, ERenderTargetStoreAction::ENoAction);
 	}
 
 	FRHISetRenderTargetsInfo Info(NumSimultaneousRenderTargets, NewRenderTargets, DepthView);
@@ -745,7 +745,7 @@ void FMetalRHICommandContext::RHIDrawPrimitive(uint32 BaseVertexIndex, uint32 Nu
 	}
 }
 
-void FMetalRHICommandContext::RHIDrawPrimitiveIndirect(FVertexBufferRHIParamRef VertexBufferRHI, uint32 ArgumentOffset)
+void FMetalRHICommandContext::RHIDrawPrimitiveIndirect(FRHIVertexBuffer* VertexBufferRHI, uint32 ArgumentOffset)
 {
 	@autoreleasepool {
 #if PLATFORM_IOS
@@ -763,7 +763,7 @@ void FMetalRHICommandContext::RHIDrawPrimitiveIndirect(FVertexBufferRHIParamRef
 	}
 }
 
-void FMetalRHICommandContext::RHIDrawIndexedPrimitive(FIndexBufferRHIParamRef IndexBufferRHI, int32 BaseVertexIndex, uint32 FirstInstance,
+void FMetalRHICommandContext::RHIDrawIndexedPrimitive(FRHIIndexBuffer* IndexBufferRHI, int32 BaseVertexIndex, uint32 FirstInstance,
 	uint32 NumVertices, uint32 StartIndex, uint32 NumPrimitives, uint32 NumInstances)
 {
 	@autoreleasepool {
@@ -781,7 +781,7 @@ void FMetalRHICommandContext::RHIDrawIndexedPrimitive(FIndexBufferRHIParamRef In
 	}
 }
 
-void FMetalRHICommandContext::RHIDrawIndexedIndirect(FIndexBufferRHIParamRef IndexBufferRHI, FStructuredBufferRHIParamRef VertexBufferRHI, int32 DrawArgumentsIndex, uint32 NumInstances)
+void FMetalRHICommandContext::RHIDrawIndexedIndirect(FRHIIndexBuffer* IndexBufferRHI, FRHIStructuredBuffer* VertexBufferRHI, int32 DrawArgumentsIndex, uint32 NumInstances)
 {
 	@autoreleasepool {
 	if (GetMetalDeviceContext().SupportsFeature(EMetalFeaturesIndirectBuffer))
@@ -806,7 +806,7 @@ void FMetalRHICommandContext::RHIDrawIndexedIndirect(FIndexBufferRHIParamRef Ind
 	}
 }
 
-void FMetalRHICommandContext::RHIDrawIndexedPrimitiveIndirect(FIndexBufferRHIParamRef IndexBufferRHI,FVertexBufferRHIParamRef VertexBufferRHI,uint32 ArgumentOffset)
+void FMetalRHICommandContext::RHIDrawIndexedPrimitiveIndirect(FRHIIndexBuffer* IndexBufferRHI, FRHIVertexBuffer* VertexBufferRHI,uint32 ArgumentOffset)
 {
 	@autoreleasepool {
 	if (GetMetalDeviceContext().SupportsFeature(EMetalFeaturesIndirectBuffer))
diff --git a/Engine/Source/Runtime/Apple/MetalRHI/Private/MetalContext.cpp b/Engine/Source/Runtime/Apple/MetalRHI/Private/MetalContext.cpp
--- a/Engine/Source/Runtime/Apple/MetalRHI/Private/MetalContext.cpp
+++ b/Engine/Source/Runtime/Apple/MetalRHI/Private/MetalContext.cpp
@@ -1158,7 +1158,7 @@ void FMetalContext::TransitionResources(FRHIUnorderedAccessView** InUAVs, int32
 	}
 }
 
-void FMetalContext::TransitionResources(FTextureRHIParamRef* InTextures, int32 NumTextures)
+void FMetalContext::TransitionResources(FRHITexture** InTextures, int32 NumTextures)
 {
 	for (uint32 i = 0; i < NumTextures; i++)
 	{
diff --git a/Engine/Source/Runtime/Apple/MetalRHI/Private/MetalContext.h b/Engine/Source/Runtime/Apple/MetalRHI/Private/MetalContext.h
--- a/Engine/Source/Runtime/Apple/MetalRHI/Private/MetalContext.h
+++ b/Engine/Source/Runtime/Apple/MetalRHI/Private/MetalContext.h
@@ -125,7 +125,7 @@ public:
 
 	// Track Write->Read transitions for TBDR Fragment->Verex fencing
 	void TransitionResources(FRHIUnorderedAccessView** InUAVs, int32 NumUAVs);
-	void TransitionResources(FTextureRHIParamRef* InTextures, int32 NumTextures);
+	void TransitionResources(FRHITexture** InTextures, int32 NumTextures);
 	
 protected:
 	/** The underlying Metal device */
diff --git a/Engine/Source/Runtime/Apple/MetalRHI/Private/MetalIndexBuffer.cpp b/Engine/Source/Runtime/Apple/MetalRHI/Private/MetalIndexBuffer.cpp
--- a/Engine/Source/Runtime/Apple/MetalRHI/Private/MetalIndexBuffer.cpp
+++ b/Engine/Source/Runtime/Apple/MetalRHI/Private/MetalIndexBuffer.cpp
@@ -99,7 +99,7 @@ FIndexBufferRHIRef FMetalDynamicRHI::RHICreateIndexBuffer(uint32 Stride, uint32
 	}
 }
 
-void FMetalDynamicRHI::RHITransferIndexBufferUnderlyingResource(FIndexBufferRHIParamRef DestIndexBuffer, FIndexBufferRHIParamRef SrcIndexBuffer)
+void FMetalDynamicRHI::RHITransferIndexBufferUnderlyingResource(FRHIIndexBuffer* DestIndexBuffer, FRHIIndexBuffer* SrcIndexBuffer)
 {
 	@autoreleasepool {
 	check(DestIndexBuffer);
@@ -118,7 +118,7 @@ void FMetalDynamicRHI::RHITransferIndexBufferUnderlyingResource(FIndexBufferRHIP
 	}
 }
 
-void* FMetalDynamicRHI::RHILockIndexBuffer(FIndexBufferRHIParamRef IndexBufferRHI, uint32 Offset, uint32 Size, EResourceLockMode LockMode)
+void* FMetalDynamicRHI::RHILockIndexBuffer(FRHIIndexBuffer* IndexBufferRHI, uint32 Offset, uint32 Size, EResourceLockMode LockMode)
 {
 	@autoreleasepool {
 	FMetalIndexBuffer* IndexBuffer = ResourceCast(IndexBufferRHI);
@@ -127,7 +127,7 @@ void* FMetalDynamicRHI::RHILockIndexBuffer(FIndexBufferRHIParamRef IndexBufferRH
 	}
 }
 
-void FMetalDynamicRHI::RHIUnlockIndexBuffer(FIndexBufferRHIParamRef IndexBufferRHI)
+void FMetalDynamicRHI::RHIUnlockIndexBuffer(FRHIIndexBuffer* IndexBufferRHI)
 {
 	@autoreleasepool {
 	FMetalIndexBuffer* IndexBuffer = ResourceCast(IndexBufferRHI);
diff --git a/Engine/Source/Runtime/Apple/MetalRHI/Private/MetalRenderTarget.cpp b/Engine/Source/Runtime/Apple/MetalRHI/Private/MetalRenderTarget.cpp
--- a/Engine/Source/Runtime/Apple/MetalRHI/Private/MetalRenderTarget.cpp
+++ b/Engine/Source/Runtime/Apple/MetalRHI/Private/MetalRenderTarget.cpp
@@ -32,7 +32,7 @@ static FAutoConsoleVariableRef CVarMetalUseTexGetBytes(
 								ECVF_RenderThreadSafe
 								);
 
-void FMetalRHICommandContext::RHICopyToResolveTarget(FTextureRHIParamRef SourceTextureRHI, FTextureRHIParamRef DestTextureRHI, const FResolveParams& ResolveParams)
+void FMetalRHICommandContext::RHICopyToResolveTarget(FRHITexture* SourceTextureRHI, FRHITexture* DestTextureRHI, const FResolveParams& ResolveParams)
 {
 	@autoreleasepool {
 	if (!SourceTextureRHI || !DestTextureRHI)
@@ -415,7 +415,7 @@ static void ConvertSurfaceDataToFColor(EPixelFormat Format, uint32 Width, uint32
 	}
 }
 
-void FMetalDynamicRHI::RHIReadSurfaceData(FTextureRHIParamRef TextureRHI, FIntRect InRect, TArray<FLinearColor>& OutData, FReadSurfaceDataFlags InFlags)
+void FMetalDynamicRHI::RHIReadSurfaceData(FRHITexture* TextureRHI, FIntRect InRect, TArray<FLinearColor>& OutData, FReadSurfaceDataFlags InFlags)
 {
 	// Use our current surface read implemtation and convert to linear - should refactor to make optimal
 	TArray<FColor> OutDataUnConverted;
@@ -430,7 +430,7 @@ void FMetalDynamicRHI::RHIReadSurfaceData(FTextureRHIParamRef TextureRHI, FIntRe
 	}
 }
 
-void FMetalDynamicRHI::RHIReadSurfaceData(FTextureRHIParamRef TextureRHI, FIntRect Rect, TArray<FColor>& OutData, FReadSurfaceDataFlags InFlags)
+void FMetalDynamicRHI::RHIReadSurfaceData(FRHITexture* TextureRHI, FIntRect Rect, TArray<FColor>& OutData, FReadSurfaceDataFlags InFlags)
 {
 	@autoreleasepool {
 	if (!ensure(TextureRHI))
@@ -563,7 +563,7 @@ void FMetalDynamicRHI::RHIReadSurfaceData(FTextureRHIParamRef TextureRHI, FIntRe
 	}
 }
 
-void FMetalDynamicRHI::RHIMapStagingSurface(FTextureRHIParamRef TextureRHI,void*& OutData,int32& OutWidth,int32& OutHeight)
+void FMetalDynamicRHI::RHIMapStagingSurface(FRHITexture* TextureRHI,void*& OutData,int32& OutWidth,int32& OutHeight)
 {
 	@autoreleasepool {
     FMetalSurface* Surface = GetMetalSurfaceFromRHITexture(TextureRHI);
@@ -576,7 +576,7 @@ void FMetalDynamicRHI::RHIMapStagingSurface(FTextureRHIParamRef TextureRHI,void*
 	}
 }
 
-void FMetalDynamicRHI::RHIUnmapStagingSurface(FTextureRHIParamRef TextureRHI)
+void FMetalDynamicRHI::RHIUnmapStagingSurface(FRHITexture* TextureRHI)
 {
 	@autoreleasepool {
     FMetalSurface* Surface = GetMetalSurfaceFromRHITexture(TextureRHI);
@@ -585,7 +585,7 @@ void FMetalDynamicRHI::RHIUnmapStagingSurface(FTextureRHIParamRef TextureRHI)
 	}
 }
 
-void FMetalDynamicRHI::RHIReadSurfaceFloatData(FTextureRHIParamRef TextureRHI, FIntRect Rect, TArray<FFloat16Color>& OutData, ECubeFace CubeFace,int32 ArrayIndex,int32 MipIndex)
+void FMetalDynamicRHI::RHIReadSurfaceFloatData(FRHITexture* TextureRHI, FIntRect Rect, TArray<FFloat16Color>& OutData, ECubeFace CubeFace,int32 ArrayIndex,int32 MipIndex)
 {
 	@autoreleasepool {
 	FMetalSurface* Surface = GetMetalSurfaceFromRHITexture(TextureRHI);
@@ -661,7 +661,7 @@ void FMetalDynamicRHI::RHIReadSurfaceFloatData(FTextureRHIParamRef TextureRHI, F
 	}
 }
 
-void FMetalDynamicRHI::RHIRead3DSurfaceFloatData(FTextureRHIParamRef TextureRHI,FIntRect InRect,FIntPoint ZMinMax,TArray<FFloat16Color>& OutData)
+void FMetalDynamicRHI::RHIRead3DSurfaceFloatData(FRHITexture* TextureRHI,FIntRect InRect,FIntPoint ZMinMax,TArray<FFloat16Color>& OutData)
 {
 	@autoreleasepool {
 	FMetalSurface* Surface = GetMetalSurfaceFromRHITexture(TextureRHI);
diff --git a/Engine/Source/Runtime/Apple/MetalRHI/Private/MetalStateCache.cpp b/Engine/Source/Runtime/Apple/MetalRHI/Private/MetalStateCache.cpp
--- a/Engine/Source/Runtime/Apple/MetalRHI/Private/MetalStateCache.cpp
+++ b/Engine/Source/Runtime/Apple/MetalRHI/Private/MetalStateCache.cpp
@@ -1201,7 +1201,7 @@ void FMetalStateCache::SetIndexType(EMetalIndexType InIndexType)
 	}
 }
 
-void FMetalStateCache::BindUniformBuffer(EMetalShaderStages const Freq, uint32 const BufferIndex, FUniformBufferRHIParamRef BufferRHI)
+void FMetalStateCache::BindUniformBuffer(EMetalShaderStages const Freq, uint32 const BufferIndex, FRHIUniformBuffer* BufferRHI)
 {
 	check(BufferIndex < ML_MaxBuffers);
 	if (BoundUniformBuffers[Freq][BufferIndex] != BufferRHI)
