UE4.23 - Deprecate FUniformBufferRHIParamRef, FIndexBufferRHIParamRef, FVertexBufferRHIParamRef, FStructuredBufferRHIParamRef, FTextureRHIParamRef, FTexture2DRHIParamRef, FTexture2DArrayRHIParamRef, F

diff --git a/Engine/Source/Runtime/Apple/MetalRHI/Private/MetalStateCache.h b/Engine/Source/Runtime/Apple/MetalRHI/Private/MetalStateCache.h
--- a/Engine/Source/Runtime/Apple/MetalRHI/Private/MetalStateCache.h
+++ b/Engine/Source/Runtime/Apple/MetalRHI/Private/MetalStateCache.h
@@ -51,7 +51,7 @@ public:
 	void SetVertexStream(uint32 const Index, FMetalBuffer* Buffer, FMetalBufferData* Bytes, uint32 const Offset, uint32 const Length);
 	void SetGraphicsPipelineState(FMetalGraphicsPipelineState* State);
 	void SetIndexType(EMetalIndexType IndexType);
-	void BindUniformBuffer(EMetalShaderStages const Freq, uint32 const BufferIndex, FUniformBufferRHIParamRef BufferRHI);
+	void BindUniformBuffer(EMetalShaderStages const Freq, uint32 const BufferIndex, FRHIUniformBuffer* BufferRHI);
 	void SetDirtyUniformBuffers(EMetalShaderStages const Freq, uint32 const Dirty);
 	
 	/*
diff --git a/Engine/Source/Runtime/Apple/MetalRHI/Private/MetalTexture.cpp b/Engine/Source/Runtime/Apple/MetalRHI/Private/MetalTexture.cpp
--- a/Engine/Source/Runtime/Apple/MetalRHI/Private/MetalTexture.cpp
+++ b/Engine/Source/Runtime/Apple/MetalRHI/Private/MetalTexture.cpp
@@ -1618,7 +1618,7 @@ bool FMetalDynamicRHI::RHIGetTextureMemoryVisualizeData( FColor* /*TextureData*/
 	return false;
 }
 
-uint32 FMetalDynamicRHI::RHIComputeMemorySize(FTextureRHIParamRef TextureRHI)
+uint32 FMetalDynamicRHI::RHIComputeMemorySize(FRHITexture* TextureRHI)
 {
 	@autoreleasepool {
 		if(!TextureRHI)
@@ -1654,7 +1654,7 @@ FTexture2DRHIRef FMetalDynamicRHI::RHIAsyncCreateTexture2D(uint32 SizeX,uint32 S
 	return FTexture2DRHIRef();
 }
 
-void FMetalDynamicRHI::RHICopySharedMips(FTexture2DRHIParamRef DestTexture2D,FTexture2DRHIParamRef SrcTexture2D)
+void FMetalDynamicRHI::RHICopySharedMips(FRHITexture2D* DestTexture2D, FRHITexture2D* SrcTexture2D)
 {
 	NOT_SUPPORTED("RHICopySharedMips");
 }
@@ -1673,13 +1673,13 @@ FTexture3DRHIRef FMetalDynamicRHI::RHICreateTexture3D(uint32 SizeX, uint32 SizeY
 	}
 }
 
-void FMetalDynamicRHI::RHIGetResourceInfo(FTextureRHIParamRef Ref, FRHIResourceInfo& OutInfo)
+void FMetalDynamicRHI::RHIGetResourceInfo(FRHITexture* Ref, FRHIResourceInfo& OutInfo)
 {
 	// @todo Needed for visualisation!!
 	// NOT_SUPPORTED("RHIGetResourceInfo");
 }
 
-void FMetalDynamicRHI::RHIGenerateMips(FTextureRHIParamRef SourceSurfaceRHI)
+void FMetalDynamicRHI::RHIGenerateMips(FRHITexture* SourceSurfaceRHI)
 {
 	@autoreleasepool {
 		FMetalSurface* Surf = GetMetalSurfaceFromRHITexture(SourceSurfaceRHI);
@@ -1770,7 +1770,7 @@ struct FMetalRHICommandAsyncReallocateTexture2D final : public FRHICommand<FMeta
 	}
 };
 
-FTexture2DRHIRef FMetalDynamicRHI::AsyncReallocateTexture2D_RenderThread(class FRHICommandListImmediate& RHICmdList, FTexture2DRHIParamRef Texture2D, int32 NewMipCount, int32 NewSizeX, int32 NewSizeY, FThreadSafeCounter* RequestStatus)
+FTexture2DRHIRef FMetalDynamicRHI::AsyncReallocateTexture2D_RenderThread(class FRHICommandListImmediate& RHICmdList, FRHITexture2D* Texture2D, int32 NewMipCount, int32 NewSizeX, int32 NewSizeY, FThreadSafeCounter* RequestStatus)
 {
 	@autoreleasepool {
 		FTexture2DRHIRef Result;
@@ -1794,20 +1794,20 @@ FTexture2DRHIRef FMetalDynamicRHI::AsyncReallocateTexture2D_RenderThread(class F
 	}
 }
 
-ETextureReallocationStatus FMetalDynamicRHI::FinalizeAsyncReallocateTexture2D_RenderThread(class FRHICommandListImmediate& RHICmdList, FTexture2DRHIParamRef Texture2D, bool bBlockUntilCompleted)
+ETextureReallocationStatus FMetalDynamicRHI::FinalizeAsyncReallocateTexture2D_RenderThread(class FRHICommandListImmediate& RHICmdList, FRHITexture2D* Texture2D, bool bBlockUntilCompleted)
 {
 	// No need to flush - does nothing
 	return GDynamicRHI->RHIFinalizeAsyncReallocateTexture2D(Texture2D, bBlockUntilCompleted);
 }
 
-ETextureReallocationStatus FMetalDynamicRHI::CancelAsyncReallocateTexture2D_RenderThread(class FRHICommandListImmediate& RHICmdList, FTexture2DRHIParamRef Texture2D, bool bBlockUntilCompleted)
+ETextureReallocationStatus FMetalDynamicRHI::CancelAsyncReallocateTexture2D_RenderThread(class FRHICommandListImmediate& RHICmdList, FRHITexture2D* Texture2D, bool bBlockUntilCompleted)
 {
 	// No need to flush - does nothing
 	return GDynamicRHI->RHICancelAsyncReallocateTexture2D(Texture2D, bBlockUntilCompleted);
 }
 
 
-FTexture2DRHIRef FMetalDynamicRHI::RHIAsyncReallocateTexture2D(FTexture2DRHIParamRef OldTextureRHI, int32 NewMipCount, int32 NewSizeX, int32 NewSizeY, FThreadSafeCounter* RequestStatus)
+FTexture2DRHIRef FMetalDynamicRHI::RHIAsyncReallocateTexture2D(FRHITexture2D* OldTextureRHI, int32 NewMipCount, int32 NewSizeX, int32 NewSizeY, FThreadSafeCounter* RequestStatus)
 {
 	@autoreleasepool {
 		FMetalTexture2D* OldTexture = ResourceCast(OldTextureRHI);
@@ -1820,17 +1820,17 @@ FTexture2DRHIRef FMetalDynamicRHI::RHIAsyncReallocateTexture2D(FTexture2DRHIPara
 	}
 }
 
-ETextureReallocationStatus FMetalDynamicRHI::RHIFinalizeAsyncReallocateTexture2D( FTexture2DRHIParamRef Texture2D, bool bBlockUntilCompleted )
+ETextureReallocationStatus FMetalDynamicRHI::RHIFinalizeAsyncReallocateTexture2D(FRHITexture2D* Texture2D, bool bBlockUntilCompleted )
 {
 	return TexRealloc_Succeeded;
 }
 
-ETextureReallocationStatus FMetalDynamicRHI::RHICancelAsyncReallocateTexture2D( FTexture2DRHIParamRef Texture2D, bool bBlockUntilCompleted )
+ETextureReallocationStatus FMetalDynamicRHI::RHICancelAsyncReallocateTexture2D(FRHITexture2D* Texture2D, bool bBlockUntilCompleted )
 {
 	return TexRealloc_Failed;
 }
 
-void* FMetalDynamicRHI::LockTexture2D_RenderThread(class FRHICommandListImmediate& RHICmdList, FTexture2DRHIParamRef Texture, uint32 MipIndex, EResourceLockMode LockMode, uint32& DestStride, bool bLockWithinMiptail, bool bNeedsDefaultRHIFlush)
+void* FMetalDynamicRHI::LockTexture2D_RenderThread(class FRHICommandListImmediate& RHICmdList, FRHITexture2D* Texture, uint32 MipIndex, EResourceLockMode LockMode, uint32& DestStride, bool bLockWithinMiptail, bool bNeedsDefaultRHIFlush)
 {
 	@autoreleasepool {
 		check(IsInRenderingThread());
@@ -1843,7 +1843,7 @@ void* FMetalDynamicRHI::LockTexture2D_RenderThread(class FRHICommandListImmediat
 	}
 }
 
-void FMetalDynamicRHI::UnlockTexture2D_RenderThread(class FRHICommandListImmediate& RHICmdList, FTexture2DRHIParamRef Texture, uint32 MipIndex, bool bLockWithinMiptail, bool bNeedsDefaultRHIFlush)
+void FMetalDynamicRHI::UnlockTexture2D_RenderThread(class FRHICommandListImmediate& RHICmdList, FRHITexture2D* Texture, uint32 MipIndex, bool bLockWithinMiptail, bool bNeedsDefaultRHIFlush)
 {
 	@autoreleasepool
 	{
@@ -1855,7 +1855,7 @@ void FMetalDynamicRHI::UnlockTexture2D_RenderThread(class FRHICommandListImmedia
 }
 
 
-void* FMetalDynamicRHI::RHILockTexture2D(FTexture2DRHIParamRef TextureRHI,uint32 MipIndex,EResourceLockMode LockMode,uint32& DestStride,bool bLockWithinMiptail)
+void* FMetalDynamicRHI::RHILockTexture2D(FRHITexture2D* TextureRHI,uint32 MipIndex,EResourceLockMode LockMode,uint32& DestStride,bool bLockWithinMiptail)
 {
 	@autoreleasepool {
 		FMetalTexture2D* Texture = ResourceCast(TextureRHI);
@@ -1863,7 +1863,7 @@ void* FMetalDynamicRHI::RHILockTexture2D(FTexture2DRHIParamRef TextureRHI,uint32
 	}
 }
 
-void FMetalDynamicRHI::RHIUnlockTexture2D(FTexture2DRHIParamRef TextureRHI,uint32 MipIndex,bool bLockWithinMiptail)
+void FMetalDynamicRHI::RHIUnlockTexture2D(FRHITexture2D* TextureRHI,uint32 MipIndex,bool bLockWithinMiptail)
 {
 	@autoreleasepool {
 		FMetalTexture2D* Texture = ResourceCast(TextureRHI);
@@ -1871,7 +1871,7 @@ void FMetalDynamicRHI::RHIUnlockTexture2D(FTexture2DRHIParamRef TextureRHI,uint3
 	}
 }
 
-void* FMetalDynamicRHI::RHILockTexture2DArray(FTexture2DArrayRHIParamRef TextureRHI,uint32 TextureIndex,uint32 MipIndex,EResourceLockMode LockMode,uint32& DestStride,bool bLockWithinMiptail)
+void* FMetalDynamicRHI::RHILockTexture2DArray(FRHITexture2DArray* TextureRHI,uint32 TextureIndex,uint32 MipIndex,EResourceLockMode LockMode,uint32& DestStride,bool bLockWithinMiptail)
 {
 	@autoreleasepool {
 		FMetalTexture2DArray* Texture = ResourceCast(TextureRHI);
@@ -1879,7 +1879,7 @@ void* FMetalDynamicRHI::RHILockTexture2DArray(FTexture2DArrayRHIParamRef Texture
 	}
 }
 
-void FMetalDynamicRHI::RHIUnlockTexture2DArray(FTexture2DArrayRHIParamRef TextureRHI,uint32 TextureIndex,uint32 MipIndex,bool bLockWithinMiptail)
+void FMetalDynamicRHI::RHIUnlockTexture2DArray(FRHITexture2DArray* TextureRHI,uint32 TextureIndex,uint32 MipIndex,bool bLockWithinMiptail)
 {
 	@autoreleasepool {
 		FMetalTexture2DArray* Texture = ResourceCast(TextureRHI);
@@ -1905,7 +1905,7 @@ static void InternalExpandR8ToStandardRGBA(uint32* pDest, const struct FUpdateTe
 }
 #endif
 
-static FMetalBuffer InternalCopyTexture2DUpdateRegion(FTexture2DRHIParamRef TextureRHI, const struct FUpdateTextureRegion2D& UpdateRegion, uint32& InOutSourcePitch, const uint8* SourceData)
+static FMetalBuffer InternalCopyTexture2DUpdateRegion(FRHITexture2D* TextureRHI, const struct FUpdateTextureRegion2D& UpdateRegion, uint32& InOutSourcePitch, const uint8* SourceData)
 {	
 	const uint32 InSourcePitch = InOutSourcePitch;
 
@@ -1945,7 +1945,7 @@ static FMetalBuffer InternalCopyTexture2DUpdateRegion(FTexture2DRHIParamRef Text
 	return Buffer;
 }
 
-static void InternalUpdateTexture2D(FMetalContext& Context, FTexture2DRHIParamRef TextureRHI, uint32 MipIndex, FUpdateTextureRegion2D const& UpdateRegion, uint32 SourcePitch, FMetalBuffer Buffer)
+static void InternalUpdateTexture2D(FMetalContext& Context, FRHITexture2D* TextureRHI, uint32 MipIndex, FUpdateTextureRegion2D const& UpdateRegion, uint32 SourcePitch, FMetalBuffer Buffer)
 {
 	FMetalTexture2D* Texture = ResourceCast(TextureRHI);
 	FMetalTexture Tex = Texture->Surface.Texture;
@@ -1981,13 +1981,13 @@ static void InternalUpdateTexture2D(FMetalContext& Context, FTexture2DRHIParamRe
 struct FMetalRHICommandUpdateTexture2D final : public FRHICommand<FMetalRHICommandUpdateTexture2D>
 {
 	FMetalContext& Context;
-	FTexture2DRHIParamRef Texture;
+	FRHITexture2D* Texture;
 	uint32 MipIndex;
 	FUpdateTextureRegion2D UpdateRegion;
 	uint32 SourcePitch;
 	FMetalBuffer SourceBuffer;
 
-	FORCEINLINE_DEBUGGABLE FMetalRHICommandUpdateTexture2D(FMetalContext& InContext, FTexture2DRHIParamRef InTexture, uint32 InMipIndex, FUpdateTextureRegion2D InUpdateRegion, uint32 InSourcePitch, FMetalBuffer InSourceBuffer)
+	FORCEINLINE_DEBUGGABLE FMetalRHICommandUpdateTexture2D(FMetalContext& InContext, FRHITexture2D* InTexture, uint32 InMipIndex, FUpdateTextureRegion2D InUpdateRegion, uint32 InSourcePitch, FMetalBuffer InSourceBuffer)
 	: Context(InContext)
 	, Texture(InTexture)
 	, MipIndex(InMipIndex)
@@ -2005,7 +2005,7 @@ struct FMetalRHICommandUpdateTexture2D final : public FRHICommand<FMetalRHIComma
 	}
 };
 
-void FMetalDynamicRHI::UpdateTexture2D_RenderThread(class FRHICommandListImmediate& RHICmdList, FTexture2DRHIParamRef Texture, uint32 MipIndex, const struct FUpdateTextureRegion2D& UpdateRegion, uint32 SourcePitch, const uint8* SourceData)
+void FMetalDynamicRHI::UpdateTexture2D_RenderThread(class FRHICommandListImmediate& RHICmdList, FRHITexture2D* Texture, uint32 MipIndex, const struct FUpdateTextureRegion2D& UpdateRegion, uint32 SourcePitch, const uint8* SourceData)
 {
 	@autoreleasepool {
 		if (RHICmdList.Bypass() || !IsRunningRHIInSeparateThread())
@@ -2020,7 +2020,7 @@ void FMetalDynamicRHI::UpdateTexture2D_RenderThread(class FRHICommandListImmedia
 	}
 }
 
-void FMetalDynamicRHI::RHIUpdateTexture2D(FTexture2DRHIParamRef TextureRHI, uint32 MipIndex, const struct FUpdateTextureRegion2D& UpdateRegion, uint32 SourcePitch, const uint8* SourceData)
+void FMetalDynamicRHI::RHIUpdateTexture2D(FRHITexture2D* TextureRHI, uint32 MipIndex, const struct FUpdateTextureRegion2D& UpdateRegion, uint32 SourcePitch, const uint8* SourceData)
 {
 	@autoreleasepool 
 	{
@@ -2055,7 +2055,7 @@ void FMetalDynamicRHI::RHIUpdateTexture2D(FTexture2DRHIParamRef TextureRHI, uint
 	}
 }
 
-static void InternalCopyTexture3DUpdateRegionData(FTexture3DRHIParamRef TextureRHI, const struct FUpdateTextureRegion3D& UpdateRegion, uint32 SourceRowPitch, uint32 SourceDepthPitch, const uint8* SourceData, uint8* DestData)
+static void InternalCopyTexture3DUpdateRegionData(FRHITexture3D* TextureRHI, const struct FUpdateTextureRegion3D& UpdateRegion, uint32 SourceRowPitch, uint32 SourceDepthPitch, const uint8* SourceData, uint8* DestData)
 {
 	// Perform safe line copy
 	FMetalTexture3D* Texture = ResourceCast(TextureRHI);	
@@ -2078,7 +2078,7 @@ static void InternalCopyTexture3DUpdateRegionData(FTexture3DRHIParamRef TextureR
 }
 
 
-static void InternalUpdateTexture3D(FMetalContext& Context, FTexture3DRHIParamRef TextureRHI, uint32 MipIndex, const FUpdateTextureRegion3D& UpdateRegion, uint32 SourceRowPitch, uint32 SourceDepthPitch, FMetalBuffer Buffer)
+static void InternalUpdateTexture3D(FMetalContext& Context, FRHITexture3D* TextureRHI, uint32 MipIndex, const FUpdateTextureRegion3D& UpdateRegion, uint32 SourceRowPitch, uint32 SourceDepthPitch, FMetalBuffer Buffer)
 {
 	FMetalTexture3D* Texture = ResourceCast(TextureRHI);
 	FMetalTexture Tex = Texture->Surface.Texture;
@@ -2111,14 +2111,14 @@ static void InternalUpdateTexture3D(FMetalContext& Context, FTexture3DRHIParamRe
 struct FMetalDynamicRHIUpdateTexture3DCommand final : public FRHICommand<FMetalDynamicRHIUpdateTexture3DCommand>
 {
 	FMetalContext& Context;
-	FTexture3DRHIParamRef DestinationTexture;
+	FRHITexture3D* DestinationTexture;
 	uint32 MipIndex;
 	FUpdateTextureRegion3D UpdateRegion;
 	uint32 SourceRowPitch;
 	uint32 SourceDepthPitch;
 	FMetalBuffer Buffer;
 	
-	FORCEINLINE_DEBUGGABLE FMetalDynamicRHIUpdateTexture3DCommand(FMetalContext& InContext, FTexture3DRHIParamRef TextureRHI, uint32 InMipIndex, const struct FUpdateTextureRegion3D& InUpdateRegion, uint32 InSourceRowPitch, uint32 InSourceDepthPitch, const uint8* SourceData)
+	FORCEINLINE_DEBUGGABLE FMetalDynamicRHIUpdateTexture3DCommand(FMetalContext& InContext, FRHITexture3D* TextureRHI, uint32 InMipIndex, const struct FUpdateTextureRegion3D& InUpdateRegion, uint32 InSourceRowPitch, uint32 InSourceDepthPitch, const uint8* SourceData)
 	: Context(InContext)
 	, DestinationTexture(TextureRHI)
 	, MipIndex(InMipIndex)
@@ -2142,7 +2142,7 @@ struct FMetalDynamicRHIUpdateTexture3DCommand final : public FRHICommand<FMetalD
 	}
 };
 
-void FMetalDynamicRHI::UpdateTexture3D_RenderThread(class FRHICommandListImmediate& RHICmdList, FTexture3DRHIParamRef Texture, uint32 MipIndex, const struct FUpdateTextureRegion3D& UpdateRegion, uint32 SourceRowPitch, uint32 SourceDepthPitch, const uint8* SourceData)
+void FMetalDynamicRHI::UpdateTexture3D_RenderThread(class FRHICommandListImmediate& RHICmdList, FRHITexture3D* Texture, uint32 MipIndex, const struct FUpdateTextureRegion3D& UpdateRegion, uint32 SourceRowPitch, uint32 SourceDepthPitch, const uint8* SourceData)
 {
 	if (RHICmdList.Bypass() || !IsRunningRHIInSeparateThread())
 	{
@@ -2154,7 +2154,7 @@ void FMetalDynamicRHI::UpdateTexture3D_RenderThread(class FRHICommandListImmedia
 	}
 }
 
-FUpdateTexture3DData FMetalDynamicRHI::BeginUpdateTexture3D_RenderThread(class FRHICommandListImmediate& RHICmdList, FTexture3DRHIParamRef Texture, uint32 MipIndex, const struct FUpdateTextureRegion3D& UpdateRegion)
+FUpdateTexture3DData FMetalDynamicRHI::BeginUpdateTexture3D_RenderThread(class FRHICommandListImmediate& RHICmdList, FRHITexture3D* Texture, uint32 MipIndex, const struct FUpdateTextureRegion3D& UpdateRegion)
 {
 	check(IsInRenderingThread());
 	
@@ -2186,7 +2186,7 @@ void FMetalDynamicRHI::EndUpdateTexture3D_RenderThread(class FRHICommandListImme
 	UpdateData.Data = nullptr;
 }
 
-void FMetalDynamicRHI::RHIUpdateTexture3D(FTexture3DRHIParamRef TextureRHI,uint32 MipIndex,const FUpdateTextureRegion3D& UpdateRegion,uint32 SourceRowPitch,uint32 SourceDepthPitch, const uint8* SourceData)
+void FMetalDynamicRHI::RHIUpdateTexture3D(FRHITexture3D* TextureRHI,uint32 MipIndex,const FUpdateTextureRegion3D& UpdateRegion,uint32 SourceRowPitch,uint32 SourceDepthPitch, const uint8* SourceData)
 {
 	@autoreleasepool {
 	
@@ -2234,7 +2234,7 @@ FTextureCubeRHIRef FMetalDynamicRHI::RHICreateTextureCubeArray(uint32 Size, uint
 	}
 }
 
-void* FMetalDynamicRHI::RHILockTextureCubeFace(FTextureCubeRHIParamRef TextureCubeRHI,uint32 FaceIndex,uint32 ArrayIndex,uint32 MipIndex,EResourceLockMode LockMode,uint32& DestStride,bool bLockWithinMiptail)
+void* FMetalDynamicRHI::RHILockTextureCubeFace(FRHITextureCube* TextureCubeRHI,uint32 FaceIndex,uint32 ArrayIndex,uint32 MipIndex,EResourceLockMode LockMode,uint32& DestStride,bool bLockWithinMiptail)
 {
 	@autoreleasepool {
 		FMetalTextureCube* TextureCube = ResourceCast(TextureCubeRHI);
@@ -2243,7 +2243,7 @@ void* FMetalDynamicRHI::RHILockTextureCubeFace(FTextureCubeRHIParamRef TextureCu
 	}
 }
 
-void FMetalDynamicRHI::RHIUnlockTextureCubeFace(FTextureCubeRHIParamRef TextureCubeRHI,uint32 FaceIndex,uint32 ArrayIndex,uint32 MipIndex,bool bLockWithinMiptail)
+void FMetalDynamicRHI::RHIUnlockTextureCubeFace(FRHITextureCube* TextureCubeRHI,uint32 FaceIndex,uint32 ArrayIndex,uint32 MipIndex,bool bLockWithinMiptail)
 {
 	@autoreleasepool {
 		FMetalTextureCube* TextureCube = ResourceCast(TextureCubeRHI);
@@ -2362,7 +2362,7 @@ FTextureReferenceRHIRef FMetalDynamicRHI::RHICreateTextureReference(FLastRenderT
 	}
 }
 
-void FMetalRHICommandContext::RHIUpdateTextureReference(FRHITextureReference* TextureRefRHI, FTextureRHIParamRef NewTextureRHI)
+void FMetalRHICommandContext::RHIUpdateTextureReference(FRHITextureReference* TextureRefRHI, FRHITexture* NewTextureRHI)
 {
 	@autoreleasepool {
 		FMetalTextureReference* TextureRef = (FMetalTextureReference*)TextureRefRHI;
@@ -2378,7 +2378,7 @@ void FMetalRHICommandContext::RHIUpdateTextureReference(FRHITextureReference* Te
 }
 
 
-void FMetalDynamicRHI::RHIBindDebugLabelName(FTextureRHIParamRef TextureRHI, const TCHAR* Name)
+void FMetalDynamicRHI::RHIBindDebugLabelName(FRHITexture* TextureRHI, const TCHAR* Name)
 {
 	@autoreleasepool {
 		FMetalSurface* Surf = GetMetalSurfaceFromRHITexture(TextureRHI);
@@ -2393,21 +2393,21 @@ void FMetalDynamicRHI::RHIBindDebugLabelName(FTextureRHIParamRef TextureRHI, con
 	}
 }
 
-void FMetalDynamicRHI::RHIVirtualTextureSetFirstMipInMemory(FTexture2DRHIParamRef TextureRHI, uint32 FirstMip)
+void FMetalDynamicRHI::RHIVirtualTextureSetFirstMipInMemory(FRHITexture2D* TextureRHI, uint32 FirstMip)
 {
 	NOT_SUPPORTED("RHIVirtualTextureSetFirstMipInMemory");
 }
 
-void FMetalDynamicRHI::RHIVirtualTextureSetFirstMipVisible(FTexture2DRHIParamRef TextureRHI, uint32 FirstMip)
+void FMetalDynamicRHI::RHIVirtualTextureSetFirstMipVisible(FRHITexture2D* TextureRHI, uint32 FirstMip)
 {
 	NOT_SUPPORTED("RHIVirtualTextureSetFirstMipVisible");
 }
 
 struct FMetalRHICommandUnaliasTextures final : public FRHICommand<FMetalRHICommandUnaliasTextures>
 {
-	TArray<FTextureRHIParamRef> Textures;
+	TArray<FRHITexture*> Textures;
 	
-	FORCEINLINE_DEBUGGABLE FMetalRHICommandUnaliasTextures(FTextureRHIParamRef* InTextures, int32 NumTextures)
+	FORCEINLINE_DEBUGGABLE FMetalRHICommandUnaliasTextures(FRHITexture** InTextures, int32 NumTextures)
 	{
 		check(InTextures && NumTextures);
 		Textures.Append(InTextures, NumTextures);
@@ -2425,7 +2425,7 @@ struct FMetalRHICommandUnaliasTextures final : public FRHICommand<FMetalRHIComma
 	}
 };
 
-void FMetalDynamicRHI::RHIAcquireTransientResource_RenderThread(FTextureRHIParamRef Texture)
+void FMetalDynamicRHI::RHIAcquireTransientResource_RenderThread(FRHITexture* Texture)
 {
 	@autoreleasepool {
 	FRHICommandListImmediate& RHICmdList = FRHICommandListExecutor::GetImmediateCommandList();
@@ -2444,9 +2444,9 @@ void FMetalDynamicRHI::RHIAcquireTransientResource_RenderThread(FTextureRHIParam
 
 struct FMetalRHICommandAliasTextures final : public FRHICommand<FMetalRHICommandAliasTextures>
 {
-	TArray<FTextureRHIParamRef> Textures;
+	TArray<FRHITexture*> Textures;
 	
-	FORCEINLINE_DEBUGGABLE FMetalRHICommandAliasTextures(FTextureRHIParamRef* InTextures, int32 NumTextures)
+	FORCEINLINE_DEBUGGABLE FMetalRHICommandAliasTextures(FRHITexture** InTextures, int32 NumTextures)
 	{
 		check(InTextures && NumTextures);
 		Textures.Append(InTextures, NumTextures);
@@ -2464,7 +2464,7 @@ struct FMetalRHICommandAliasTextures final : public FRHICommand<FMetalRHICommand
 	}
 };
 
-void FMetalDynamicRHI::RHIDiscardTransientResource_RenderThread(FTextureRHIParamRef Texture)
+void FMetalDynamicRHI::RHIDiscardTransientResource_RenderThread(FRHITexture* Texture)
 {
 	@autoreleasepool {
 	FRHICommandListImmediate& RHICmdList = FRHICommandListExecutor::GetImmediateCommandList();
@@ -2499,7 +2499,7 @@ struct FMetalRHICommandAliasBuffer final : public FRHICommand<FMetalRHICommandAl
 	}
 };
 
-void FMetalDynamicRHI::RHIAcquireTransientResource_RenderThread(FVertexBufferRHIParamRef Buffer)
+void FMetalDynamicRHI::RHIAcquireTransientResource_RenderThread(FRHIVertexBuffer* Buffer)
 {
 	@autoreleasepool {
 	FMetalVertexBuffer* MetalBuffer = ResourceCast(Buffer);
@@ -2515,14 +2515,14 @@ void FMetalDynamicRHI::RHIAcquireTransientResource_RenderThread(FVertexBufferRHI
 	}
 	}
 }
-void FMetalDynamicRHI::RHIDiscardTransientResource_RenderThread(FVertexBufferRHIParamRef Buffer)
+void FMetalDynamicRHI::RHIDiscardTransientResource_RenderThread(FRHIVertexBuffer* Buffer)
 {
 	@autoreleasepool {
 	FMetalVertexBuffer* MetalBuffer = ResourceCast(Buffer);
 	MetalBuffer->Alias();
 	}
 }
-void FMetalDynamicRHI::RHIAcquireTransientResource_RenderThread(FStructuredBufferRHIParamRef Buffer)
+void FMetalDynamicRHI::RHIAcquireTransientResource_RenderThread(FRHIStructuredBuffer* Buffer)
 {
 	@autoreleasepool {
 	FMetalStructuredBuffer* MetalBuffer = ResourceCast(Buffer);
@@ -2538,7 +2538,7 @@ void FMetalDynamicRHI::RHIAcquireTransientResource_RenderThread(FStructuredBuffe
 	}
 	}
 }
-void FMetalDynamicRHI::RHIDiscardTransientResource_RenderThread(FStructuredBufferRHIParamRef Buffer)
+void FMetalDynamicRHI::RHIDiscardTransientResource_RenderThread(FRHIStructuredBuffer* Buffer)
 {
 	@autoreleasepool {
 	FMetalStructuredBuffer* MetalBuffer = ResourceCast(Buffer);
@@ -2548,12 +2548,12 @@ void FMetalDynamicRHI::RHIDiscardTransientResource_RenderThread(FStructuredBuffe
 
 struct FRHICopySubTextureRegion final : public FRHICommand<FRHICopySubTextureRegion>
 {
-	FTexture2DRHIParamRef SourceTexture;
-	FTexture2DRHIParamRef DestinationTexture;
+	FRHITexture2D* SourceTexture;
+	FRHITexture2D* DestinationTexture;
 	FBox2D SourceBox;
 	FBox2D DestinationBox;
 	
-	FORCEINLINE_DEBUGGABLE FRHICopySubTextureRegion(FTexture2DRHIParamRef InSourceTexture, FTexture2DRHIParamRef InDestinationTexture, FBox2D InSourceBox, FBox2D InDestinationBox)
+	FORCEINLINE_DEBUGGABLE FRHICopySubTextureRegion(FRHITexture2D* InSourceTexture, FRHITexture2D* InDestinationTexture, FBox2D InSourceBox, FBox2D InDestinationBox)
 	: SourceTexture(InSourceTexture)
 	, DestinationTexture(InDestinationTexture)
 	, SourceBox(InSourceBox)
@@ -2567,7 +2567,7 @@ struct FRHICopySubTextureRegion final : public FRHICommand<FRHICopySubTextureReg
 	}
 };
 
-void FMetalDynamicRHI::RHICopySubTextureRegion_RenderThread(class FRHICommandListImmediate& RHICmdList, FTexture2DRHIParamRef SourceTexture, FTexture2DRHIParamRef DestinationTexture, FBox2D SourceBox, FBox2D DestinationBox)
+void FMetalDynamicRHI::RHICopySubTextureRegion_RenderThread(class FRHICommandListImmediate& RHICmdList, FRHITexture2D* SourceTexture, FRHITexture2D* DestinationTexture, FBox2D SourceBox, FBox2D DestinationBox)
 {
 	@autoreleasepool
 	{
@@ -2582,7 +2582,7 @@ void FMetalDynamicRHI::RHICopySubTextureRegion_RenderThread(class FRHICommandLis
 	}
 }
 
-void FMetalDynamicRHI::RHICopySubTextureRegion(FTexture2DRHIParamRef SourceTexture, FTexture2DRHIParamRef DestinationTexture, FBox2D SourceBox, FBox2D DestinationBox)
+void FMetalDynamicRHI::RHICopySubTextureRegion(FRHITexture2D* SourceTexture, FRHITexture2D* DestinationTexture, FBox2D SourceBox, FBox2D DestinationBox)
 {
 	@autoreleasepool {
 		check(SourceTexture);
@@ -2666,7 +2666,7 @@ void FMetalDynamicRHI::RHICopySubTextureRegion(FTexture2DRHIParamRef SourceTextu
 
 
 
-void FMetalRHICommandContext::RHICopyTexture(FTextureRHIParamRef SourceTextureRHI, FTextureRHIParamRef DestTextureRHI, const FRHICopyTextureInfo& CopyInfo)
+void FMetalRHICommandContext::RHICopyTexture(FRHITexture* SourceTextureRHI, FRHITexture* DestTextureRHI, const FRHICopyTextureInfo& CopyInfo)
 {
 	if (!SourceTextureRHI || !DestTextureRHI || SourceTextureRHI == DestTextureRHI)
 	{
diff --git a/Engine/Source/Runtime/Apple/MetalRHI/Private/MetalUAV.cpp b/Engine/Source/Runtime/Apple/MetalRHI/Private/MetalUAV.cpp
--- a/Engine/Source/Runtime/Apple/MetalRHI/Private/MetalUAV.cpp
+++ b/Engine/Source/Runtime/Apple/MetalRHI/Private/MetalUAV.cpp
@@ -57,12 +57,12 @@ ns::AutoReleased<FMetalTexture> FMetalShaderResourceView::GetLinearTexture(bool
 	return NewLinearTexture;
 }
 
-FUnorderedAccessViewRHIRef FMetalDynamicRHI::RHICreateUnorderedAccessView_RenderThread(class FRHICommandListImmediate& RHICmdList, FStructuredBufferRHIParamRef StructuredBuffer, bool bUseUAVCounter, bool bAppendBuffer)
+FUnorderedAccessViewRHIRef FMetalDynamicRHI::RHICreateUnorderedAccessView_RenderThread(class FRHICommandListImmediate& RHICmdList, FRHIStructuredBuffer* StructuredBuffer, bool bUseUAVCounter, bool bAppendBuffer)
 {
 	return GDynamicRHI->RHICreateUnorderedAccessView(StructuredBuffer, bUseUAVCounter, bAppendBuffer);
 }
 
-FUnorderedAccessViewRHIRef FMetalDynamicRHI::RHICreateUnorderedAccessView_RenderThread(class FRHICommandListImmediate& RHICmdList, FTextureRHIParamRef Texture, uint32 MipLevel)
+FUnorderedAccessViewRHIRef FMetalDynamicRHI::RHICreateUnorderedAccessView_RenderThread(class FRHICommandListImmediate& RHICmdList, FRHITexture* Texture, uint32 MipLevel)
 {
 	FMetalSurface* Surface = (FMetalSurface*)Texture->GetTextureBaseRHI();
 	FMetalTexture Tex = Surface->Texture;
@@ -77,21 +77,21 @@ FUnorderedAccessViewRHIRef FMetalDynamicRHI::RHICreateUnorderedAccessView_Render
 	}
 }
 
-FUnorderedAccessViewRHIRef FMetalDynamicRHI::RHICreateUnorderedAccessView_RenderThread(class FRHICommandListImmediate& RHICmdList, FVertexBufferRHIParamRef VertexBuffer, uint8 Format)
+FUnorderedAccessViewRHIRef FMetalDynamicRHI::RHICreateUnorderedAccessView_RenderThread(class FRHICommandListImmediate& RHICmdList, FRHIVertexBuffer* VertexBuffer, uint8 Format)
 {
 	FUnorderedAccessViewRHIRef Result = GDynamicRHI->RHICreateUnorderedAccessView(VertexBuffer, Format);
 	if (IsRunningRHIInSeparateThread() && !RHICmdList.Bypass()) 	{ 		RHICmdList.RHIThreadFence(true); 	}
 	return Result;
 }
 
-FUnorderedAccessViewRHIRef FMetalDynamicRHI::RHICreateUnorderedAccessView_RenderThread(class FRHICommandListImmediate& RHICmdList, FIndexBufferRHIParamRef IndexBuffer, uint8 Format)
+FUnorderedAccessViewRHIRef FMetalDynamicRHI::RHICreateUnorderedAccessView_RenderThread(class FRHICommandListImmediate& RHICmdList, FRHIIndexBuffer* IndexBuffer, uint8 Format)
 {
 	FUnorderedAccessViewRHIRef Result = GDynamicRHI->RHICreateUnorderedAccessView(IndexBuffer, Format);
 	if (IsRunningRHIInSeparateThread() && !RHICmdList.Bypass()) 	{ 		RHICmdList.RHIThreadFence(true); 	}
 	return Result;
 }
 
-FUnorderedAccessViewRHIRef FMetalDynamicRHI::RHICreateUnorderedAccessView(FStructuredBufferRHIParamRef StructuredBufferRHI, bool bUseUAVCounter, bool bAppendBuffer)
+FUnorderedAccessViewRHIRef FMetalDynamicRHI::RHICreateUnorderedAccessView(FRHIStructuredBuffer* StructuredBufferRHI, bool bUseUAVCounter, bool bAppendBuffer)
 {
 	@autoreleasepool {
 	FMetalStructuredBuffer* StructuredBuffer = ResourceCast(StructuredBufferRHI);
@@ -109,7 +109,7 @@ FUnorderedAccessViewRHIRef FMetalDynamicRHI::RHICreateUnorderedAccessView(FStruc
 	}
 }
 
-FUnorderedAccessViewRHIRef FMetalDynamicRHI::RHICreateUnorderedAccessView(FTextureRHIParamRef TextureRHI, uint32 MipLevel)
+FUnorderedAccessViewRHIRef FMetalDynamicRHI::RHICreateUnorderedAccessView(FRHITexture* TextureRHI, uint32 MipLevel)
 {
 	@autoreleasepool {
 	FMetalShaderResourceView* SRV = new FMetalShaderResourceView;
@@ -139,7 +139,7 @@ FUnorderedAccessViewRHIRef FMetalDynamicRHI::RHICreateUnorderedAccessView(FTextu
 	}
 }
 
-FUnorderedAccessViewRHIRef FMetalDynamicRHI::RHICreateUnorderedAccessView(FVertexBufferRHIParamRef VertexBufferRHI, uint8 Format)
+FUnorderedAccessViewRHIRef FMetalDynamicRHI::RHICreateUnorderedAccessView(FRHIVertexBuffer* VertexBufferRHI, uint8 Format)
 {
 	@autoreleasepool {
 	FMetalVertexBuffer* VertexBuffer = ResourceCast(VertexBufferRHI);
@@ -163,7 +163,7 @@ FUnorderedAccessViewRHIRef FMetalDynamicRHI::RHICreateUnorderedAccessView(FVerte
 	}
 }
 
-FUnorderedAccessViewRHIRef FMetalDynamicRHI::RHICreateUnorderedAccessView(FIndexBufferRHIParamRef IndexBufferRHI, uint8 Format)
+FUnorderedAccessViewRHIRef FMetalDynamicRHI::RHICreateUnorderedAccessView(FRHIIndexBuffer* IndexBufferRHI, uint8 Format)
 {
 	@autoreleasepool {
 		FMetalIndexBuffer* IndexBuffer = ResourceCast(IndexBufferRHI);
@@ -187,40 +187,40 @@ FUnorderedAccessViewRHIRef FMetalDynamicRHI::RHICreateUnorderedAccessView(FIndex
 	}
 }
 
-FShaderResourceViewRHIRef FMetalDynamicRHI::CreateShaderResourceView_RenderThread(class FRHICommandListImmediate& RHICmdList, FVertexBufferRHIParamRef VertexBuffer, uint32 Stride, uint8 Format)
+FShaderResourceViewRHIRef FMetalDynamicRHI::CreateShaderResourceView_RenderThread(class FRHICommandListImmediate& RHICmdList, FRHIVertexBuffer* VertexBuffer, uint32 Stride, uint8 Format)
 {
 	FShaderResourceViewRHIRef Result = GDynamicRHI->RHICreateShaderResourceView(VertexBuffer, Stride, Format);
 	if (IsRunningRHIInSeparateThread() && !RHICmdList.Bypass()) 	{ 		RHICmdList.RHIThreadFence(true); 	}
 	return Result;
 }
 
-FShaderResourceViewRHIRef FMetalDynamicRHI::CreateShaderResourceView_RenderThread(class FRHICommandListImmediate& RHICmdList, FIndexBufferRHIParamRef Buffer)
+FShaderResourceViewRHIRef FMetalDynamicRHI::CreateShaderResourceView_RenderThread(class FRHICommandListImmediate& RHICmdList, FRHIIndexBuffer* Buffer)
 {
 	FShaderResourceViewRHIRef Result = GDynamicRHI->RHICreateShaderResourceView(Buffer);
 	if (IsRunningRHIInSeparateThread() && !RHICmdList.Bypass()) 	{ 		RHICmdList.RHIThreadFence(true); 	}
 	return Result;
 }
 
-FShaderResourceViewRHIRef FMetalDynamicRHI::RHICreateShaderResourceView_RenderThread(class FRHICommandListImmediate& RHICmdList, FVertexBufferRHIParamRef VertexBuffer, uint32 Stride, uint8 Format)
+FShaderResourceViewRHIRef FMetalDynamicRHI::RHICreateShaderResourceView_RenderThread(class FRHICommandListImmediate& RHICmdList, FRHIVertexBuffer* VertexBuffer, uint32 Stride, uint8 Format)
 {
 	FShaderResourceViewRHIRef Result = GDynamicRHI->RHICreateShaderResourceView(VertexBuffer, Stride, Format);
 	if (IsRunningRHIInSeparateThread() && !RHICmdList.Bypass()) 	{ 		RHICmdList.RHIThreadFence(true); 	}
 	return Result;
 }
 
-FShaderResourceViewRHIRef FMetalDynamicRHI::RHICreateShaderResourceView_RenderThread(class FRHICommandListImmediate& RHICmdList, FIndexBufferRHIParamRef Buffer)
+FShaderResourceViewRHIRef FMetalDynamicRHI::RHICreateShaderResourceView_RenderThread(class FRHICommandListImmediate& RHICmdList, FRHIIndexBuffer* Buffer)
 {
 	FShaderResourceViewRHIRef Result = GDynamicRHI->RHICreateShaderResourceView(Buffer);
 	if (IsRunningRHIInSeparateThread() && !RHICmdList.Bypass()) 	{ 		RHICmdList.RHIThreadFence(true); 	}
 	return Result;
 }
 
-FShaderResourceViewRHIRef FMetalDynamicRHI::RHICreateShaderResourceView_RenderThread(class FRHICommandListImmediate& RHICmdList, FStructuredBufferRHIParamRef StructuredBuffer)
+FShaderResourceViewRHIRef FMetalDynamicRHI::RHICreateShaderResourceView_RenderThread(class FRHICommandListImmediate& RHICmdList, FRHIStructuredBuffer* StructuredBuffer)
 {
 	return GDynamicRHI->RHICreateShaderResourceView(StructuredBuffer);
 }
 
-FShaderResourceViewRHIRef FMetalDynamicRHI::RHICreateShaderResourceView_RenderThread(class FRHICommandListImmediate& RHICmdList, FTextureRHIParamRef Texture2DRHI, const FRHITextureSRVCreateInfo& CreateInfo)
+FShaderResourceViewRHIRef FMetalDynamicRHI::RHICreateShaderResourceView_RenderThread(class FRHICommandListImmediate& RHICmdList, FRHITexture* Texture2DRHI, const FRHITextureSRVCreateInfo& CreateInfo)
 {
 	FMetalSurface* Surface = (FMetalSurface*)Texture2DRHI->GetTextureBaseRHI();
 	FMetalTexture Tex = Surface->Texture;
@@ -235,7 +235,7 @@ FShaderResourceViewRHIRef FMetalDynamicRHI::RHICreateShaderResourceView_RenderTh
 	}
 }
 
-FShaderResourceViewRHIRef FMetalDynamicRHI::RHICreateShaderResourceView(FTextureRHIParamRef Texture2DRHI, const FRHITextureSRVCreateInfo& CreateInfo)
+FShaderResourceViewRHIRef FMetalDynamicRHI::RHICreateShaderResourceView(FRHITexture* Texture2DRHI, const FRHITextureSRVCreateInfo& CreateInfo)
 {
 	@autoreleasepool {
 		FMetalShaderResourceView* SRV = new FMetalShaderResourceView;
@@ -261,7 +261,7 @@ FShaderResourceViewRHIRef FMetalDynamicRHI::RHICreateShaderResourceView(FTexture
 	}
 }
 
-FShaderResourceViewRHIRef FMetalDynamicRHI::RHICreateShaderResourceView(FStructuredBufferRHIParamRef StructuredBufferRHI)
+FShaderResourceViewRHIRef FMetalDynamicRHI::RHICreateShaderResourceView(FRHIStructuredBuffer* StructuredBufferRHI)
 {
 	FMetalStructuredBuffer* StructuredBuffer = ResourceCast(StructuredBufferRHI);
 
@@ -274,7 +274,7 @@ FShaderResourceViewRHIRef FMetalDynamicRHI::RHICreateShaderResourceView(FStructu
 	return SRV;
 }
 
-FShaderResourceViewRHIRef FMetalDynamicRHI::RHICreateShaderResourceView(FVertexBufferRHIParamRef VertexBufferRHI, uint32 Stride, uint8 Format)
+FShaderResourceViewRHIRef FMetalDynamicRHI::RHICreateShaderResourceView(FRHIVertexBuffer* VertexBufferRHI, uint32 Stride, uint8 Format)
 {
 	@autoreleasepool {
 	if (!VertexBufferRHI)
@@ -308,7 +308,7 @@ FShaderResourceViewRHIRef FMetalDynamicRHI::RHICreateShaderResourceView(FVertexB
 	}
 }
 
-FShaderResourceViewRHIRef FMetalDynamicRHI::RHICreateShaderResourceView(FIndexBufferRHIParamRef BufferRHI)
+FShaderResourceViewRHIRef FMetalDynamicRHI::RHICreateShaderResourceView(FRHIIndexBuffer* BufferRHI)
 {
 	@autoreleasepool {
 	if (!BufferRHI)
@@ -339,7 +339,7 @@ FShaderResourceViewRHIRef FMetalDynamicRHI::RHICreateShaderResourceView(FIndexBu
 	}
 }
 
-void FMetalDynamicRHI::RHIUpdateShaderResourceView(FRHIShaderResourceView* SRVRHI, FVertexBufferRHIParamRef VertexBufferRHI, uint32 Stride, uint8 Format)
+void FMetalDynamicRHI::RHIUpdateShaderResourceView(FRHIShaderResourceView* SRVRHI, FRHIVertexBuffer* VertexBufferRHI, uint32 Stride, uint8 Format)
 {
 	check(SRVRHI);
 	FMetalShaderResourceView* SRV = ResourceCast(SRVRHI);
@@ -364,7 +364,7 @@ void FMetalDynamicRHI::RHIUpdateShaderResourceView(FRHIShaderResourceView* SRVRH
 	}
 }
 
-void FMetalDynamicRHI::RHIUpdateShaderResourceView(FRHIShaderResourceView* SRVRHI, FIndexBufferRHIParamRef IndexBufferRHI)
+void FMetalDynamicRHI::RHIUpdateShaderResourceView(FRHIShaderResourceView* SRVRHI, FRHIIndexBuffer* IndexBufferRHI)
 {
 	check(SRVRHI);
 	FMetalShaderResourceView* SRV = ResourceCast(SRVRHI);
@@ -742,7 +742,7 @@ void FMetalRHICommandContext::RHITransitionResources(EResourceTransitionAccess T
 	}
 }
 
-void FMetalRHICommandContext::RHITransitionResources(EResourceTransitionAccess TransitionType, FTextureRHIParamRef* InTextures, int32 NumTextures)
+void FMetalRHICommandContext::RHITransitionResources(EResourceTransitionAccess TransitionType, FRHITexture** InTextures, int32 NumTextures)
 {
 	@autoreleasepool
 	{
@@ -779,7 +779,7 @@ void FMetalGPUFence::WriteInternal(mtlpp::CommandBuffer& CmdBuffer)
 	check(Fence);
 }
 
-void FMetalRHICommandContext::RHICopyToStagingBuffer(FVertexBufferRHIParamRef SourceBufferRHI, FRHIStagingBuffer* DestinationStagingBufferRHI, uint32 Offset, uint32 NumBytes)
+void FMetalRHICommandContext::RHICopyToStagingBuffer(FRHIVertexBuffer* SourceBufferRHI, FRHIStagingBuffer* DestinationStagingBufferRHI, uint32 Offset, uint32 NumBytes)
 {
 	@autoreleasepool {
 		check(DestinationStagingBufferRHI);
diff --git a/Engine/Source/Runtime/Apple/MetalRHI/Private/MetalUniformBuffer.cpp b/Engine/Source/Runtime/Apple/MetalRHI/Private/MetalUniformBuffer.cpp
--- a/Engine/Source/Runtime/Apple/MetalRHI/Private/MetalUniformBuffer.cpp
+++ b/Engine/Source/Runtime/Apple/MetalRHI/Private/MetalUniformBuffer.cpp
@@ -683,7 +683,7 @@ struct FMetalRHICommandUpateUniformBuffer : public FRHICommand<FMetalRHICommandU
 	}
 };
 
-void FMetalDynamicRHI::RHIUpdateUniformBuffer(FUniformBufferRHIParamRef UniformBufferRHI, const void* Contents)
+void FMetalDynamicRHI::RHIUpdateUniformBuffer(FRHIUniformBuffer* UniformBufferRHI, const void* Contents)
 {
 	@autoreleasepool {
 	// check((IsInRenderingThread() || IsInRHIThread()) && !IsInParallelRenderingThread());
diff --git a/Engine/Source/Runtime/Apple/MetalRHI/Private/MetalVertexBuffer.cpp b/Engine/Source/Runtime/Apple/MetalRHI/Private/MetalVertexBuffer.cpp
--- a/Engine/Source/Runtime/Apple/MetalRHI/Private/MetalVertexBuffer.cpp
+++ b/Engine/Source/Runtime/Apple/MetalRHI/Private/MetalVertexBuffer.cpp
@@ -584,7 +584,7 @@ FVertexBufferRHIRef FMetalDynamicRHI::RHICreateVertexBuffer(uint32 Size, uint32
 	}
 }
 
-void* FMetalDynamicRHI::RHILockVertexBuffer(FVertexBufferRHIParamRef VertexBufferRHI, uint32 Offset, uint32 Size, EResourceLockMode LockMode)
+void* FMetalDynamicRHI::RHILockVertexBuffer(FRHIVertexBuffer* VertexBufferRHI, uint32 Offset, uint32 Size, EResourceLockMode LockMode)
 {
 	@autoreleasepool {
 	FMetalVertexBuffer* VertexBuffer = ResourceCast(VertexBufferRHI);
@@ -594,7 +594,7 @@ void* FMetalDynamicRHI::RHILockVertexBuffer(FVertexBufferRHIParamRef VertexBuffe
 	}
 }
 
-void FMetalDynamicRHI::RHIUnlockVertexBuffer(FVertexBufferRHIParamRef VertexBufferRHI)
+void FMetalDynamicRHI::RHIUnlockVertexBuffer(FRHIVertexBuffer* VertexBufferRHI)
 {
 	@autoreleasepool {
 	FMetalVertexBuffer* VertexBuffer = ResourceCast(VertexBufferRHI);
@@ -603,7 +603,7 @@ void FMetalDynamicRHI::RHIUnlockVertexBuffer(FVertexBufferRHIParamRef VertexBuff
 	}
 }
 
-void FMetalDynamicRHI::RHICopyVertexBuffer(FVertexBufferRHIParamRef SourceBufferRHI,FVertexBufferRHIParamRef DestBufferRHI)
+void FMetalDynamicRHI::RHICopyVertexBuffer(FRHIVertexBuffer* SourceBufferRHI, FRHIVertexBuffer* DestBufferRHI)
 {
 	@autoreleasepool {
 		FMetalVertexBuffer* SrcVertexBuffer = ResourceCast(SourceBufferRHI);
@@ -742,7 +742,7 @@ FVertexBufferRHIRef FMetalDynamicRHI::CreateVertexBuffer_RenderThread(class FRHI
 	}
 }
 
-void FMetalDynamicRHI::RHITransferVertexBufferUnderlyingResource(FVertexBufferRHIParamRef DestVertexBuffer, FVertexBufferRHIParamRef SrcVertexBuffer)
+void FMetalDynamicRHI::RHITransferVertexBufferUnderlyingResource(FRHIVertexBuffer* DestVertexBuffer, FRHIVertexBuffer* SrcVertexBuffer)
 {
 	check(DestVertexBuffer);
 	FMetalVertexBuffer* Dest = ResourceCast(DestVertexBuffer);
diff --git a/Engine/Source/Runtime/Apple/MetalRHI/Private/MetalViewport.cpp b/Engine/Source/Runtime/Apple/MetalRHI/Private/MetalViewport.cpp
--- a/Engine/Source/Runtime/Apple/MetalRHI/Private/MetalViewport.cpp
+++ b/Engine/Source/Runtime/Apple/MetalRHI/Private/MetalViewport.cpp
@@ -236,7 +236,7 @@ void FMetalViewport::Resize(uint32 InSizeX, uint32 InSizeY, bool bInIsFullscreen
         FTexture2DRHIRef DoubleBuffer;
         if (GMetalSupportsIntermediateBackBuffer)
         {
-            NewBackBuffer = (FMetalTexture2D*)(FTexture2DRHIParamRef)GDynamicRHI->RHICreateTexture2D(InSizeX, InSizeY, Format, 1, 1, TexCreate_RenderTargetable, CreateInfo);
+            NewBackBuffer = (FMetalTexture2D*)(FRHITexture2D*)GDynamicRHI->RHICreateTexture2D(InSizeX, InSizeY, Format, 1, 1, TexCreate_RenderTargetable, CreateInfo);
             
             if (GMetalSeparatePresentThread)
             {
@@ -246,7 +246,7 @@ void FMetalViewport::Resize(uint32 InSizeX, uint32 InSizeY, bool bInIsFullscreen
         }
         else
         {
-            NewBackBuffer = (FMetalTexture2D*)(FTexture2DRHIParamRef)GDynamicRHI->RHICreateTexture2D(InSizeX, InSizeY, Format, 1, 1, TexCreate_RenderTargetable | TexCreate_Presentable, CreateInfo);
+            NewBackBuffer = (FMetalTexture2D*)(FRHITexture2D*)GDynamicRHI->RHICreateTexture2D(InSizeX, InSizeY, Format, 1, 1, TexCreate_RenderTargetable | TexCreate_Presentable, CreateInfo);
         }
         ((FMetalTexture2D*)NewBackBuffer.GetReference())->Surface.Viewport = this;
         
@@ -606,12 +606,12 @@ void FMetalDynamicRHI::RHITick( float DeltaTime )
  *	Viewport functions.
  *=============================================================================*/
 
-void FMetalRHICommandContext::RHIBeginDrawingViewport(FRHIViewport* ViewportRHI, FTextureRHIParamRef RenderTargetRHI)
+void FMetalRHICommandContext::RHIBeginDrawingViewport(FRHIViewport* ViewportRHI, FRHITexture* RenderTargetRHI)
 {
 	check(false);
 }
 
-void FMetalRHIImmediateCommandContext::RHIBeginDrawingViewport(FRHIViewport* ViewportRHI, FTextureRHIParamRef RenderTargetRHI)
+void FMetalRHIImmediateCommandContext::RHIBeginDrawingViewport(FRHIViewport* ViewportRHI, FRHITexture* RenderTargetRHI)
 {
 	@autoreleasepool {
 	FMetalViewport* Viewport = ResourceCast(ViewportRHI);
diff --git a/Engine/Source/Runtime/Apple/MetalRHI/Public/MetalRHI.h b/Engine/Source/Runtime/Apple/MetalRHI/Public/MetalRHI.h
--- a/Engine/Source/Runtime/Apple/MetalRHI/Public/MetalRHI.h
+++ b/Engine/Source/Runtime/Apple/MetalRHI/Public/MetalRHI.h
@@ -64,27 +64,27 @@ public:
 	virtual TRefCountPtr<FRHIComputePipelineState> RHICreateComputePipelineState(FRHIComputeShader* ComputeShader) final override;
 	virtual FUniformBufferRHIRef RHICreateUniformBuffer(const void* Contents, const FRHIUniformBufferLayout& Layout, EUniformBufferUsage Usage, EUniformBufferValidation Validation) final override;
 	virtual FIndexBufferRHIRef RHICreateIndexBuffer(uint32 Stride, uint32 Size, uint32 InUsage, FRHIResourceCreateInfo& CreateInfo) final override;
-	virtual void* RHILockIndexBuffer(FIndexBufferRHIParamRef IndexBuffer, uint32 Offset, uint32 Size, EResourceLockMode LockMode) final override;
-	virtual void RHIUnlockIndexBuffer(FIndexBufferRHIParamRef IndexBuffer) final override;
-	virtual void RHITransferIndexBufferUnderlyingResource(FIndexBufferRHIParamRef DestIndexBuffer, FIndexBufferRHIParamRef SrcIndexBuffer) final override;
+	virtual void* RHILockIndexBuffer(FRHIIndexBuffer* IndexBuffer, uint32 Offset, uint32 Size, EResourceLockMode LockMode) final override;
+	virtual void RHIUnlockIndexBuffer(FRHIIndexBuffer* IndexBuffer) final override;
+	virtual void RHITransferIndexBufferUnderlyingResource(FRHIIndexBuffer* DestIndexBuffer, FRHIIndexBuffer* SrcIndexBuffer) final override;
 	virtual FVertexBufferRHIRef RHICreateVertexBuffer(uint32 Size, uint32 InUsage, FRHIResourceCreateInfo& CreateInfo) final override;
-	virtual void* RHILockVertexBuffer(FVertexBufferRHIParamRef VertexBuffer, uint32 Offset, uint32 SizeRHI, EResourceLockMode LockMode) final override;
-	virtual void RHIUnlockVertexBuffer(FVertexBufferRHIParamRef VertexBuffer) final override;
-	virtual void RHICopyVertexBuffer(FVertexBufferRHIParamRef SourceBuffer,FVertexBufferRHIParamRef DestBuffer) final override;
-	virtual void RHITransferVertexBufferUnderlyingResource(FVertexBufferRHIParamRef DestVertexBuffer, FVertexBufferRHIParamRef SrcVertexBuffer) final override;
+	virtual void* RHILockVertexBuffer(FRHIVertexBuffer* VertexBuffer, uint32 Offset, uint32 SizeRHI, EResourceLockMode LockMode) final override;
+	virtual void RHIUnlockVertexBuffer(FRHIVertexBuffer* VertexBuffer) final override;
+	virtual void RHICopyVertexBuffer(FRHIVertexBuffer* SourceBuffer, FRHIVertexBuffer* DestBuffer) final override;
+	virtual void RHITransferVertexBufferUnderlyingResource(FRHIVertexBuffer* DestVertexBuffer, FRHIVertexBuffer* SrcVertexBuffer) final override;
 	virtual FStructuredBufferRHIRef RHICreateStructuredBuffer(uint32 Stride, uint32 Size, uint32 InUsage, FRHIResourceCreateInfo& CreateInfo) final override;
-	virtual void* RHILockStructuredBuffer(FStructuredBufferRHIParamRef StructuredBuffer, uint32 Offset, uint32 SizeRHI, EResourceLockMode LockMode) final override;
-	virtual void RHIUnlockStructuredBuffer(FStructuredBufferRHIParamRef StructuredBuffer) final override;
-	virtual FUnorderedAccessViewRHIRef RHICreateUnorderedAccessView(FStructuredBufferRHIParamRef StructuredBuffer, bool bUseUAVCounter, bool bAppendBuffer) final override;
-	virtual FUnorderedAccessViewRHIRef RHICreateUnorderedAccessView(FTextureRHIParamRef Texture, uint32 MipLevel) final override;
-	virtual FUnorderedAccessViewRHIRef RHICreateUnorderedAccessView(FVertexBufferRHIParamRef VertexBuffer, uint8 Format) final override;
-	virtual FUnorderedAccessViewRHIRef RHICreateUnorderedAccessView(FIndexBufferRHIParamRef IndexBuffer, uint8 Format) final override;
-	virtual FShaderResourceViewRHIRef RHICreateShaderResourceView(FStructuredBufferRHIParamRef StructuredBuffer) final override;
-	virtual FShaderResourceViewRHIRef RHICreateShaderResourceView(FTextureRHIParamRef Texture2DRHI, const FRHITextureSRVCreateInfo& CreateInfo) final override;
-	virtual FShaderResourceViewRHIRef RHICreateShaderResourceView(FVertexBufferRHIParamRef VertexBuffer, uint32 Stride, uint8 Format) final override;
-	virtual FShaderResourceViewRHIRef RHICreateShaderResourceView(FIndexBufferRHIParamRef Buffer) final override;
-	virtual void RHIUpdateShaderResourceView(FRHIShaderResourceView* SRV, FVertexBufferRHIParamRef VertexBuffer, uint32 Stride, uint8 Format) final override;
-	virtual void RHIUpdateShaderResourceView(FRHIShaderResourceView* SRV, FIndexBufferRHIParamRef IndexBuffer) final override;
+	virtual void* RHILockStructuredBuffer(FRHIStructuredBuffer* StructuredBuffer, uint32 Offset, uint32 SizeRHI, EResourceLockMode LockMode) final override;
+	virtual void RHIUnlockStructuredBuffer(FRHIStructuredBuffer* StructuredBuffer) final override;
+	virtual FUnorderedAccessViewRHIRef RHICreateUnorderedAccessView(FRHIStructuredBuffer* StructuredBuffer, bool bUseUAVCounter, bool bAppendBuffer) final override;
+	virtual FUnorderedAccessViewRHIRef RHICreateUnorderedAccessView(FRHITexture* Texture, uint32 MipLevel) final override;
+	virtual FUnorderedAccessViewRHIRef RHICreateUnorderedAccessView(FRHIVertexBuffer* VertexBuffer, uint8 Format) final override;
+	virtual FUnorderedAccessViewRHIRef RHICreateUnorderedAccessView(FRHIIndexBuffer* IndexBuffer, uint8 Format) final override;
+	virtual FShaderResourceViewRHIRef RHICreateShaderResourceView(FRHIStructuredBuffer* StructuredBuffer) final override;
+	virtual FShaderResourceViewRHIRef RHICreateShaderResourceView(FRHITexture* Texture2DRHI, const FRHITextureSRVCreateInfo& CreateInfo) final override;
+	virtual FShaderResourceViewRHIRef RHICreateShaderResourceView(FRHIVertexBuffer* VertexBuffer, uint32 Stride, uint8 Format) final override;
+	virtual FShaderResourceViewRHIRef RHICreateShaderResourceView(FRHIIndexBuffer* Buffer) final override;
+	virtual void RHIUpdateShaderResourceView(FRHIShaderResourceView* SRV, FRHIVertexBuffer* VertexBuffer, uint32 Stride, uint8 Format) final override;
+	virtual void RHIUpdateShaderResourceView(FRHIShaderResourceView* SRV, FRHIIndexBuffer* IndexBuffer) final override;
 	virtual uint64 RHICalcTexture2DPlatformSize(uint32 SizeX, uint32 SizeY, uint8 Format, uint32 NumMips, uint32 NumSamples, uint32 Flags, uint32& OutAlign) final override;
 	virtual uint64 RHICalcTexture3DPlatformSize(uint32 SizeX, uint32 SizeY, uint32 SizeZ, uint8 Format, uint32 NumMips, uint32 Flags, uint32& OutAlign) final override;
 	virtual uint64 RHICalcTextureCubePlatformSize(uint32 Size, uint8 Format, uint32 NumMips, uint32 Flags, uint32& OutAlign) final override;
@@ -94,32 +94,32 @@ public:
 	virtual FTexture2DRHIRef RHICreateTexture2D(uint32 SizeX, uint32 SizeY, uint8 Format, uint32 NumMips, uint32 NumSamples, uint32 Flags, FRHIResourceCreateInfo& CreateInfo) final override;
 	virtual FTexture2DRHIRef RHICreateTextureExternal2D(uint32 SizeX, uint32 SizeY, uint8 Format, uint32 NumMips, uint32 NumSamples, uint32 Flags, FRHIResourceCreateInfo& Info) final override;
 	virtual FTexture2DRHIRef RHIAsyncCreateTexture2D(uint32 SizeX, uint32 SizeY, uint8 Format, uint32 NumMips, uint32 Flags, void** InitialMipData, uint32 NumInitialMips) final override;
-	virtual void RHICopySharedMips(FTexture2DRHIParamRef DestTexture2D, FTexture2DRHIParamRef SrcTexture2D) final override;
+	virtual void RHICopySharedMips(FRHITexture2D* DestTexture2D, FRHITexture2D* SrcTexture2D) final override;
 	virtual FTexture2DArrayRHIRef RHICreateTexture2DArray(uint32 SizeX, uint32 SizeY, uint32 SizeZ, uint8 Format, uint32 NumMips, uint32 NumSamples, uint32 Flags, FRHIResourceCreateInfo& CreateInfo) final override;
 	virtual FTexture3DRHIRef RHICreateTexture3D(uint32 SizeX, uint32 SizeY, uint32 SizeZ, uint8 Format, uint32 NumMips, uint32 Flags, FRHIResourceCreateInfo& CreateInfo) final override;
-	virtual void RHIGetResourceInfo(FTextureRHIParamRef Ref, FRHIResourceInfo& OutInfo) final override;
-	virtual void RHIGenerateMips(FTextureRHIParamRef Texture) final override;
-	virtual uint32 RHIComputeMemorySize(FTextureRHIParamRef TextureRHI) final override;
-	virtual FTexture2DRHIRef RHIAsyncReallocateTexture2D(FTexture2DRHIParamRef Texture2D, int32 NewMipCount, int32 NewSizeX, int32 NewSizeY, FThreadSafeCounter* RequestStatus) final override;
-	virtual ETextureReallocationStatus RHIFinalizeAsyncReallocateTexture2D(FTexture2DRHIParamRef Texture2D, bool bBlockUntilCompleted) final override;
-	virtual ETextureReallocationStatus RHICancelAsyncReallocateTexture2D(FTexture2DRHIParamRef Texture2D, bool bBlockUntilCompleted) final override;
-	virtual void* RHILockTexture2D(FTexture2DRHIParamRef Texture, uint32 MipIndex, EResourceLockMode LockMode, uint32& DestStride, bool bLockWithinMiptail) final override;
-	virtual void RHIUnlockTexture2D(FTexture2DRHIParamRef Texture, uint32 MipIndex, bool bLockWithinMiptail) final override;
-	virtual void* RHILockTexture2DArray(FTexture2DArrayRHIParamRef Texture, uint32 TextureIndex, uint32 MipIndex, EResourceLockMode LockMode, uint32& DestStride, bool bLockWithinMiptail) final override;
-	virtual void RHIUnlockTexture2DArray(FTexture2DArrayRHIParamRef Texture, uint32 TextureIndex, uint32 MipIndex, bool bLockWithinMiptail) final override;
-	virtual void RHIUpdateTexture2D(FTexture2DRHIParamRef Texture, uint32 MipIndex, const struct FUpdateTextureRegion2D& UpdateRegion, uint32 SourcePitch, const uint8* SourceData) final override;
-	virtual void RHIUpdateTexture3D(FTexture3DRHIParamRef Texture, uint32 MipIndex, const struct FUpdateTextureRegion3D& UpdateRegion, uint32 SourceRowPitch, uint32 SourceDepthPitch, const uint8* SourceData) final override;
+	virtual void RHIGetResourceInfo(FRHITexture* Ref, FRHIResourceInfo& OutInfo) final override;
+	virtual void RHIGenerateMips(FRHITexture* Texture) final override;
+	virtual uint32 RHIComputeMemorySize(FRHITexture* TextureRHI) final override;
+	virtual FTexture2DRHIRef RHIAsyncReallocateTexture2D(FRHITexture2D* Texture2D, int32 NewMipCount, int32 NewSizeX, int32 NewSizeY, FThreadSafeCounter* RequestStatus) final override;
+	virtual ETextureReallocationStatus RHIFinalizeAsyncReallocateTexture2D(FRHITexture2D* Texture2D, bool bBlockUntilCompleted) final override;
+	virtual ETextureReallocationStatus RHICancelAsyncReallocateTexture2D(FRHITexture2D* Texture2D, bool bBlockUntilCompleted) final override;
+	virtual void* RHILockTexture2D(FRHITexture2D* Texture, uint32 MipIndex, EResourceLockMode LockMode, uint32& DestStride, bool bLockWithinMiptail) final override;
+	virtual void RHIUnlockTexture2D(FRHITexture2D* Texture, uint32 MipIndex, bool bLockWithinMiptail) final override;
+	virtual void* RHILockTexture2DArray(FRHITexture2DArray* Texture, uint32 TextureIndex, uint32 MipIndex, EResourceLockMode LockMode, uint32& DestStride, bool bLockWithinMiptail) final override;
+	virtual void RHIUnlockTexture2DArray(FRHITexture2DArray* Texture, uint32 TextureIndex, uint32 MipIndex, bool bLockWithinMiptail) final override;
+	virtual void RHIUpdateTexture2D(FRHITexture2D* Texture, uint32 MipIndex, const struct FUpdateTextureRegion2D& UpdateRegion, uint32 SourcePitch, const uint8* SourceData) final override;
+	virtual void RHIUpdateTexture3D(FRHITexture3D* Texture, uint32 MipIndex, const struct FUpdateTextureRegion3D& UpdateRegion, uint32 SourceRowPitch, uint32 SourceDepthPitch, const uint8* SourceData) final override;
 	virtual FTextureCubeRHIRef RHICreateTextureCube(uint32 Size, uint8 Format, uint32 NumMips, uint32 Flags, FRHIResourceCreateInfo& CreateInfo) final override;
 	virtual FTextureCubeRHIRef RHICreateTextureCubeArray(uint32 Size, uint32 ArraySize, uint8 Format, uint32 NumMips, uint32 Flags, FRHIResourceCreateInfo& CreateInfo) final override;
-	virtual void* RHILockTextureCubeFace(FTextureCubeRHIParamRef Texture, uint32 FaceIndex, uint32 ArrayIndex, uint32 MipIndex, EResourceLockMode LockMode, uint32& DestStride, bool bLockWithinMiptail) final override;
-	virtual void RHIUnlockTextureCubeFace(FTextureCubeRHIParamRef Texture, uint32 FaceIndex, uint32 ArrayIndex, uint32 MipIndex, bool bLockWithinMiptail) final override;
-	virtual void RHIBindDebugLabelName(FTextureRHIParamRef Texture, const TCHAR* Name) final override;
-	virtual void RHIReadSurfaceData(FTextureRHIParamRef Texture,FIntRect Rect,TArray<FColor>& OutData,FReadSurfaceDataFlags InFlags) final override;
-	virtual void RHIReadSurfaceData(FTextureRHIParamRef TextureRHI, FIntRect InRect, TArray<FLinearColor>& OutData, FReadSurfaceDataFlags InFlags) final override;
-	virtual void RHIMapStagingSurface(FTextureRHIParamRef Texture,void*& OutData,int32& OutWidth,int32& OutHeight) final override;
-	virtual void RHIUnmapStagingSurface(FTextureRHIParamRef Texture) final override;
-	virtual void RHIReadSurfaceFloatData(FTextureRHIParamRef Texture,FIntRect Rect,TArray<FFloat16Color>& OutData,ECubeFace CubeFace,int32 ArrayIndex,int32 MipIndex) final override;
-	virtual void RHIRead3DSurfaceFloatData(FTextureRHIParamRef Texture,FIntRect Rect,FIntPoint ZMinMax,TArray<FFloat16Color>& OutData) final override;
+	virtual void* RHILockTextureCubeFace(FRHITextureCube* Texture, uint32 FaceIndex, uint32 ArrayIndex, uint32 MipIndex, EResourceLockMode LockMode, uint32& DestStride, bool bLockWithinMiptail) final override;
+	virtual void RHIUnlockTextureCubeFace(FRHITextureCube* Texture, uint32 FaceIndex, uint32 ArrayIndex, uint32 MipIndex, bool bLockWithinMiptail) final override;
+	virtual void RHIBindDebugLabelName(FRHITexture* Texture, const TCHAR* Name) final override;
+	virtual void RHIReadSurfaceData(FRHITexture* Texture,FIntRect Rect,TArray<FColor>& OutData,FReadSurfaceDataFlags InFlags) final override;
+	virtual void RHIReadSurfaceData(FRHITexture* TextureRHI, FIntRect InRect, TArray<FLinearColor>& OutData, FReadSurfaceDataFlags InFlags) final override;
+	virtual void RHIMapStagingSurface(FRHITexture* Texture,void*& OutData,int32& OutWidth,int32& OutHeight) final override;
+	virtual void RHIUnmapStagingSurface(FRHITexture* Texture) final override;
+	virtual void RHIReadSurfaceFloatData(FRHITexture* Texture,FIntRect Rect,TArray<FFloat16Color>& OutData,ECubeFace CubeFace,int32 ArrayIndex,int32 MipIndex) final override;
+	virtual void RHIRead3DSurfaceFloatData(FRHITexture* Texture,FIntRect Rect,FIntPoint ZMinMax,TArray<FFloat16Color>& OutData) final override;
 	virtual FRenderQueryRHIRef RHICreateRenderQuery(ERenderQueryType QueryType) final override;
 	virtual bool RHIGetRenderQueryResult(FRHIRenderQuery* RenderQuery, uint64& OutResult, bool bWait) final override;
 	virtual FTexture2DRHIRef RHIGetViewportBackBuffer(FRHIViewport* Viewport) final override;
@@ -132,12 +132,12 @@ public:
 	virtual void RHIResizeViewport(FRHIViewport* Viewport, uint32 SizeX, uint32 SizeY, bool bIsFullscreen) final override;
 	virtual void RHIResizeViewport(FRHIViewport* Viewport, uint32 SizeX, uint32 SizeY, bool bIsFullscreen, EPixelFormat PreferredPixelFormat) final override;
 	virtual void RHITick(float DeltaTime) final override;
-	virtual void RHISetStreamOutTargets(uint32 NumTargets,const FVertexBufferRHIParamRef* VertexBuffers,const uint32* Offsets) final override;
+	virtual void RHISetStreamOutTargets(uint32 NumTargets,const FRHIVertexBuffer** VertexBuffers,const uint32* Offsets) final override;
 	virtual void RHIBlockUntilGPUIdle() final override;
 	virtual bool RHIGetAvailableResolutions(FScreenResolutionArray& Resolutions, bool bIgnoreRefreshRate) final override;
 	virtual void RHIGetSupportedResolution(uint32& Width, uint32& Height) final override;
-	virtual void RHIVirtualTextureSetFirstMipInMemory(FTexture2DRHIParamRef Texture, uint32 FirstMip) final override;
-	virtual void RHIVirtualTextureSetFirstMipVisible(FTexture2DRHIParamRef Texture, uint32 FirstMip) final override;
+	virtual void RHIVirtualTextureSetFirstMipInMemory(FRHITexture2D* Texture, uint32 FirstMip) final override;
+	virtual void RHIVirtualTextureSetFirstMipVisible(FRHITexture2D* Texture, uint32 FirstMip) final override;
 	virtual void RHIExecuteCommandList(FRHICommandList* CmdList) final override;
 	virtual void* RHIGetNativeDevice() final override;
 	virtual class IRHICommandContext* RHIGetDefaultContext() final override;
@@ -148,17 +148,17 @@ public:
 	virtual FIndexBufferRHIRef CreateIndexBuffer_RenderThread(class FRHICommandListImmediate& RHICmdList, uint32 Stride, uint32 Size, uint32 InUsage, FRHIResourceCreateInfo& CreateInfo) final override;
 	virtual FStructuredBufferRHIRef CreateStructuredBuffer_RenderThread(class FRHICommandListImmediate& RHICmdList, uint32 Stride, uint32 Size, uint32 InUsage, FRHIResourceCreateInfo& CreateInfo) final override;
 	
-	virtual FTexture2DRHIRef AsyncReallocateTexture2D_RenderThread(class FRHICommandListImmediate& RHICmdList, FTexture2DRHIParamRef Texture2D, int32 NewMipCount, int32 NewSizeX, int32 NewSizeY, FThreadSafeCounter* RequestStatus) final override;
-	virtual ETextureReallocationStatus FinalizeAsyncReallocateTexture2D_RenderThread(class FRHICommandListImmediate& RHICmdList, FTexture2DRHIParamRef Texture2D, bool bBlockUntilCompleted) final override;
-	virtual ETextureReallocationStatus CancelAsyncReallocateTexture2D_RenderThread(class FRHICommandListImmediate& RHICmdList, FTexture2DRHIParamRef Texture2D, bool bBlockUntilCompleted) final override;
+	virtual FTexture2DRHIRef AsyncReallocateTexture2D_RenderThread(class FRHICommandListImmediate& RHICmdList, FRHITexture2D* Texture2D, int32 NewMipCount, int32 NewSizeX, int32 NewSizeY, FThreadSafeCounter* RequestStatus) final override;
+	virtual ETextureReallocationStatus FinalizeAsyncReallocateTexture2D_RenderThread(class FRHICommandListImmediate& RHICmdList, FRHITexture2D* Texture2D, bool bBlockUntilCompleted) final override;
+	virtual ETextureReallocationStatus CancelAsyncReallocateTexture2D_RenderThread(class FRHICommandListImmediate& RHICmdList, FRHITexture2D* Texture2D, bool bBlockUntilCompleted) final override;
 
-	virtual void* LockTexture2D_RenderThread(class FRHICommandListImmediate& RHICmdList, FTexture2DRHIParamRef Texture, uint32 MipIndex, EResourceLockMode LockMode, uint32& DestStride, bool bLockWithinMiptail, bool bNeedsDefaultRHIFlush = true) final override;
-	virtual void UnlockTexture2D_RenderThread(class FRHICommandListImmediate& RHICmdList, FTexture2DRHIParamRef Texture, uint32 MipIndex, bool bLockWithinMiptail, bool bNeedsDefaultRHIFlush = true) final override;
-	virtual void UpdateTexture2D_RenderThread(class FRHICommandListImmediate& RHICmdList, FTexture2DRHIParamRef Texture, uint32 MipIndex, const struct FUpdateTextureRegion2D& UpdateRegion, uint32 SourcePitch, const uint8* SourceData) final override;
+	virtual void* LockTexture2D_RenderThread(class FRHICommandListImmediate& RHICmdList, FRHITexture2D* Texture, uint32 MipIndex, EResourceLockMode LockMode, uint32& DestStride, bool bLockWithinMiptail, bool bNeedsDefaultRHIFlush = true) final override;
+	virtual void UnlockTexture2D_RenderThread(class FRHICommandListImmediate& RHICmdList, FRHITexture2D* Texture, uint32 MipIndex, bool bLockWithinMiptail, bool bNeedsDefaultRHIFlush = true) final override;
+	virtual void UpdateTexture2D_RenderThread(class FRHICommandListImmediate& RHICmdList, FRHITexture2D* Texture, uint32 MipIndex, const struct FUpdateTextureRegion2D& UpdateRegion, uint32 SourcePitch, const uint8* SourceData) final override;
 	
-	virtual FUpdateTexture3DData BeginUpdateTexture3D_RenderThread(class FRHICommandListImmediate& RHICmdList, FTexture3DRHIParamRef Texture, uint32 MipIndex, const struct FUpdateTextureRegion3D& UpdateRegion) final override;
+	virtual FUpdateTexture3DData BeginUpdateTexture3D_RenderThread(class FRHICommandListImmediate& RHICmdList, FRHITexture3D* Texture, uint32 MipIndex, const struct FUpdateTextureRegion3D& UpdateRegion) final override;
 	virtual void EndUpdateTexture3D_RenderThread(class FRHICommandListImmediate& RHICmdList, FUpdateTexture3DData& UpdateData) final override;
-	virtual void UpdateTexture3D_RenderThread(class FRHICommandListImmediate& RHICmdList, FTexture3DRHIParamRef Texture, uint32 MipIndex, const struct FUpdateTextureRegion3D& UpdateRegion, uint32 SourceRowPitch, uint32 SourceDepthPitch, const uint8* SourceData) final override;
+	virtual void UpdateTexture3D_RenderThread(class FRHICommandListImmediate& RHICmdList, FRHITexture3D* Texture, uint32 MipIndex, const struct FUpdateTextureRegion3D& UpdateRegion, uint32 SourceRowPitch, uint32 SourceDepthPitch, const uint8* SourceData) final override;
 	
 	virtual FTexture2DRHIRef RHICreateTexture2D_RenderThread(class FRHICommandListImmediate& RHICmdList, uint32 SizeX, uint32 SizeY, uint8 Format, uint32 NumMips, uint32 NumSamples, uint32 Flags, FRHIResourceCreateInfo& CreateInfo) final override;
 	virtual FTexture2DArrayRHIRef RHICreateTexture2DArray_RenderThread(class FRHICommandListImmediate& RHICmdList, uint32 SizeX, uint32 SizeY, uint32 SizeZ, uint8 Format, uint32 NumMips, uint32 NumSamples, uint32 Flags, FRHIResourceCreateInfo& CreateInfo) final override;
@@ -168,17 +168,17 @@ public:
 	
 	virtual FRenderQueryRHIRef RHICreateRenderQuery_RenderThread(class FRHICommandListImmediate& RHICmdList, ERenderQueryType QueryType) final override;
 	
-	virtual FUnorderedAccessViewRHIRef RHICreateUnorderedAccessView_RenderThread(class FRHICommandListImmediate& RHICmdList, FStructuredBufferRHIParamRef StructuredBuffer, bool bUseUAVCounter, bool bAppendBuffer) final override;
-	virtual FUnorderedAccessViewRHIRef RHICreateUnorderedAccessView_RenderThread(class FRHICommandListImmediate& RHICmdList, FTextureRHIParamRef Texture, uint32 MipLevel) final override;
-	virtual FUnorderedAccessViewRHIRef RHICreateUnorderedAccessView_RenderThread(class FRHICommandListImmediate& RHICmdList, FVertexBufferRHIParamRef VertexBuffer, uint8 Format) final override;
-	virtual FUnorderedAccessViewRHIRef RHICreateUnorderedAccessView_RenderThread(class FRHICommandListImmediate& RHICmdList, FIndexBufferRHIParamRef IndexBuffer, uint8 Format) final override;
+	virtual FUnorderedAccessViewRHIRef RHICreateUnorderedAccessView_RenderThread(class FRHICommandListImmediate& RHICmdList, FRHIStructuredBuffer* StructuredBuffer, bool bUseUAVCounter, bool bAppendBuffer) final override;
+	virtual FUnorderedAccessViewRHIRef RHICreateUnorderedAccessView_RenderThread(class FRHICommandListImmediate& RHICmdList, FRHITexture* Texture, uint32 MipLevel) final override;
+	virtual FUnorderedAccessViewRHIRef RHICreateUnorderedAccessView_RenderThread(class FRHICommandListImmediate& RHICmdList, FRHIVertexBuffer* VertexBuffer, uint8 Format) final override;
+	virtual FUnorderedAccessViewRHIRef RHICreateUnorderedAccessView_RenderThread(class FRHICommandListImmediate& RHICmdList, FRHIIndexBuffer* IndexBuffer, uint8 Format) final override;
 	
-	virtual FShaderResourceViewRHIRef RHICreateShaderResourceView_RenderThread(class FRHICommandListImmediate& RHICmdList, FVertexBufferRHIParamRef VertexBuffer, uint32 Stride, uint8 Format) final override;
-	virtual FShaderResourceViewRHIRef CreateShaderResourceView_RenderThread(class FRHICommandListImmediate& RHICmdList, FIndexBufferRHIParamRef Buffer) final override;
-	virtual FShaderResourceViewRHIRef CreateShaderResourceView_RenderThread(class FRHICommandListImmediate& RHICmdList, FVertexBufferRHIParamRef VertexBuffer, uint32 Stride, uint8 Format) final override;
-	virtual FShaderResourceViewRHIRef RHICreateShaderResourceView_RenderThread(class FRHICommandListImmediate& RHICmdList, FIndexBufferRHIParamRef Buffer) final override;
-	virtual FShaderResourceViewRHIRef RHICreateShaderResourceView_RenderThread(class FRHICommandListImmediate& RHICmdList, FStructuredBufferRHIParamRef StructuredBuffer) final override;
-	virtual FShaderResourceViewRHIRef RHICreateShaderResourceView_RenderThread(class FRHICommandListImmediate& RHICmdList, FTextureRHIParamRef Texture2DRHI, const FRHITextureSRVCreateInfo& CreateInfo) final override;
+	virtual FShaderResourceViewRHIRef RHICreateShaderResourceView_RenderThread(class FRHICommandListImmediate& RHICmdList, FRHIVertexBuffer* VertexBuffer, uint32 Stride, uint8 Format) final override;
+	virtual FShaderResourceViewRHIRef CreateShaderResourceView_RenderThread(class FRHICommandListImmediate& RHICmdList, FRHIIndexBuffer* Buffer) final override;
+	virtual FShaderResourceViewRHIRef CreateShaderResourceView_RenderThread(class FRHICommandListImmediate& RHICmdList, FRHIVertexBuffer* VertexBuffer, uint32 Stride, uint8 Format) final override;
+	virtual FShaderResourceViewRHIRef RHICreateShaderResourceView_RenderThread(class FRHICommandListImmediate& RHICmdList, FRHIIndexBuffer* Buffer) final override;
+	virtual FShaderResourceViewRHIRef RHICreateShaderResourceView_RenderThread(class FRHICommandListImmediate& RHICmdList, FRHIStructuredBuffer* StructuredBuffer) final override;
+	virtual FShaderResourceViewRHIRef RHICreateShaderResourceView_RenderThread(class FRHICommandListImmediate& RHICmdList, FRHITexture* Texture2DRHI, const FRHITextureSRVCreateInfo& CreateInfo) final override;
 
 	virtual FComputeFenceRHIRef RHICreateComputeFence(const FName& Name) final override;
 	virtual FGPUFenceRHIRef RHICreateGPUFence(const FName &Name) final override;
@@ -189,16 +189,16 @@ public:
 	virtual void* LockStagingBuffer_RenderThread(class FRHICommandListImmediate& RHICmdList, FRHIStagingBuffer* StagingBuffer, uint32 Offset, uint32 SizeRHI) final override;
 	virtual void UnlockStagingBuffer_RenderThread(class FRHICommandListImmediate& RHICmdList, FRHIStagingBuffer* StagingBuffer) final override;
 	
-	virtual void RHIAcquireTransientResource_RenderThread(FTextureRHIParamRef Texture) final override;
-	virtual void RHIDiscardTransientResource_RenderThread(FTextureRHIParamRef Texture) final override;
-	virtual void RHIAcquireTransientResource_RenderThread(FVertexBufferRHIParamRef Buffer) final override;
-	virtual void RHIDiscardTransientResource_RenderThread(FVertexBufferRHIParamRef Buffer) final override;
-	virtual void RHIAcquireTransientResource_RenderThread(FStructuredBufferRHIParamRef Buffer) final override;
-	virtual void RHIDiscardTransientResource_RenderThread(FStructuredBufferRHIParamRef Buffer) final override;
+	virtual void RHIAcquireTransientResource_RenderThread(FRHITexture* Texture) final override;
+	virtual void RHIDiscardTransientResource_RenderThread(FRHITexture* Texture) final override;
+	virtual void RHIAcquireTransientResource_RenderThread(FRHIVertexBuffer* Buffer) final override;
+	virtual void RHIDiscardTransientResource_RenderThread(FRHIVertexBuffer* Buffer) final override;
+	virtual void RHIAcquireTransientResource_RenderThread(FRHIStructuredBuffer* Buffer) final override;
+	virtual void RHIDiscardTransientResource_RenderThread(FRHIStructuredBuffer* Buffer) final override;
 	
 	/* Copy the source box pixels in the destination box texture, return true if implemented for the current platform*/
-	virtual void RHICopySubTextureRegion_RenderThread(class FRHICommandListImmediate& RHICmdList, FTexture2DRHIParamRef SourceTexture, FTexture2DRHIParamRef DestinationTexture, FBox2D SourceBox, FBox2D DestinationBox) final override;
-	virtual void RHICopySubTextureRegion(FTexture2DRHIParamRef SourceTexture, FTexture2DRHIParamRef DestinationTexture, FBox2D SourceBox, FBox2D DestinationBox) final override;
+	virtual void RHICopySubTextureRegion_RenderThread(class FRHICommandListImmediate& RHICmdList, FRHITexture2D* SourceTexture, FRHITexture2D* DestinationTexture, FBox2D SourceBox, FBox2D DestinationBox) final override;
+	virtual void RHICopySubTextureRegion(FRHITexture2D* SourceTexture, FRHITexture2D* DestinationTexture, FBox2D SourceBox, FBox2D DestinationBox) final override;
 
 	virtual FVertexShaderRHIRef CreateVertexShader_RenderThread(class FRHICommandListImmediate& RHICmdList, const TArray<uint8>& Code) final override;
 	virtual FVertexShaderRHIRef CreateVertexShader_RenderThread(class FRHICommandListImmediate& RHICmdList, FRHIShaderLibrary* Library, FSHAHash Hash) final override;
@@ -216,7 +216,7 @@ public:
 	virtual FDomainShaderRHIRef CreateDomainShader_RenderThread(class FRHICommandListImmediate& RHICmdList, FRHIShaderLibrary* Library, FSHAHash Hash) final override;
 	virtual FRHIShaderLibraryRef RHICreateShaderLibrary_RenderThread(class FRHICommandListImmediate& RHICmdList, EShaderPlatform Platform, FString FilePath, FString Name) final override;
 
-	virtual void RHIUpdateUniformBuffer(FUniformBufferRHIParamRef UniformBufferRHI, const void* Contents) final override;
+	virtual void RHIUpdateUniformBuffer(FRHIUniformBuffer* UniformBufferRHI, const void* Contents) final override;
 
 	virtual void RHICalibrateTimers() final override;
 
diff --git a/Engine/Source/Runtime/Apple/MetalRHI/Public/MetalRHIContext.h b/Engine/Source/Runtime/Apple/MetalRHI/Public/MetalRHIContext.h
--- a/Engine/Source/Runtime/Apple/MetalRHI/Public/MetalRHIContext.h
+++ b/Engine/Source/Runtime/Apple/MetalRHI/Public/MetalRHIContext.h
@@ -33,7 +33,7 @@ public:
 	
 	virtual void RHIDispatchComputeShader(uint32 ThreadGroupCountX, uint32 ThreadGroupCountY, uint32 ThreadGroupCountZ) final override;
 	
-	virtual void RHIDispatchIndirectComputeShader(FVertexBufferRHIParamRef ArgumentBuffer, uint32 ArgumentOffset) final override;
+	virtual void RHIDispatchIndirectComputeShader(FRHIVertexBuffer* ArgumentBuffer, uint32 ArgumentOffset) final override;
 	
 	virtual void RHIAutomaticCacheFlushAfterComputeShader(bool bEnable) final override;
 	
@@ -47,7 +47,7 @@ public:
 	/** Clears a UAV to the multi-component value provided. */
 	virtual void RHIClearTinyUAV(FRHIUnorderedAccessView* UnorderedAccessViewRHI, const uint32* Values) final override;
 	
-	virtual void RHICopyTexture(FTextureRHIParamRef SourceTextureRHI, FTextureRHIParamRef DestTextureRHI, const FRHICopyTextureInfo& CopyInfo) final override;
+	virtual void RHICopyTexture(FRHITexture* SourceTextureRHI, FRHITexture* DestTextureRHI, const FRHICopyTextureInfo& CopyInfo) final override;
 	
 	/**
 	 * Resolves from one texture to another.
@@ -55,7 +55,7 @@ public:
 	 * @param DestTexture - texture to resolve to, 0 is silenty ignored
 	 * @param ResolveParams - optional resolve params
 	 */
-	virtual void RHICopyToResolveTarget(FTextureRHIParamRef SourceTexture, FTextureRHIParamRef DestTexture, const FResolveParams& ResolveParams) final override;
+	virtual void RHICopyToResolveTarget(FRHITexture* SourceTexture, FRHITexture* DestTexture, const FResolveParams& ResolveParams) final override;
 	
 	virtual void RHIBeginRenderQuery(FRHIRenderQuery* RenderQuery) final override;
 	
@@ -70,7 +70,7 @@ public:
 	virtual void RHIDiscardRenderTargets(bool Depth, bool Stencil, uint32 ColorBitMask) final override;
 	
 	// This method is queued with an RHIThread, otherwise it will flush after it is queued; without an RHI thread there is no benefit to queuing this frame advance commands
-	virtual void RHIBeginDrawingViewport(FRHIViewport* Viewport, FTextureRHIParamRef RenderTargetRHI) override;
+	virtual void RHIBeginDrawingViewport(FRHIViewport* Viewport, FRHITexture* RenderTargetRHI) override;
 	
 	// This method is queued with an RHIThread, otherwise it will flush after it is queued; without an RHI thread there is no benefit to queuing this frame advance commands
 	virtual void RHIEndDrawingViewport(FRHIViewport* Viewport, bool bPresent, bool bLockToVsync) override;
@@ -95,7 +95,7 @@ public:
 	// This method is queued with an RHIThread, otherwise it will flush after it is queued; without an RHI thread there is no benefit to queuing this frame advance commands
 	virtual void RHIEndScene() override;
 	
-	virtual void RHISetStreamSource(uint32 StreamIndex, FVertexBufferRHIParamRef VertexBuffer, uint32 Offset) final override;
+	virtual void RHISetStreamSource(uint32 StreamIndex, FRHIVertexBuffer* VertexBuffer, uint32 Offset) final override;
 
 	// @param MinX including like Win32 RECT
 	// @param MinY including like Win32 RECT
@@ -114,22 +114,22 @@ public:
 	virtual void RHISetGraphicsPipelineState(FRHIGraphicsPipelineState* GraphicsState) final override;
 	
 	/** Set the shader resource view of a surface.  This is used for binding TextureMS parameter types that need a multi sampled view. */
-	virtual void RHISetShaderTexture(FRHIVertexShader* VertexShader, uint32 TextureIndex, FTextureRHIParamRef NewTexture) final override;
+	virtual void RHISetShaderTexture(FRHIVertexShader* VertexShader, uint32 TextureIndex, FRHITexture* NewTexture) final override;
 	
 	/** Set the shader resource view of a surface.  This is used for binding TextureMS parameter types that need a multi sampled view. */
-	virtual void RHISetShaderTexture(FRHIHullShader* HullShader, uint32 TextureIndex, FTextureRHIParamRef NewTexture) final override;
+	virtual void RHISetShaderTexture(FRHIHullShader* HullShader, uint32 TextureIndex, FRHITexture* NewTexture) final override;
 	
 	/** Set the shader resource view of a surface.  This is used for binding TextureMS parameter types that need a multi sampled view. */
-	virtual void RHISetShaderTexture(FRHIDomainShader* DomainShader, uint32 TextureIndex, FTextureRHIParamRef NewTexture) final override;
+	virtual void RHISetShaderTexture(FRHIDomainShader* DomainShader, uint32 TextureIndex, FRHITexture* NewTexture) final override;
 	
 	/** Set the shader resource view of a surface.  This is used for binding TextureMS parameter types that need a multi sampled view. */
-	virtual void RHISetShaderTexture(FRHIGeometryShader* GeometryShader, uint32 TextureIndex, FTextureRHIParamRef NewTexture) final override;
+	virtual void RHISetShaderTexture(FRHIGeometryShader* GeometryShader, uint32 TextureIndex, FRHITexture* NewTexture) final override;
 	
 	/** Set the shader resource view of a surface.  This is used for binding TextureMS parameter types that need a multi sampled view. */
-	virtual void RHISetShaderTexture(FRHIPixelShader* PixelShader, uint32 TextureIndex, FTextureRHIParamRef NewTexture) final override;
+	virtual void RHISetShaderTexture(FRHIPixelShader* PixelShader, uint32 TextureIndex, FRHITexture* NewTexture) final override;
 	
 	/** Set the shader resource view of a surface.  This is used for binding TextureMS parameter types that need a multi sampled view. */
-	virtual void RHISetShaderTexture(FRHIComputeShader* PixelShader, uint32 TextureIndex, FTextureRHIParamRef NewTexture) final override;
+	virtual void RHISetShaderTexture(FRHIComputeShader* PixelShader, uint32 TextureIndex, FRHITexture* NewTexture) final override;
 	
 	/**
 	 * Sets sampler state.
@@ -197,17 +197,17 @@ public:
 	
 	virtual void RHISetShaderResourceViewParameter(FRHIGeometryShader* GeometryShader, uint32 SamplerIndex, FRHIShaderResourceView* SRV) final override;
 	
-	virtual void RHISetShaderUniformBuffer(FRHIVertexShader* VertexShader, uint32 BufferIndex, FUniformBufferRHIParamRef Buffer) final override;
+	virtual void RHISetShaderUniformBuffer(FRHIVertexShader* VertexShader, uint32 BufferIndex, FRHIUniformBuffer* Buffer) final override;
 	
-	virtual void RHISetShaderUniformBuffer(FRHIHullShader* HullShader, uint32 BufferIndex, FUniformBufferRHIParamRef Buffer) final override;
+	virtual void RHISetShaderUniformBuffer(FRHIHullShader* HullShader, uint32 BufferIndex, FRHIUniformBuffer* Buffer) final override;
 	
-	virtual void RHISetShaderUniformBuffer(FRHIDomainShader* DomainShader, uint32 BufferIndex, FUniformBufferRHIParamRef Buffer) final override;
+	virtual void RHISetShaderUniformBuffer(FRHIDomainShader* DomainShader, uint32 BufferIndex, FRHIUniformBuffer* Buffer) final override;
 	
-	virtual void RHISetShaderUniformBuffer(FRHIGeometryShader* GeometryShader, uint32 BufferIndex, FUniformBufferRHIParamRef Buffer) final override;
+	virtual void RHISetShaderUniformBuffer(FRHIGeometryShader* GeometryShader, uint32 BufferIndex, FRHIUniformBuffer* Buffer) final override;
 	
-	virtual void RHISetShaderUniformBuffer(FRHIPixelShader* PixelShader, uint32 BufferIndex, FUniformBufferRHIParamRef Buffer) final override;
+	virtual void RHISetShaderUniformBuffer(FRHIPixelShader* PixelShader, uint32 BufferIndex, FRHIUniformBuffer* Buffer) final override;
 	
-	virtual void RHISetShaderUniformBuffer(FRHIComputeShader* ComputeShader, uint32 BufferIndex, FUniformBufferRHIParamRef Buffer) final override;
+	virtual void RHISetShaderUniformBuffer(FRHIComputeShader* ComputeShader, uint32 BufferIndex, FRHIUniformBuffer* Buffer) final override;
 	
 	virtual void RHISetShaderParameter(FRHIVertexShader* VertexShader, uint32 BufferIndex, uint32 BaseIndex, uint32 NumBytes, const void* NewValue) final override;
 	
@@ -231,14 +231,14 @@ public:
 	
 	virtual void RHIDrawPrimitive(uint32 BaseVertexIndex, uint32 NumPrimitives, uint32 NumInstances) final override;
 	
-	virtual void RHIDrawPrimitiveIndirect(FVertexBufferRHIParamRef ArgumentBuffer, uint32 ArgumentOffset) final override;
+	virtual void RHIDrawPrimitiveIndirect(FRHIVertexBuffer* ArgumentBuffer, uint32 ArgumentOffset) final override;
 	
-	virtual void RHIDrawIndexedIndirect(FIndexBufferRHIParamRef IndexBufferRHI, FStructuredBufferRHIParamRef ArgumentsBufferRHI, int32 DrawArgumentsIndex, uint32 NumInstances) final override;
+	virtual void RHIDrawIndexedIndirect(FRHIIndexBuffer* IndexBufferRHI, FRHIStructuredBuffer* ArgumentsBufferRHI, int32 DrawArgumentsIndex, uint32 NumInstances) final override;
 	
 	// @param NumPrimitives need to be >0
-	virtual void RHIDrawIndexedPrimitive(FIndexBufferRHIParamRef IndexBuffer, int32 BaseVertexIndex, uint32 FirstInstance, uint32 NumVertices, uint32 StartIndex, uint32 NumPrimitives, uint32 NumInstances) final override;
+	virtual void RHIDrawIndexedPrimitive(FRHIIndexBuffer* IndexBuffer, int32 BaseVertexIndex, uint32 FirstInstance, uint32 NumVertices, uint32 StartIndex, uint32 NumPrimitives, uint32 NumInstances) final override;
 	
-	virtual void RHIDrawIndexedPrimitiveIndirect(FIndexBufferRHIParamRef IndexBuffer, FVertexBufferRHIParamRef ArgumentBuffer, uint32 ArgumentOffset) final override;
+	virtual void RHIDrawIndexedPrimitiveIndirect(FRHIIndexBuffer* IndexBuffer, FRHIVertexBuffer* ArgumentBuffer, uint32 ArgumentOffset) final override;
 
 	/**
 	* Sets Depth Bounds Testing with the given min/max depth.
@@ -252,7 +252,7 @@ public:
 	
 	virtual void RHIPopEvent() final override;
 	
-	virtual void RHIUpdateTextureReference(FRHITextureReference* TextureRef, FTextureRHIParamRef NewTexture) final override;
+	virtual void RHIUpdateTextureReference(FRHITextureReference* TextureRef, FRHITexture* NewTexture) final override;
 	
 	/**
 	 * Explicitly transition a UAV from readable -> writable by the GPU or vice versa.
@@ -267,7 +267,7 @@ public:
 	 */
 	virtual void RHITransitionResources(EResourceTransitionAccess TransitionType, EResourceTransitionPipeline TransitionPipeline, FRHIUnorderedAccessView** InUAVs, int32 NumUAVs, FRHIComputeFence* WriteComputeFence) final override;
 	
-	virtual void RHICopyToStagingBuffer(FVertexBufferRHIParamRef SourceBufferRHI, FRHIStagingBuffer* DestinationStagingBufferRHI, uint32 Offset, uint32 NumBytes) final override;
+	virtual void RHICopyToStagingBuffer(FRHIVertexBuffer* SourceBufferRHI, FRHIStagingBuffer* DestinationStagingBufferRHI, uint32 Offset, uint32 NumBytes) final override;
 	virtual void RHIWriteGPUFence(FRHIGPUFence* FenceRHI) final override;
 
 	/**
@@ -279,7 +279,7 @@ public:
 	 * @param InTextures - array of texture objects to transition
 	 * @param NumTextures - number of textures to transition
 	 */
-	virtual void RHITransitionResources(EResourceTransitionAccess TransitionType, FTextureRHIParamRef* InTextures, int32 NumTextures) final override;
+	virtual void RHITransitionResources(EResourceTransitionAccess TransitionType, FRHITexture** InTextures, int32 NumTextures) final override;
 	
 	/**
 	 * Compute queue will wait for the fence to be written before continuing.
@@ -338,7 +338,7 @@ public:
 	FMetalRHIImmediateCommandContext(class FMetalProfiler* InProfiler, FMetalContext* WrapContext);
 
 	// FRHICommandContext API accessible only on the immediate device context
-	virtual void RHIBeginDrawingViewport(FRHIViewport* Viewport, FTextureRHIParamRef RenderTargetRHI) final override;
+	virtual void RHIBeginDrawingViewport(FRHIViewport* Viewport, FRHITexture* RenderTargetRHI) final override;
 	virtual void RHIEndDrawingViewport(FRHIViewport* Viewport, bool bPresent, bool bLockToVsync) final override;
 	virtual void RHIBeginFrame() final override;
 	virtual void RHIEndFrame() final override;
