Metal compile fix on clang.

#rnx

diff --git a/Engine/Source/Runtime/Apple/MetalRHI/Private/MetalCommands.cpp b/Engine/Source/Runtime/Apple/MetalRHI/Private/MetalCommands.cpp
--- a/Engine/Source/Runtime/Apple/MetalRHI/Private/MetalCommands.cpp
+++ b/Engine/Source/Runtime/Apple/MetalRHI/Private/MetalCommands.cpp
@@ -102,7 +102,7 @@ void FMetalRHICommandContext::RHISetStreamSource(uint32 StreamIndex, FRHIVertexB
 	}
 }
 
-void FMetalDynamicRHI::RHISetStreamOutTargets(uint32 NumTargets, const FRHIVertexBuffer** VertexBuffers, const uint32* Offsets)
+void FMetalDynamicRHI::RHISetStreamOutTargets(uint32 NumTargets, FRHIVertexBuffer* const * VertexBuffers, const uint32* Offsets)
 {
 	NOT_SUPPORTED("RHISetStreamOutTargets");
 
diff --git a/Engine/Source/Runtime/Apple/MetalRHI/Public/MetalRHI.h b/Engine/Source/Runtime/Apple/MetalRHI/Public/MetalRHI.h
--- a/Engine/Source/Runtime/Apple/MetalRHI/Public/MetalRHI.h
+++ b/Engine/Source/Runtime/Apple/MetalRHI/Public/MetalRHI.h
@@ -132,7 +132,7 @@ public:
 	virtual void RHIResizeViewport(FRHIViewport* Viewport, uint32 SizeX, uint32 SizeY, bool bIsFullscreen) final override;
 	virtual void RHIResizeViewport(FRHIViewport* Viewport, uint32 SizeX, uint32 SizeY, bool bIsFullscreen, EPixelFormat PreferredPixelFormat) final override;
 	virtual void RHITick(float DeltaTime) final override;
-	virtual void RHISetStreamOutTargets(uint32 NumTargets,const FRHIVertexBuffer** VertexBuffers,const uint32* Offsets) final override;
+	virtual void RHISetStreamOutTargets(uint32 NumTargets,FRHIVertexBuffer* const * VertexBuffers,const uint32* Offsets) final override;
 	virtual void RHIBlockUntilGPUIdle() final override;
 	virtual bool RHIGetAvailableResolutions(FScreenResolutionArray& Resolutions, bool bIgnoreRefreshRate) final override;
 	virtual void RHIGetSupportedResolution(uint32& Width, uint32& Height) final override;
