Avoid storing/resolving depth on iOS when not needed

diff --git a/Engine/Source/Runtime/Apple/MetalRHI/Private/MetalStateCache.cpp b/Engine/Source/Runtime/Apple/MetalRHI/Private/MetalStateCache.cpp
--- a/Engine/Source/Runtime/Apple/MetalRHI/Private/MetalStateCache.cpp
+++ b/Engine/Source/Runtime/Apple/MetalRHI/Private/MetalStateCache.cpp
@@ -798,7 +798,7 @@ bool FMetalStateCache::SetRenderPassInfo(FRHIRenderPassInfo const& InRenderTarge
 				{
 					DepthAttachment.SetLoadAction(mtlpp::LoadAction::Clear);
 					
-					if (bSupportsMSAADepthResolve && Surface.MSAATexture)
+					if (bSupportsMSAADepthResolve && Surface.MSAATexture && DepthStoreAction == ERenderTargetStoreAction::EMultisampleResolve)
 					{
 						HighLevelStoreAction = ERenderTargetStoreAction::EMultisampleResolve;
 					}
@@ -807,6 +807,10 @@ bool FMetalStateCache::SetRenderPassInfo(FRHIRenderPassInfo const& InRenderTarge
 						HighLevelStoreAction = ERenderTargetStoreAction::ENoAction;
 					}
 				}
+                else
+                {
+                	HighLevelStoreAction = DepthStoreAction;
+                }
 #endif
                 //needed to quiet the metal validation that runs when you end renderpass. (it requires some kind of 'resolve' for an msaa target)
 				//But with deferredstore we don't set the real one until submit time.
@@ -815,7 +819,7 @@ bool FMetalStateCache::SetRenderPassInfo(FRHIRenderPassInfo const& InRenderTarge
 				DepthAttachment.SetClearDepth(DepthClearValue);
 				check(SampleCount > 0);
 
-				if (Surface.MSAATexture && bSupportsMSAADepthResolve)
+				if (Surface.MSAATexture && bSupportsMSAADepthResolve && DepthAttachment.GetStoreAction() != mtlpp::StoreAction::DontCare)
 				{
                     if (!bDepthStencilSampleCountMismatchFixup)
                     {
@@ -865,14 +869,18 @@ bool FMetalStateCache::SetRenderPassInfo(FRHIRenderPassInfo const& InRenderTarge
 					HighLevelStoreAction = ERenderTargetStoreAction::EStore;
 				}
 				
-					bool bStencilMemoryless = false;
+				bool bStencilMemoryless = false;
 #if PLATFORM_IOS
 				if (StencilTexture.GetStorageMode() == mtlpp::StorageMode::Memoryless)
 				{
-						bStencilMemoryless = true;
+					bStencilMemoryless = true;
 					HighLevelStoreAction = ERenderTargetStoreAction::ENoAction;
 					StencilAttachment.SetLoadAction(mtlpp::LoadAction::Clear);
 				}
+				else
+				{
+					HighLevelStoreAction = StencilStoreAction;
+				}
 #endif
 				
 				// For the case where Depth+Stencil is MSAA we can't Resolve depth and Store stencil - we can only Resolve + DontCare or StoreResolve + Store (on newer H/W and iOS).
diff --git a/Engine/Source/Runtime/Renderer/Private/MobileShadingRenderer.cpp b/Engine/Source/Runtime/Renderer/Private/MobileShadingRenderer.cpp
--- a/Engine/Source/Runtime/Renderer/Private/MobileShadingRenderer.cpp
+++ b/Engine/Source/Runtime/Renderer/Private/MobileShadingRenderer.cpp
@@ -455,7 +455,7 @@ void FMobileSceneRenderer::Render(FRHICommandListImmediate& RHICmdList)
 			DepthTargetAction = EDepthStencilTargetActions::ClearDepthStencil_StoreDepthStencil;
 		}
 						
-		if (bKeepDepthContent)
+		if (bKeepDepthContent && !bMobileMSAA)
 		{
 			// store depth if post-processing/capture needs it
 			DepthTargetAction = EDepthStencilTargetActions::ClearDepthStencil_StoreDepthStencil;
@@ -536,7 +536,7 @@ void FMobileSceneRenderer::Render(FRHICommandListImmediate& RHICmdList)
 			ExclusiveDepthStencil = FExclusiveDepthStencil::DepthRead_StencilWrite;
 		}
 		
-		if (bKeepDepthContent)
+		if (bKeepDepthContent && !bMobileMSAA)
 		{
 			DepthTargetAction = EDepthStencilTargetActions::LoadDepthStencil_StoreDepthStencil;
 		}
