Fixed render thread time counting on iOS

diff --git a/Engine/Source/Runtime/Apple/MetalRHI/Private/MetalContext.cpp b/Engine/Source/Runtime/Apple/MetalRHI/Private/MetalContext.cpp
--- a/Engine/Source/Runtime/Apple/MetalRHI/Private/MetalContext.cpp
+++ b/Engine/Source/Runtime/Apple/MetalRHI/Private/MetalContext.cpp
@@ -655,7 +655,8 @@ void FMetalDeviceContext::EndDrawingViewport(FMetalViewport* Viewport, bool bPre
 	// We may be limiting our framerate to the display link
 	if( FrameReadyEvent != nullptr && !GMetalSeparatePresentThread )
 	{
-		FrameReadyEvent->Wait();
+		bool bIgnoreThreadIdleStats = true; // Idle time is already counted by the caller
+		FrameReadyEvent->Wait(MAX_uint32, bIgnoreThreadIdleStats);
 	}
 	
 	Viewport->ReleaseDrawable();
