One more fix for render thread time counting on iOS

diff --git a/Engine/Source/Runtime/Apple/MetalRHI/Private/MetalViewport.cpp b/Engine/Source/Runtime/Apple/MetalRHI/Private/MetalViewport.cpp
--- a/Engine/Source/Runtime/Apple/MetalRHI/Private/MetalViewport.cpp
+++ b/Engine/Source/Runtime/Apple/MetalRHI/Private/MetalViewport.cpp
@@ -322,7 +322,7 @@ mtlpp::Drawable FMetalViewport::GetDrawable(EMetalViewportAccessFlag Accessor)
 				}
 				else
 				{
-					FPlatformProcess::Sleep(0.001f);
+					FPlatformProcess::SleepNoStats(0.001f);
 				}
 			}
 			while (Drawable == nil || Size.width != BackBuffer[GetViewportIndex(Accessor)]->GetSizeX() || Size.height != BackBuffer[GetViewportIndex(Accessor)]->GetSizeY());
