#ocn - Fix to reduce memory usage with new audio mixer.

diff --git a/Engine/Source/Runtime/Apple/AudioMixerAudioUnit/Private/AudioMixerPlatformAudioUnit.h b/Engine/Source/Runtime/Apple/AudioMixerAudioUnit/Private/AudioMixerPlatformAudioUnit.h
--- a/Engine/Source/Runtime/Apple/AudioMixerAudioUnit/Private/AudioMixerPlatformAudioUnit.h
+++ b/Engine/Source/Runtime/Apple/AudioMixerAudioUnit/Private/AudioMixerPlatformAudioUnit.h
@@ -43,7 +43,14 @@ namespace Audio
         virtual int32 GetNumFrames(const int32 InNumReqestedFrames) override;
         virtual void ResumeContext() override;
         virtual void SuspendContext() override;
-        //~ End IAudioMixerPlatformInterface
+		
+		/** Whether or not the platform supports realtime decompression. */
+		virtual bool SupportsRealtimeDecompression() const override { return true; }
+		
+		/** Whether or not the platform disables caching of decompressed PCM data (i.e. to save memory on fixed memory platforms) */
+		virtual bool DisablePCMAudioCaching() const override { return true; }
+		
+      //~ End IAudioMixerPlatformInterface
         
 	private:
 		AudioStreamBasicDescription OutputFormat;
