enable purgeable state for IOS 12 and up

diff --git a/Engine/Source/Runtime/Apple/MetalRHI/Private/MetalBuffer.cpp b/Engine/Source/Runtime/Apple/MetalRHI/Private/MetalBuffer.cpp
--- a/Engine/Source/Runtime/Apple/MetalRHI/Private/MetalBuffer.cpp
+++ b/Engine/Source/Runtime/Apple/MetalRHI/Private/MetalBuffer.cpp
@@ -1265,7 +1265,10 @@ FMetalTexture FMetalTexturePool::CreateTexture(mtlpp::Device Device, mtlpp::Text
 		Pool.Remove(Descriptor);
 		if (GMetalResourcePurgeInPool)
 		{
-        	Texture.SetPurgeableState(mtlpp::PurgeableState::NonVolatile);
+			if (@available(iOS 12.0, macOS 10.13, *))
+			{
+        		Texture.SetPurgeableState(mtlpp::PurgeableState::NonVolatile);
+			}
         }
 	}
 	else
@@ -1296,7 +1299,10 @@ void FMetalTexturePool::ReleaseTexture(FMetalTexture& Texture)
 	
 	if (GMetalResourcePurgeInPool && Texture.SetPurgeableState(mtlpp::PurgeableState::KeepCurrent) == mtlpp::PurgeableState::NonVolatile)
 	{
-		Texture.SetPurgeableState(mtlpp::PurgeableState::Volatile);
+		if (@available(iOS 12.0, macOS 10.13, *))
+		{
+			Texture.SetPurgeableState(mtlpp::PurgeableState::Volatile);
+		}
 	}
 	
 	FScopeLock Lock(&PoolMutex);
@@ -1318,9 +1324,15 @@ void FMetalTexturePool::Drain(bool const bForce)
 			{
 				It.RemoveCurrent();
 			}
-            else if (GMetalResourcePurgeInPool && (GFrameNumberRenderThread - It->Key.freedFrame) >= PurgeAfterNumFrames)
+            else
             {
-                It->Value.SetPurgeableState(mtlpp::PurgeableState::Empty);
+				if (GMetalResourcePurgeInPool && (GFrameNumberRenderThread - It->Key.freedFrame) >= PurgeAfterNumFrames)
+				{
+					if (@available(iOS 12.0, macOS 10.13, *))
+					{
+						It->Value.SetPurgeableState(mtlpp::PurgeableState::Empty);
+					}
+				}
             }
         }
     }
@@ -1514,7 +1526,7 @@ FMetalBuffer FMetalResourceHeap::CreateBuffer(uint32 Size, uint32 Alignment, uin
 				 	return Found->NewBuffer(BlockSize);
 				 }
 				 else
-				{
+				 {
                     Buffer = ManagedBuffers.CreatePooledResource(FMetalPooledBufferArgs(Queue->GetDevice(), BlockSize, Flags, StorageMode));
 					if (GMetalResourcePurgeInPool)
 					{
@@ -1595,7 +1607,10 @@ FMetalBuffer FMetalResourceHeap::CreateBuffer(uint32 Size, uint32 Alignment, uin
                     Buffer = Buffers[Storage].CreatePooledResource(FMetalPooledBufferArgs(Queue->GetDevice(), BlockSize, Flags, StorageMode));
 					if (GMetalResourcePurgeInPool)
 					{
-                    	Buffer.SetPurgeableState(mtlpp::PurgeableState::NonVolatile);
+						if (@available(iOS 12.0, macOS 10.13, *))
+						{
+                    		Buffer.SetPurgeableState(mtlpp::PurgeableState::NonVolatile);
+						}
 					}
 					DEC_MEMORY_STAT_BY(STAT_MetalBufferUnusedMemory, Buffer.GetLength());
 					DEC_MEMORY_STAT_BY(STAT_MetalPooledBufferUnusedMemory, Buffer.GetLength());
@@ -1643,7 +1658,10 @@ void FMetalResourceHeap::ReleaseBuffer(FMetalBuffer& Buffer)
 		
 		if (GMetalResourcePurgeInPool)
 		{
-        	Buffer.SetPurgeableState(mtlpp::PurgeableState::Volatile);
+			if (@available(iOS 12.0, macOS 10.13, *))
+			{
+        		Buffer.SetPurgeableState(mtlpp::PurgeableState::Volatile);
+			}
 		}
         
 		switch (StorageMode)
