Fixed an iOS RHI bug where cube maps were temporarily allocating six times as much memory as they need while copying loaded texture faces into the Metal cubemap. This was causing memory spikes on load.

diff --git a/Engine/Source/Runtime/Apple/MetalRHI/Private/MetalTexture.cpp b/Engine/Source/Runtime/Apple/MetalRHI/Private/MetalTexture.cpp
--- a/Engine/Source/Runtime/Apple/MetalRHI/Private/MetalTexture.cpp
+++ b/Engine/Source/Runtime/Apple/MetalRHI/Private/MetalTexture.cpp
@@ -1142,10 +1142,10 @@ FMetalSurface::~FMetalSurface()
 	}
 }
 
-FMetalBuffer FMetalSurface::AllocSurface(uint32 MipIndex, uint32 ArrayIndex, EResourceLockMode LockMode, uint32& DestStride)
+FMetalBuffer FMetalSurface::AllocSurface(uint32 MipIndex, uint32 ArrayIndex, EResourceLockMode LockMode, uint32& DestStride, bool SingleLayer /*= false*/)
 {
 	// get size and stride
-	uint32 MipBytes = GetMipSize(MipIndex, &DestStride, false);
+	uint32 MipBytes = GetMipSize(MipIndex, &DestStride, SingleLayer);
 	
 	// allocate some temporary memory
 	mtlpp::ResourceOptions ResMode = FMetalCommandQueue::GetCompatibleResourceOptions((mtlpp::ResourceOptions)(mtlpp::ResourceOptions::StorageModeShared | (!(PLATFORM_MAC && PixelFormat == PF_G8 && (Flags & TexCreate_SRGB)) ? mtlpp::ResourceOptions::CpuCacheModeWriteCombined : 0)));
@@ -1289,7 +1289,7 @@ void FMetalSurface::UpdateSurface(FMetalBuffer& Buffer, uint32 MipIndex, uint32
 #endif
 }
 
-void* FMetalSurface::Lock(uint32 MipIndex, uint32 ArrayIndex, EResourceLockMode LockMode, uint32& DestStride)
+void* FMetalSurface::Lock(uint32 MipIndex, uint32 ArrayIndex, EResourceLockMode LockMode, uint32& DestStride, bool SingleLayer /*= false*/)
 {
 	// get size and stride
 	uint32 MipBytes = GetMipSize(MipIndex, &DestStride, false);
@@ -1297,7 +1297,7 @@ void* FMetalSurface::Lock(uint32 MipIndex, uint32 ArrayIndex, EResourceLockMode
 	// allocate some temporary memory
 	if(!LockedMemory[MipIndex])
 	{
-		LockedMemory[MipIndex] = AllocSurface(MipIndex, ArrayIndex, LockMode, DestStride);
+		LockedMemory[MipIndex] = AllocSurface(MipIndex, ArrayIndex, LockMode, DestStride, SingleLayer);
 	}
 	
 	switch(LockMode)
@@ -2239,7 +2239,7 @@ void* FMetalDynamicRHI::RHILockTextureCubeFace(FRHITextureCube* TextureCubeRHI,u
 	@autoreleasepool {
 		FMetalTextureCube* TextureCube = ResourceCast(TextureCubeRHI);
 		uint32 MetalFace = GetMetalCubeFace((ECubeFace)FaceIndex);
-		return TextureCube->Surface.Lock(MipIndex, MetalFace + (6 * ArrayIndex), LockMode, DestStride);
+		return TextureCube->Surface.Lock(MipIndex, MetalFace + (6 * ArrayIndex), LockMode, DestStride, true);
 	}
 }
 
diff --git a/Engine/Source/Runtime/Apple/MetalRHI/Public/MetalResources.h b/Engine/Source/Runtime/Apple/MetalRHI/Public/MetalResources.h
--- a/Engine/Source/Runtime/Apple/MetalRHI/Public/MetalResources.h
+++ b/Engine/Source/Runtime/Apple/MetalRHI/Public/MetalResources.h
@@ -450,7 +450,7 @@ public:
 	void PrepareTextureView();
 	
 	/** @returns A newly allocated buffer object large enough for the surface within the texture specified. */
-	FMetalBuffer AllocSurface(uint32 MipIndex, uint32 ArrayIndex, EResourceLockMode LockMode, uint32& DestStride);
+	FMetalBuffer AllocSurface(uint32 MipIndex, uint32 ArrayIndex, EResourceLockMode LockMode, uint32& DestStride, bool SingleLayer = false);
 
 	/** Apply the data in Buffer to the surface specified. */
 	void UpdateSurface(FMetalBuffer& Buffer, uint32 MipIndex, uint32 ArrayIndex);
@@ -460,7 +460,7 @@ public:
 	 * @param ArrayIndex Index of the texture array/face in the form Index*6+Face
 	 * @return A pointer to the specified texture data.
 	 */
-	void* Lock(uint32 MipIndex, uint32 ArrayIndex, EResourceLockMode LockMode, uint32& DestStride);
+	void* Lock(uint32 MipIndex, uint32 ArrayIndex, EResourceLockMode LockMode, uint32& DestStride, bool SingleLayer = false);
 
 	/** Unlocks a previously locked mip-map.
 	 * @param ArrayIndex Index of the texture array/face in the form Index*6+Face
