- Fixed MetalTexture.SetPurgeableState for ios11, macos12 in a couple cases

diff --git a/Engine/Source/Runtime/Apple/MetalRHI/Private/MetalBuffer.cpp b/Engine/Source/Runtime/Apple/MetalRHI/Private/MetalBuffer.cpp
--- a/Engine/Source/Runtime/Apple/MetalRHI/Private/MetalBuffer.cpp
+++ b/Engine/Source/Runtime/Apple/MetalRHI/Private/MetalBuffer.cpp
@@ -1297,9 +1297,9 @@ void FMetalTexturePool::ReleaseTexture(FMetalTexture& Texture)
 	Descriptor.usage = Texture.GetUsage();
 	Descriptor.freedFrame = GFrameNumberRenderThread;
 	
-	if (GMetalResourcePurgeInPool && Texture.SetPurgeableState(mtlpp::PurgeableState::KeepCurrent) == mtlpp::PurgeableState::NonVolatile)
+	if (@available(iOS 12.0, macOS 10.13, *))
 	{
-		if (@available(iOS 12.0, macOS 10.13, *))
+		if (GMetalResourcePurgeInPool && Texture.SetPurgeableState(mtlpp::PurgeableState::KeepCurrent) == mtlpp::PurgeableState::NonVolatile)
 		{
 			Texture.SetPurgeableState(mtlpp::PurgeableState::Volatile);
 		}
@@ -1528,9 +1528,12 @@ FMetalBuffer FMetalResourceHeap::CreateBuffer(uint32 Size, uint32 Alignment, uin
 				 else
 				 {
                     Buffer = ManagedBuffers.CreatePooledResource(FMetalPooledBufferArgs(Queue->GetDevice(), BlockSize, Flags, StorageMode));
-					if (GMetalResourcePurgeInPool)
+					if (@available(iOS 12.0, macOS 10.13, *))
 					{
-                    	Buffer.SetPurgeableState(mtlpp::PurgeableState::NonVolatile);
+						if (GMetalResourcePurgeInPool)
+						{
+							Buffer.SetPurgeableState(mtlpp::PurgeableState::NonVolatile);
+						}
 					}
 					DEC_MEMORY_STAT_BY(STAT_MetalBufferUnusedMemory, Buffer.GetLength());
 					DEC_MEMORY_STAT_BY(STAT_MetalPooledBufferUnusedMemory, Buffer.GetLength());
