Assert that Metal heaps and sub-allocated buffers are actually created properly to catch OoM errors. #rb none #jira UE-62235

diff --git a/Engine/Source/Runtime/Apple/MetalRHI/Private/MetalBuffer.cpp b/Engine/Source/Runtime/Apple/MetalRHI/Private/MetalBuffer.cpp
--- a/Engine/Source/Runtime/Apple/MetalRHI/Private/MetalBuffer.cpp
+++ b/Engine/Source/Runtime/Apple/MetalRHI/Private/MetalBuffer.cpp
@@ -155,6 +155,7 @@ FMetalSubBufferHeap::FMetalSubBufferHeap(NSUInteger Size, NSUInteger Alignment,
 		Desc.SetSize(FullSize);
 		Desc.SetStorageMode(Storage);
 		ParentHeap = GetMetalDeviceContext().GetDevice().NewHeap(Desc);
+		check(ParentHeap.GetPtr());
 #if STATS || ENABLE_LOW_LEVEL_MEM_TRACKER
 		MetalLLM::LogAllocHeap(GetMetalDeviceContext().GetDevice(), ParentHeap);
 #endif
@@ -162,7 +163,8 @@ FMetalSubBufferHeap::FMetalSubBufferHeap(NSUInteger Size, NSUInteger Alignment,
 	else
 	{
 		ParentBuffer = MTLPP_VALIDATE(mtlpp::Device, GetMetalDeviceContext().GetDevice(), SafeGetRuntimeDebuggingLevel() >= EMetalDebugLevelValidation, NewBuffer(FullSize, Options));
-		check(ParentBuffer.GetPtr() && ParentBuffer.GetLength() >= FullSize);
+		check(ParentBuffer.GetPtr());
+		check(ParentBuffer.GetLength() >= FullSize);
 #if STATS || ENABLE_LOW_LEVEL_MEM_TRACKER
 		MetalLLM::LogAllocBuffer(GetMetalDeviceContext().GetDevice(), ParentBuffer);
 #endif
@@ -451,7 +453,8 @@ FMetalSubBufferLinear::FMetalSubBufferLinear(NSUInteger Size, NSUInteger Alignme
 	
 	mtlpp::StorageMode Storage = (mtlpp::StorageMode)((Options & mtlpp::ResourceStorageModeMask) >> mtlpp::ResourceStorageModeShift);
 	ParentBuffer = MTLPP_VALIDATE(mtlpp::Device, GetMetalDeviceContext().GetDevice(), SafeGetRuntimeDebuggingLevel() >= EMetalDebugLevelValidation, NewBuffer(FullSize, Options));
-	check(ParentBuffer.GetPtr() && ParentBuffer.GetLength() >= FullSize);
+	check(ParentBuffer.GetPtr());
+	check(ParentBuffer.GetLength() >= FullSize);
 #if STATS || ENABLE_LOW_LEVEL_MEM_TRACKER
 	MetalLLM::LogAllocBuffer(GetMetalDeviceContext().GetDevice(), ParentBuffer);
 #endif
@@ -586,6 +589,7 @@ FMetalSubBufferMagazine::FMetalSubBufferMagazine(NSUInteger Size, NSUInteger Chu
 		Desc.SetSize(FullSize);
 		Desc.SetStorageMode(Storage);
 		ParentHeap = GetMetalDeviceContext().GetDevice().NewHeap(Desc);
+		check(ParentHeap.GetPtr());
 #if STATS || ENABLE_LOW_LEVEL_MEM_TRACKER
 		MetalLLM::LogAllocHeap(GetMetalDeviceContext().GetDevice(), ParentHeap);
 #endif
@@ -593,7 +597,8 @@ FMetalSubBufferMagazine::FMetalSubBufferMagazine(NSUInteger Size, NSUInteger Chu
 	else
 	{
 		ParentBuffer = MTLPP_VALIDATE(mtlpp::Device, GetMetalDeviceContext().GetDevice(), SafeGetRuntimeDebuggingLevel() >= EMetalDebugLevelValidation, NewBuffer(FullSize, Options));
-		check(ParentBuffer.GetPtr() && ParentBuffer.GetLength() >= FullSize);
+		check(ParentBuffer.GetPtr());
+		check(ParentBuffer.GetLength() >= FullSize);
 #if STATS || ENABLE_LOW_LEVEL_MEM_TRACKER
 		MetalLLM::LogAllocBuffer(GetMetalDeviceContext().GetDevice(), ParentBuffer);
 #endif
