UE-64066 TM-Gameplay Load Game does not work on TVOS

#tvOS
#4.21

Blocker: GRHIHDRDisplayOutputFormat must have a default value for tvOS (PF_B8G8R8A8, just like mac or ios for non-hdr). The issue appeared at CL4092793 [at] //UE4/Dev-Rendering, 5/24/2018

diff --git a/Engine/Source/Runtime/Apple/MetalRHI/Private/MetalRHI.cpp b/Engine/Source/Runtime/Apple/MetalRHI/Private/MetalRHI.cpp
--- a/Engine/Source/Runtime/Apple/MetalRHI/Private/MetalRHI.cpp
+++ b/Engine/Source/Runtime/Apple/MetalRHI/Private/MetalRHI.cpp
@@ -500,7 +500,7 @@ FMetalDynamicRHI::FMetalDynamicRHI(ERHIFeatureLevel::Type RequestedFeatureLevel)
 	GRHISupportsBaseVertexIndex = false;
 	GRHISupportsFirstInstance = false; // Supported on macOS & iOS but not tvOS.
 	GRHISupportsHDROutput = false;
-	GRHIHDRDisplayOutputFormat = PF_Unknown;
+	GRHIHDRDisplayOutputFormat = PF_B8G8R8A8; // must have a default value for non-hdr, just like mac or ios
 #else
 	// Only A9+ can support this, so for now we need to limit this to the desktop-forward renderer only.
 	GRHISupportsBaseVertexIndex = [Device supportsFeatureSet:MTLFeatureSet_iOS_GPUFamily3_v1] && (GMaxRHIFeatureLevel >= ERHIFeatureLevel::SM5);
