Turning GPU Command Buffer timing back on. Enables correct profiling.

diff --git a/Engine/Source/Runtime/Apple/MetalRHI/Private/MetalCommandQueue.cpp b/Engine/Source/Runtime/Apple/MetalRHI/Private/MetalCommandQueue.cpp
--- a/Engine/Source/Runtime/Apple/MetalRHI/Private/MetalCommandQueue.cpp
+++ b/Engine/Source/Runtime/Apple/MetalRHI/Private/MetalCommandQueue.cpp
@@ -82,12 +82,8 @@ FMetalCommandQueue::FMetalCommandQueue(mtlpp::Device InDevice, uint32 const MaxN
 		
 		if(Vers.majorVersion > 10 || (Vers.majorVersion == 10 && Vers.minorVersion >= 3))
         {
-            // This isn't currently working, with GPUStartTime and GPUStopTime usually coming back as zero.
-            // The docs say this means the GPU isn't finished with the command buffer yet, but we are running
-            // in the completed handler and the status is MTLCommandBufferStatusCompleted, so it has to be
-            // finished. My guess is that the command buffer is empty so the GPU isn't executing it.
-            //Features |= EMetalFeaturesGPUCommandBufferTimes;
-            
+            Features |= EMetalFeaturesGPUCommandBufferTimes;
+
 			Features |= EMetalFeaturesLinearTextures;
 			// InjectCurves() does not work with this
 			//Features |= EMetalFeaturesEfficientBufferBlits;
