Returns the default context from GetCurrentContext if none is set. Fixes crash in UE-64369. This is probably not 100% correct - we should fix all the cases where GetCurrentContext is called and a context is not set.

diff --git a/Engine/Source/Runtime/Apple/MetalRHI/Private/MetalContext.cpp b/Engine/Source/Runtime/Apple/MetalRHI/Private/MetalContext.cpp
--- a/Engine/Source/Runtime/Apple/MetalRHI/Private/MetalContext.cpp
+++ b/Engine/Source/Runtime/Apple/MetalRHI/Private/MetalContext.cpp
@@ -941,6 +941,17 @@ void FMetalContext::InsertCommandBufferFence(FMetalCommandBufferFence& Fence, mt
 FMetalContext* FMetalContext::GetCurrentContext()
 {
 	FMetalContext* Current = (FMetalContext*)FPlatformTLS::GetTlsValue(CurrentContextTLSSlot);
+	
+	if (!Current)
+	{
+		// If we are executing this outside of a pass we'll return the default.
+		// TODO This needs further investigation. We should fix all the cases that call this without
+		// a context set.
+		FMetalRHICommandContext* CmdContext = static_cast<FMetalRHICommandContext*>(RHIGetDefaultContext());
+		check(CmdContext);
+		Current = &CmdContext->GetInternalContext();
+	}
+	
 	check(Current);
 	return Current;
 }
