Sort out the way we report R11G11B10 support for Metal so that it should work correctly on iOS & tvOS.

diff --git a/Engine/Source/Runtime/Apple/MetalRHI/Private/MetalRHI.cpp b/Engine/Source/Runtime/Apple/MetalRHI/Private/MetalRHI.cpp
--- a/Engine/Source/Runtime/Apple/MetalRHI/Private/MetalRHI.cpp
+++ b/Engine/Source/Runtime/Apple/MetalRHI/Private/MetalRHI.cpp
@@ -638,8 +638,11 @@ FMetalDynamicRHI::FMetalDynamicRHI(ERHIFeatureLevel::Type RequestedFeatureLevel)
 	GPixelFormats[PF_PLATFORM_HDR_0		].PlatformFormat	= (uint32)mtlpp::PixelFormat::BGR10_XR_sRGB;
 	GPixelFormats[PF_PLATFORM_HDR_0		].Supported			= GRHISupportsHDROutput;
 		
-#if !PLATFORM_TVOS
+#if PLATFORM_TVOS
+	if (![Device supportsFeatureSet:MTLFeatureSet_tvOS_GPUFamily2_v1])
+#else
 	if (![Device supportsFeatureSet:MTLFeatureSet_iOS_GPUFamily3_v2])
+#endif
 	{
 		GPixelFormats[PF_FloatRGB			].PlatformFormat 	= (uint32)mtlpp::PixelFormat::RGBA16Float;
 		GPixelFormats[PF_FloatRGBA			].BlockBytes		= 8;
@@ -647,12 +650,12 @@ FMetalDynamicRHI::FMetalDynamicRHI(ERHIFeatureLevel::Type RequestedFeatureLevel)
 		GPixelFormats[PF_FloatR11G11B10		].BlockBytes		= 8;
 	}
 	else
-#endif
 	{
 		GPixelFormats[PF_FloatRGB			].PlatformFormat	= (uint32)mtlpp::PixelFormat::RG11B10Float;
 		GPixelFormats[PF_FloatRGB			].BlockBytes		= 4;
 		GPixelFormats[PF_FloatR11G11B10		].PlatformFormat	= (uint32)mtlpp::PixelFormat::RG11B10Float;
 		GPixelFormats[PF_FloatR11G11B10		].BlockBytes		= 4;
+		GPixelFormats[PF_FloatR11G11B10		].Supported			= true;
 	}
 	
 	if (FMetalCommandQueue::SupportsFeature(EMetalFeaturesStencilView) && FMetalCommandQueue::SupportsFeature(EMetalFeaturesCombinedDepthStencil) && !FParse::Param(FCommandLine::Get(),TEXT("metalforceseparatedepthstencil")))
@@ -665,8 +668,10 @@ FMetalDynamicRHI::FMetalDynamicRHI(ERHIFeatureLevel::Type RequestedFeatureLevel)
 		GPixelFormats[PF_DepthStencil		].PlatformFormat	= (uint32)mtlpp::PixelFormat::Depth32Float;
 		GPixelFormats[PF_DepthStencil		].BlockBytes		= 4;
 	}
+	GPixelFormats[PF_DepthStencil		].Supported			= true;
 	GPixelFormats[PF_ShadowDepth		].PlatformFormat	= (uint32)mtlpp::PixelFormat::Depth32Float;
 	GPixelFormats[PF_ShadowDepth		].BlockBytes		= 4;
+	GPixelFormats[PF_ShadowDepth		].Supported			= true;
 		
 	GPixelFormats[PF_BC5				].PlatformFormat	= (uint32)mtlpp::PixelFormat::Invalid;
 	GPixelFormats[PF_R5G6B5_UNORM		].PlatformFormat	= (uint32)mtlpp::PixelFormat::B5G6R5Unorm;
@@ -679,6 +684,7 @@ FMetalDynamicRHI::FMetalDynamicRHI(ERHIFeatureLevel::Type RequestedFeatureLevel)
 	GPixelFormats[PF_FloatRGB			].BlockBytes		= 4;
 	GPixelFormats[PF_FloatR11G11B10		].PlatformFormat	= (uint32)mtlpp::PixelFormat::RG11B10Float;
 	GPixelFormats[PF_FloatR11G11B10		].BlockBytes		= 4;
+	GPixelFormats[PF_FloatR11G11B10		].Supported			= true;
 	
 	// Only one HDR format for OSX.
 	GPixelFormats[PF_PLATFORM_HDR_0		].BlockSizeX		= 1;
@@ -699,6 +705,7 @@ FMetalDynamicRHI::FMetalDynamicRHI(ERHIFeatureLevel::Type RequestedFeatureLevel)
 		GPixelFormats[PF_DepthStencil	].PlatformFormat	= (uint32)mtlpp::PixelFormat::Depth32Float_Stencil8;
 	}
 	GPixelFormats[PF_DepthStencil		].BlockBytes		= 4;
+	GPixelFormats[PF_DepthStencil		].Supported			= true;
 	if (bSupportsD16)
 	{
 		GPixelFormats[PF_ShadowDepth		].PlatformFormat	= (uint32)mtlpp::PixelFormat::Depth16Unorm;
@@ -709,6 +716,7 @@ FMetalDynamicRHI::FMetalDynamicRHI(ERHIFeatureLevel::Type RequestedFeatureLevel)
 		GPixelFormats[PF_ShadowDepth		].PlatformFormat	= (uint32)mtlpp::PixelFormat::Depth32Float;
 		GPixelFormats[PF_ShadowDepth		].BlockBytes		= 4;
 	}
+	GPixelFormats[PF_ShadowDepth		].Supported			= true;
 	if(bSupportsD24S8)
 	{
 		GPixelFormats[PF_D24			].PlatformFormat	= (uint32)mtlpp::PixelFormat::Depth24Unorm_Stencil8;
