Prevent heap-use-after-free bugs in Metal timer queries.

diff --git a/Engine/Source/Runtime/Apple/MetalRHI/Private/MetalQuery.cpp b/Engine/Source/Runtime/Apple/MetalRHI/Private/MetalQuery.cpp
--- a/Engine/Source/Runtime/Apple/MetalRHI/Private/MetalQuery.cpp
+++ b/Engine/Source/Runtime/Apple/MetalRHI/Private/MetalQuery.cpp
@@ -230,6 +230,8 @@ void FMetalRenderQuery::End(FMetalContext* Context)
 		}
 		case RQT_AbsoluteTime:
 		{
+			AddRef();
+			
 			// Reset the result availability state
 			Buffer.SourceBuffer.SafeRelease();
 			Buffer.Offset = 0;
@@ -262,6 +264,7 @@ void FMetalRenderQuery::End(FMetalContext* Context)
 						Result = (FPlatformTime::ToMilliseconds64(StatSample.Array[0]) * 1000.0);
 					}
 					[StatSample release];
+					this->Release();
 				});
 			}
 			else
@@ -271,6 +274,7 @@ void FMetalRenderQuery::End(FMetalContext* Context)
 				Context->InsertCommandBufferFence(*(Buffer.CommandBufferFence), [this](mtlpp::CommandBuffer const&)
 				{
 					Result = (FPlatformTime::ToMilliseconds64(mach_absolute_time()) * 1000.0);
+					this->Release();
 				});
 				
 				// Submit the current command buffer, marking this is as a break of a logical command buffer for render restart purposes
