Re-enable Managed buffer sub-allocation as it works now!

diff --git a/Engine/Source/Runtime/Apple/MetalRHI/Private/MetalBuffer.cpp b/Engine/Source/Runtime/Apple/MetalRHI/Private/MetalBuffer.cpp
--- a/Engine/Source/Runtime/Apple/MetalRHI/Private/MetalBuffer.cpp
+++ b/Engine/Source/Runtime/Apple/MetalRHI/Private/MetalBuffer.cpp
@@ -1274,27 +1274,27 @@ FMetalBuffer FMetalResourceHeap::CreateBuffer(uint32 Size, uint32 Alignment, mtl
 				FScopeLock Lock(&Mutex);
 
 				// Disabled Managed sub-allocation as it seems inexplicably slow on the GPU				
-//				if (!bForceUnique && BlockSize <= HeapSizes[NumHeapSizes - 1])
-//				{
-//					FMetalSubBufferLinear* Found = nullptr;
-//					for (FMetalSubBufferLinear* Heap : ManagedSubHeaps)
-//					{
-//						if (Heap->CanAllocateSize(BlockSize))
-//						{
-//							Found = Heap;
-//							break;
-//						}
-//					}
-//					if (!Found)
-//					{
-//						Found = new FMetalSubBufferLinear(HeapAllocSizes[NumHeapSizes - 1], BufferOffsetAlignment, mtlpp::ResourceOptions((NSUInteger)Options & (mtlpp::ResourceStorageModeMask)), Mutex);
-//						ManagedSubHeaps.Add(Found);
-//					}
-//					check(Found);
-//					
-//					return Found->NewBuffer(BlockSize);
-//				}
-//				else
+				 if (!bForceUnique && BlockSize <= HeapSizes[NumHeapSizes - 1])
+				 {
+				 	FMetalSubBufferLinear* Found = nullptr;
+				 	for (FMetalSubBufferLinear* Heap : ManagedSubHeaps)
+				 	{
+				 		if (Heap->CanAllocateSize(BlockSize))
+				 		{
+				 			Found = Heap;
+				 			break;
+				 		}
+				 	}
+				 	if (!Found)
+				 	{
+				 		Found = new FMetalSubBufferLinear(HeapAllocSizes[NumHeapSizes - 1], BufferOffsetAlignment, mtlpp::ResourceOptions((NSUInteger)Options & mtlpp::ResourceStorageModeMask), Mutex);
+				 		ManagedSubHeaps.Add(Found);
+				 	}
+				 	check(Found);
+					
+				 	return Found->NewBuffer(BlockSize);
+				 }
+				 else
 				{
 					Buffer = ManagedBuffers.CreatePooledResource(FMetalPooledBufferArgs(Queue->GetDevice(), BlockSize, StorageMode));
 					DEC_MEMORY_STAT_BY(STAT_MetalBufferUnusedMemory, Buffer.GetLength());
