Add -metalfence  and -metaluntracktextures for testing MTLFence which is a prerequisite for all new Metal features - so far this doesn't work.

diff --git a/Engine/Source/Runtime/Apple/MetalRHI/Private/MetalCommandQueue.cpp b/Engine/Source/Runtime/Apple/MetalRHI/Private/MetalCommandQueue.cpp
--- a/Engine/Source/Runtime/Apple/MetalRHI/Private/MetalCommandQueue.cpp
+++ b/Engine/Source/Runtime/Apple/MetalRHI/Private/MetalCommandQueue.cpp
@@ -181,7 +181,12 @@ FMetalCommandQueue::FMetalCommandQueue(mtlpp::Device InDevice, uint32 const MaxN
     
     if(Device.SupportsFeatureSet(mtlpp::FeatureSet::macOS_GPUFamily1_v3) && FPlatformMisc::MacOSXVersionCompare(10,13,0) >= 0)
     {
-        Features |= EMetalFeaturesMultipleViewports | EMetalFeaturesGPUCommandBufferTimes | EMetalFeaturesGPUCaptureManager | EMetalFeaturesAbsoluteTimeQueries | EMetalFeaturesSupportsVSyncToggle /*| EMetalFeaturesHeaps | EMetalFeaturesFences*/;
+        Features |= EMetalFeaturesMultipleViewports | EMetalFeaturesGPUCommandBufferTimes | EMetalFeaturesGPUCaptureManager | EMetalFeaturesAbsoluteTimeQueries | EMetalFeaturesSupportsVSyncToggle;
+		
+		if (FParse::Param(FCommandLine::Get(),TEXT("metalfence")))
+		{
+			Features |=  EMetalFeaturesHeaps | EMetalFeaturesFences;
+		}
     }
 	else
 	// Time query emulation breaks on AMD < 10.13 - disable by default until they can explain why, should work everywhere else.
@@ -241,7 +246,7 @@ FMetalCommandQueue::FMetalCommandQueue(mtlpp::Device InDevice, uint32 const MaxN
 			PermittedOptions |= mtlpp::ResourceOptions::StorageModeMemoryless;
 		}
 #endif
-		if (Features & EMetalFeaturesFences)
+		if ((Features & EMetalFeaturesFences) && (!PLATFORM_MAC || (FParse::Param(FCommandLine::Get(),TEXT("metaluntracktextures")))))
 		{
 			PermittedOptions |= mtlpp::ResourceOptions::HazardTrackingModeUntracked;
 		}
