Missed a function in PrimitiveType removal.

diff --git a/Engine/Source/Runtime/Apple/MetalRHI/Private/MetalCommands.cpp b/Engine/Source/Runtime/Apple/MetalRHI/Private/MetalCommands.cpp
--- a/Engine/Source/Runtime/Apple/MetalRHI/Private/MetalCommands.cpp
+++ b/Engine/Source/Runtime/Apple/MetalRHI/Private/MetalCommands.cpp
@@ -746,7 +746,7 @@ void FMetalRHICommandContext::RHIDrawIndexedIndirect(FIndexBufferRHIParamRef Ind
 	}
 }
 
-void FMetalRHICommandContext::RHIDrawIndexedPrimitiveIndirect(uint32 PrimitiveType_DEPRECATED,FIndexBufferRHIParamRef IndexBufferRHI,FVertexBufferRHIParamRef VertexBufferRHI,uint32 ArgumentOffset)
+void FMetalRHICommandContext::RHIDrawIndexedPrimitiveIndirect(FIndexBufferRHIParamRef IndexBufferRHI,FVertexBufferRHIParamRef VertexBufferRHI,uint32 ArgumentOffset)
 {
 	@autoreleasepool {
 	if (GetMetalDeviceContext().SupportsFeature(EMetalFeaturesIndirectBuffer))
