Another missing piece of the primtype deprecation.

diff --git a/Engine/Source/Runtime/Apple/MetalRHI/Private/MetalRenderTarget.cpp b/Engine/Source/Runtime/Apple/MetalRHI/Private/MetalRenderTarget.cpp
--- a/Engine/Source/Runtime/Apple/MetalRHI/Private/MetalRenderTarget.cpp
+++ b/Engine/Source/Runtime/Apple/MetalRHI/Private/MetalRenderTarget.cpp
@@ -197,7 +197,7 @@ void FMetalRHICommandContext::RHICopyToResolveTarget(FTextureRHIParamRef SourceT
 
 				RHIUnlockVertexBuffer(VertexBufferRHI);
 				RHICmdList.SetStreamSource(0, VertexBufferRHI, 0);
-				RHICmdList.DrawPrimitive(PT_TriangleStrip, 0, 2, 1);
+				RHICmdList.DrawPrimitive(0, 2, 1);
 
                 RHICmdList.Flush();
             }
