Duplicate MetalPipeline change from CL #4370914 to catch incorrect render-target array format configuration for PSOs.

diff --git a/Engine/Source/Runtime/Apple/MetalRHI/Private/MetalPipeline.cpp b/Engine/Source/Runtime/Apple/MetalRHI/Private/MetalPipeline.cpp
--- a/Engine/Source/Runtime/Apple/MetalRHI/Private/MetalPipeline.cpp
+++ b/Engine/Source/Runtime/Apple/MetalRHI/Private/MetalPipeline.cpp
@@ -395,9 +395,9 @@ static FMetalShaderPipeline* CreateMTLRenderPipeline(bool const bSync, FMetalGra
         for (uint32 i = 0; i < NumActiveTargets; i++)
         {
             EPixelFormat TargetFormat = Init.RenderTargetFormats[i];
-            if (TargetFormat == PF_Unknown)
+            if (TargetFormat == PF_Unknown && PixelShader && (((PixelShader->Bindings.InOutMask & 0x7fff) & (1 << i))))
             {
-                UE_CLOG(PixelShader && (((PixelShader->Bindings.InOutMask & 0x7fff) & (1 << i))), LogMetal, Warning, TEXT("Pipeline pixel shader expects target %u to be bound but it isn't."), i);
+				UE_LOG(LogMetal, Fatal, TEXT("Pipeline pixel shader expects target %u to be bound but it isn't: %s"), i, *FString(PixelShader->GetSourceCode()));
                 continue;
             }
             
@@ -418,7 +418,7 @@ static FMetalShaderPipeline* CreateMTLRenderPipeline(bool const bSync, FMetalGra
             Attachment.SetPixelFormat(MetalFormat);
             
             mtlpp::RenderPipelineColorAttachmentDescriptor Blend = BlendState->RenderTargetStates[i].BlendState;
-            if(Attachment)
+            if(TargetFormat != PF_Unknown)
             {
                 // assign each property manually, would be nice if this was faster
                 Attachment.SetBlendingEnabled(Blend.IsBlendingEnabled());
