Removed vestigial and broken FMetalShaderParameterCache::CommitPackedUniformBuffers - this doesn't work anymore because of the changes to the way shaders are compiled and isn't something I wish to support.

diff --git a/Engine/Source/Runtime/Apple/MetalRHI/Private/MetalShaders.cpp b/Engine/Source/Runtime/Apple/MetalRHI/Private/MetalShaders.cpp
--- a/Engine/Source/Runtime/Apple/MetalRHI/Private/MetalShaders.cpp
+++ b/Engine/Source/Runtime/Apple/MetalRHI/Private/MetalShaders.cpp
@@ -1423,55 +1423,3 @@ void FMetalShaderParameterCache::CommitPackedGlobals(FMetalStateCache* Cache, FM
 		}
 	}
 }
-
-void FMetalShaderParameterCache::CommitPackedUniformBuffers(FMetalStateCache* Cache, TRefCountPtr<FMetalGraphicsPipelineState> BoundShaderState, FMetalComputeShader* ComputeShader, int32 Stage, const FRHIUniformBuffer** RHIUniformBuffers, const TArray<CrossCompiler::FUniformBufferCopyInfo>& UniformBuffersCopyInfo)
-{
-//	SCOPE_CYCLE_COUNTER(STAT_MetalConstantBufferUpdateTime);
-	// Uniform Buffers are split into precision/type; the list of RHI UBs is traversed and if a new one was set, its
-	// contents are copied per precision/type into corresponding scratch buffers which are then uploaded to the program
-	if (Stage == CrossCompiler::SHADER_STAGE_PIXEL && !IsValidRef(BoundShaderState->PixelShader))
-	{
-		return;
-	}
-
-	auto& Bindings = [this, &Stage, &BoundShaderState, ComputeShader]() -> FMetalShaderBindings& {
-		switch(Stage) {
-			default: check(0);
-			case CrossCompiler::SHADER_STAGE_VERTEX: return BoundShaderState->VertexShader->Bindings;
-			case CrossCompiler::SHADER_STAGE_PIXEL: return BoundShaderState->PixelShader->Bindings;
-			case CrossCompiler::SHADER_STAGE_COMPUTE: return ComputeShader->Bindings;
-			case CrossCompiler::SHADER_STAGE_HULL: return BoundShaderState->HullShader->Bindings;
-			case CrossCompiler::SHADER_STAGE_DOMAIN: return BoundShaderState->DomainShader->Bindings;
-		}
-	}();
-
-	if (!Bindings.bHasRegularUniformBuffers)
-	{
-		check(Bindings.NumUniformBuffers <= ML_MaxBuffers);
-		int32 LastInfoIndex = 0;
-		for (int32 BufferIndex = 0; BufferIndex < Bindings.NumUniformBuffers; ++BufferIndex)
-		{
-			const FRHIUniformBuffer* RHIUniformBuffer = RHIUniformBuffers[BufferIndex];
-			check(RHIUniformBuffer);
-			FMetalUniformBuffer* EmulatedUniformBuffer = (FMetalUniformBuffer*)RHIUniformBuffer;
-			const uint32* RESTRICT SourceData = (uint32 const*)((uint8 const*)EmulatedUniformBuffer->GetData());
-			for (int32 InfoIndex = LastInfoIndex; InfoIndex < UniformBuffersCopyInfo.Num(); ++InfoIndex)
-			{
-				const CrossCompiler::FUniformBufferCopyInfo& Info = UniformBuffersCopyInfo[InfoIndex];
-				if (Info.SourceUBIndex == BufferIndex)
-				{
-					float* RESTRICT ScratchMem = (float*)PackedGlobalUniforms[Info.DestUBTypeIndex]->Data;
-					ScratchMem += Info.DestOffsetInFloats;
-					FMemory::Memcpy(ScratchMem, SourceData + Info.SourceOffsetInFloats, Info.SizeInFloats * sizeof(float));
-					PackedGlobalUniformDirty[Info.DestUBTypeIndex].LowVector = FMath::Min(PackedGlobalUniformDirty[Info.DestUBTypeIndex].LowVector, uint32(Info.DestOffsetInFloats / SizeOfFloat));
-					PackedGlobalUniformDirty[Info.DestUBTypeIndex].HighVector = FMath::Max(PackedGlobalUniformDirty[Info.DestUBTypeIndex].HighVector, uint32(((Info.DestOffsetInFloats + Info.SizeInFloats) * sizeof(float) + SizeOfFloat - 1) / SizeOfFloat));
-				}
-				else
-				{
-					LastInfoIndex = InfoIndex;
-					break;
-				}
-			}
-		}
-	}
-}
diff --git a/Engine/Source/Runtime/Apple/MetalRHI/Private/MetalStateCache.cpp b/Engine/Source/Runtime/Apple/MetalRHI/Private/MetalStateCache.cpp
--- a/Engine/Source/Runtime/Apple/MetalRHI/Private/MetalStateCache.cpp
+++ b/Engine/Source/Runtime/Apple/MetalRHI/Private/MetalStateCache.cpp
@@ -1719,13 +1719,11 @@ void FMetalStateCache::CommitRenderResources(FMetalCommandEncoder* Raster)
 	check(IsValidRef(GraphicsPSO));
     
     SetResourcesFromTables(GraphicsPSO->VertexShader, CrossCompiler::SHADER_STAGE_VERTEX);
-    GetShaderParameters(CrossCompiler::SHADER_STAGE_VERTEX).CommitPackedUniformBuffers(this, GraphicsPSO, nullptr, CrossCompiler::SHADER_STAGE_VERTEX, GetBoundUniformBuffers(SF_Vertex), GraphicsPSO->VertexShader->UniformBuffersCopyInfo);
     GetShaderParameters(CrossCompiler::SHADER_STAGE_VERTEX).CommitPackedGlobals(this, Raster, SF_Vertex, GraphicsPSO->VertexShader->Bindings);
 	
     if (IsValidRef(GraphicsPSO->PixelShader))
     {
     	SetResourcesFromTables(GraphicsPSO->PixelShader, CrossCompiler::SHADER_STAGE_PIXEL);
-        GetShaderParameters(CrossCompiler::SHADER_STAGE_PIXEL).CommitPackedUniformBuffers(this, GraphicsPSO, nullptr, CrossCompiler::SHADER_STAGE_PIXEL, GetBoundUniformBuffers(SF_Pixel), GraphicsPSO->PixelShader->UniformBuffersCopyInfo);
         GetShaderParameters(CrossCompiler::SHADER_STAGE_PIXEL).CommitPackedGlobals(this, Raster, SF_Pixel, GraphicsPSO->PixelShader->Bindings);
     }
 }
@@ -1736,23 +1734,19 @@ void FMetalStateCache::CommitTessellationResources(FMetalCommandEncoder* Raster,
     check(IsValidRef(GraphicsPSO->HullShader) && IsValidRef(GraphicsPSO->DomainShader));
     
     SetResourcesFromTables(GraphicsPSO->VertexShader, CrossCompiler::SHADER_STAGE_VERTEX);
-    GetShaderParameters(CrossCompiler::SHADER_STAGE_VERTEX).CommitPackedUniformBuffers(this, GraphicsPSO, nullptr, CrossCompiler::SHADER_STAGE_VERTEX, GetBoundUniformBuffers(SF_Vertex), GraphicsPSO->VertexShader->UniformBuffersCopyInfo);
     GetShaderParameters(CrossCompiler::SHADER_STAGE_VERTEX).CommitPackedGlobals(this, Compute, SF_Vertex, GraphicsPSO->VertexShader->Bindings);
 	
     if (IsValidRef(GraphicsPSO->PixelShader))
     {
     	SetResourcesFromTables(GraphicsPSO->PixelShader, CrossCompiler::SHADER_STAGE_PIXEL);
-        GetShaderParameters(CrossCompiler::SHADER_STAGE_PIXEL).CommitPackedUniformBuffers(this, GraphicsPSO, nullptr, CrossCompiler::SHADER_STAGE_PIXEL, GetBoundUniformBuffers(SF_Pixel), GraphicsPSO->PixelShader->UniformBuffersCopyInfo);
         GetShaderParameters(CrossCompiler::SHADER_STAGE_PIXEL).CommitPackedGlobals(this, Raster, SF_Pixel, GraphicsPSO->PixelShader->Bindings);
     }
     
     SetResourcesFromTables(GraphicsPSO->HullShader, CrossCompiler::SHADER_STAGE_HULL);
-    GetShaderParameters(CrossCompiler::SHADER_STAGE_HULL).CommitPackedUniformBuffers(this, GraphicsPSO, nullptr, CrossCompiler::SHADER_STAGE_HULL, GetBoundUniformBuffers(SF_Hull), GraphicsPSO->HullShader->UniformBuffersCopyInfo);
     GetShaderParameters(CrossCompiler::SHADER_STAGE_HULL).CommitPackedGlobals(this, Compute, SF_Hull, GraphicsPSO->HullShader->Bindings);
 	
 	SetResourcesFromTables(GraphicsPSO->DomainShader, CrossCompiler::SHADER_STAGE_DOMAIN);
-    GetShaderParameters(CrossCompiler::SHADER_STAGE_DOMAIN).CommitPackedUniformBuffers(this, GraphicsPSO, nullptr, CrossCompiler::SHADER_STAGE_DOMAIN, GetBoundUniformBuffers(SF_Domain), GraphicsPSO->DomainShader->UniformBuffersCopyInfo);
-    GetShaderParameters(CrossCompiler::SHADER_STAGE_DOMAIN).CommitPackedGlobals(this, Raster, SF_Domain, GraphicsPSO->DomainShader->Bindings);
+   	GetShaderParameters(CrossCompiler::SHADER_STAGE_DOMAIN).CommitPackedGlobals(this, Raster, SF_Domain, GraphicsPSO->DomainShader->Bindings);
 }
 
 void FMetalStateCache::CommitComputeResources(FMetalCommandEncoder* Compute)
@@ -1760,7 +1754,6 @@ void FMetalStateCache::CommitComputeResources(FMetalCommandEncoder* Compute)
 	check(IsValidRef(ComputeShader));
 	SetResourcesFromTables(ComputeShader, CrossCompiler::SHADER_STAGE_COMPUTE);
 	
-	GetShaderParameters(CrossCompiler::SHADER_STAGE_COMPUTE).CommitPackedUniformBuffers(this, GraphicsPSO, ComputeShader, CrossCompiler::SHADER_STAGE_COMPUTE, GetBoundUniformBuffers(SF_Compute), ComputeShader->UniformBuffersCopyInfo);
 	GetShaderParameters(CrossCompiler::SHADER_STAGE_COMPUTE).CommitPackedGlobals(this, Compute, SF_Compute, ComputeShader->Bindings);
 }
 
diff --git a/Engine/Source/Runtime/Apple/MetalRHI/Public/MetalResources.h b/Engine/Source/Runtime/Apple/MetalRHI/Public/MetalResources.h
--- a/Engine/Source/Runtime/Apple/MetalRHI/Public/MetalResources.h
+++ b/Engine/Source/Runtime/Apple/MetalRHI/Public/MetalResources.h
@@ -932,7 +932,6 @@ public:
 	 * Commit shader parameters to the currently bound program.
 	 */
 	void CommitPackedGlobals(class FMetalStateCache* Cache, class FMetalCommandEncoder* Encoder, EShaderFrequency Frequency, const FMetalShaderBindings& Bindings);
-	void CommitPackedUniformBuffers(class FMetalStateCache* Cache, TRefCountPtr<FMetalGraphicsPipelineState> BoundShaderState, FMetalComputeShader* ComputeShader, int32 Stage, const FRHIUniformBuffer** UniformBuffers, const TArray<CrossCompiler::FUniformBufferCopyInfo>& UniformBuffersCopyInfo);
 
 private:
 	/** CPU memory block for storing uniform values. */
