iOS 11 and above should be able to suballocate, the -metalfence command-line turns on MTLFence, -metalheap turns on MTLHeap and -metaluntracktextures allows us to use the HazardTrackingModeUntracked  when MTLFence is enabled.

diff --git a/Engine/Source/Runtime/Apple/MetalRHI/Private/MetalBuffer.cpp b/Engine/Source/Runtime/Apple/MetalRHI/Private/MetalBuffer.cpp
--- a/Engine/Source/Runtime/Apple/MetalRHI/Private/MetalBuffer.cpp
+++ b/Engine/Source/Runtime/Apple/MetalRHI/Private/MetalBuffer.cpp
@@ -149,7 +149,7 @@ FMetalSubBufferHeap::FMetalSubBufferHeap(NSUInteger Size, NSUInteger Alignment,
 	check(Storage != mtlpp::StorageMode::Managed /* Managed memory cannot be safely suballocated! When you overwrite existing data the GPU buffer is immediately disposed of! */);
 #endif
 
-	if (bSupportsHeaps && (!PLATFORM_MAC || Storage == mtlpp::StorageMode::Private))
+	if (bSupportsHeaps && Storage == mtlpp::StorageMode::Private)
 	{
 		mtlpp::HeapDescriptor Desc;
 		Desc.SetSize(FullSize);
@@ -580,7 +580,7 @@ FMetalSubBufferMagazine::FMetalSubBufferMagazine(NSUInteger Size, NSUInteger Chu
 #endif
 
 	static bool bSupportsHeaps = GetMetalDeviceContext().SupportsFeature(EMetalFeaturesHeaps);
-	if (bSupportsHeaps && (!PLATFORM_MAC || Storage == mtlpp::StorageMode::Private))
+	if (bSupportsHeaps && Storage == mtlpp::StorageMode::Private)
 	{
 		mtlpp::HeapDescriptor Desc;
 		Desc.SetSize(FullSize);
@@ -1258,8 +1258,9 @@ FMetalBuffer FMetalResourceHeap::CreateBuffer(uint32 Size, uint32 Alignment, mtl
 	LLM_SCOPE_METAL(ELLMTagMetal::Buffers);
 	LLM_PLATFORM_SCOPE_METAL(ELLMTagMetal::Buffers);
 	
+	static bool bSupportsHeaps = GetMetalDeviceContext().SupportsFeature(EMetalFeaturesHeaps);
 	static bool bSupportsBufferSubAllocation = FMetalCommandQueue::SupportsFeature(EMetalFeaturesBufferSubAllocation);
-	bForceUnique |= !(bSupportsBufferSubAllocation);
+	bForceUnique |= (!bSupportsBufferSubAllocation && !bSupportsHeaps);
 	
 	FMetalBuffer Buffer;
 	uint32 BlockSize = Align(Size, Alignment);
diff --git a/Engine/Source/Runtime/Apple/MetalRHI/Private/MetalCommandQueue.cpp b/Engine/Source/Runtime/Apple/MetalRHI/Private/MetalCommandQueue.cpp
--- a/Engine/Source/Runtime/Apple/MetalRHI/Private/MetalCommandQueue.cpp
+++ b/Engine/Source/Runtime/Apple/MetalRHI/Private/MetalCommandQueue.cpp
@@ -72,7 +72,17 @@ FMetalCommandQueue::FMetalCommandQueue(mtlpp::Device InDevice, uint32 const MaxN
 		
 		if([Device supportsFeatureSet:MTLFeatureSet_iOS_GPUFamily3_v2] || [Device supportsFeatureSet:MTLFeatureSet_iOS_GPUFamily2_v3] || [Device supportsFeatureSet:MTLFeatureSet_iOS_GPUFamily1_v3])
 		{
-			Features |= EMetalFeaturesStencilView | EMetalFeaturesFunctionConstants | EMetalFeaturesGraphicsUAVs | EMetalFeaturesMemoryLessResources /*| EMetalFeaturesHeaps | EMetalFeaturesFences*/;
+			Features |= EMetalFeaturesStencilView | EMetalFeaturesFunctionConstants | EMetalFeaturesGraphicsUAVs | EMetalFeaturesMemoryLessResources;
+			
+			if (FParse::Param(FCommandLine::Get(),TEXT("metalfence")))
+			{
+				Features |= EMetalFeaturesFences;
+			}
+			
+			if (FParse::Param(FCommandLine::Get(),TEXT("metalheap")))
+			{
+				Features |= EMetalFeaturesHeaps;
+			}
 		}
 		
 		if([Device supportsFeatureSet:MTLFeatureSet_iOS_GPUFamily3_v2])
@@ -117,7 +127,7 @@ FMetalCommandQueue::FMetalCommandQueue(mtlpp::Device InDevice, uint32 const MaxN
 		
 		if(Vers.majorVersion >= 11)
 		{
-			Features |= EMetalFeaturesPresentMinDuration | EMetalFeaturesGPUCaptureManager;
+			Features |= EMetalFeaturesPresentMinDuration | EMetalFeaturesGPUCaptureManager | EMetalFeaturesBufferSubAllocation;
         }
 #endif
 	}
@@ -185,7 +195,12 @@ FMetalCommandQueue::FMetalCommandQueue(mtlpp::Device InDevice, uint32 const MaxN
 		
 		if (FParse::Param(FCommandLine::Get(),TEXT("metalfence")))
 		{
-			Features |=  EMetalFeaturesHeaps | EMetalFeaturesFences;
+			Features |= EMetalFeaturesFences;
+		}
+		
+		if (FParse::Param(FCommandLine::Get(),TEXT("metalheap")))
+		{
+			Features |= EMetalFeaturesHeaps;
 		}
     }
 	else
@@ -246,7 +261,7 @@ FMetalCommandQueue::FMetalCommandQueue(mtlpp::Device InDevice, uint32 const MaxN
 			PermittedOptions |= mtlpp::ResourceOptions::StorageModeMemoryless;
 		}
 #endif
-		if ((Features & EMetalFeaturesFences) && (!PLATFORM_MAC || (FParse::Param(FCommandLine::Get(),TEXT("metaluntracktextures")))))
+		if ((Features & EMetalFeaturesFences) && FParse::Param(FCommandLine::Get(),TEXT("metaluntracktextures")))
 		{
 			PermittedOptions |= mtlpp::ResourceOptions::HazardTrackingModeUntracked;
 		}
