Rewrite MTLFence support to try and make it viable again.
- Adds an explicit feature EMetalFeaturesParallelRenderEncoders to make handling the different behaviour needed for parallel encoders and context easier.
- This implementation is incomplete and does not work correctly on macOS, on iOS it is inefficient as it is inserting Fragment->Vertex barriers across all encoders.
- The purpose of the exercise is to find where the logic is breaking down and fix the behaviour so that we can implement the proper fences and formerly adopt this feature.

diff --git a/Engine/Source/Runtime/Apple/MetalRHI/Private/MetalCommandEncoder.cpp b/Engine/Source/Runtime/Apple/MetalRHI/Private/MetalCommandEncoder.cpp
--- a/Engine/Source/Runtime/Apple/MetalRHI/Private/MetalCommandEncoder.cpp
+++ b/Engine/Source/Runtime/Apple/MetalRHI/Private/MetalCommandEncoder.cpp
@@ -286,7 +286,7 @@ mtlpp::BlitCommandEncoder& FMetalCommandEncoder::GetBlitCommandEncoder(void)
 	return BlitCommandEncoder;
 }
 
-mtlpp::Fence FMetalCommandEncoder::GetEncoderFence(void) const
+FMetalFence* FMetalCommandEncoder::GetEncoderFence(void) const
 {
 	return EncoderFence;
 }
@@ -305,7 +305,7 @@ void FMetalCommandEncoder::BeginParallelRenderCommandEncoding(uint32 NumChildren
 	
 	EncoderNum++;
 	
-	check(!EncoderFence.IsValid());
+	check(!EncoderFence);
 	NSString* Label = nil;
 	
 	if(GetEmitDrawEvents())
@@ -354,7 +354,7 @@ void FMetalCommandEncoder::BeginRenderCommandEncoding(void)
 		METAL_DEBUG_LAYER(EMetalDebugLevelFastValidation, RenderEncoderDebug = FMetalRenderCommandEncoderDebugging::Get(RenderCommandEncoder));
 	}
 	
-	check(!EncoderFence.IsValid());
+	check(!EncoderFence);
 	NSString* Label = nil;
 	
 	if(GetEmitDrawEvents())
@@ -384,7 +384,10 @@ void FMetalCommandEncoder::BeginRenderCommandEncoding(void)
 	}
 	METAL_STATISTIC(FMetalProfiler::GetProfiler()->BeginEncoder(CommandBufferStats, RenderCommandEncoder));
 	
-	EncoderFence = FMetalFence(CommandList.GetCommandQueue().CreateFence(Label));
+	if (!CommandList.IsParallel())
+	{
+		EncoderFence = CommandList.GetCommandQueue().CreateFence(Label);
+	}
 }
 
 void FMetalCommandEncoder::BeginComputeCommandEncoding(void)
@@ -397,7 +400,7 @@ void FMetalCommandEncoder::BeginComputeCommandEncoding(void)
 
 	EncoderNum++;
 	
-	check(!EncoderFence.IsValid());
+	check(!EncoderFence);
 	NSString* Label = nil;
 	
 	if(GetEmitDrawEvents())
@@ -419,10 +422,10 @@ void FMetalCommandEncoder::BeginComputeCommandEncoding(void)
 		}
 	}
 	METAL_STATISTIC(FMetalProfiler::GetProfiler()->BeginEncoder(CommandBufferStats, ComputeCommandEncoder));
-	EncoderFence = FMetalFence(CommandList.GetCommandQueue().CreateFence(Label));
+	EncoderFence = CommandList.GetCommandQueue().CreateFence(Label);
 }
 
-void FMetalCommandEncoder::BeginBlitCommandEncoding(void)
+void FMetalCommandEncoder::BeginBlitCommandEncoding(bool const bSuppressFence)
 {
 	check(CommandBuffer);
 	check(IsRenderCommandEncoderActive() == false && IsComputeCommandEncoderActive() == false && IsBlitCommandEncoderActive() == false);
@@ -432,7 +435,7 @@ void FMetalCommandEncoder::BeginBlitCommandEncoding(void)
 	
 	EncoderNum++;
 	
-	check(!EncoderFence.IsValid());
+	check(!EncoderFence);
 	NSString* Label = nil;
 	
 	if(GetEmitDrawEvents())
@@ -454,20 +457,23 @@ void FMetalCommandEncoder::BeginBlitCommandEncoding(void)
 		}
 	}
 	METAL_STATISTIC(FMetalProfiler::GetProfiler()->BeginEncoder(CommandBufferStats, BlitCommandEncoder));
-	EncoderFence = FMetalFence(CommandList.GetCommandQueue().CreateFence(Label));
+	if (!bSuppressFence)
+	{
+		EncoderFence = CommandList.GetCommandQueue().CreateFence(Label);
+	}
 }
 
-mtlpp::Fence FMetalCommandEncoder::EndEncoding(void)
+FMetalFence* FMetalCommandEncoder::EndEncoding(void)
 {
 	static bool bSupportsFences = CommandList.GetCommandQueue().SupportsFeature(EMetalFeaturesFences);
-	mtlpp::Fence Fence;
+	FMetalFence* Fence = nullptr;
 	@autoreleasepool
 	{
 		if(IsRenderCommandEncoderActive())
 		{
 			if (RenderCommandEncoder)
 			{
-				check(!bSupportsFences || EncoderFence.IsValid());
+				check(!bSupportsFences || EncoderFence || CommandList.IsParallel());
 				static bool bDeferredStoreActions = CommandList.GetCommandQueue().SupportsFeature(EMetalFeaturesDeferredStoreActions);
 				if (bDeferredStoreActions && ParallelRenderCommandEncoder.GetPtr() == nil)
 				{
@@ -497,14 +503,18 @@ mtlpp::Fence FMetalCommandEncoder::EndEncoding(void)
 					}
 				}
 				
-				Fence = EncoderFence;
-				UpdateFence(Fence);
+				if (EncoderFence && EncoderFence->NeedsWrite(mtlpp::RenderStages::Fragment))
+				{
+					Fence = EncoderFence;
+				}
+				UpdateFence(EncoderFence);
 				
 				METAL_STATISTIC(FMetalProfiler::GetProfiler()->EndEncoder(CommandBufferStats, RenderCommandEncoder));
 				RenderCommandEncoder.EndEncoding();
 				METAL_DEBUG_LAYER(EMetalDebugLevelFastValidation, RenderEncoderDebug.EndEncoder());
 				RenderCommandEncoder = nil;
-				EncoderFence.Reset();
+				SafeReleaseMetalFence(EncoderFence);
+				EncoderFence = nullptr;
 			}
 			
 			if (ParallelRenderCommandEncoder && IsParallel())
@@ -586,27 +596,35 @@ mtlpp::Fence FMetalCommandEncoder::EndEncoding(void)
 		}
 		else if(IsComputeCommandEncoderActive())
 		{
-			check(!bSupportsFences || EncoderFence.IsValid());
-			Fence = EncoderFence;
-			UpdateFence(Fence);
+			check(!bSupportsFences || EncoderFence);
+			if (EncoderFence && EncoderFence->NeedsWrite(mtlpp::RenderStages::Fragment))
+			{
+				Fence = EncoderFence;
+			}
+			UpdateFence(EncoderFence);
 			
 			METAL_STATISTIC(FMetalProfiler::GetProfiler()->EndEncoder(CommandBufferStats, ComputeCommandEncoder));
 			ComputeCommandEncoder.EndEncoding();
 			METAL_DEBUG_LAYER(EMetalDebugLevelFastValidation, ComputeEncoderDebug.EndEncoder());
 			ComputeCommandEncoder = nil;
-			EncoderFence.Reset();
+			SafeReleaseMetalFence(EncoderFence);
+			EncoderFence = nullptr;
 		}
 		else if(IsBlitCommandEncoderActive())
 		{
-			check(!bSupportsFences || EncoderFence.IsValid());
-			Fence = EncoderFence;
-			UpdateFence(Fence);
+			// check(!bSupportsFences || EncoderFence);
+			if (EncoderFence && EncoderFence->NeedsWrite(mtlpp::RenderStages::Fragment))
+			{
+				Fence = EncoderFence;
+			}
+			UpdateFence(EncoderFence);
 			
 			METAL_STATISTIC(FMetalProfiler::GetProfiler()->EndEncoder(CommandBufferStats, RenderCommandEncoder));
 			BlitCommandEncoder.EndEncoding();
 			METAL_DEBUG_LAYER(EMetalDebugLevelFastValidation, BlitEncoderDebug.EndEncoder());
 			BlitCommandEncoder = nil;
-			EncoderFence.Reset();
+			SafeReleaseMetalFence(EncoderFence);
+			EncoderFence = nullptr;
 		}
 	}
 	
@@ -645,54 +663,67 @@ void FMetalCommandEncoder::AddCompletionHandler(mtlpp::CommandBufferHandler Hand
 	Block_release(HeapHandler);
 }
 
-void FMetalCommandEncoder::UpdateFence(mtlpp::Fence Fence)
+void FMetalCommandEncoder::UpdateFence(FMetalFence* Fence)
 {
 	check(IsRenderCommandEncoderActive() || IsComputeCommandEncoderActive() || IsBlitCommandEncoderActive());
 	static bool bSupportsFences = CommandList.GetCommandQueue().SupportsFeature(EMetalFeaturesFences);
-	if (bSupportsFences METAL_DEBUG_OPTION(|| CommandList.GetCommandQueue().GetRuntimeDebuggingLevel() >= EMetalDebugLevelValidation))
+	if ((bSupportsFences METAL_DEBUG_OPTION(|| CommandList.GetCommandQueue().GetRuntimeDebuggingLevel() >= EMetalDebugLevelValidation)) && Fence && Fence->NeedsWrite(mtlpp::RenderStages::Fragment))
 	{
 		check(Fence);
-		mtlpp::Fence InnerFence = METAL_DEBUG_OPTION(CommandList.GetCommandQueue().GetRuntimeDebuggingLevel() >= EMetalDebugLevelValidation ? mtlpp::Fence(((FMetalDebugFence*)Fence.GetPtr()).Inner) :) Fence;
+		mtlpp::Fence InnerFence = METAL_DEBUG_OPTION(CommandList.GetCommandQueue().GetRuntimeDebuggingLevel() >= EMetalDebugLevelValidation ? mtlpp::Fence(((FMetalDebugFence*)Fence->GetPtr()).Inner) :) *Fence;
 		if (RenderCommandEncoder)
 		{
-			RenderCommandEncoder.UpdateFence(InnerFence, (mtlpp::RenderStages)(mtlpp::RenderStages::Vertex | mtlpp::RenderStages::Fragment));
-			METAL_DEBUG_LAYER(EMetalDebugLevelFastValidation, RenderEncoderDebug.AddUpdateFence(Fence));
+//			Fence->Write(mtlpp::RenderStages::Vertex);
+			Fence->Write(mtlpp::RenderStages::Fragment);
+			RenderCommandEncoder.UpdateFence(InnerFence, (mtlpp::RenderStages)(/*mtlpp::RenderStages::Vertex |*/ mtlpp::RenderStages::Fragment));
+			METAL_DEBUG_LAYER(EMetalDebugLevelFastValidation, RenderEncoderDebug.AddUpdateFence(*Fence));
 		}
 		else if (ComputeCommandEncoder)
 		{
+//			Fence->Write(mtlpp::RenderStages::Vertex);
+			Fence->Write(mtlpp::RenderStages::Fragment);
 			ComputeCommandEncoder.UpdateFence(InnerFence);
-			METAL_DEBUG_LAYER(EMetalDebugLevelFastValidation, ComputeEncoderDebug.AddUpdateFence(Fence));
+			METAL_DEBUG_LAYER(EMetalDebugLevelFastValidation, ComputeEncoderDebug.AddUpdateFence(*Fence));
 		}
 		else if (BlitCommandEncoder)
 		{
+//			Fence->Write(mtlpp::RenderStages::Vertex);
+			Fence->Write(mtlpp::RenderStages::Fragment);
 			BlitCommandEncoder.UpdateFence(InnerFence);
-			METAL_DEBUG_LAYER(EMetalDebugLevelFastValidation, BlitEncoderDebug.AddUpdateFence(Fence));
+			METAL_DEBUG_LAYER(EMetalDebugLevelFastValidation, BlitEncoderDebug.AddUpdateFence(*Fence));
 		}
 	}
 }
 
-void FMetalCommandEncoder::WaitForFence(mtlpp::Fence Fence)
+void FMetalCommandEncoder::WaitForFence(FMetalFence* Fence)
 {
 	check(IsRenderCommandEncoderActive() || IsComputeCommandEncoderActive() || IsBlitCommandEncoderActive());
 	static bool bSupportsFences = CommandList.GetCommandQueue().SupportsFeature(EMetalFeaturesFences);
-	if (bSupportsFences METAL_DEBUG_OPTION(|| CommandList.GetCommandQueue().GetRuntimeDebuggingLevel() >= EMetalDebugLevelValidation))
+	if ((bSupportsFences METAL_DEBUG_OPTION(|| CommandList.GetCommandQueue().GetRuntimeDebuggingLevel() >= EMetalDebugLevelValidation)) && Fence && Fence->NeedsWait(mtlpp::RenderStages::Fragment))
 	{
 		check(Fence);
-		mtlpp::Fence InnerFence = METAL_DEBUG_OPTION(CommandList.GetCommandQueue().GetRuntimeDebuggingLevel() >= EMetalDebugLevelValidation ? mtlpp::Fence(((FMetalDebugFence*)Fence.GetPtr()).Inner) :) Fence;
+		
+		mtlpp::Fence InnerFence = METAL_DEBUG_OPTION(CommandList.GetCommandQueue().GetRuntimeDebuggingLevel() >= EMetalDebugLevelValidation ? mtlpp::Fence(((FMetalDebugFence*)Fence->GetPtr()).Inner) :) *Fence;
 		if (RenderCommandEncoder)
 		{
-			RenderCommandEncoder.WaitForFence(InnerFence, (mtlpp::RenderStages)(mtlpp::RenderStages::Vertex | mtlpp::RenderStages::Fragment));
-			METAL_DEBUG_LAYER(EMetalDebugLevelFastValidation, RenderEncoderDebug.AddWaitFence(Fence));
+//			Fence->Wait(mtlpp::RenderStages::Vertex);
+			Fence->Wait(mtlpp::RenderStages::Fragment);
+			RenderCommandEncoder.WaitForFence(InnerFence, (mtlpp::RenderStages)(mtlpp::RenderStages::Vertex /*| mtlpp::RenderStages::Fragment*/));
+			METAL_DEBUG_LAYER(EMetalDebugLevelFastValidation, RenderEncoderDebug.AddWaitFence(*Fence));
 		}
 		else if (ComputeCommandEncoder)
 		{
+//			Fence->Wait(mtlpp::RenderStages::Vertex);
+			Fence->Wait(mtlpp::RenderStages::Fragment);
 			ComputeCommandEncoder.WaitForFence(InnerFence);
-			METAL_DEBUG_LAYER(EMetalDebugLevelFastValidation, ComputeEncoderDebug.AddWaitFence(Fence));
+			METAL_DEBUG_LAYER(EMetalDebugLevelFastValidation, ComputeEncoderDebug.AddWaitFence(*Fence));
 		}
 		else if (BlitCommandEncoder)
 		{
+//			Fence->Wait(mtlpp::RenderStages::Vertex);
+			Fence->Wait(mtlpp::RenderStages::Fragment);
 			BlitCommandEncoder.WaitForFence(InnerFence);
-			METAL_DEBUG_LAYER(EMetalDebugLevelFastValidation, BlitEncoderDebug.AddWaitFence(Fence));
+			METAL_DEBUG_LAYER(EMetalDebugLevelFastValidation, BlitEncoderDebug.AddWaitFence(*Fence));
 		}
 	}
 }
diff --git a/Engine/Source/Runtime/Apple/MetalRHI/Private/MetalCommandEncoder.h b/Engine/Source/Runtime/Apple/MetalRHI/Private/MetalCommandEncoder.h
--- a/Engine/Source/Runtime/Apple/MetalRHI/Private/MetalCommandEncoder.h
+++ b/Engine/Source/Runtime/Apple/MetalRHI/Private/MetalCommandEncoder.h
@@ -128,7 +128,7 @@ public:
 	mtlpp::BlitCommandEncoder& GetBlitCommandEncoder(void);
 	
 	/** @returns The MTLFence for the current encoder or nil if there isn't one. */
-	mtlpp::Fence GetEncoderFence(void) const;
+	FMetalFence* GetEncoderFence(void) const;
 	
 #pragma mark - Public Command Encoder Mutators -
 
@@ -147,10 +147,10 @@ public:
 	void BeginComputeCommandEncoding(void);
 	
 	/** Begins encoding blit commands into the current command buffer. No other encoder may be active. */
-	void BeginBlitCommandEncoding(void);
+	void BeginBlitCommandEncoding(bool const bSuppressFence = false);
 	
 	/** Declare that all command generation from this encoder is complete, and detach from the MTLCommandBuffer if there is an encoder active or does nothing if there isn't. */
-	mtlpp::Fence EndEncoding(void);
+	FMetalFence* EndEncoding(void);
 	
 	/** Initialises a fence for the current command-buffer optionally adding a command-buffer completion handler to the command-buffer */
 	void InsertCommandBufferFence(FMetalCommandBufferFence& Fence, mtlpp::CommandBufferHandler Handler);
@@ -159,10 +159,10 @@ public:
 	void AddCompletionHandler(mtlpp::CommandBufferHandler Handler);
 	
 	/** Update the event to capture all GPU work so far enqueued by this encoder. */
-	void UpdateFence(mtlpp::Fence Fence);
+	void UpdateFence(FMetalFence* Fence);
 	
 	/** Prevent further GPU work until the event is reached. */
-	void WaitForFence(mtlpp::Fence Fence);
+	void WaitForFence(FMetalFence* Fence);
 
 #pragma mark - Public Debug Support -
 	
@@ -433,7 +433,7 @@ private:
 	METAL_DEBUG_ONLY(FMetalBlitCommandEncoderDebugging BlitEncoderDebug);
 	METAL_DEBUG_ONLY(FMetalParallelRenderCommandEncoderDebugging ParallelEncoderDebug);
 	
-	FMetalFence EncoderFence;
+	FMetalFence* EncoderFence;
 #if ENABLE_METAL_GPUPROFILE
 	FMetalCommandBufferStats* CommandBufferStats;
 #endif
diff --git a/Engine/Source/Runtime/Apple/MetalRHI/Private/MetalCommandList.h b/Engine/Source/Runtime/Apple/MetalRHI/Private/MetalCommandList.h
--- a/Engine/Source/Runtime/Apple/MetalRHI/Private/MetalCommandList.h
+++ b/Engine/Source/Runtime/Apple/MetalRHI/Private/MetalCommandList.h
@@ -71,7 +71,7 @@ public:
 	 * True iff the command-list is part of an MTLParallelRenderCommandEncoder pass.
 	 * @returns True iff the command-list is part of an MTLParallelRenderCommandEncoder pass, false for immediate and parallel-command-buffer contexts.
 	 */
-	bool IsParallel(void) const { return !bImmediate && Num > 0 && !IsRHIDeviceNVIDIA(); }
+	bool IsParallel(void) const { return !bImmediate && Num > 0 && FMetalCommandQueue::SupportsFeature(EMetalFeaturesParallelRenderEncoders); }
 	
 	/**
 	 * The index of this command-list within the parallel pass.
diff --git a/Engine/Source/Runtime/Apple/MetalRHI/Private/MetalCommandQueue.cpp b/Engine/Source/Runtime/Apple/MetalRHI/Private/MetalCommandQueue.cpp
--- a/Engine/Source/Runtime/Apple/MetalRHI/Private/MetalCommandQueue.cpp
+++ b/Engine/Source/Runtime/Apple/MetalRHI/Private/MetalCommandQueue.cpp
@@ -127,7 +127,7 @@ FMetalCommandQueue::FMetalCommandQueue(mtlpp::Device InDevice, uint32 const MaxN
 		
 		if(Vers.majorVersion >= 11)
 		{
-			Features |= EMetalFeaturesPresentMinDuration | EMetalFeaturesGPUCaptureManager | EMetalFeaturesBufferSubAllocation;
+			Features |= EMetalFeaturesPresentMinDuration | EMetalFeaturesGPUCaptureManager | EMetalFeaturesBufferSubAllocation | EMetalFeaturesParallelRenderEncoders;
         }
 #endif
 	}
@@ -166,6 +166,16 @@ FMetalCommandQueue::FMetalCommandQueue(mtlpp::Device InDevice, uint32 const MaxN
 				Features |= EMetalFeaturesPrivateBufferSubAllocation;
 			}
 		}
+		
+		// On 10.13.5+ we can use MTLParallelRenderEncoder
+		if (FPlatformMisc::MacOSXVersionCompare(10,13,5) >= 0)
+		{
+			// Except on Nvidia for the moment
+			if ([Device.GetName().GetPtr() rangeOfString:@"Nvidia" options:NSCaseInsensitiveSearch].location == NSNotFound)
+			{
+				Features |= EMetalFeaturesParallelRenderEncoders;
+			}
+		}
 
 		// Turn on Linear Texture UAVs! Avoids the need to have function-constants which reduces initial runtime shader compile time
 		if (MaxShaderVersion >= 3 && FPlatformMisc::MacOSXVersionCompare(10,13,5) >= 0)
@@ -236,6 +246,9 @@ FMetalCommandQueue::FMetalCommandQueue(mtlpp::Device InDevice, uint32 const MaxN
 			{
 				GSupportsTimestampRenderQueries = true;
 				Features |= EMetalFeaturesStatistics;
+				
+				// Stats doesn't support Parallel Encoders yet
+				Features &= ~(EMetalFeaturesParallelRenderEncoders);
 			}
 			else
 			{
@@ -345,28 +358,30 @@ void FMetalCommandQueue::SubmitCommandBuffers(TArray<mtlpp::CommandBuffer> Buffe
 	}
 }
 
-mtlpp::Fence FMetalCommandQueue::CreateFence(ns::String const& Label) const
+FMetalFence* FMetalCommandQueue::CreateFence(ns::String const& Label) const
 {
-	mtlpp::Fence InternalFence;
-	if(Features & EMetalFeaturesFences)
+	if ((Features & EMetalFeaturesFences) != 0)
 	{
-		InternalFence = Device.NewFence();
-	}
+		FMetalFence* InternalFence = FMetalFencePool::Get().AllocateFence();
+		NSString* String = [NSString stringWithFormat:@"%p: %@", InternalFence->GetPtr(), Label.GetPtr()];
 #if METAL_DEBUG_OPTIONS
-	if (RuntimeDebuggingLevel >= EMetalDebugLevelValidation)
-	{
-		FMetalDebugFence* Fence = [[FMetalDebugFence new] autorelease];
-		Fence.Inner = InternalFence;
-		InternalFence = Fence;
-		Fence.label = Label;
+		if (RuntimeDebuggingLevel >= EMetalDebugLevelValidation)
+		{
+			FMetalDebugFence* Fence = (FMetalDebugFence*)InternalFence->GetPtr();
+			Fence.label = String;
+		}
+		else
+#endif
+		if(InternalFence && Label)
+		{
+			InternalFence->SetLabel(String);
+		}
+		return InternalFence;
 	}
 	else
-#endif
-	if(InternalFence && Label)
 	{
-		InternalFence.SetLabel(Label);
+		return nullptr;
 	}
-	return InternalFence;
 }
 
 void FMetalCommandQueue::GetCommittedCommandBufferFences(TArray<mtlpp::CommandBufferFence>& Fences)
diff --git a/Engine/Source/Runtime/Apple/MetalRHI/Private/MetalCommandQueue.h b/Engine/Source/Runtime/Apple/MetalRHI/Private/MetalCommandQueue.h
--- a/Engine/Source/Runtime/Apple/MetalRHI/Private/MetalCommandQueue.h
+++ b/Engine/Source/Runtime/Apple/MetalRHI/Private/MetalCommandQueue.h
@@ -97,6 +97,8 @@ typedef NS_OPTIONS(uint64, EMetalFeatures)
 	EMetalFeaturesTextureBuffers = 1llu << 38llu,
 	/** Supports max. compute threads per threadgroup */
 	EMetalFeaturesMaxThreadsPerThreadgroup = 1llu << 39llu,
+	/** Supports parallel render encoders */
+	EMetalFeaturesParallelRenderEncoders = 1llu << 40llu,
 };
 
 /**
@@ -141,7 +143,7 @@ public:
 	void SubmitCommandBuffers(TArray<mtlpp::CommandBuffer> BufferList, uint32 Index, uint32 Count);
 
 	/** @returns Creates a new MTLFence or nil if this is unsupported */
-	mtlpp::Fence CreateFence(ns::String const& Label) const;
+	FMetalFence* CreateFence(ns::String const& Label) const;
 	
 	/** @params Fences An array of command-buffer fences for the committed command-buffers */
 	void GetCommittedCommandBufferFences(TArray<mtlpp::CommandBufferFence>& Fences);
diff --git a/Engine/Source/Runtime/Apple/MetalRHI/Private/MetalContext.cpp b/Engine/Source/Runtime/Apple/MetalRHI/Private/MetalContext.cpp
--- a/Engine/Source/Runtime/Apple/MetalRHI/Private/MetalContext.cpp
+++ b/Engine/Source/Runtime/Apple/MetalRHI/Private/MetalContext.cpp
@@ -329,6 +329,11 @@ FMetalDeviceContext* FMetalDeviceContext::CreateDeviceContext()
 	FMetalCommandQueue* Queue = new FMetalCommandQueue(Device, GMetalCommandQueueSize);
 	check(Queue);
 	
+	if (FMetalCommandQueue::SupportsFeature(EMetalFeaturesFences))
+	{
+		FMetalFencePool::Get().Initialise(Device);
+	}
+	
 	return new FMetalDeviceContext(Device, DeviceIndex, Queue);
 }
 
@@ -467,6 +472,10 @@ void FMetalDeviceContext::ClearFreeList()
 					Heap.ReleaseTexture(nullptr, Texture);
 				}
 			}
+			for ( FMetalFence* Fence : Pair->FenceFreeList )
+			{
+				FMetalFencePool::Get().ReleaseFence(Fence);
+			}
 			delete Pair;
 			DelayedFreeLists.RemoveAt(Index, 1, false);
 		}
@@ -567,7 +576,7 @@ bool FMetalDeviceContext::FMetalDelayedFreeList::IsComplete() const
 	return bFinished;
 }
 
-void FMetalDeviceContext::FlushFreeList()
+void FMetalDeviceContext::FlushFreeList(bool const bFlushFences)
 {
 	FMetalDelayedFreeList* NewList = new FMetalDelayedFreeList;
 	
@@ -579,6 +588,10 @@ void FMetalDeviceContext::FlushFreeList()
 	NewList->UsedBuffers = MoveTemp(UsedBuffers);
 	NewList->UsedTextures = MoveTemp(UsedTextures);
 	NewList->ObjectFreeList = ObjectFreeList;
+	if (bFlushFences)
+	{
+		NewList->FenceFreeList = MoveTemp(FenceFreeList);
+	}
 #if METAL_DEBUG_OPTIONS
 	if (FrameFences.Num())
 	{
@@ -677,7 +690,7 @@ void FMetalDeviceContext::ReleaseTexture(FMetalTexture& Texture)
 	}
 }
 
-void FMetalDeviceContext::ReleaseFence(mtlpp::Fence Fence)
+void FMetalDeviceContext::ReleaseFence(FMetalFence* Fence)
 {
 #if METAL_DEBUG_OPTIONS
 	if(GetCommandList().GetCommandQueue().GetRuntimeDebuggingLevel() >= EMetalDebugLevelValidation)
@@ -687,7 +700,16 @@ void FMetalDeviceContext::ReleaseFence(mtlpp::Fence Fence)
 	}
 #endif
 	
-	ReleaseObject([Fence.GetPtr() retain]);
+	if (GIsMetalInitialized) // @todo zebra: there seems to be some race condition at exit when the framerate is very low
+	{
+		check(Fence);
+		FreeListMutex.Lock();
+		if(!FenceFreeList.Contains(Fence))
+		{
+			FenceFreeList.Add(Fence);
+		}
+		FreeListMutex.Unlock();
+	}
 }
 
 FMetalTexture FMetalDeviceContext::CreateTexture(FMetalSurface* Surface, mtlpp::TextureDescriptor Descriptor)
@@ -733,10 +755,10 @@ void FMetalDeviceContext::ReleaseBuffer(FMetalBuffer& Buffer)
 
 struct FMetalRHICommandUpdateFence final : public FRHICommand<FMetalRHICommandUpdateFence>
 {
-	FMetalFence Fence;
+	FMetalFence* Fence;
 	uint32 Num;
 	
-	FORCEINLINE_DEBUGGABLE FMetalRHICommandUpdateFence(FMetalFence const& InFence, uint32 InNum)
+	FORCEINLINE_DEBUGGABLE FMetalRHICommandUpdateFence(FMetalFence* InFence, uint32 InNum)
 	: Fence(InFence)
 	, Num(InNum)
 	{
@@ -777,8 +799,8 @@ FMetalRHICommandContext* FMetalDeviceContext::AcquireContext(int32 NewIndex, int
 	EndLabel = [NSString stringWithFormat:@"End Parallel Context Index %d Num %d", NewIndex, NewNum];
 #endif
 	
-	FMetalFence StartFence(NewIndex == 0 ? CommandList.GetCommandQueue().CreateFence(StartLabel) : ParallelFences[NewIndex - 1]);
-	FMetalFence EndFence(CommandList.GetCommandQueue().CreateFence(EndLabel));
+	FMetalFence* StartFence(NewIndex == 0 ? CommandList.GetCommandQueue().CreateFence(StartLabel) : ParallelFences[NewIndex - 1]);
+	FMetalFence* EndFence(CommandList.GetCommandQueue().CreateFence(EndLabel));
 	ParallelFences[NewIndex] = EndFence;
 	
 	// Give the context the fences so that we can properly order the parallel contexts.
@@ -951,19 +973,19 @@ void FMetalContext::MakeCurrent(FMetalContext* Context)
 }
 #endif
 
-void FMetalContext::SetParallelPassFences(mtlpp::Fence Start, mtlpp::Fence End)
+void FMetalContext::SetParallelPassFences(FMetalFence* Start, FMetalFence* End)
 {
-	check(!StartFence.IsValid() && !EndFence.IsValid());
+	check(!StartFence && !EndFence);
 	StartFence = Start;
 	EndFence = End;
 }
 
-FMetalFence const& FMetalContext::GetParallelPassStartFence(void) const
+FMetalFence* FMetalContext::GetParallelPassStartFence(void) const
 {
 	return StartFence;
 }
 
-FMetalFence const& FMetalContext::GetParallelPassEndFence(void) const
+FMetalFence* FMetalContext::GetParallelPassEndFence(void) const
 {
 	return EndFence;
 }
@@ -977,13 +999,8 @@ void FMetalContext::InitFrame(bool const bImmediateContext, uint32 Index, uint32
 	// Reset cached state in the encoder
 	StateCache.Reset();
 
-	bool bStatistics = false;
-#if METAL_STATISTICS
-	bStatistics = GetCommandQueue().GetStatistics() != nullptr;
-#endif
-
 	// Sets the index of the parallel context within the pass
-	if (!bImmediateContext && !bStatistics)
+	if (!bImmediateContext && FMetalCommandQueue::SupportsFeature(EMetalFeaturesParallelRenderEncoders))
 	{
 		CommandList.SetParallelIndex(Index, Num);
 	}
@@ -999,7 +1016,8 @@ void FMetalContext::InitFrame(bool const bImmediateContext, uint32 Index, uint32
 	RenderPass.Begin(StartFence);
 	
 	// Unset the start fence, the render-pass owns it and we can consider it encoded now!
-	StartFence.Reset();
+	SafeReleaseMetalFence(StartFence);
+	StartFence = nullptr;
 	
 	// make sure first SetRenderTarget goes through
 	StateCache.InvalidateRenderTargets();
@@ -1011,7 +1029,8 @@ void FMetalContext::FinishFrame()
 	RenderPass.Update(EndFence);
 	
 	// Unset the end fence, the render-pass owns it and we can consider it encoded now!
-	EndFence.Reset();
+	SafeReleaseMetalFence(EndFence);
+	EndFence = nullptr;
 	
 	// End the render pass
 	RenderPass.End();
@@ -1497,7 +1516,8 @@ void FMetalDeviceContext::SetParallelRenderPassDescriptor(FRHISetRenderTargetsIn
 	if (!RenderPass.IsWithinParallelPass())
 	{
 		RenderPass.Begin(EndFence);
-		EndFence.Reset();
+		SafeReleaseMetalFence(EndFence);
+		EndFence = nullptr;
 		StateCache.InvalidateRenderTargets();
 		SetRenderTargetsInfo(TargetInfo, false);
 	}
@@ -1519,8 +1539,9 @@ void FMetalDeviceContext::EndParallelRenderCommandEncoding(void)
 	if (FPlatformAtomics::InterlockedDecrement(&ActiveParallelContexts) == 0)
 	{
 		RenderPass.EndRenderPass();
-		RenderPass.Begin(StartFence);
-		StartFence.Reset();
+		RenderPass.Begin(StartFence, true);
+		SafeReleaseMetalFence(StartFence);
+		StartFence = nullptr;
 		FPlatformAtomics::AtomicStore(&NumParallelContextsInPass, 0);
 	}
 }
@@ -1572,7 +1593,7 @@ public:
 			
 			if (Index == (Num - 1))
 			{
-				mtlpp::Fence Fence = CmdContext->GetInternalContext().GetParallelPassEndFence();
+				FMetalFence* Fence = CmdContext->GetInternalContext().GetParallelPassEndFence();
 				GetMetalDeviceContext().SetParallelPassFences(Fence, nil);
 			}
 
diff --git a/Engine/Source/Runtime/Apple/MetalRHI/Private/MetalContext.h b/Engine/Source/Runtime/Apple/MetalRHI/Private/MetalContext.h
--- a/Engine/Source/Runtime/Apple/MetalRHI/Private/MetalContext.h
+++ b/Engine/Source/Runtime/Apple/MetalRHI/Private/MetalContext.h
@@ -116,9 +116,9 @@ public:
 	static FMetalContext* GetCurrentContext();
 #endif
 	
-	void SetParallelPassFences(mtlpp::Fence Start, mtlpp::Fence End);
-	FMetalFence const& GetParallelPassStartFence(void) const;
-	FMetalFence const& GetParallelPassEndFence(void) const;
+	void SetParallelPassFences(FMetalFence* Start, FMetalFence* End);
+	FMetalFence* GetParallelPassStartFence(void) const;
+	FMetalFence* GetParallelPassEndFence(void) const;
 	
 	void InitFrame(bool const bImmediateContext, uint32 Index, uint32 Num);
 	void FinishFrame();
@@ -146,10 +146,10 @@ protected:
 	TSharedPtr<FMetalQueryBufferPool, ESPMode::ThreadSafe> QueryBuffer;
 	
 	/** Initial fence to wait on for parallel contexts */
-	FMetalFence StartFence;
+	FMetalFence* StartFence;
 	
 	/** Fence to update at the end for parallel contexts */
-	FMetalFence EndFence;
+	FMetalFence* EndFence;
 	
 #if ENABLE_METAL_GPUPROFILE
 	/** the slot to store a per-thread context ref */
@@ -182,10 +182,10 @@ public:
 	void ReleaseObject(id Object);
 	void ReleaseTexture(FMetalSurface* Surface, FMetalTexture& Texture);
 	void ReleaseTexture(FMetalTexture& Texture);
-	void ReleaseFence(mtlpp::Fence Fence);
+	void ReleaseFence(FMetalFence* Fence);
 	
 	void BeginFrame();
-	void FlushFreeList();
+	void FlushFreeList(bool const bFlushFences = true);
 	void ClearFreeList();
 	void DrainHeap();
 	void EndFrame();
@@ -238,6 +238,7 @@ private:
 	/** Free lists for releasing objects only once it is safe to do so */
 	TSet<FMetalBuffer> UsedBuffers;
 	TSet<FMetalTexture> UsedTextures;
+	TSet<FMetalFence*> FenceFreeList;
 	TSet<id> ObjectFreeList;
 	struct FMetalDelayedFreeList
 	{
@@ -245,6 +246,7 @@ private:
 		TArray<mtlpp::CommandBufferFence> Fences;
 		TSet<FMetalBuffer> UsedBuffers;
 		TSet<FMetalTexture> UsedTextures;
+		TSet<FMetalFence*> FenceFreeList;
 		TSet<id> ObjectFreeList;
 #if METAL_DEBUG_OPTIONS
 		int32 DeferCount;
@@ -254,7 +256,7 @@ private:
 	
 #if METAL_DEBUG_OPTIONS
 	/** The list of fences for the current frame */
-	TArray<mtlpp::Fence> FrameFences;
+	TArray<FMetalFence*> FrameFences;
     
     FCriticalSection ActiveBuffersMutex;
     
@@ -266,7 +268,7 @@ private:
 	TLockFreePointerListLIFO<FMetalRHICommandContext> ParallelContexts;
 	
 	/** Fences for parallel execution */
-	TArray<mtlpp::Fence> ParallelFences;
+	TArray<FMetalFence*> ParallelFences;
 	
 	/** Critical section for FreeList */
 	FCriticalSection FreeListMutex;
diff --git a/Engine/Source/Runtime/Apple/MetalRHI/Private/MetalFence.cpp b/Engine/Source/Runtime/Apple/MetalRHI/Private/MetalFence.cpp
--- a/Engine/Source/Runtime/Apple/MetalRHI/Private/MetalFence.cpp
+++ b/Engine/Source/Runtime/Apple/MetalRHI/Private/MetalFence.cpp
@@ -98,9 +98,89 @@ APPLE_PLATFORM_OBJECT_ALLOC_OVERRIDES(FMetalDebugFence)
 #if METAL_DEBUG_OPTIONS
 void FMetalFence::Validate(void) const
 {
-	if (GetMetalDeviceContext().GetCommandQueue().GetRuntimeDebuggingLevel() >= EMetalDebugLevelValidation && Object)
+	if (GetMetalDeviceContext().GetCommandQueue().GetRuntimeDebuggingLevel() >= EMetalDebugLevelValidation && GetPtr())
 	{
-		[(FMetalDebugFence*)Object.GetPtr() validate];
+		[(FMetalDebugFence*)GetPtr() validate];
 	}
 }
 #endif
+
+FMetalFencePool FMetalFencePool::sSelf;
+#if METAL_DEBUG_OPTIONS
+extern int32 GMetalRuntimeDebugLevel;
+#endif
+
+void FMetalFencePool::Initialise(mtlpp::Device const& InDevice)
+{
+	Device = InDevice;
+	for (int32 i = 0; i < FMetalFencePool::NumFences; i++)
+	{
+#if METAL_DEBUG_OPTIONS
+		if (GMetalRuntimeDebugLevel >= EMetalDebugLevelValidation)
+		{
+			FMetalDebugFence* Fence = [[FMetalDebugFence new] autorelease];
+			Fence.Inner = Device.NewFence();
+			FMetalFence* F = new FMetalFence(Fence);
+			Fences.Add(F);
+			Lifo.Push(F);
+		}
+		else
+#endif
+		{
+			FMetalFence* F = new FMetalFence(Device.NewFence());
+			Fences.Add(F);
+			Lifo.Push(F);
+		}
+	}
+	Count = FMetalFencePool::NumFences;
+}
+
+FMetalFence* FMetalFencePool::AllocateFence()
+{
+	FMetalFence* Fence = Lifo.Pop();
+	if (Fence)
+	{
+		FPlatformAtomics::InterlockedDecrement(&Count);
+		{
+			FScopeLock Lock(&Mutex);
+			check(Fences.Contains(Fence));
+			Fences.Remove(Fence);
+		}
+	}
+	check(Fence);
+	Fence->Reset();
+	return Fence;
+}
+
+void FMetalFence::ValidateUsage(FMetalFence* InFence)
+{
+	if (InFence)
+	{
+		if (InFence->NumWrites(mtlpp::RenderStages::Vertex) != InFence->NumWaits(mtlpp::RenderStages::Vertex))
+		{
+			UE_LOG(LogMetal, Warning, TEXT("%p (%s) writes %d waits %d"), InFence, *FString(InFence->GetLabel()), (uint32)InFence->NumWrites(mtlpp::RenderStages::Vertex), (uint32)InFence->NumWaits(mtlpp::RenderStages::Vertex));
+		}
+		if (InFence->NumWrites(mtlpp::RenderStages::Fragment) != InFence->NumWaits(mtlpp::RenderStages::Fragment))
+		{
+			UE_LOG(LogMetal, Warning, TEXT("%p (%s) writes %d waits %d"), InFence, *FString(InFence->GetLabel()), (uint32)InFence->NumWrites(mtlpp::RenderStages::Fragment), (uint32)InFence->NumWaits(mtlpp::RenderStages::Fragment));
+		}
+	}
+}
+
+void FMetalFencePool::ReleaseFence(FMetalFence* const InFence)
+{
+	if (InFence)
+	{
+		FMetalFence::ValidateUsage(InFence);
+		
+		{
+			FScopeLock Lock(&Mutex);
+			check(!Fences.Contains(InFence));
+			Fences.Add(InFence);
+		}
+		
+		FPlatformAtomics::InterlockedIncrement(&Count);
+		check(Count <= FMetalFencePool::NumFences);
+		Lifo.Push(InFence);
+	}
+}
diff --git a/Engine/Source/Runtime/Apple/MetalRHI/Private/MetalFence.h b/Engine/Source/Runtime/Apple/MetalRHI/Private/MetalFence.h
--- a/Engine/Source/Runtime/Apple/MetalRHI/Private/MetalFence.h
+++ b/Engine/Source/Runtime/Apple/MetalRHI/Private/MetalFence.h
@@ -61,18 +61,31 @@
 -(void) waitForFence:(id <MTLFence>)fence beforeStages:(MTLRenderStages)stages;
 @end
 
-class FMetalFence
+class FMetalFence : public mtlpp::Fence
 {
+	enum
+	{
+		NumFenceStages = 2
+	};
 public:
 	FMetalFence()
+	: mtlpp::Fence(nil)
 	{
-		
+		for (uint32 i = 0; i < NumFenceStages; i++)
+		{
+			Writes[i] = 0;
+			Waits[i] = 0;
+		}
 	}
 	
 	explicit FMetalFence(mtlpp::Fence Obj)
-	: Object(Obj)
+	: mtlpp::Fence(Obj)
 	{
-		
+		for (uint32 i = 0; i < NumFenceStages; i++)
+		{
+			Writes[i] = 0;
+			Waits[i] = 0;
+		}
 	}
 	
 	explicit FMetalFence(FMetalFence const& Other)
@@ -82,37 +95,35 @@ public:
 	
 	~FMetalFence()
 	{
-		SafeReleaseMetalFence(Object);
 	}
 	
 	FMetalFence& operator=(FMetalFence const& Other)
 	{
 		if (&Other != this)
 		{
-			operator=(Other.Object);
+			mtlpp::Fence::operator=(Other);
 		}
 		return *this;
 	}
 	
 	FMetalFence& operator=(mtlpp::Fence const& Other)
 	{
-		if (Other.GetPtr() != Object.GetPtr())
+		if (Other.GetPtr() != GetPtr())
 		{
 			METAL_DEBUG_OPTION(Validate());
-			SafeReleaseMetalFence(Object);
-			Object = Other;
+			mtlpp::Fence::operator=(Other);
 		}
 		return *this;
 	}
 	
-	operator mtlpp::Fence() const
+	FMetalFence& operator=(mtlpp::Fence&& Other)
 	{
-		return Object;
-	}
-	
-	mtlpp::Fence operator*() const
-	{
-		return Object;
+		if (Other.GetPtr() != GetPtr())
+		{
+			METAL_DEBUG_OPTION(Validate());
+			mtlpp::Fence::operator=(Other);
+		}
+		return *this;
 	}
 	
 #if METAL_DEBUG_OPTIONS
@@ -121,14 +132,74 @@ public:
 	
 	void Reset(void)
 	{
-		Object = nil;
+		for (uint32 i = 0; i < NumFenceStages; i++)
+		{
+			Writes[i] = 0;
+			Waits[i] = 0;
+		}
+	}
+	
+	void Write(mtlpp::RenderStages Stage)
+	{
+		Writes[(uint32)Stage - 1u]++;
+	}
+	
+	void Wait(mtlpp::RenderStages Stage)
+	{
+		Waits[(uint32)Stage - 1u]++;
+	}
+	
+	int8 NumWrites(mtlpp::RenderStages Stage) const
+	{
+		return Writes[(uint32)Stage - 1u];
+	}
+	
+	int8 NumWaits(mtlpp::RenderStages Stage) const
+	{
+		return Waits[(uint32)Stage - 1u];
 	}
 	
-	bool IsValid() const
+	bool NeedsWrite(mtlpp::RenderStages Stage) const
 	{
-		return Object.GetPtr() != nil;
+		return Writes[(uint32)Stage - 1u] == 0 || (Waits[(uint32)Stage - 1u] > Writes[(uint32)Stage - 1u]);
 	}
 	
+	bool NeedsWait(mtlpp::RenderStages Stage) const
+	{
+		return Waits[(uint32)Stage - 1u] == 0 || (Writes[(uint32)Stage - 1u] > Waits[(uint32)Stage - 1u]);
+	}
+	
+	static void ValidateUsage(FMetalFence* Fence);
+	
+private:
+	int8 Writes[NumFenceStages];
+	int8 Waits[NumFenceStages];
+};
+
+class FMetalFencePool
+{
+	enum
+	{
+		NumFences = 2048
+	};
+public:
+	FMetalFencePool() {}
+	
+	static FMetalFencePool& Get()
+	{
+		return sSelf;
+	}
+	
+	void Initialise(mtlpp::Device const& InDevice);
+	
+	FMetalFence* AllocateFence();
+	void ReleaseFence(FMetalFence* const InFence);
+	
 private:
-	mtlpp::Fence Object;
+	int32 Count;
+	mtlpp::Device Device;
+	TSet<FMetalFence*> Fences;
+	FCriticalSection Mutex;
+	TLockFreePointerListLIFO<FMetalFence> Lifo;
+	static FMetalFencePool sSelf;
 };
diff --git a/Engine/Source/Runtime/Apple/MetalRHI/Private/MetalRHI.cpp b/Engine/Source/Runtime/Apple/MetalRHI/Private/MetalRHI.cpp
--- a/Engine/Source/Runtime/Apple/MetalRHI/Private/MetalRHI.cpp
+++ b/Engine/Source/Runtime/Apple/MetalRHI/Private/MetalRHI.cpp
@@ -1104,7 +1104,7 @@ void FMetalDynamicRHI::RHIFlushResources()
 {
 	@autoreleasepool {
 		((FMetalDeviceContext*)ImmediateContext.Context)->DrainHeap();
-		((FMetalDeviceContext*)ImmediateContext.Context)->FlushFreeList();
+		((FMetalDeviceContext*)ImmediateContext.Context)->FlushFreeList(false);
 		ImmediateContext.Context->SubmitCommandBufferAndWait();
 		((FMetalDeviceContext*)ImmediateContext.Context)->ClearFreeList();
 		ImmediateContext.Context->GetCurrentState().Reset();
diff --git a/Engine/Source/Runtime/Apple/MetalRHI/Private/MetalRHIContext.cpp b/Engine/Source/Runtime/Apple/MetalRHI/Private/MetalRHIContext.cpp
--- a/Engine/Source/Runtime/Apple/MetalRHI/Private/MetalRHIContext.cpp
+++ b/Engine/Source/Runtime/Apple/MetalRHI/Private/MetalRHIContext.cpp
@@ -50,14 +50,14 @@ void SafeReleaseMetalBuffer(FMetalBuffer& Buffer)
 	}
 }
 
-void SafeReleaseMetalFence(id Object)
+void SafeReleaseMetalFence(FMetalFence* Object)
 {
 	if(GIsMetalInitialized && GDynamicRHI && Object)
 	{
 		FMetalRHICommandContext* Context = static_cast<FMetalRHICommandContext*>(RHIGetDefaultContext());
 		if(Context)
 		{
-			((FMetalDeviceContext&)Context->GetInternalContext()).ReleaseFence((id<MTLFence>)Object);
+			((FMetalDeviceContext&)Context->GetInternalContext()).ReleaseFence(Object);
 			return;
 		}
 	}
diff --git a/Engine/Source/Runtime/Apple/MetalRHI/Private/MetalRHIPrivate.h b/Engine/Source/Runtime/Apple/MetalRHI/Private/MetalRHIPrivate.h
--- a/Engine/Source/Runtime/Apple/MetalRHI/Private/MetalRHIPrivate.h
+++ b/Engine/Source/Runtime/Apple/MetalRHI/Private/MetalRHIPrivate.h
@@ -146,7 +146,7 @@ void SafeReleaseMetalTexture(FMetalTexture& Object);
 void SafeReleaseMetalBuffer(FMetalBuffer& Buffer);
 
 // Safely release a fence, correctly handling cases where fences aren't supported or the debug implementation is used.
-void SafeReleaseMetalFence(id Object);
+void SafeReleaseMetalFence(class FMetalFence* Object);
 
 // Access the underlying surface object from any kind of texture
 FMetalSurface* GetMetalSurfaceFromRHITexture(FRHITexture* Texture);
diff --git a/Engine/Source/Runtime/Apple/MetalRHI/Private/MetalRenderPass.cpp b/Engine/Source/Runtime/Apple/MetalRHI/Private/MetalRenderPass.cpp
--- a/Engine/Source/Runtime/Apple/MetalRHI/Private/MetalRenderPass.cpp
+++ b/Engine/Source/Runtime/Apple/MetalRHI/Private/MetalRenderPass.cpp
@@ -57,13 +57,27 @@ FMetalRenderPass::~FMetalRenderPass(void)
 {
 	check(!CurrentEncoder.GetCommandBuffer());
 	check(!PrologueEncoder.GetCommandBuffer());
-	check(!PassStartFence.IsValid());
+	check(!PassStartFence);
 }
 
-void FMetalRenderPass::Begin(mtlpp::Fence Fence)
+void FMetalRenderPass::Begin(FMetalFence* Fence, bool const bParallelBegin)
 {
-	check(!PassStartFence.IsValid());
-	PassStartFence = Fence;
+	if (!bParallelBegin || !FMetalCommandQueue::SupportsFeature(EMetalFeaturesParallelRenderEncoders))
+	{
+		check(!PassStartFence || Fence == nullptr);
+		if (Fence)
+		{
+			PassStartFence = Fence;
+		}
+	}
+	else
+	{
+		check(!CurrentEncoderFence || Fence == nullptr);
+		if (Fence)
+		{
+			CurrentEncoderFence = Fence;
+		}
+	}
 	
 	if (!CmdList.IsParallel() && !CurrentEncoder.GetCommandBuffer())
 	{
@@ -72,24 +86,29 @@ void FMetalRenderPass::Begin(mtlpp::Fence Fence)
 	}
 }
 
-void FMetalRenderPass::Wait(mtlpp::Fence Fence)
+void FMetalRenderPass::Wait(FMetalFence* Fence)
 {
 	if (Fence)
 	{
 		if (PrologueEncoder.IsBlitCommandEncoderActive() || PrologueEncoder.IsComputeCommandEncoderActive())
 		{
 			PrologueEncoder.WaitForFence(Fence);
+			FMetalFence::ValidateUsage(Fence);
 		}
-		if (CurrentEncoder.IsRenderCommandEncoderActive() || CurrentEncoder.IsBlitCommandEncoderActive() || CurrentEncoder.IsComputeCommandEncoderActive())
+		else if (CurrentEncoder.IsRenderCommandEncoderActive() || CurrentEncoder.IsBlitCommandEncoderActive() || CurrentEncoder.IsComputeCommandEncoderActive())
 		{
 			CurrentEncoder.WaitForFence(Fence);
+			FMetalFence::ValidateUsage(Fence);
 		}
+        else
+        {
+            PassStartFence = Fence;
+        }
 	}
-	PassStartFence = Fence;
 		
 }
 
-void FMetalRenderPass::Update(mtlpp::Fence Fence)
+void FMetalRenderPass::Update(FMetalFence* Fence)
 {
 	if (Fence)
 	{
@@ -97,15 +116,20 @@ void FMetalRenderPass::Update(mtlpp::Fence Fence)
 		// the higher-level can generate empty contexts but we have no sane way to deal with that.
 		if (!CurrentEncoder.IsRenderCommandEncoderActive() && !CurrentEncoder.IsBlitCommandEncoderActive() && !CurrentEncoder.IsComputeCommandEncoderActive())
 		{
-			ConditionalSwitchToCompute();
+			ConditionalSwitchToBlit(true);
 		}
 		CurrentEncoder.UpdateFence(Fence);
 		State.FlushVisibilityResults(CurrentEncoder);
-		CurrentEncoderFence = CurrentEncoder.EndEncoding();
+		FMetalFence* NewFence = CurrentEncoder.EndEncoding();
+		check(!CurrentEncoderFence || !NewFence);
+		if (NewFence)
+		{
+			CurrentEncoderFence = NewFence;
+		}
 	}
 }
 
-mtlpp::Fence FMetalRenderPass::Submit(EMetalSubmitFlags Flags)
+FMetalFence* FMetalRenderPass::Submit(EMetalSubmitFlags Flags)
 {
 	if (CurrentEncoder.GetCommandBuffer() || (Flags & EMetalSubmitFlagsAsyncCommandBuffer))
 	{
@@ -206,17 +230,20 @@ void FMetalRenderPass::BeginRenderPass(mtlpp::RenderPassDescriptor RenderPass)
 	if (!GMetalDeferRenderPasses || !State.CanRestartRenderPass() || CmdList.IsParallel())
 	{
 		CurrentEncoder.BeginRenderCommandEncoding();
-		if (PassStartFence.IsValid())
+		if (PassStartFence)
 		{
 			CurrentEncoder.WaitForFence(PassStartFence);
+			PassStartFence = nullptr;
 		}
-		if (CurrentEncoderFence.IsValid())
+		if (CurrentEncoderFence)
 		{
 			CurrentEncoder.WaitForFence(CurrentEncoderFence);
+			CurrentEncoderFence = nullptr;
 		}
-		if (PrologueEncoderFence.IsValid())
+		if (PrologueEncoderFence)
 		{
 			CurrentEncoder.WaitForFence(PrologueEncoderFence);
+			PrologueEncoderFence = nullptr;
 		}
 		State.SetRenderStoreActions(CurrentEncoder, false);
 		check(CurrentEncoder.IsRenderCommandEncoderActive());
@@ -314,17 +341,20 @@ void FMetalRenderPass::RestartRenderPass(mtlpp::RenderPassDescriptor RenderPass)
 	
 	CurrentEncoder.SetRenderPassDescriptor(RenderPassDesc);
 	CurrentEncoder.BeginRenderCommandEncoding();
-	if (CurrentEncoderFence.IsValid())
+	if (CurrentEncoderFence)
 	{
 		CurrentEncoder.WaitForFence(CurrentEncoderFence);
+		CurrentEncoderFence = nullptr;
 	}
-	else if (PassStartFence.IsValid())
+	else if (PassStartFence)
 	{
 		CurrentEncoder.WaitForFence(PassStartFence);
+		PassStartFence = nullptr;
 	}
-	if (PrologueEncoderFence.IsValid())
+	if (PrologueEncoderFence)
 	{
 		CurrentEncoder.WaitForFence(PrologueEncoderFence);
+		PrologueEncoderFence = nullptr;
 	}
 	State.SetRenderStoreActions(CurrentEncoder, false);
 	
@@ -708,7 +738,7 @@ void FMetalRenderPass::DispatchIndirect(FMetalVertexBuffer* ArgumentBuffer, uint
 	ConditionalSubmit();
 }
 
-mtlpp::Fence FMetalRenderPass::EndRenderPass(void)
+FMetalFence* FMetalRenderPass::EndRenderPass(void)
 {
 	if (bWithinRenderPass)
 	{
@@ -963,27 +993,42 @@ void FMetalRenderPass::AsyncGenerateMipmapsForTexture(FMetalTexture const& Textu
 	METAL_DEBUG_LAYER(EMetalDebugLevelFastValidation, PrologueEncoder.GetBlitCommandEncoderDebugging().GenerateMipmaps(Texture));
 }
 
-mtlpp::Fence FMetalRenderPass::End(void)
+FMetalFence* FMetalRenderPass::End(void)
 {
 	// EndEncoding should provide the encoder fence...
 	if (PrologueEncoder.IsBlitCommandEncoderActive() || PrologueEncoder.IsComputeCommandEncoderActive())
 	{
 		PrologueEncoderFence = PrologueEncoder.EndEncoding();
 	}
-	if (CurrentEncoder.IsRenderCommandEncoderActive() || CurrentEncoder.IsBlitCommandEncoderActive() || CurrentEncoder.IsComputeCommandEncoderActive())
+	
+	if (CmdList.IsImmediate() && IsWithinParallelPass() && CurrentEncoder.IsParallelRenderCommandEncoderActive())
+	{
+		State.SetRenderStoreActions(CurrentEncoder, false);
+		CurrentEncoder.EndEncoding();
+		
+		ConditionalSwitchToBlit(true);
+		CurrentEncoder.WaitForFence(PassStartFence);
+		CurrentEncoder.WaitForFence(CurrentEncoderFence);
+		CurrentEncoderFence = CurrentEncoder.EndEncoding();
+	}
+	else if (CurrentEncoder.IsRenderCommandEncoderActive() || CurrentEncoder.IsBlitCommandEncoderActive() || CurrentEncoder.IsComputeCommandEncoderActive())
 	{
 		State.FlushVisibilityResults(CurrentEncoder);
-
-		if (CmdList.IsImmediate() && IsWithinParallelPass())
-		{
-			State.SetRenderStoreActions(CurrentEncoder, false);
-		}
+		CurrentEncoder.WaitForFence(PassStartFence);
+		CurrentEncoder.WaitForFence(CurrentEncoderFence);
+		CurrentEncoderFence = CurrentEncoder.EndEncoding();
+	}
+	else if (PassStartFence)
+	{
+		ConditionalSwitchToBlit(true);
+		CurrentEncoder.WaitForFence(PassStartFence);
+		CurrentEncoder.WaitForFence(CurrentEncoderFence);
 		CurrentEncoderFence = CurrentEncoder.EndEncoding();
 	}
 	
-	State.SetRenderTargetsActive(false);
+	PassStartFence = nullptr;
 	
-	PassStartFence.Reset();
+	State.SetRenderTargetsActive(false);
 	
 	RenderPassDesc = nil;
 	bWithinRenderPass = false;
@@ -1129,22 +1174,24 @@ void FMetalRenderPass::ConditionalSwitchToTessellation(void)
 		}
 		PrologueEncoder.BeginComputeCommandEncoding();
 		
-		// check(PrologueEncoderFence.IsValid() || CurrentEncoderFence.IsValid() || PassStartFence.IsValid());
-		if (PrologueEncoderFence.IsValid())
+		if (PrologueEncoderFence)
 		{
 			PrologueEncoder.WaitForFence(PrologueEncoderFence);
+			PrologueEncoderFence = nullptr;
 		}
-		if (CurrentEncoderFence.IsValid())
+		if (CurrentEncoderFence)
 		{
 			PrologueEncoder.WaitForFence(CurrentEncoderFence);
+			CurrentEncoderFence = nullptr;
 		}
-		if (PassStartFence.IsValid())
+		if (PassStartFence)
 		{
 			PrologueEncoder.WaitForFence(PassStartFence);
+			PassStartFence = nullptr;
 		}
 		PrologueEncoderFence = PrologueEncoder.GetEncoderFence();
 #if METAL_DEBUG_OPTIONS
-		if (GetEmitDrawEvents() && PrologueEncoderFence.IsValid())
+		if (GetEmitDrawEvents() && PrologueEncoderFence)
 		{
 			if (CmdList.GetCommandQueue().GetRuntimeDebuggingLevel() >= EMetalDebugLevelValidation)
 			{
@@ -1165,6 +1212,7 @@ void FMetalRenderPass::ConditionalSwitchToTessellation(void)
 	else if (bCreatePrologueEncoder)
 	{
 		CurrentEncoder.WaitForFence(PrologueEncoderFence);
+		PrologueEncoderFence = nullptr;
 	}
 	
 	check(CurrentEncoder.IsRenderCommandEncoderActive());
@@ -1191,20 +1239,22 @@ void FMetalRenderPass::ConditionalSwitchToCompute(void)
 	{
 		State.SetStateDirty();
 		CurrentEncoder.BeginComputeCommandEncoding();
-		if (CurrentEncoderFence.IsValid())
+		if (CurrentEncoderFence)
 		{
 			CurrentEncoder.WaitForFence(CurrentEncoderFence);
+			CurrentEncoderFence = nullptr;
 		}
-		if (PrologueEncoderFence.IsValid())
+		if (PrologueEncoderFence)
 		{
 			CurrentEncoder.WaitForFence(PrologueEncoderFence);
+			PrologueEncoderFence = nullptr;
 		}
 	}
 	
 	check(CurrentEncoder.IsComputeCommandEncoderActive());
 }
 
-void FMetalRenderPass::ConditionalSwitchToBlit(void)
+void FMetalRenderPass::ConditionalSwitchToBlit(bool const bSuppressFence)
 {
 	check(CurrentEncoder.GetCommandBuffer());
 	check(!CurrentEncoder.IsParallel());
@@ -1222,14 +1272,16 @@ void FMetalRenderPass::ConditionalSwitchToBlit(void)
 	
 	if (!CurrentEncoder.IsBlitCommandEncoderActive())
 	{
-		CurrentEncoder.BeginBlitCommandEncoding();
-		if (CurrentEncoderFence.IsValid())
+		CurrentEncoder.BeginBlitCommandEncoding(bSuppressFence);
+		if (CurrentEncoderFence)
 		{
 			CurrentEncoder.WaitForFence(CurrentEncoderFence);
+			CurrentEncoderFence = nullptr;
 		}
-		if (PrologueEncoderFence.IsValid())
+		if (PrologueEncoderFence)
 		{
 			CurrentEncoder.WaitForFence(PrologueEncoderFence);
+			PrologueEncoderFence = nullptr;
 		}
 	}
 	
@@ -1250,21 +1302,19 @@ void FMetalRenderPass::ConditionalSwitchToAsyncBlit(void)
 			PrologueEncoder.StartCommandBuffer();
 		}
 		PrologueEncoder.BeginBlitCommandEncoding();
-		if (PrologueEncoderFence.IsValid())
+		if (PrologueEncoderFence)
 		{
 			PrologueEncoder.WaitForFence(PrologueEncoderFence);
+			PrologueEncoderFence = nullptr;
 		}
-		if (CurrentEncoderFence.IsValid())
-		{
-			PrologueEncoder.WaitForFence(CurrentEncoderFence);
-		}
-		if (PassStartFence.IsValid())
+		if (PassStartFence)
 		{
 			PrologueEncoder.WaitForFence(PassStartFence);
+			PassStartFence = nullptr;
 		}
 		PrologueEncoderFence = PrologueEncoder.GetEncoderFence();
 #if METAL_DEBUG_OPTIONS
-		if (GetEmitDrawEvents() && PrologueEncoderFence.IsValid())
+		if (GetEmitDrawEvents() && PrologueEncoderFence)
 		{
 			if (CmdList.GetCommandQueue().GetRuntimeDebuggingLevel() >= EMetalDebugLevelValidation)
 			{
@@ -1281,6 +1331,7 @@ void FMetalRenderPass::ConditionalSwitchToAsyncBlit(void)
 		if (CurrentEncoder.IsRenderCommandEncoderActive() || CurrentEncoder.IsComputeCommandEncoderActive() || CurrentEncoder.IsBlitCommandEncoderActive())
 		{
 			CurrentEncoder.WaitForFence(PrologueEncoderFence);
+			PrologueEncoderFence = nullptr;
 		}
 	}
 	
diff --git a/Engine/Source/Runtime/Apple/MetalRHI/Private/MetalRenderPass.h b/Engine/Source/Runtime/Apple/MetalRHI/Private/MetalRenderPass.h
--- a/Engine/Source/Runtime/Apple/MetalRHI/Private/MetalRenderPass.h
+++ b/Engine/Source/Runtime/Apple/MetalRHI/Private/MetalRenderPass.h
@@ -22,11 +22,11 @@ public:
 	~FMetalRenderPass(void);
 	
 #pragma mark -
-    void Begin(mtlpp::Fence Fence);
+    void Begin(FMetalFence* Fence, bool const bParallelBegin = false);
 	
-	void Wait(mtlpp::Fence Fence);
+	void Wait(FMetalFence* Fence);
 
-	void Update(mtlpp::Fence Fence);
+	void Update(FMetalFence* Fence);
 	
     void BeginParallelRenderPass(mtlpp::RenderPassDescriptor RenderPass, uint32 NumParallelContextsInPass);
 
@@ -52,7 +52,7 @@ public:
     
     void DispatchIndirect(FMetalVertexBuffer* ArgumentBufferRHI, uint32 ArgumentOffset);
     
-    mtlpp::Fence EndRenderPass(void);
+    FMetalFence* EndRenderPass(void);
     
     void CopyFromTextureToBuffer(FMetalTexture const& Texture, uint32 sourceSlice, uint32 sourceLevel, mtlpp::Origin sourceOrigin, mtlpp::Size sourceSize, FMetalBuffer const& toBuffer, uint32 destinationOffset, uint32 destinationBytesPerRow, uint32 destinationBytesPerImage, mtlpp::BlitOption options);
     
@@ -78,9 +78,9 @@ public:
 	
 	void AsyncGenerateMipmapsForTexture(FMetalTexture const& Texture);
 	
-    mtlpp::Fence Submit(EMetalSubmitFlags SubmissionFlags);
+    FMetalFence* Submit(EMetalSubmitFlags SubmissionFlags);
     
-    mtlpp::Fence End(void);
+    FMetalFence* End(void);
 	
 	void InsertCommandBufferFence(FMetalCommandBufferFence& Fence, mtlpp::CommandBufferHandler Handler);
 	
@@ -142,7 +142,7 @@ private:
     void ConditionalSwitchToRender(void);
     void ConditionalSwitchToTessellation(void);
     void ConditionalSwitchToCompute(void);
-	void ConditionalSwitchToBlit(void);
+	void ConditionalSwitchToBlit(bool const bSuppressFence = false);
 	void ConditionalSwitchToAsyncBlit(void);
 	
     void PrepareToRender(uint32 PrimType);
@@ -172,9 +172,9 @@ private:
 	// Disjoint ranges *are* permitted!
 	TMap<id<MTLBuffer>, TArray<NSRange>> OutstandingBufferUploads;
     
-    FMetalFence PassStartFence;
-    FMetalFence CurrentEncoderFence;
-    FMetalFence PrologueEncoderFence;
+    FMetalFence* PassStartFence;
+    FMetalFence* CurrentEncoderFence;
+    FMetalFence* PrologueEncoderFence;
     
     mtlpp::RenderPassDescriptor RenderPassDesc;
     
diff --git a/Engine/Source/Runtime/Apple/MetalRHI/Private/MetalStateCache.cpp b/Engine/Source/Runtime/Apple/MetalRHI/Private/MetalStateCache.cpp
--- a/Engine/Source/Runtime/Apple/MetalRHI/Private/MetalStateCache.cpp
+++ b/Engine/Source/Runtime/Apple/MetalRHI/Private/MetalStateCache.cpp
@@ -1856,9 +1856,9 @@ void FMetalStateCache::FlushVisibilityResults(FMetalCommandEncoder& CommandEncod
 #if PLATFORM_MAC
 	if(VisibilityResults && VisibilityResults->Buffer && VisibilityResults->Buffer.GetStorageMode() == mtlpp::StorageMode::Managed && VisibilityWritten && CommandEncoder.IsRenderCommandEncoderActive())
 	{
-		mtlpp::Fence Fence = CommandEncoder.EndEncoding();
+		FMetalFence* Fence = CommandEncoder.EndEncoding();
 		
-		CommandEncoder.BeginBlitCommandEncoding();
+        CommandEncoder.BeginBlitCommandEncoding();
 		CommandEncoder.WaitForFence(Fence);
 		
 		mtlpp::BlitCommandEncoder& Encoder = CommandEncoder.GetBlitCommandEncoder();
diff --git a/Engine/Source/Runtime/Apple/MetalRHI/Private/MetalUAV.cpp b/Engine/Source/Runtime/Apple/MetalRHI/Private/MetalUAV.cpp
--- a/Engine/Source/Runtime/Apple/MetalRHI/Private/MetalUAV.cpp
+++ b/Engine/Source/Runtime/Apple/MetalRHI/Private/MetalUAV.cpp
@@ -768,6 +768,13 @@ FComputeFenceRHIRef FMetalDynamicRHI::RHICreateComputeFence(const FName& Name)
 	}
 }
 
+void FMetalComputeFence::Write(FMetalFence* InFence)
+{
+	check(!Fence);
+	Fence = InFence;
+	FRHIComputeFence::WriteFence();
+}
+
 void FMetalComputeFence::Wait(FMetalContext& Context)
 {
 	if (Context.GetCurrentCommandBuffer())
@@ -775,6 +782,13 @@ void FMetalComputeFence::Wait(FMetalContext& Context)
 		Context.SubmitCommandsHint(EMetalSubmitFlagsNone);
 	}
 	Context.GetCurrentRenderPass().Begin(Fence);
+	Fence = nullptr;
+}
+
+void FMetalComputeFence::Reset()
+{
+	FRHIComputeFence::Reset();
+	Fence = nullptr;
 }
 
 void FMetalRHICommandContext::RHITransitionResources(EResourceTransitionAccess TransitionType, EResourceTransitionPipeline TransitionPipeline, FUnorderedAccessViewRHIParamRef* InUAVs, int32 NumUAVs, FComputeFenceRHIParamRef WriteComputeFence)
diff --git a/Engine/Source/Runtime/Apple/MetalRHI/Public/MetalResources.h b/Engine/Source/Runtime/Apple/MetalRHI/Public/MetalResources.h
--- a/Engine/Source/Runtime/Apple/MetalRHI/Public/MetalResources.h
+++ b/Engine/Source/Runtime/Apple/MetalRHI/Public/MetalResources.h
@@ -950,35 +950,29 @@ private:
 	void ResizeGlobalUniforms(uint32 TypeIndex, uint32 UniformArraySize);
 };
 
+class FMetalFence;
+
 class FMetalComputeFence : public FRHIComputeFence
 {
 public:
 	
 	FMetalComputeFence(FName InName)
 	: FRHIComputeFence(InName)
+	, Fence(nullptr)
 	{}
 	
 	virtual ~FMetalComputeFence()
 	{
 	}
 	
-	virtual void Reset() final override
-	{
-		FRHIComputeFence::Reset();
-		Fence = nil;
-	}
+	virtual void Reset() final override;
 	
-	void Write(mtlpp::Fence InFence)
-	{
-		check(Fence.GetPtr() == nil);
-		Fence = InFence;
-		FRHIComputeFence::WriteFence();
-	}
+	void Write(FMetalFence* InFence);
 	
 	void Wait(FMetalContext& Context);
 	
 private:
-	mtlpp::Fence Fence;
+	FMetalFence* Fence;
 };
 
 class FMetalGPUFence final : public FRHIGPUFence
