Added -metalunretained to force unretained references on Metal command-queues to allow investigation and testing of potential fixes.

diff --git a/Engine/Source/Runtime/Apple/MetalRHI/Private/MetalCommandQueue.cpp b/Engine/Source/Runtime/Apple/MetalRHI/Private/MetalCommandQueue.cpp
--- a/Engine/Source/Runtime/Apple/MetalRHI/Private/MetalCommandQueue.cpp
+++ b/Engine/Source/Runtime/Apple/MetalRHI/Private/MetalCommandQueue.cpp
@@ -293,9 +293,10 @@ FMetalCommandQueue::~FMetalCommandQueue(void)
 mtlpp::CommandBuffer FMetalCommandQueue::CreateCommandBuffer(void)
 {
 #if PLATFORM_MAC
-	static bool bUnretainedRefs = !FParse::Param(FCommandLine::Get(),TEXT("metalretainrefs"))
-	&& ([Device.GetName() rangeOfString:@"Nvidia" options:NSCaseInsensitiveSearch].location == NSNotFound)
-	&& ([Device.GetName() rangeOfString:@"Intel" options:NSCaseInsensitiveSearch].location == NSNotFound || FPlatformMisc::MacOSXVersionCompare(10,13,0) >= 0);
+	static bool bUnretainedRefs = FParse::Param(FCommandLine::Get(),TEXT("metalunretained"))
+	|| (!FParse::Param(FCommandLine::Get(),TEXT("metalretainrefs"))
+		&& ([Device.GetName() rangeOfString:@"Nvidia" options:NSCaseInsensitiveSearch].location == NSNotFound)
+		&& ([Device.GetName() rangeOfString:@"Intel" options:NSCaseInsensitiveSearch].location == NSNotFound || FPlatformMisc::MacOSXVersionCompare(10,13,0) >= 0));
 #else
 	static bool bUnretainedRefs = !FParse::Param(FCommandLine::Get(),TEXT("metalretainrefs"));
 #endif
