Zero-filling private buffers when using MTLHeap is causing problems - I'm still investigating that but for now don't do it.

diff --git a/Engine/Source/Runtime/Apple/MetalRHI/Private/MetalIndexBuffer.cpp b/Engine/Source/Runtime/Apple/MetalRHI/Private/MetalIndexBuffer.cpp
--- a/Engine/Source/Runtime/Apple/MetalRHI/Private/MetalIndexBuffer.cpp
+++ b/Engine/Source/Runtime/Apple/MetalRHI/Private/MetalIndexBuffer.cpp
@@ -54,7 +54,7 @@ FIndexBufferRHIRef FMetalDynamicRHI::RHICreateIndexBuffer(uint32 Stride,uint32 S
 			IndexBuffer->CPUBuffer = nil;
 		}
 
-		if (GMetalBufferZeroFill)
+		if (GMetalBufferZeroFill && !FMetalCommandQueue::SupportsFeature(EMetalFeaturesHeaps))
 		{
 			GetMetalDeviceContext().FillBuffer(IndexBuffer->Buffer, ns::Range(0, IndexBuffer->Buffer.GetLength()), 0);
 		}
@@ -116,7 +116,7 @@ struct FMetalRHICommandInitialiseIndexBuffer : public FRHICommand<FMetalRHIComma
 				Buffer->LastUpdate = GFrameNumberRenderThread;
 			}
 		}
-		else if (GMetalBufferZeroFill)
+		else if (GMetalBufferZeroFill && !FMetalCommandQueue::SupportsFeature(EMetalFeaturesHeaps))
 		{
 			GetMetalDeviceContext().FillBuffer(Buffer->Buffer, ns::Range(0, Buffer->Buffer.GetLength()), 0);
 		}
diff --git a/Engine/Source/Runtime/Apple/MetalRHI/Private/MetalStructuredBuffer.cpp b/Engine/Source/Runtime/Apple/MetalRHI/Private/MetalStructuredBuffer.cpp
--- a/Engine/Source/Runtime/Apple/MetalRHI/Private/MetalStructuredBuffer.cpp
+++ b/Engine/Source/Runtime/Apple/MetalRHI/Private/MetalStructuredBuffer.cpp
@@ -40,7 +40,7 @@ FStructuredBufferRHIRef FMetalDynamicRHI::RHICreateStructuredBuffer(uint32 Strid
 			Buffer->CPUBuffer = nil;
 		}
 
-		if (GMetalBufferZeroFill)
+		if (GMetalBufferZeroFill && !FMetalCommandQueue::SupportsFeature(EMetalFeaturesHeaps))
 		{
 			GetMetalDeviceContext().FillBuffer(Buffer->Buffer, ns::Range(0, Buffer->Buffer.GetLength()), 0);
 		}
@@ -102,7 +102,7 @@ struct FMetalRHICommandInitialiseStructuredBuffer : public FRHICommand<FMetalRHI
 			}
 		}
 #if PLATFORM_MAC
-		else if (GMetalBufferZeroFill)
+		else if (GMetalBufferZeroFill && !FMetalCommandQueue::SupportsFeature(EMetalFeaturesHeaps))
 		{
 			GetMetalDeviceContext().FillBuffer(Buffer->Buffer, ns::Range(0, Buffer->Buffer.GetLength()), 0);
 		}
diff --git a/Engine/Source/Runtime/Apple/MetalRHI/Private/MetalUniformBuffer.cpp b/Engine/Source/Runtime/Apple/MetalRHI/Private/MetalUniformBuffer.cpp
--- a/Engine/Source/Runtime/Apple/MetalRHI/Private/MetalUniformBuffer.cpp
+++ b/Engine/Source/Runtime/Apple/MetalRHI/Private/MetalUniformBuffer.cpp
@@ -31,7 +31,7 @@ struct FMetalRHICommandInitialiseUniformdBuffer : public FRHICommand<FMetalRHICo
 			LLM_SCOPE(ELLMTag::VertexBuffer);
 			SafeReleaseMetalBuffer(Buffer->CPUBuffer);
 		}
-		else if (GMetalBufferZeroFill)
+		else if (GMetalBufferZeroFill && !FMetalCommandQueue::SupportsFeature(EMetalFeaturesHeaps))
 		{
 			GetMetalDeviceContext().FillBuffer(Buffer->Buffer, ns::Range(0, Buffer->Buffer.GetLength()), 0);
 		}
diff --git a/Engine/Source/Runtime/Apple/MetalRHI/Private/MetalVertexBuffer.cpp b/Engine/Source/Runtime/Apple/MetalRHI/Private/MetalVertexBuffer.cpp
--- a/Engine/Source/Runtime/Apple/MetalRHI/Private/MetalVertexBuffer.cpp
+++ b/Engine/Source/Runtime/Apple/MetalRHI/Private/MetalVertexBuffer.cpp
@@ -536,7 +536,7 @@ FVertexBufferRHIRef FMetalDynamicRHI::RHICreateVertexBuffer(uint32 Size, uint32
 			VertexBuffer->CPUBuffer = nil;
 		}
 
-		if (GMetalBufferZeroFill)
+		if (GMetalBufferZeroFill && !FMetalCommandQueue::SupportsFeature(EMetalFeaturesHeaps))
 		{
 			GetMetalDeviceContext().FillBuffer(VertexBuffer->Buffer, ns::Range(0, VertexBuffer->Buffer.GetLength()), 0);
 		}
@@ -622,7 +622,7 @@ struct FMetalRHICommandInitialiseVertexBuffer : public FRHICommand<FMetalRHIComm
 				Buffer->LastUpdate = GFrameNumberRenderThread;
 			}
 		}
-		else if (GMetalBufferZeroFill)
+		else if (GMetalBufferZeroFill && !FMetalCommandQueue::SupportsFeature(EMetalFeaturesHeaps))
 		{
 			GetMetalDeviceContext().FillBuffer(Buffer->Buffer, ns::Range(0, Buffer->Buffer.GetLength()), 0);
 		}
