Only the real MTLBuffer 'heap' needs to call MaxAvailableSizeWithAlignment - the magazine can't fragment as all allocations are the same size. Makes traces bearable again.

diff --git a/Engine/Source/Runtime/Apple/MetalRHI/Private/MetalBuffer.cpp b/Engine/Source/Runtime/Apple/MetalRHI/Private/MetalBuffer.cpp
--- a/Engine/Source/Runtime/Apple/MetalRHI/Private/MetalBuffer.cpp
+++ b/Engine/Source/Runtime/Apple/MetalRHI/Private/MetalBuffer.cpp
@@ -722,14 +722,7 @@ NSUInteger FMetalSubBufferMagazine::GetUsedSize() const
 
 NSUInteger FMetalSubBufferMagazine::GetFreeSize() const 
 {
-	if (ParentHeap)
-	{
-		return ParentHeap.MaxAvailableSizeWithAlignment(MinAlign);
-	}
-	else
-	{
-		return GetSize() - GetUsedSize();
-	}
+	return GetSize() - GetUsedSize();
 }
 
 void FMetalSubBufferMagazine::SetLabel(const ns::String& label)
