Strictly order prologue encoders by keeping the previous prologue fence around and reissuing it as needed - this ensures that we have matched update/wait events and that the prologue encoders complete in the order we expect. One note is that with MTLFence and MTLHeap enabled *all* prologue work completes prior to the main encoding which is what I would have expected to occur previously but it is possible that the drivers were getting clever.

diff --git a/Engine/Source/Runtime/Apple/MetalRHI/Private/MetalCommandEncoder.cpp b/Engine/Source/Runtime/Apple/MetalRHI/Private/MetalCommandEncoder.cpp
--- a/Engine/Source/Runtime/Apple/MetalRHI/Private/MetalCommandEncoder.cpp
+++ b/Engine/Source/Runtime/Apple/MetalRHI/Private/MetalCommandEncoder.cpp
@@ -676,7 +676,6 @@ void FMetalCommandEncoder::UpdateFence(FMetalFence* Fence)
 	static bool bSupportsFences = CommandList.GetCommandQueue().SupportsFeature(EMetalFeaturesFences);
 	if ((bSupportsFences METAL_DEBUG_OPTION(|| CommandList.GetCommandQueue().GetRuntimeDebuggingLevel() >= EMetalDebugLevelValidation)) && Fence && Fence->NeedsWrite(mtlpp::RenderStages::Fragment))
 	{
-		check(Fence);
 		mtlpp::Fence InnerFence = METAL_DEBUG_OPTION(CommandList.GetCommandQueue().GetRuntimeDebuggingLevel() >= EMetalDebugLevelValidation ? mtlpp::Fence(((FMetalDebugFence*)Fence->GetPtr()).Inner) :) *Fence;
 		if (RenderCommandEncoder)
 		{
@@ -708,8 +707,6 @@ void FMetalCommandEncoder::WaitForFence(FMetalFence* Fence)
 	static bool bSupportsFences = CommandList.GetCommandQueue().SupportsFeature(EMetalFeaturesFences);
 	if ((bSupportsFences METAL_DEBUG_OPTION(|| CommandList.GetCommandQueue().GetRuntimeDebuggingLevel() >= EMetalDebugLevelValidation)) && Fence && Fence->NeedsWait(mtlpp::RenderStages::Fragment))
 	{
-		check(Fence);
-		
 		mtlpp::Fence InnerFence = METAL_DEBUG_OPTION(CommandList.GetCommandQueue().GetRuntimeDebuggingLevel() >= EMetalDebugLevelValidation ? mtlpp::Fence(((FMetalDebugFence*)Fence->GetPtr()).Inner) :) *Fence;
 		if (RenderCommandEncoder)
 		{
@@ -735,6 +732,46 @@ void FMetalCommandEncoder::WaitForFence(FMetalFence* Fence)
 	}
 }
 
+void FMetalCommandEncoder::WaitAndUpdateFence(FMetalFence* Fence)
+{
+	check(IsRenderCommandEncoderActive() || IsComputeCommandEncoderActive() || IsBlitCommandEncoderActive());
+	static bool bSupportsFences = CommandList.GetCommandQueue().SupportsFeature(EMetalFeaturesFences);
+	if ((bSupportsFences METAL_DEBUG_OPTION(|| CommandList.GetCommandQueue().GetRuntimeDebuggingLevel() >= EMetalDebugLevelValidation)) && Fence)
+	{
+		mtlpp::Fence InnerFence = METAL_DEBUG_OPTION(CommandList.GetCommandQueue().GetRuntimeDebuggingLevel() >= EMetalDebugLevelValidation ? mtlpp::Fence(((FMetalDebugFence*)Fence->GetPtr()).Inner) :) *Fence;
+		if (RenderCommandEncoder)
+		{
+			//			Fence->Wait(mtlpp::RenderStages::Vertex);
+			Fence->Wait(mtlpp::RenderStages::Fragment);
+			RenderCommandEncoder.WaitForFence(InnerFence, (mtlpp::RenderStages)(mtlpp::RenderStages::Vertex /*| mtlpp::RenderStages::Fragment*/));
+			METAL_DEBUG_LAYER(EMetalDebugLevelFastValidation, RenderEncoderDebug.AddWaitFence(*Fence));
+			RenderCommandEncoder.UpdateFence(InnerFence, (mtlpp::RenderStages)(/*mtlpp::RenderStages::Vertex |*/ mtlpp::RenderStages::Fragment));
+			METAL_DEBUG_LAYER(EMetalDebugLevelFastValidation, RenderEncoderDebug.AddUpdateFence(*Fence));
+			Fence->Write(mtlpp::RenderStages::Fragment);
+		}
+		else if (ComputeCommandEncoder)
+		{
+			//			Fence->Wait(mtlpp::RenderStages::Vertex);
+			Fence->Wait(mtlpp::RenderStages::Fragment);
+			ComputeCommandEncoder.WaitForFence(InnerFence);
+			METAL_DEBUG_LAYER(EMetalDebugLevelFastValidation, ComputeEncoderDebug.AddWaitFence(*Fence));
+			ComputeCommandEncoder.UpdateFence(InnerFence);
+			METAL_DEBUG_LAYER(EMetalDebugLevelFastValidation, ComputeEncoderDebug.AddUpdateFence(*Fence));
+			Fence->Write(mtlpp::RenderStages::Fragment);
+		}
+		else if (BlitCommandEncoder)
+		{
+			//			Fence->Wait(mtlpp::RenderStages::Vertex);
+			Fence->Wait(mtlpp::RenderStages::Fragment);
+			BlitCommandEncoder.WaitForFence(InnerFence);
+			METAL_DEBUG_LAYER(EMetalDebugLevelFastValidation, BlitEncoderDebug.AddWaitFence(*Fence));
+			BlitCommandEncoder.UpdateFence(InnerFence);
+			METAL_DEBUG_LAYER(EMetalDebugLevelFastValidation, BlitEncoderDebug.AddUpdateFence(*Fence));
+			Fence->Write(mtlpp::RenderStages::Fragment);
+		}
+	}
+}
+
 #pragma mark - Public Debug Support -
 
 void FMetalCommandEncoder::InsertDebugSignpost(ns::String const& String)
diff --git a/Engine/Source/Runtime/Apple/MetalRHI/Private/MetalCommandEncoder.h b/Engine/Source/Runtime/Apple/MetalRHI/Private/MetalCommandEncoder.h
--- a/Engine/Source/Runtime/Apple/MetalRHI/Private/MetalCommandEncoder.h
+++ b/Engine/Source/Runtime/Apple/MetalRHI/Private/MetalCommandEncoder.h
@@ -163,6 +163,12 @@ public:
 	
 	/** Prevent further GPU work until the event is reached. */
 	void WaitForFence(FMetalFence* Fence);
+	
+	/**
+	 * Prevent further GPU work until the event is reached and then kick the fence again so future encoders may wait on it.
+	 * Unlike the functions above this is not guarded and the caller must ensure proper matching waits & updates.
+	 */
+	void WaitAndUpdateFence(FMetalFence* Fence);
 
 #pragma mark - Public Debug Support -
 	
diff --git a/Engine/Source/Runtime/Apple/MetalRHI/Private/MetalRenderPass.cpp b/Engine/Source/Runtime/Apple/MetalRHI/Private/MetalRenderPass.cpp
--- a/Engine/Source/Runtime/Apple/MetalRHI/Private/MetalRenderPass.cpp
+++ b/Engine/Source/Runtime/Apple/MetalRHI/Private/MetalRenderPass.cpp
@@ -47,6 +47,7 @@ FMetalRenderPass::FMetalRenderPass(FMetalCommandList& InCmdList, FMetalStateCach
 , PassStartFence(nil)
 , CurrentEncoderFence(nil)
 , PrologueEncoderFence(nil)
+, LastPrologueEncoderFence(nil)
 , RenderPassDesc(nil)
 , ComputeDispatchType(mtlpp::DispatchType::Serial)
 , NumOutstandingOps(0)
@@ -143,6 +144,7 @@ FMetalFence* FMetalRenderPass::Submit(EMetalSubmitFlags Flags)
 		{
 			check(PrologueEncoder.GetCommandBuffer());
 			PrologueEncoderFence = PrologueEncoder.EndEncoding();
+			LastPrologueEncoderFence = PrologueEncoderFence;
 		}
 		if (PrologueEncoder.GetCommandBuffer())
 		{
@@ -187,6 +189,7 @@ void FMetalRenderPass::BeginParallelRenderPass(mtlpp::RenderPassDescriptor Rende
 		if (PrologueEncoder.IsBlitCommandEncoderActive() || PrologueEncoder.IsComputeCommandEncoderActive())
 		{
 			PrologueEncoderFence = PrologueEncoder.EndEncoding();
+			LastPrologueEncoderFence = PrologueEncoderFence;
 		}
 		if (CurrentEncoder.IsRenderCommandEncoderActive() || CurrentEncoder.IsBlitCommandEncoderActive() || CurrentEncoder.IsComputeCommandEncoderActive())
 		{
@@ -220,6 +223,7 @@ void FMetalRenderPass::BeginRenderPass(mtlpp::RenderPassDescriptor RenderPass)
 	if (PrologueEncoder.IsBlitCommandEncoderActive() || PrologueEncoder.IsComputeCommandEncoderActive())
 	{
 		PrologueEncoderFence = PrologueEncoder.EndEncoding();
+		LastPrologueEncoderFence = PrologueEncoderFence;
 	}
 	if (CurrentEncoder.IsRenderCommandEncoderActive() || CurrentEncoder.IsBlitCommandEncoderActive() || CurrentEncoder.IsComputeCommandEncoderActive())
 	{
@@ -1005,6 +1009,7 @@ FMetalFence* FMetalRenderPass::End(void)
 	if (PrologueEncoder.IsBlitCommandEncoderActive() || PrologueEncoder.IsComputeCommandEncoderActive())
 	{
 		PrologueEncoderFence = PrologueEncoder.EndEncoding();
+		LastPrologueEncoderFence = PrologueEncoderFence;
 	}
 	
 	if (CmdList.IsImmediate() && IsWithinParallelPass() && CurrentEncoder.IsParallelRenderCommandEncoderActive())
@@ -1163,6 +1168,7 @@ void FMetalRenderPass::ConditionalSwitchToTessellation(void)
 	if (PrologueEncoder.IsBlitCommandEncoderActive())
 	{
 		PrologueEncoderFence = PrologueEncoder.EndEncoding();
+		LastPrologueEncoderFence = PrologueEncoderFence;
 	}
 	
 	if (CurrentEncoder.IsComputeCommandEncoderActive() || CurrentEncoder.IsBlitCommandEncoderActive())
@@ -1184,6 +1190,12 @@ void FMetalRenderPass::ConditionalSwitchToTessellation(void)
 		{
 			PrologueEncoder.WaitForFence(PrologueEncoderFence);
 			PrologueEncoderFence = nullptr;
+			LastPrologueEncoderFence = nullptr;
+		}
+		if (LastPrologueEncoderFence)
+		{
+			PrologueEncoder.WaitAndUpdateFence(LastPrologueEncoderFence);
+			LastPrologueEncoderFence = nullptr;
 		}
 		if (CurrentEncoderFence)
 		{
@@ -1196,6 +1208,7 @@ void FMetalRenderPass::ConditionalSwitchToTessellation(void)
 			PassStartFence = nullptr;
 		}
 		PrologueEncoderFence = PrologueEncoder.GetEncoderFence();
+		LastPrologueEncoderFence = PrologueEncoderFence;
 #if METAL_DEBUG_OPTIONS
 		if (GetEmitDrawEvents() && PrologueEncoderFence)
 		{
@@ -1299,6 +1312,7 @@ void FMetalRenderPass::ConditionalSwitchToAsyncBlit(void)
 	if (PrologueEncoder.IsComputeCommandEncoderActive())
 	{
 		PrologueEncoderFence = PrologueEncoder.EndEncoding();
+		LastPrologueEncoderFence = PrologueEncoderFence;
 	}
 	
 	if (!PrologueEncoder.IsBlitCommandEncoderActive())
@@ -1312,6 +1326,12 @@ void FMetalRenderPass::ConditionalSwitchToAsyncBlit(void)
 		{
 			PrologueEncoder.WaitForFence(PrologueEncoderFence);
 			PrologueEncoderFence = nullptr;
+			LastPrologueEncoderFence = nullptr;
+		}
+		if (LastPrologueEncoderFence)
+		{
+			PrologueEncoder.WaitAndUpdateFence(LastPrologueEncoderFence);
+			LastPrologueEncoderFence = nullptr;
 		}
 		if (PassStartFence)
 		{
@@ -1319,6 +1339,7 @@ void FMetalRenderPass::ConditionalSwitchToAsyncBlit(void)
 			PassStartFence = nullptr;
 		}
 		PrologueEncoderFence = PrologueEncoder.GetEncoderFence();
+		LastPrologueEncoderFence = PrologueEncoderFence;
 #if METAL_DEBUG_OPTIONS
 		if (GetEmitDrawEvents() && PrologueEncoderFence)
 		{
diff --git a/Engine/Source/Runtime/Apple/MetalRHI/Private/MetalRenderPass.h b/Engine/Source/Runtime/Apple/MetalRHI/Private/MetalRenderPass.h
--- a/Engine/Source/Runtime/Apple/MetalRHI/Private/MetalRenderPass.h
+++ b/Engine/Source/Runtime/Apple/MetalRHI/Private/MetalRenderPass.h
@@ -177,6 +177,7 @@ private:
     FMetalFence* PassStartFence;
     FMetalFence* CurrentEncoderFence;
     FMetalFence* PrologueEncoderFence;
+	FMetalFence* LastPrologueEncoderFence;
     
     mtlpp::RenderPassDescriptor RenderPassDesc;
     
