Concurrent compute dispatch requires MTLFence to ensure proper ordering.

diff --git a/Engine/Source/Runtime/Apple/MetalRHI/Private/MetalRHIContext.cpp b/Engine/Source/Runtime/Apple/MetalRHI/Private/MetalRHIContext.cpp
--- a/Engine/Source/Runtime/Apple/MetalRHI/Private/MetalRHIContext.cpp
+++ b/Engine/Source/Runtime/Apple/MetalRHI/Private/MetalRHIContext.cpp
@@ -82,7 +82,7 @@ FMetalRHICommandContext::~FMetalRHICommandContext()
 FMetalRHIComputeContext::FMetalRHIComputeContext(class FMetalProfiler* InProfiler, FMetalContext* WrapContext)
 : FMetalRHICommandContext(InProfiler, WrapContext)
 {
-	if (FApplePlatformMisc::IsOSAtLeastVersion((uint32[]){10, 14, 0}, (uint32[]){12, 0, 0}, (uint32[]){12, 0, 0}))
+	if (FMetalCommandQueue::SupportsFeature(EMetalFeaturesFences) && FApplePlatformMisc::IsOSAtLeastVersion((uint32[]){10, 14, 0}, (uint32[]){12, 0, 0}, (uint32[]){12, 0, 0}))
 	{
 		WrapContext->GetCurrentRenderPass().SetDispatchType(mtlpp::DispatchType::Concurrent);
 	}
