Fix mismatched fences when using parallel command buffer encoding not MTLParallelCommandEncoder.

diff --git a/Engine/Source/Runtime/Apple/MetalRHI/Private/MetalCommandEncoder.cpp b/Engine/Source/Runtime/Apple/MetalRHI/Private/MetalCommandEncoder.cpp
--- a/Engine/Source/Runtime/Apple/MetalRHI/Private/MetalCommandEncoder.cpp
+++ b/Engine/Source/Runtime/Apple/MetalRHI/Private/MetalCommandEncoder.cpp
@@ -384,7 +384,7 @@ void FMetalCommandEncoder::BeginRenderCommandEncoding(void)
 	}
 	METAL_STATISTIC(FMetalProfiler::GetProfiler()->BeginEncoder(CommandBufferStats, RenderCommandEncoder));
 	
-	if (!CommandList.IsParallel())
+	if (CommandList.IsImmediate())
 	{
 		EncoderFence = CommandList.GetCommandQueue().CreateFence(Label);
 	}
@@ -480,7 +480,7 @@ FMetalFence* FMetalCommandEncoder::EndEncoding(void)
 		{
 			if (RenderCommandEncoder)
 			{
-				check(!bSupportsFences || EncoderFence || CommandList.IsParallel());
+				check(!bSupportsFences || EncoderFence || !CommandList.IsImmediate());
 				static bool bDeferredStoreActions = CommandList.GetCommandQueue().SupportsFeature(EMetalFeaturesDeferredStoreActions);
 				if (bDeferredStoreActions && ParallelRenderCommandEncoder.GetPtr() == nil)
 				{
