Add tracking of MTLFences.

diff --git a/Engine/Source/Runtime/Apple/MetalRHI/Private/MetalCommandEncoder.cpp b/Engine/Source/Runtime/Apple/MetalRHI/Private/MetalCommandEncoder.cpp
--- a/Engine/Source/Runtime/Apple/MetalRHI/Private/MetalCommandEncoder.cpp
+++ b/Engine/Source/Runtime/Apple/MetalRHI/Private/MetalCommandEncoder.cpp
@@ -676,6 +676,14 @@ void FMetalCommandEncoder::UpdateFence(FMetalFence* Fence)
 	static bool bSupportsFences = CommandList.GetCommandQueue().SupportsFeature(EMetalFeaturesFences);
 	if ((bSupportsFences METAL_DEBUG_OPTION(|| CommandList.GetCommandQueue().GetRuntimeDebuggingLevel() >= EMetalDebugLevelValidation)) && Fence && Fence->NeedsWrite(mtlpp::RenderStages::Fragment))
 	{
+#if ENABLE_METAL_GPUPROFILE
+		FMetalProfiler* Profiler = FMetalProfiler::GetProfiler();
+		if (Profiler)
+		{
+			Profiler->EncodeFence(GetCommandBufferStats(), TEXT("UpdateFence"), Fence);
+		}
+#endif
+		
 		mtlpp::Fence InnerFence = METAL_DEBUG_OPTION(CommandList.GetCommandQueue().GetRuntimeDebuggingLevel() >= EMetalDebugLevelValidation ? mtlpp::Fence(((FMetalDebugFence*)Fence->GetPtr()).Inner) :) *Fence;
 		if (RenderCommandEncoder)
 		{
@@ -707,6 +715,14 @@ void FMetalCommandEncoder::WaitForFence(FMetalFence* Fence)
 	static bool bSupportsFences = CommandList.GetCommandQueue().SupportsFeature(EMetalFeaturesFences);
 	if ((bSupportsFences METAL_DEBUG_OPTION(|| CommandList.GetCommandQueue().GetRuntimeDebuggingLevel() >= EMetalDebugLevelValidation)) && Fence && Fence->NeedsWait(mtlpp::RenderStages::Fragment))
 	{
+#if ENABLE_METAL_GPUPROFILE
+		FMetalProfiler* Profiler = FMetalProfiler::GetProfiler();
+		if (Profiler)
+		{
+			Profiler->EncodeFence(GetCommandBufferStats(), TEXT("WaitForFence"), Fence);
+		}
+#endif
+		
 		mtlpp::Fence InnerFence = METAL_DEBUG_OPTION(CommandList.GetCommandQueue().GetRuntimeDebuggingLevel() >= EMetalDebugLevelValidation ? mtlpp::Fence(((FMetalDebugFence*)Fence->GetPtr()).Inner) :) *Fence;
 		if (RenderCommandEncoder)
 		{
@@ -738,6 +754,14 @@ void FMetalCommandEncoder::WaitAndUpdateFence(FMetalFence* Fence)
 	static bool bSupportsFences = CommandList.GetCommandQueue().SupportsFeature(EMetalFeaturesFences);
 	if ((bSupportsFences METAL_DEBUG_OPTION(|| CommandList.GetCommandQueue().GetRuntimeDebuggingLevel() >= EMetalDebugLevelValidation)) && Fence)
 	{
+#if ENABLE_METAL_GPUPROFILE
+		FMetalProfiler* Profiler = FMetalProfiler::GetProfiler();
+		if (Profiler)
+		{
+			Profiler->EncodeFence(GetCommandBufferStats(), TEXT("WaitAndUpdateFence"), Fence);
+		}
+#endif
+		
 		mtlpp::Fence InnerFence = METAL_DEBUG_OPTION(CommandList.GetCommandQueue().GetRuntimeDebuggingLevel() >= EMetalDebugLevelValidation ? mtlpp::Fence(((FMetalDebugFence*)Fence->GetPtr()).Inner) :) *Fence;
 		if (RenderCommandEncoder)
 		{
diff --git a/Engine/Source/Runtime/Apple/MetalRHI/Private/MetalProfiler.cpp b/Engine/Source/Runtime/Apple/MetalRHI/Private/MetalProfiler.cpp
--- a/Engine/Source/Runtime/Apple/MetalRHI/Private/MetalProfiler.cpp
+++ b/Engine/Source/Runtime/Apple/MetalRHI/Private/MetalProfiler.cpp
@@ -447,7 +447,7 @@ FString IMetalStatsScope::GetJSONRepresentation(uint32 Pid)
 		if (GPUStartTime && GPUEndTime)
 		{
 			uint64 ChildStartCallTime = GPUStartTime;
-			uint64 ChildDrawCallTime = GPUEndTime - GPUStartTime;
+			uint64 ChildDrawCallTime = FMath::Max(GPUEndTime - GPUStartTime, 1llu);
 			
 			if (DrawStat.PSOPerformanceStats)
 			{
@@ -569,7 +569,7 @@ FString IMetalStatsScope::GetJSONRepresentation(uint32 Pid)
 	if (CPUStartTime && CPUEndTime)
 	{
 		uint64 ChildStartCallTime = CPUStartTime;
-		uint64 ChildDrawCallTime = CPUEndTime - CPUStartTime;
+		uint64 ChildDrawCallTime = FMath::Max(CPUEndTime - CPUStartTime, 1llu);
 		
 		JSONOutput += FString::Printf(TEXT("{\"pid\":%d, \"tid\":%d, \"ph\": \"X\", \"name\": \"%s\", \"ts\": %llu, \"dur\": %llu, \"args\":{\"num_child\":%u}},\n"),
 			 Pid,
@@ -602,6 +602,23 @@ FMetalEventStats::FMetalEventStats(const TCHAR* InName, FColor Color)
 	EndSample = nullptr;
 }
 
+FMetalEventStats::FMetalEventStats(const TCHAR* InName, uint64 InGPUIdx)
+{
+	Name = InName;
+	
+	CPUThreadIndex = FPlatformTLS::GetCurrentThreadId();
+	GPUThreadIndex = InGPUIdx;
+	
+	CPUStartTime = FPlatformTime::ToMilliseconds64(mach_absolute_time()) * 1000.0;
+	CPUEndTime = 0;
+	
+	GPUStartTime = 0;
+	GPUEndTime = 0;
+	
+	StartSample = nullptr;
+	EndSample = nullptr;
+}
+
 FMetalEventStats::~FMetalEventStats()
 {
 	[StartSample release];
@@ -910,6 +927,11 @@ void FMetalEncoderStats::End(mtlpp::CommandBuffer const& Buffer)
 	EndSample = [Stats->RegisterEncoderStatistics(CmdBufferStats, (EMetalSamples)EndPoint) retain];
 }
 
+void FMetalEncoderStats::EncodeFence(FMetalEventStats* Stat)
+{
+	Children.Add(Stat);
+}
+
 void FMetalEncoderStats::EncodeDraw(char const* DrawCall, uint32 RHIPrimitives, uint32 RHIVertices, uint32 RHIInstances)
 {
 	check(CmdBuffer);
@@ -1233,6 +1255,12 @@ FMetalProfiler* FMetalProfiler::CreateProfiler(FMetalContext *InContext)
 	if (!Self)
 	{
 		Self = new FMetalProfiler(InContext);
+		
+		int32 CaptureFrames = 0;
+		if (FParse::Value(FCommandLine::Get(), TEXT("MetalProfileFrames="), CaptureFrames))
+		{
+			Self->BeginCapture(CaptureFrames);
+		}
 	}
 	return Self;
 }
@@ -1477,6 +1505,24 @@ void FMetalProfiler::AddCommandBuffer(FMetalCommandBufferStats *CommandBuffer)
 	}
 }
 
+void FMetalProfiler::EncodeFence(FMetalCommandBufferStats* CmdBufStats, const TCHAR* Name, FMetalFence* Fence)
+{
+#if METAL_STATISTICS
+	if (MetalGPUProfilerIsInSafeThread() && Fence && bEnabled && StatisticsAPI)
+	{
+		if (!Context->GetCurrentCommandBuffer().GetPtr() || !StatisticsAPI->GetLastStatisticsSample(Context->GetCurrentCommandBuffer().GetPtr()))
+		{
+			Context->GetCurrentRenderPass().InsertDebugEncoder();
+		}
+		
+		FMetalEventStats* Event = new FMetalEventStats(*FString::Printf(TEXT("%s: %s"), Name, *FString(Fence->GetLabel())), 0);
+		Event->Start(Context->GetCurrentCommandBuffer());
+		Event->End(Context->GetCurrentCommandBuffer());
+		CmdBufStats->ActiveEncoderStats->EncodeFence(Event);
+	}
+#endif
+}
+
 void FMetalProfiler::PushEvent(const TCHAR *Name, FColor Color)
 {
 #if METAL_STATISTICS
@@ -1628,7 +1674,7 @@ void FMetalProfiler::SaveTrace()
 			if (Event->DriverStats.Num())
 			{
 				uint64 ChildStartCallTime = Event->CPUStartTime;
-				uint64 ChildDrawCallTime = Event->CPUEndTime - Event->CPUStartTime;
+				uint64 ChildDrawCallTime = FMath::Max(Event->CPUEndTime - Event->CPUStartTime, 1llu);
 				
 				FString Output;
 				FString DriverStats;
diff --git a/Engine/Source/Runtime/Apple/MetalRHI/Private/MetalProfiler.h b/Engine/Source/Runtime/Apple/MetalRHI/Private/MetalProfiler.h
--- a/Engine/Source/Runtime/Apple/MetalRHI/Private/MetalProfiler.h
+++ b/Engine/Source/Runtime/Apple/MetalRHI/Private/MetalProfiler.h
@@ -218,6 +218,7 @@ struct FMetalDisplayStats : public IMetalStatsScope
 struct FMetalEventStats : public IMetalStatsScope
 {
 	FMetalEventStats(const TCHAR* Name, FColor Color);
+	FMetalEventStats(const TCHAR* Name, uint64 InGPUIdx);
 	virtual ~FMetalEventStats();
 	
 	virtual void Start(mtlpp::CommandBuffer const& Buffer) final override;
@@ -278,6 +279,7 @@ struct FMetalEncoderStats : public IMetalStatsScope
 	void EncodeBlit(char const* DrawCall);
 	void EncodeDispatch(char const* DrawCall);
 	void EncodePipeline(FMetalShaderPipeline* PipelineStat);
+	void EncodeFence(FMetalEventStats* Stat);
 	
 	id<IMetalCommandBufferStats> CmdBufferStats;
 	ns::AutoReleased<mtlpp::CommandBuffer> CmdBuffer;
@@ -423,6 +425,8 @@ public:
 	void RemoveCounter(NSString* Counter);
 	TMap<FString, EMTLCounterType> const& GetCounterTypes() const { return CounterTypes; }
 	
+	void EncodeFence(FMetalCommandBufferStats* CmdBufStats, const TCHAR* Name, FMetalFence* Fence);
+	
 	void DumpPipeline(FMetalShaderPipeline* PipelineStat);
 #endif
 	
