Fences need an encoder to already be active.

diff --git a/Engine/Source/Runtime/Apple/MetalRHI/Private/MetalProfiler.cpp b/Engine/Source/Runtime/Apple/MetalRHI/Private/MetalProfiler.cpp
--- a/Engine/Source/Runtime/Apple/MetalRHI/Private/MetalProfiler.cpp
+++ b/Engine/Source/Runtime/Apple/MetalRHI/Private/MetalProfiler.cpp
@@ -1508,13 +1508,8 @@ void FMetalProfiler::AddCommandBuffer(FMetalCommandBufferStats *CommandBuffer)
 void FMetalProfiler::EncodeFence(FMetalCommandBufferStats* CmdBufStats, const TCHAR* Name, FMetalFence* Fence)
 {
 #if METAL_STATISTICS
-	if (MetalGPUProfilerIsInSafeThread() && Fence && bEnabled && StatisticsAPI)
+	if (MetalGPUProfilerIsInSafeThread() && Fence && bEnabled && StatisticsAPI && CmdBufStats->ActiveEncoderStats)
 	{
-		if (!Context->GetCurrentCommandBuffer().GetPtr() || !StatisticsAPI->GetLastStatisticsSample(Context->GetCurrentCommandBuffer().GetPtr()))
-		{
-			Context->GetCurrentRenderPass().InsertDebugEncoder();
-		}
-		
 		FMetalEventStats* Event = new FMetalEventStats(*FString::Printf(TEXT("%s: %s"), Name, *FString(Fence->GetLabel())), 0);
 		Event->Start(Context->GetCurrentCommandBuffer());
 		Event->End(Context->GetCurrentCommandBuffer());
