More work toward having MTLFence behave correctly.

diff --git a/Engine/Source/Runtime/Apple/MetalRHI/Private/MetalCommandEncoder.cpp b/Engine/Source/Runtime/Apple/MetalRHI/Private/MetalCommandEncoder.cpp
--- a/Engine/Source/Runtime/Apple/MetalRHI/Private/MetalCommandEncoder.cpp
+++ b/Engine/Source/Runtime/Apple/MetalRHI/Private/MetalCommandEncoder.cpp
@@ -30,6 +30,10 @@ FMetalCommandEncoder::FMetalCommandEncoder(FMetalCommandList& CmdList)
 #if ENABLE_METAL_GPUPROFILE
 , CommandBufferStats(nullptr)
 #endif
+#if METAL_DEBUG_OPTIONS
+, WaitCount(0)
+, UpdateCount(0)
+#endif
 , DebugGroups([NSMutableArray new])
 , EncoderNum(0)
 {
@@ -432,7 +436,7 @@ void FMetalCommandEncoder::BeginComputeCommandEncoding(mtlpp::DispatchType Type)
 	EncoderFence = CommandList.GetCommandQueue().CreateFence(Label);
 }
 
-void FMetalCommandEncoder::BeginBlitCommandEncoding(bool const bSuppressFence)
+void FMetalCommandEncoder::BeginBlitCommandEncoding(void)
 {
 	check(CommandBuffer);
 	check(IsRenderCommandEncoderActive() == false && IsComputeCommandEncoderActive() == false && IsBlitCommandEncoderActive() == false);
@@ -464,10 +468,8 @@ void FMetalCommandEncoder::BeginBlitCommandEncoding(bool const bSuppressFence)
 		}
 	}
 	METAL_STATISTIC(FMetalProfiler::GetProfiler()->BeginEncoder(CommandBufferStats, BlitCommandEncoder));
-	if (!bSuppressFence)
-	{
-		EncoderFence = CommandList.GetCommandQueue().CreateFence(Label);
-	}
+	
+	EncoderFence = CommandList.GetCommandQueue().CreateFence(Label);
 }
 
 FMetalFence* FMetalCommandEncoder::EndEncoding(void)
@@ -516,6 +518,16 @@ FMetalFence* FMetalCommandEncoder::EndEncoding(void)
 				}
 				UpdateFence(EncoderFence);
 				
+#if METAL_DEBUG_OPTIONS
+				if (bSupportsFences && (!WaitCount || !UpdateCount))
+				{
+					UE_LOG(LogMetal, Error, TEXT("%s has incorrect fence waits (%u) vs. updates (%u)."), *FString(RenderCommandEncoder.GetLabel()), WaitCount, UpdateCount);
+					check(UpdateCount);
+				}
+				WaitCount = 0;
+				UpdateCount = 0;
+#endif
+				
 				METAL_STATISTIC(FMetalProfiler::GetProfiler()->EndEncoder(CommandBufferStats, RenderCommandEncoder));
 				RenderCommandEncoder.EndEncoding();
 				METAL_DEBUG_LAYER(EMetalDebugLevelFastValidation, RenderEncoderDebug.EndEncoder());
@@ -610,6 +622,16 @@ FMetalFence* FMetalCommandEncoder::EndEncoding(void)
 			}
 			UpdateFence(EncoderFence);
 			
+#if METAL_DEBUG_OPTIONS
+			if (bSupportsFences && (!WaitCount || !UpdateCount))
+			{
+				UE_LOG(LogMetal, Error, TEXT("%s has incorrect fence waits (%u) vs. updates (%u)."), *FString(ComputeCommandEncoder.GetLabel()), WaitCount, UpdateCount);
+				check(UpdateCount);
+			}
+			WaitCount = 0;
+			UpdateCount = 0;
+#endif
+			
 			METAL_STATISTIC(FMetalProfiler::GetProfiler()->EndEncoder(CommandBufferStats, ComputeCommandEncoder));
 			ComputeCommandEncoder.EndEncoding();
 			METAL_DEBUG_LAYER(EMetalDebugLevelFastValidation, ComputeEncoderDebug.EndEncoder());
@@ -626,6 +648,16 @@ FMetalFence* FMetalCommandEncoder::EndEncoding(void)
 			}
 			UpdateFence(EncoderFence);
 			
+#if METAL_DEBUG_OPTIONS
+			if (bSupportsFences && (!WaitCount || !UpdateCount))
+			{
+				UE_LOG(LogMetal, Error, TEXT("%s has incorrect fence waits (%u) vs. updates (%u)."), *FString(BlitCommandEncoder.GetLabel()), WaitCount, UpdateCount);
+				check(UpdateCount);
+			}
+			WaitCount = 0;
+			UpdateCount = 0;
+#endif
+			
 			METAL_STATISTIC(FMetalProfiler::GetProfiler()->EndEncoder(CommandBufferStats, RenderCommandEncoder));
 			BlitCommandEncoder.EndEncoding();
 			METAL_DEBUG_LAYER(EMetalDebugLevelFastValidation, BlitEncoderDebug.EndEncoder());
@@ -676,13 +708,14 @@ void FMetalCommandEncoder::UpdateFence(FMetalFence* Fence)
 	static bool bSupportsFences = CommandList.GetCommandQueue().SupportsFeature(EMetalFeaturesFences);
 	if ((bSupportsFences METAL_DEBUG_OPTION(|| CommandList.GetCommandQueue().GetRuntimeDebuggingLevel() >= EMetalDebugLevelValidation)) && Fence && Fence->NeedsWrite(mtlpp::RenderStages::Fragment))
 	{
-#if ENABLE_METAL_GPUPROFILE
+#if ENABLE_METAL_GPUPROFILE && METAL_STATISTICS
 		FMetalProfiler* Profiler = FMetalProfiler::GetProfiler();
 		if (Profiler)
 		{
-			Profiler->EncodeFence(GetCommandBufferStats(), TEXT("UpdateFence"), Fence);
+			Profiler->EncodeFence(GetCommandBufferStats(), TEXT("UpdateFence"), Fence, EMTLFenceTypeUpdate);
 		}
 #endif
+		METAL_DEBUG_LAYER(EMetalDebugLevelFastValidation, UpdateCount++);
 		
 		mtlpp::Fence InnerFence = METAL_DEBUG_OPTION(CommandList.GetCommandQueue().GetRuntimeDebuggingLevel() >= EMetalDebugLevelValidation ? mtlpp::Fence(((FMetalDebugFence*)Fence->GetPtr()).Inner) :) *Fence;
 		if (RenderCommandEncoder)
@@ -715,13 +748,14 @@ void FMetalCommandEncoder::WaitForFence(FMetalFence* Fence)
 	static bool bSupportsFences = CommandList.GetCommandQueue().SupportsFeature(EMetalFeaturesFences);
 	if ((bSupportsFences METAL_DEBUG_OPTION(|| CommandList.GetCommandQueue().GetRuntimeDebuggingLevel() >= EMetalDebugLevelValidation)) && Fence && Fence->NeedsWait(mtlpp::RenderStages::Fragment))
 	{
-#if ENABLE_METAL_GPUPROFILE
+#if ENABLE_METAL_GPUPROFILE && METAL_STATISTICS
 		FMetalProfiler* Profiler = FMetalProfiler::GetProfiler();
 		if (Profiler)
 		{
-			Profiler->EncodeFence(GetCommandBufferStats(), TEXT("WaitForFence"), Fence);
+			Profiler->EncodeFence(GetCommandBufferStats(), TEXT("WaitForFence"), Fence, EMTLFenceTypeWait);
 		}
 #endif
+		METAL_DEBUG_LAYER(EMetalDebugLevelFastValidation, WaitCount++);
 		
 		mtlpp::Fence InnerFence = METAL_DEBUG_OPTION(CommandList.GetCommandQueue().GetRuntimeDebuggingLevel() >= EMetalDebugLevelValidation ? mtlpp::Fence(((FMetalDebugFence*)Fence->GetPtr()).Inner) :) *Fence;
 		if (RenderCommandEncoder)
@@ -754,13 +788,16 @@ void FMetalCommandEncoder::WaitAndUpdateFence(FMetalFence* Fence)
 	static bool bSupportsFences = CommandList.GetCommandQueue().SupportsFeature(EMetalFeaturesFences);
 	if ((bSupportsFences METAL_DEBUG_OPTION(|| CommandList.GetCommandQueue().GetRuntimeDebuggingLevel() >= EMetalDebugLevelValidation)) && Fence)
 	{
-#if ENABLE_METAL_GPUPROFILE
+#if ENABLE_METAL_GPUPROFILE && METAL_STATISTICS
 		FMetalProfiler* Profiler = FMetalProfiler::GetProfiler();
 		if (Profiler)
 		{
-			Profiler->EncodeFence(GetCommandBufferStats(), TEXT("WaitAndUpdateFence"), Fence);
+			Profiler->EncodeFence(GetCommandBufferStats(), TEXT("WaitForFence"), Fence, EMTLFenceTypeWait);
+			Profiler->EncodeFence(GetCommandBufferStats(), TEXT("UpdateFence"), Fence, EMTLFenceTypeUpdate);
 		}
 #endif
+		METAL_DEBUG_LAYER(EMetalDebugLevelFastValidation, WaitCount++);
+		METAL_DEBUG_LAYER(EMetalDebugLevelFastValidation, UpdateCount++);
 		
 		mtlpp::Fence InnerFence = METAL_DEBUG_OPTION(CommandList.GetCommandQueue().GetRuntimeDebuggingLevel() >= EMetalDebugLevelValidation ? mtlpp::Fence(((FMetalDebugFence*)Fence->GetPtr()).Inner) :) *Fence;
 		if (RenderCommandEncoder)
diff --git a/Engine/Source/Runtime/Apple/MetalRHI/Private/MetalCommandEncoder.h b/Engine/Source/Runtime/Apple/MetalRHI/Private/MetalCommandEncoder.h
--- a/Engine/Source/Runtime/Apple/MetalRHI/Private/MetalCommandEncoder.h
+++ b/Engine/Source/Runtime/Apple/MetalRHI/Private/MetalCommandEncoder.h
@@ -147,7 +147,7 @@ public:
 	void BeginComputeCommandEncoding(mtlpp::DispatchType Type = mtlpp::DispatchType::Serial);
 	
 	/** Begins encoding blit commands into the current command buffer. No other encoder may be active. */
-	void BeginBlitCommandEncoding(bool const bSuppressFence = false);
+	void BeginBlitCommandEncoding(void);
 	
 	/** Declare that all command generation from this encoder is complete, and detach from the MTLCommandBuffer if there is an encoder active or does nothing if there isn't. */
 	FMetalFence* EndEncoding(void);
@@ -443,7 +443,11 @@ private:
 #if ENABLE_METAL_GPUPROFILE
 	FMetalCommandBufferStats* CommandBufferStats;
 #endif
-	
+#if METAL_DEBUG_OPTIONS
+	uint8 WaitCount;
+	uint8 UpdateCount;
+#endif
+
 	TArray<ns::Object<mtlpp::CommandBufferHandler>> CompletionHandlers;
 	NSMutableArray* DebugGroups;
     
diff --git a/Engine/Source/Runtime/Apple/MetalRHI/Private/MetalProfiler.cpp b/Engine/Source/Runtime/Apple/MetalRHI/Private/MetalProfiler.cpp
--- a/Engine/Source/Runtime/Apple/MetalRHI/Private/MetalProfiler.cpp
+++ b/Engine/Source/Runtime/Apple/MetalRHI/Private/MetalProfiler.cpp
@@ -925,10 +925,26 @@ void FMetalEncoderStats::End(mtlpp::CommandBuffer const& Buffer)
 	CPUEndTime = FPlatformTime::ToMilliseconds64(mach_absolute_time()) * 1000.0;
 	IMetalStatistics* Stats = FMetalProfiler::GetStatistics();
 	EndSample = [Stats->RegisterEncoderStatistics(CmdBufferStats, (EMetalSamples)EndPoint) retain];
+	for (auto Stat : FenceUpdates)
+	{
+		Stat->StartSample = [EndSample retain];
+		Stat->EndSample = [EndSample retain];
+		check(Stat->StartSample && Stat->EndSample);
+	}
 }
 
-void FMetalEncoderStats::EncodeFence(FMetalEventStats* Stat)
+void FMetalEncoderStats::EncodeFence(FMetalEventStats* Stat, EMTLFenceType Type)
 {
+	if(Type == EMTLFenceTypeWait)
+	{
+		Stat->StartSample = [StartSample retain];
+		Stat->EndSample = [StartSample retain];
+		check(Stat->StartSample && Stat->EndSample);
+	}
+	else
+	{
+		FenceUpdates.Add(Stat);
+	}
 	Children.Add(Stat);
 }
 
@@ -1505,15 +1521,13 @@ void FMetalProfiler::AddCommandBuffer(FMetalCommandBufferStats *CommandBuffer)
 	}
 }
 
-void FMetalProfiler::EncodeFence(FMetalCommandBufferStats* CmdBufStats, const TCHAR* Name, FMetalFence* Fence)
+void FMetalProfiler::EncodeFence(FMetalCommandBufferStats* CmdBufStats, const TCHAR* Name, FMetalFence* Fence, EMTLFenceType Type)
 {
 #if METAL_STATISTICS
 	if (MetalGPUProfilerIsInSafeThread() && Fence && bEnabled && StatisticsAPI && CmdBufStats->ActiveEncoderStats)
 	{
 		FMetalEventStats* Event = new FMetalEventStats(*FString::Printf(TEXT("%s: %s"), Name, *FString(Fence->GetLabel())), 0);
-		Event->Start(Context->GetCurrentCommandBuffer());
-		Event->End(Context->GetCurrentCommandBuffer());
-		CmdBufStats->ActiveEncoderStats->EncodeFence(Event);
+		CmdBufStats->ActiveEncoderStats->EncodeFence(Event, Type);
 	}
 #endif
 }
diff --git a/Engine/Source/Runtime/Apple/MetalRHI/Private/MetalProfiler.h b/Engine/Source/Runtime/Apple/MetalRHI/Private/MetalProfiler.h
--- a/Engine/Source/Runtime/Apple/MetalRHI/Private/MetalProfiler.h
+++ b/Engine/Source/Runtime/Apple/MetalRHI/Private/MetalProfiler.h
@@ -264,6 +264,11 @@ struct FMetalShaderPipelineStats : public IMetalStatsScope
 	FMetalShaderPipeline* Pipeline;
 };
 
+enum EMTLFenceType
+{
+	EMTLFenceTypeWait,
+	EMTLFenceTypeUpdate,
+};
 struct FMetalEncoderStats : public IMetalStatsScope
 {
 	FMetalEncoderStats(mtlpp::RenderCommandEncoder const& Encoder, uint64 GPUThreadIndex);
@@ -279,7 +284,7 @@ struct FMetalEncoderStats : public IMetalStatsScope
 	void EncodeBlit(char const* DrawCall);
 	void EncodeDispatch(char const* DrawCall);
 	void EncodePipeline(FMetalShaderPipeline* PipelineStat);
-	void EncodeFence(FMetalEventStats* Stat);
+	void EncodeFence(FMetalEventStats* Stat, EMTLFenceType Type);
 	
 	id<IMetalCommandBufferStats> CmdBufferStats;
 	ns::AutoReleased<mtlpp::CommandBuffer> CmdBuffer;
@@ -287,6 +292,7 @@ struct FMetalEncoderStats : public IMetalStatsScope
 	uint32 EndPoint;
 	id<IMetalStatisticsSamples> StartSample;
 	id<IMetalStatisticsSamples> EndSample;
+	TArray<FMetalEventStats*> FenceUpdates;
 };
 #endif
 
@@ -425,7 +431,7 @@ public:
 	void RemoveCounter(NSString* Counter);
 	TMap<FString, EMTLCounterType> const& GetCounterTypes() const { return CounterTypes; }
 	
-	void EncodeFence(FMetalCommandBufferStats* CmdBufStats, const TCHAR* Name, FMetalFence* Fence);
+	void EncodeFence(FMetalCommandBufferStats* CmdBufStats, const TCHAR* Name, FMetalFence* Fence, EMTLFenceType Type);
 	
 	void DumpPipeline(FMetalShaderPipeline* PipelineStat);
 #endif
diff --git a/Engine/Source/Runtime/Apple/MetalRHI/Private/MetalRenderPass.cpp b/Engine/Source/Runtime/Apple/MetalRHI/Private/MetalRenderPass.cpp
--- a/Engine/Source/Runtime/Apple/MetalRHI/Private/MetalRenderPass.cpp
+++ b/Engine/Source/Runtime/Apple/MetalRHI/Private/MetalRenderPass.cpp
@@ -45,6 +45,7 @@ FMetalRenderPass::FMetalRenderPass(FMetalCommandList& InCmdList, FMetalStateCach
 , CurrentEncoder(InCmdList)
 , PrologueEncoder(InCmdList)
 , PassStartFence(nil)
+, ParallelPassEndFence(nil)
 , CurrentEncoderFence(nil)
 , PrologueEncoderFence(nil)
 , LastPrologueEncoderFence(nil)
@@ -79,10 +80,10 @@ void FMetalRenderPass::Begin(FMetalFence* Fence, bool const bParallelBegin)
 	}
 	else
 	{
-		check(!CurrentEncoderFence || Fence == nullptr);
+		check(!ParallelPassEndFence || Fence == nullptr);
 		if (Fence)
 		{
-			CurrentEncoderFence = Fence;
+			ParallelPassEndFence = Fence;
 		}
 	}
 	
@@ -123,7 +124,7 @@ void FMetalRenderPass::Update(FMetalFence* Fence)
 		// the higher-level can generate empty contexts but we have no sane way to deal with that.
 		if (!CurrentEncoder.IsRenderCommandEncoderActive() && !CurrentEncoder.IsBlitCommandEncoderActive() && !CurrentEncoder.IsComputeCommandEncoderActive())
 		{
-			ConditionalSwitchToBlit(true);
+			ConditionalSwitchToBlit();
 		}
 		CurrentEncoder.UpdateFence(Fence);
 		State.FlushVisibilityResults(CurrentEncoder);
@@ -172,6 +173,11 @@ FMetalFence* FMetalRenderPass::Submit(EMetalSubmitFlags Flags)
 	{
 		PrologueEncoder.Reset();
 		CurrentEncoder.Reset();
+		
+		if (!LastPrologueEncoderFence || !LastPrologueEncoderFence->NeedsWait(mtlpp::RenderStages::Fragment))
+		{
+			LastPrologueEncoderFence = CurrentEncoderFence;
+		}
 	}
 	
 	return CurrentEncoderFence;
@@ -245,6 +251,11 @@ void FMetalRenderPass::BeginRenderPass(mtlpp::RenderPassDescriptor RenderPass)
 			CurrentEncoder.WaitForFence(PassStartFence);
 			PassStartFence = nullptr;
 		}
+		if (ParallelPassEndFence)
+		{
+			CurrentEncoder.WaitForFence(ParallelPassEndFence);
+			ParallelPassEndFence = nullptr;
+		}
 		if (CurrentEncoderFence)
 		{
 			CurrentEncoder.WaitForFence(CurrentEncoderFence);
@@ -361,6 +372,11 @@ void FMetalRenderPass::RestartRenderPass(mtlpp::RenderPassDescriptor RenderPass)
 		CurrentEncoder.WaitForFence(PassStartFence);
 		PassStartFence = nullptr;
 	}
+	if (ParallelPassEndFence)
+	{
+		CurrentEncoder.WaitForFence(ParallelPassEndFence);
+		ParallelPassEndFence = nullptr;
+	}
 	if (PrologueEncoderFence)
 	{
 		CurrentEncoder.WaitForFence(PrologueEncoderFence);
@@ -1017,9 +1033,10 @@ FMetalFence* FMetalRenderPass::End(void)
 		State.SetRenderStoreActions(CurrentEncoder, false);
 		CurrentEncoder.EndEncoding();
 		
-		ConditionalSwitchToBlit(true);
+		ConditionalSwitchToBlit();
 		CurrentEncoder.WaitForFence(PassStartFence);
 		CurrentEncoder.WaitForFence(CurrentEncoderFence);
+		CurrentEncoder.WaitForFence(ParallelPassEndFence);
 		CurrentEncoderFence = CurrentEncoder.EndEncoding();
 	}
 	else if (CurrentEncoder.IsRenderCommandEncoderActive() || CurrentEncoder.IsBlitCommandEncoderActive() || CurrentEncoder.IsComputeCommandEncoderActive())
@@ -1027,16 +1044,27 @@ FMetalFence* FMetalRenderPass::End(void)
 		State.FlushVisibilityResults(CurrentEncoder);
 		CurrentEncoder.WaitForFence(PassStartFence);
 		CurrentEncoder.WaitForFence(CurrentEncoderFence);
+		CurrentEncoder.WaitForFence(ParallelPassEndFence);
 		CurrentEncoderFence = CurrentEncoder.EndEncoding();
 	}
 	else if (PassStartFence)
 	{
-		ConditionalSwitchToBlit(true);
+		ConditionalSwitchToBlit();
+		CurrentEncoder.WaitForFence(PassStartFence);
+		CurrentEncoder.WaitForFence(CurrentEncoderFence);
+		CurrentEncoder.WaitForFence(ParallelPassEndFence);
+		CurrentEncoderFence = CurrentEncoder.EndEncoding();
+	}
+	else if (ParallelPassEndFence)
+	{
+		ConditionalSwitchToBlit();
 		CurrentEncoder.WaitForFence(PassStartFence);
 		CurrentEncoder.WaitForFence(CurrentEncoderFence);
+		CurrentEncoder.WaitForFence(ParallelPassEndFence);
 		CurrentEncoderFence = CurrentEncoder.EndEncoding();
 	}
 	
+	ParallelPassEndFence = nullptr;
 	PassStartFence = nullptr;
 	
 	State.SetRenderTargetsActive(false);
@@ -1207,6 +1235,11 @@ void FMetalRenderPass::ConditionalSwitchToTessellation(void)
 			PrologueEncoder.WaitForFence(PassStartFence);
 			PassStartFence = nullptr;
 		}
+		if (ParallelPassEndFence)
+		{
+			PrologueEncoder.WaitForFence(ParallelPassEndFence);
+			ParallelPassEndFence = nullptr;
+		}
 		PrologueEncoderFence = PrologueEncoder.GetEncoderFence();
 		LastPrologueEncoderFence = PrologueEncoderFence;
 #if METAL_DEBUG_OPTIONS
@@ -1273,7 +1306,7 @@ void FMetalRenderPass::ConditionalSwitchToCompute(void)
 	check(CurrentEncoder.IsComputeCommandEncoderActive());
 }
 
-void FMetalRenderPass::ConditionalSwitchToBlit(bool const bSuppressFence)
+void FMetalRenderPass::ConditionalSwitchToBlit(void)
 {
 	check(CurrentEncoder.GetCommandBuffer());
 	check(!CurrentEncoder.IsParallel());
@@ -1291,7 +1324,7 @@ void FMetalRenderPass::ConditionalSwitchToBlit(bool const bSuppressFence)
 	
 	if (!CurrentEncoder.IsBlitCommandEncoderActive())
 	{
-		CurrentEncoder.BeginBlitCommandEncoding(bSuppressFence);
+		CurrentEncoder.BeginBlitCommandEncoding();
 		if (CurrentEncoderFence)
 		{
 			CurrentEncoder.WaitForFence(CurrentEncoderFence);
@@ -1324,9 +1357,12 @@ void FMetalRenderPass::ConditionalSwitchToAsyncBlit(void)
 		PrologueEncoder.BeginBlitCommandEncoding();
 		if (PrologueEncoderFence)
 		{
+			if (PrologueEncoderFence->NeedsWait(mtlpp::RenderStages::Fragment))
+			{
+				LastPrologueEncoderFence = nullptr;
+			}
 			PrologueEncoder.WaitForFence(PrologueEncoderFence);
 			PrologueEncoderFence = nullptr;
-			LastPrologueEncoderFence = nullptr;
 		}
 		if (LastPrologueEncoderFence)
 		{
@@ -1338,6 +1374,11 @@ void FMetalRenderPass::ConditionalSwitchToAsyncBlit(void)
 			PrologueEncoder.WaitForFence(PassStartFence);
 			PassStartFence = nullptr;
 		}
+		if (ParallelPassEndFence)
+		{
+			PrologueEncoder.WaitForFence(ParallelPassEndFence);
+			ParallelPassEndFence = nullptr;
+		}
 		PrologueEncoderFence = PrologueEncoder.GetEncoderFence();
 		LastPrologueEncoderFence = PrologueEncoderFence;
 #if METAL_DEBUG_OPTIONS
diff --git a/Engine/Source/Runtime/Apple/MetalRHI/Private/MetalRenderPass.h b/Engine/Source/Runtime/Apple/MetalRHI/Private/MetalRenderPass.h
--- a/Engine/Source/Runtime/Apple/MetalRHI/Private/MetalRenderPass.h
+++ b/Engine/Source/Runtime/Apple/MetalRHI/Private/MetalRenderPass.h
@@ -144,7 +144,7 @@ private:
     void ConditionalSwitchToRender(void);
     void ConditionalSwitchToTessellation(void);
     void ConditionalSwitchToCompute(void);
-	void ConditionalSwitchToBlit(bool const bSuppressFence = false);
+	void ConditionalSwitchToBlit(void);
 	void ConditionalSwitchToAsyncBlit(void);
 	
     void PrepareToRender(uint32 PrimType);
@@ -175,7 +175,8 @@ private:
 	TMap<id<MTLBuffer>, TArray<NSRange>> OutstandingBufferUploads;
     
     FMetalFence* PassStartFence;
-    FMetalFence* CurrentEncoderFence;
+    FMetalFence* ParallelPassEndFence;
+	FMetalFence* CurrentEncoderFence;
     FMetalFence* PrologueEncoderFence;
 	FMetalFence* LastPrologueEncoderFence;
     
