Consistent conditions for MTLFence debug code.

diff --git a/Engine/Source/Runtime/Apple/MetalRHI/Private/MetalCommandEncoder.cpp b/Engine/Source/Runtime/Apple/MetalRHI/Private/MetalCommandEncoder.cpp
--- a/Engine/Source/Runtime/Apple/MetalRHI/Private/MetalCommandEncoder.cpp
+++ b/Engine/Source/Runtime/Apple/MetalRHI/Private/MetalCommandEncoder.cpp
@@ -519,10 +519,10 @@ FMetalFence* FMetalCommandEncoder::EndEncoding(void)
 				UpdateFence(EncoderFence);
 				
 #if METAL_DEBUG_OPTIONS
-				if (bSupportsFences && (!WaitCount || !UpdateCount))
+				if (bSupportsFences && SafeGetRuntimeDebuggingLevel() >= EMetalDebugLevelFastValidation && (!WaitCount || !UpdateCount))
 				{
 					UE_LOG(LogMetal, Error, TEXT("%s has incorrect fence waits (%u) vs. updates (%u)."), *FString(RenderCommandEncoder.GetLabel()), WaitCount, UpdateCount);
-					check(UpdateCount);
+					
 				}
 				WaitCount = 0;
 				UpdateCount = 0;
@@ -623,10 +623,10 @@ FMetalFence* FMetalCommandEncoder::EndEncoding(void)
 			UpdateFence(EncoderFence);
 			
 #if METAL_DEBUG_OPTIONS
-			if (bSupportsFences && (!WaitCount || !UpdateCount))
+			if (bSupportsFences && SafeGetRuntimeDebuggingLevel() >= EMetalDebugLevelFastValidation && (!WaitCount || !UpdateCount))
 			{
 				UE_LOG(LogMetal, Error, TEXT("%s has incorrect fence waits (%u) vs. updates (%u)."), *FString(ComputeCommandEncoder.GetLabel()), WaitCount, UpdateCount);
-				check(UpdateCount);
+				
 			}
 			WaitCount = 0;
 			UpdateCount = 0;
@@ -649,10 +649,10 @@ FMetalFence* FMetalCommandEncoder::EndEncoding(void)
 			UpdateFence(EncoderFence);
 			
 #if METAL_DEBUG_OPTIONS
-			if (bSupportsFences && (!WaitCount || !UpdateCount))
+			if (bSupportsFences && SafeGetRuntimeDebuggingLevel() >= EMetalDebugLevelFastValidation && (!WaitCount || !UpdateCount))
 			{
 				UE_LOG(LogMetal, Error, TEXT("%s has incorrect fence waits (%u) vs. updates (%u)."), *FString(BlitCommandEncoder.GetLabel()), WaitCount, UpdateCount);
-				check(UpdateCount);
+				
 			}
 			WaitCount = 0;
 			UpdateCount = 0;
