Don't track Synchronise operations as blits - they apparently aren't.

diff --git a/Engine/Source/Runtime/Apple/MetalRHI/Private/MetalRenderPass.cpp b/Engine/Source/Runtime/Apple/MetalRHI/Private/MetalRenderPass.cpp
--- a/Engine/Source/Runtime/Apple/MetalRHI/Private/MetalRenderPass.cpp
+++ b/Engine/Source/Runtime/Apple/MetalRHI/Private/MetalRenderPass.cpp
@@ -861,7 +861,7 @@ void FMetalRenderPass::SynchronizeTexture(FMetalTexture const& Texture, uint32 S
 	mtlpp::BlitCommandEncoder& Encoder = CurrentEncoder.GetBlitCommandEncoder();
 	check(Encoder.GetPtr());
 	
-	METAL_GPUPROFILE(FMetalProfiler::GetProfiler()->EncodeBlit(CurrentEncoder.GetCommandBufferStats(), __FUNCTION__));
+	// METAL_GPUPROFILE(FMetalProfiler::GetProfiler()->EncodeBlit(CurrentEncoder.GetCommandBufferStats(), __FUNCTION__));
 	MTLPP_VALIDATE(mtlpp::BlitCommandEncoder, Encoder, SafeGetRuntimeDebuggingLevel() >= EMetalDebugLevelValidation, Synchronize(Texture, Slice, Level));
 	METAL_DEBUG_LAYER(EMetalDebugLevelFastValidation, CurrentEncoder.GetBlitCommandEncoderDebugging().Synchronize(Texture, Slice, Level));
 	ConditionalSubmit();
@@ -876,7 +876,7 @@ void FMetalRenderPass::SynchroniseResource(mtlpp::Resource const& Resource)
 	mtlpp::BlitCommandEncoder& Encoder = CurrentEncoder.GetBlitCommandEncoder();
 	check(Encoder.GetPtr());
 	
-	METAL_GPUPROFILE(FMetalProfiler::GetProfiler()->EncodeBlit(CurrentEncoder.GetCommandBufferStats(), __FUNCTION__));
+	// METAL_GPUPROFILE(FMetalProfiler::GetProfiler()->EncodeBlit(CurrentEncoder.GetCommandBufferStats(), __FUNCTION__));
 	MTLPP_VALIDATE(mtlpp::BlitCommandEncoder, Encoder, SafeGetRuntimeDebuggingLevel() >= EMetalDebugLevelValidation, Synchronize(Resource));
 	METAL_DEBUG_LAYER(EMetalDebugLevelFastValidation, CurrentEncoder.GetBlitCommandEncoderDebugging().Synchronize(Resource));
 	ConditionalSubmit();
diff --git a/Engine/Source/Runtime/Apple/MetalRHI/Private/MetalStateCache.cpp b/Engine/Source/Runtime/Apple/MetalRHI/Private/MetalStateCache.cpp
--- a/Engine/Source/Runtime/Apple/MetalRHI/Private/MetalStateCache.cpp
+++ b/Engine/Source/Runtime/Apple/MetalRHI/Private/MetalStateCache.cpp
@@ -1863,7 +1863,7 @@ void FMetalStateCache::FlushVisibilityResults(FMetalCommandEncoder& CommandEncod
 		
 		mtlpp::BlitCommandEncoder& Encoder = CommandEncoder.GetBlitCommandEncoder();
 
-		METAL_GPUPROFILE(FMetalProfiler::GetProfiler()->EncodeBlit(CommandEncoder.GetCommandBufferStats(), __FUNCTION__));
+		// METAL_GPUPROFILE(FMetalProfiler::GetProfiler()->EncodeBlit(CommandEncoder.GetCommandBufferStats(), __FUNCTION__));
 		MTLPP_VALIDATE(mtlpp::BlitCommandEncoder, Encoder, SafeGetRuntimeDebuggingLevel() >= EMetalDebugLevelValidation, Synchronize(VisibilityResults->Buffer));
 		METAL_DEBUG_LAYER(EMetalDebugLevelFastValidation, CommandEncoder.GetBlitCommandEncoderDebugging().Synchronize(VisibilityResults->Buffer));
 		
