More work to debug blit event ordering bugs.

diff --git a/Engine/Source/Runtime/Apple/MetalRHI/Private/MetalProfiler.cpp b/Engine/Source/Runtime/Apple/MetalRHI/Private/MetalProfiler.cpp
--- a/Engine/Source/Runtime/Apple/MetalRHI/Private/MetalProfiler.cpp
+++ b/Engine/Source/Runtime/Apple/MetalRHI/Private/MetalProfiler.cpp
@@ -804,6 +804,28 @@ FMetalOperationStats::FMetalOperationStats(char const* DrawCall, uint64 InGPUThr
 	RHIInstances = 0;
 }
 
+FMetalOperationStats::FMetalOperationStats(FString DrawCall, uint64 InGPUThreadIndex, uint32 InStartPoint, uint32 InEndPoint)
+{
+	Name = DrawCall;
+	CmdBufferStats = nullptr;
+	
+	CPUThreadIndex = FPlatformTLS::GetCurrentThreadId();
+	GPUThreadIndex = InGPUThreadIndex;
+	
+	CPUStartTime = FPlatformTime::ToMilliseconds64(mach_absolute_time()) * 1000.0;
+	CPUEndTime = 0;
+	
+	GPUStartTime = 0;
+	GPUEndTime = 0;
+	
+	StartPoint = InStartPoint;
+	EndPoint = InEndPoint;
+	DrawStats = nullptr;
+	RHIPrimitives = 0;
+	RHIVertices = 0;
+	RHIInstances = 0;
+}
+
 FMetalOperationStats::~FMetalOperationStats()
 {
 	delete DrawStats;
@@ -969,6 +991,16 @@ void FMetalEncoderStats::EncodeBlit(char const* DrawCall)
 	
 }
 
+void FMetalEncoderStats::EncodeBlit(FString DrawCall)
+{
+	check(CmdBuffer);
+	FMetalOperationStats* Draw = new FMetalOperationStats(DrawCall, GPUThreadIndex, EMetalSampleBeforeBlit, EMetalSampleAfterBlit);
+	Draw->CmdBufferStats = CmdBufferStats;
+	Children.Add(Draw);
+	Draw->Start(CmdBuffer);
+	Draw->End(CmdBuffer);
+}
+
 void FMetalEncoderStats::EncodeDispatch(char const* DrawCall)
 {
 	check(CmdBuffer);
@@ -1410,6 +1442,17 @@ void FMetalProfiler::EncodeBlit(FMetalCommandBufferStats* CmdBufStats, char cons
 #endif
 }
 
+void FMetalProfiler::EncodeBlit(FMetalCommandBufferStats* CmdBufStats, FString DrawCall)
+{
+	if (MetalGPUProfilerIsInSafeThread())
+		FMetalGPUProfiler::RegisterGPUWork(1, 1);
+	
+#if METAL_STATISTICS
+	if (StatisticsAPI)
+		CmdBufStats->ActiveEncoderStats->EncodeBlit(DrawCall);
+#endif
+}
+
 void FMetalProfiler::EncodeDispatch(FMetalCommandBufferStats* CmdBufStats, char const* DrawCall)
 {
 	if (MetalGPUProfilerIsInSafeThread())
diff --git a/Engine/Source/Runtime/Apple/MetalRHI/Private/MetalProfiler.h b/Engine/Source/Runtime/Apple/MetalRHI/Private/MetalProfiler.h
--- a/Engine/Source/Runtime/Apple/MetalRHI/Private/MetalProfiler.h
+++ b/Engine/Source/Runtime/Apple/MetalRHI/Private/MetalProfiler.h
@@ -241,6 +241,7 @@ struct FMetalOperationStats : public IMetalStatsScope
 {
 	FMetalOperationStats(char const* DrawCall, uint64 GPUThreadIndex, uint32 StartPoint, uint32 EndPoint, uint32 RHIPrimitives, uint32 RHIVertices, uint32 RHIInstances);
 	FMetalOperationStats(char const* DrawCall, uint64 GPUThreadIndex, uint32 StartPoint, uint32 EndPoint);
+	FMetalOperationStats(FString DrawCall, uint64 GPUThreadIndex, uint32 StartPoint, uint32 EndPoint);
 	virtual ~FMetalOperationStats();
 	
 	virtual void Start(mtlpp::CommandBuffer const& Buffer) final override;
@@ -283,6 +284,7 @@ struct FMetalEncoderStats : public IMetalStatsScope
 	
 	void EncodeDraw(char const* DrawCall, uint32 RHIPrimitives, uint32 RHIVertices, uint32 RHIInstances);
 	void EncodeBlit(char const* DrawCall);
+	void EncodeBlit(FString DrawCall);
 	void EncodeDispatch(char const* DrawCall);
 	void EncodePipeline(FMetalShaderPipeline* PipelineStat);
 	void EncodeFence(FMetalEventStats* Stat, EMTLFenceType Type);
@@ -408,6 +410,7 @@ public:
 	
 	void EncodeDraw(FMetalCommandBufferStats* CmdBufStats, char const* DrawCall, uint32 RHIPrimitives, uint32 RHIVertices, uint32 RHIInstances);
 	void EncodeBlit(FMetalCommandBufferStats* CmdBufStats, char const* DrawCall);
+	void EncodeBlit(FMetalCommandBufferStats* CmdBufStats, FString DrawCall);
 	void EncodeDispatch(FMetalCommandBufferStats* CmdBufStats, char const* DrawCall);
 	
 #if METAL_STATISTICS
diff --git a/Engine/Source/Runtime/Apple/MetalRHI/Private/MetalRenderPass.cpp b/Engine/Source/Runtime/Apple/MetalRHI/Private/MetalRenderPass.cpp
--- a/Engine/Source/Runtime/Apple/MetalRHI/Private/MetalRenderPass.cpp
+++ b/Engine/Source/Runtime/Apple/MetalRHI/Private/MetalRenderPass.cpp
@@ -894,14 +894,14 @@ void FMetalRenderPass::FillBuffer(FMetalBuffer const& Buffer, ns::Range Range, u
 	{
 		ConditionalSwitchToAsyncBlit();
 		TargetEncoder = PrologueEncoder.GetBlitCommandEncoder();
-		METAL_GPUPROFILE(FMetalProfiler::GetProfiler()->EncodeBlit(PrologueEncoder.GetCommandBufferStats(), __FUNCTION__));
+		METAL_GPUPROFILE(FMetalProfiler::GetProfiler()->EncodeBlit(PrologueEncoder.GetCommandBufferStats(), FString::Printf(TEXT("FillBuffer: %p %llu %llu"), Buffer.GetPtr(), Buffer.GetOffset() + Range.Location, Range.Length)));
 		METAL_DEBUG_LAYER(EMetalDebugLevelFastValidation, Debugging = PrologueEncoder.GetBlitCommandEncoderDebugging());
 	}
 	else
 	{
 		ConditionalSwitchToBlit();
 		TargetEncoder = CurrentEncoder.GetBlitCommandEncoder();
-		METAL_GPUPROFILE(FMetalProfiler::GetProfiler()->EncodeBlit(CurrentEncoder.GetCommandBufferStats(), __FUNCTION__));
+		METAL_GPUPROFILE(FMetalProfiler::GetProfiler()->EncodeBlit(CurrentEncoder.GetCommandBufferStats(), FString::Printf(TEXT("FillBuffer: %p %llu %llu"), Buffer.GetPtr(), Buffer.GetOffset() + Range.Location, Range.Length)));
 		METAL_DEBUG_LAYER(EMetalDebugLevelFastValidation, Debugging = CurrentEncoder.GetBlitCommandEncoderDebugging());
 	}
 	
@@ -989,14 +989,14 @@ void FMetalRenderPass::AsyncCopyFromBufferToBuffer(FMetalBuffer const& SourceBuf
 	{
 		ConditionalSwitchToAsyncBlit();
 		TargetEncoder = PrologueEncoder.GetBlitCommandEncoder();
-		METAL_GPUPROFILE(FMetalProfiler::GetProfiler()->EncodeBlit(PrologueEncoder.GetCommandBufferStats(), __FUNCTION__));
+		METAL_GPUPROFILE(FMetalProfiler::GetProfiler()->EncodeBlit(PrologueEncoder.GetCommandBufferStats(), FString::Printf(TEXT("AsyncCopyFromBufferToBuffer: %p %llu %llu"), DestinationBuffer.GetPtr(), DestinationBuffer.GetOffset() + DestinationOffset, Size)));
 		METAL_DEBUG_LAYER(EMetalDebugLevelFastValidation, Debugging = PrologueEncoder.GetBlitCommandEncoderDebugging());
 	}
 	else
 	{
 		ConditionalSwitchToBlit();
 		TargetEncoder = CurrentEncoder.GetBlitCommandEncoder();
-		METAL_GPUPROFILE(FMetalProfiler::GetProfiler()->EncodeBlit(CurrentEncoder.GetCommandBufferStats(), __FUNCTION__));
+		METAL_GPUPROFILE(FMetalProfiler::GetProfiler()->EncodeBlit(CurrentEncoder.GetCommandBufferStats(), FString::Printf(TEXT("AsyncCopyFromBufferToBuffer: %p %llu %llu"), DestinationBuffer.GetPtr(), DestinationBuffer.GetOffset() + DestinationOffset, Size)));
 		METAL_DEBUG_LAYER(EMetalDebugLevelFastValidation, Debugging = CurrentEncoder.GetBlitCommandEncoderDebugging());
 	}
 	
