More improvements to MetalProfiler.

diff --git a/Engine/Source/Runtime/Apple/MetalRHI/Private/MetalProfiler.cpp b/Engine/Source/Runtime/Apple/MetalRHI/Private/MetalProfiler.cpp
--- a/Engine/Source/Runtime/Apple/MetalRHI/Private/MetalProfiler.cpp
+++ b/Engine/Source/Runtime/Apple/MetalRHI/Private/MetalProfiler.cpp
@@ -457,6 +457,12 @@ FString IMetalStatsScope::GetJSONRepresentation(uint32 Pid)
 				Occupancy.Add(TEXT("Compute Shader Max theoretical occupancy"), TEXT("0"));
 
 				FString PSOStats;
+				
+				if (Parent.Len())
+				{
+					PSOStats += FString::Printf(TEXT(",\"Parent\":\"%s\""), *Parent);
+				}
+				
 				for(id Key in DrawStat.PSOPerformanceStats)
 				{
 					NSString* ShaderName = (NSString*)Key;
@@ -502,6 +508,11 @@ FString IMetalStatsScope::GetJSONRepresentation(uint32 Pid)
 			else
 			{
 				FString CustomCounters;
+				if (Parent.Len())
+				{
+					CustomCounters += FString::Printf(TEXT(",\"Parent\":\"%s\""), *Parent);
+				}
+				
 				TMap<FString, FMetalProfiler::EMTLCounterType> const& CounterTypes = FMetalProfiler::GetProfiler()->GetCounterTypes();
 				for(auto const& Pair : DrawStat.Counters)
 				{
@@ -590,7 +601,7 @@ FMetalEventStats::FMetalEventStats(const TCHAR* InName, FColor Color)
 	Name = InName;
 	
 	CPUThreadIndex = FPlatformTLS::GetCurrentThreadId();
-	GPUThreadIndex = 1;
+	GPUThreadIndex = 2;
 	
 	CPUStartTime = FPlatformTime::ToMilliseconds64(mach_absolute_time()) * 1000.0;
 	CPUEndTime = 0;
@@ -1045,7 +1056,7 @@ FMetalCommandBufferStats::FMetalCommandBufferStats(mtlpp::CommandBuffer const& B
 	}
 #endif
 	
-	Name = FString::Printf(TEXT("CommandBuffer: %s"), *FString(CmdBuffer.GetLabel().GetPtr()));
+	Name = FString::Printf(TEXT("CommandBuffer: %p %s"), CmdBuffer.GetPtr(), *FString(CmdBuffer.GetLabel().GetPtr()));
 	
 	CPUThreadIndex = FPlatformTLS::GetCurrentThreadId();
 	GPUThreadIndex = InGPUThreadIndex;
@@ -1117,7 +1128,8 @@ void FMetalCommandBufferStats::BeginEncoder(mtlpp::RenderCommandEncoder const& E
 {
 	check(!ActiveEncoderStats);
 	
-	ActiveEncoderStats = new FMetalEncoderStats(Encoder, GPUThreadIndex);
+	ActiveEncoderStats = new FMetalEncoderStats(Encoder, GPUThreadIndex+1);
+	ActiveEncoderStats->Parent = Name;
 	ActiveEncoderStats->CmdBufferStats = CmdBufferStats;
 	Children.Add(ActiveEncoderStats);
 	ActiveEncoderStats->Start(CmdBuffer);
@@ -1127,7 +1139,8 @@ void FMetalCommandBufferStats::BeginEncoder(mtlpp::BlitCommandEncoder const& Enc
 {
 	check(!ActiveEncoderStats);
 	
-	ActiveEncoderStats = new FMetalEncoderStats(Encoder, GPUThreadIndex);
+	ActiveEncoderStats = new FMetalEncoderStats(Encoder, GPUThreadIndex+1);
+	ActiveEncoderStats->Parent = Name;
 	ActiveEncoderStats->CmdBufferStats = CmdBufferStats;
 	Children.Add(ActiveEncoderStats);
 	ActiveEncoderStats->Start(CmdBuffer);
@@ -1137,7 +1150,8 @@ void FMetalCommandBufferStats::BeginEncoder(mtlpp::ComputeCommandEncoder const&
 {
 	check(!ActiveEncoderStats);
 	
-	ActiveEncoderStats = new FMetalEncoderStats(Encoder, GPUThreadIndex);
+	ActiveEncoderStats = new FMetalEncoderStats(Encoder, GPUThreadIndex+1);
+	ActiveEncoderStats->Parent = Name;
 	ActiveEncoderStats->CmdBufferStats = CmdBufferStats;
 	Children.Add(ActiveEncoderStats);
 	ActiveEncoderStats->Start(CmdBuffer);
@@ -1569,7 +1583,7 @@ void FMetalProfiler::EncodeFence(FMetalCommandBufferStats* CmdBufStats, const TC
 #if METAL_STATISTICS
 	if (MetalGPUProfilerIsInSafeThread() && Fence && bEnabled && StatisticsAPI && CmdBufStats->ActiveEncoderStats)
 	{
-		FMetalEventStats* Event = new FMetalEventStats(*FString::Printf(TEXT("%s: %s"), Name, *FString(Fence->GetLabel())), 0);
+		FMetalEventStats* Event = new FMetalEventStats(*FString::Printf(TEXT("%s: %s"), Name, *FString(Fence->GetLabel())), 1);
 		CmdBufStats->ActiveEncoderStats->EncodeFence(Event, Type);
 	}
 #endif
@@ -1659,13 +1673,20 @@ void FMetalProfiler::SaveTrace()
 		
 		for (int32 GPUIndex = 0; GPUIndex <= 0/*MaxGPUIndex*/; ++GPUIndex)
 		{
-			FString Output = FString::Printf(TEXT("{\"pid\":%d, \"tid\":%d, \"ph\": \"M\", \"name\": \"thread_name\", \"args\":{\"name\":\"GPU %d\"}},{\"pid\":%d, \"tid\":%d, \"ph\": \"M\", \"name\": \"thread_sort_index\", \"args\":{\"sort_index\": %d}},\n"),
+			FString Output = FString::Printf(TEXT("{\"pid\":%d, \"tid\":%d, \"ph\": \"M\", \"name\": \"thread_name\", \"args\":{\"name\":\"GPU %d Command Buffers\"}},{\"pid\":%d, \"tid\":%d, \"ph\": \"M\", \"name\": \"thread_sort_index\", \"args\":{\"sort_index\": %d}},\n"),
 											 Pid, GPUIndex, GPUIndex, Pid, GPUIndex, SortIndex
 											 );
 			
 			WriteString(OutputFile, TCHAR_TO_UTF8(*Output));
 			SortIndex++;
 			
+			Output = FString::Printf(TEXT("{\"pid\":%d, \"tid\":%d, \"ph\": \"M\", \"name\": \"thread_name\", \"args\":{\"name\":\"GPU %d Operations\"}},{\"pid\":%d, \"tid\":%d, \"ph\": \"M\", \"name\": \"thread_sort_index\", \"args\":{\"sort_index\": %d}},\n"),
+											 Pid, GPUIndex+SortIndex, GPUIndex, Pid, GPUIndex+SortIndex, SortIndex
+											 );
+			
+			WriteString(OutputFile, TCHAR_TO_UTF8(*Output));
+			SortIndex++;
+			
 			Output = FString::Printf(TEXT("{\"pid\":%d, \"tid\":%d, \"ph\": \"M\", \"name\": \"thread_name\", \"args\":{\"name\":\"Render Events %d\"}},{\"pid\":%d, \"tid\":%d, \"ph\": \"M\", \"name\": \"thread_sort_index\", \"args\":{\"sort_index\": %d}},\n"),
 									 Pid, GPUIndex+SortIndex, GPUIndex, Pid, GPUIndex+SortIndex, SortIndex
 									 );
@@ -1736,7 +1757,7 @@ void FMetalProfiler::SaveTrace()
 					
 					if (Pair.Key.Contains(TEXT("Device Utilization")))
 					{
-						Output += FString::Printf(TEXT("{\"pid\":%d, \"tid\":2, \"ph\": \"C\", \"name\": \"%s\", \"ts\": %llu, \"args\":{ \"%s\": %0.8f }},\n"),
+						Output += FString::Printf(TEXT("{\"pid\":%d, \"tid\":3, \"ph\": \"C\", \"name\": \"%s\", \"ts\": %llu, \"args\":{ \"%s\": %0.8f }},\n"),
 												   Pid,
 												   *Pair.Key,
 												   ChildStartCallTime,
@@ -1745,7 +1766,7 @@ void FMetalProfiler::SaveTrace()
 					}
 				}
 				
-				Output += FString::Printf(TEXT("{\"pid\":%d, \"tid\":2, \"ph\": \"X\", \"name\": \"Driver Stats\", \"ts\": %llu, \"dur\": %llu, \"args\":{\"num_child\":%d %s}},\n"),
+				Output += FString::Printf(TEXT("{\"pid\":%d, \"tid\":3, \"ph\": \"X\", \"name\": \"Driver Stats\", \"ts\": %llu, \"dur\": %llu, \"args\":{\"num_child\":%d %s}},\n"),
 											  Pid,
 											  ChildStartCallTime,
 											  ChildDrawCallTime,
diff --git a/Engine/Source/Runtime/Apple/MetalRHI/Private/MetalProfiler.h b/Engine/Source/Runtime/Apple/MetalRHI/Private/MetalProfiler.h
--- a/Engine/Source/Runtime/Apple/MetalRHI/Private/MetalProfiler.h
+++ b/Engine/Source/Runtime/Apple/MetalRHI/Private/MetalProfiler.h
@@ -165,6 +165,7 @@ private:
 struct IMetalStatsScope
 {
 	FString Name;
+	FString Parent;
 	TArray<IMetalStatsScope*> Children;
 	
 	uint64 CPUStartTime;
