Capture all the startup resource creation when profiling from application start.

diff --git a/Engine/Source/Runtime/Apple/MetalRHI/Private/MetalProfiler.cpp b/Engine/Source/Runtime/Apple/MetalRHI/Private/MetalProfiler.cpp
--- a/Engine/Source/Runtime/Apple/MetalRHI/Private/MetalProfiler.cpp
+++ b/Engine/Source/Runtime/Apple/MetalRHI/Private/MetalProfiler.cpp
@@ -197,7 +197,7 @@ mtlpp::CommandBufferHandler FMetalEventNode::Stop(void)
 
 bool MetalGPUProfilerIsInSafeThread()
 {
-	return IsInRHIThread() || IsInActualRenderingThread();
+	return (GIsMetalInitialized && !GIsRHIInitialized) || (IsInRHIThread() || IsInActualRenderingThread());
 }
 	
 /** Start this frame of per tracking */
diff --git a/Engine/Source/Runtime/Apple/MetalRHI/Private/MetalRHI.cpp b/Engine/Source/Runtime/Apple/MetalRHI/Private/MetalRHI.cpp
--- a/Engine/Source/Runtime/Apple/MetalRHI/Private/MetalRHI.cpp
+++ b/Engine/Source/Runtime/Apple/MetalRHI/Private/MetalRHI.cpp
@@ -824,8 +824,11 @@ FMetalDynamicRHI::FMetalDynamicRHI(ERHIFeatureLevel::Type RequestedFeatureLevel)
 	ImmediateContext.Profiler = nullptr;
 #if ENABLE_METAL_GPUPROFILE
 	ImmediateContext.Profiler = FMetalProfiler::CreateProfiler(ImmediateContext.Context);
+	if (ImmediateContext.Profiler)
+		ImmediateContext.Profiler->BeginFrame();
 #endif
 		
+		
 	// Notify all initialized FRenderResources that there's a valid RHI device to create their RHI resources for now.
 	for(TLinkedList<FRenderResource*>::TIterator ResourceIt(FRenderResource::GetResourceList());ResourceIt;ResourceIt.Next())
 	{
@@ -838,6 +841,11 @@ FMetalDynamicRHI::FMetalDynamicRHI(ERHIFeatureLevel::Type RequestedFeatureLevel)
 	}
 	
 	AsyncComputeContext = GSupportsEfficientAsyncCompute ? new FMetalRHIComputeContext(ImmediateContext.Profiler, new FMetalContext(ImmediateContext.Context->GetDevice(), ImmediateContext.Context->GetCommandQueue(), true)) : nullptr;
+
+#if ENABLE_METAL_GPUPROFILE
+		if (ImmediateContext.Profiler)
+			ImmediateContext.Profiler->EndFrame();
+#endif
 	}
 }
 
