We've made so many changes that we can now render directly into the Metal drawable texture even with an RHI thread enabled - movie capture and screenshots might not work without the intermediate back-buffer but it is faster to avoid it.

diff --git a/Engine/Source/Runtime/Apple/MetalRHI/Private/MetalRHI.cpp b/Engine/Source/Runtime/Apple/MetalRHI/Private/MetalRHI.cpp
--- a/Engine/Source/Runtime/Apple/MetalRHI/Private/MetalRHI.cpp
+++ b/Engine/Source/Runtime/Apple/MetalRHI/Private/MetalRHI.cpp
@@ -399,16 +399,6 @@ FMetalDynamicRHI::FMetalDynamicRHI(ERHIFeatureLevel::Type RequestedFeatureLevel)
 #endif
 		GSupportsEfficientAsyncCompute = GRHISupportsParallelRHIExecute && (IsRHIDeviceAMD() || PLATFORM_IOS); // Only AMD currently support async. compute and it requires parallel execution to be useful.
 		GSupportsParallelOcclusionQueries = GRHISupportsRHIThread;
-		
-		// We must always use an intermediate back-buffer for the RHI thread to work properly at present.
-		if(GRHISupportsRHIThread)
-		{
-			static auto CVarSupportsIntermediateBackBuffer = IConsoleManager::Get().FindConsoleVariable(TEXT("rhi.Metal.SupportsIntermediateBackBuffer"));
-			if(CVarSupportsIntermediateBackBuffer && CVarSupportsIntermediateBackBuffer->GetInt() != 1)
-			{
-				CVarSupportsIntermediateBackBuffer->Set(1);
-			}
-		}
 	}
 	else
 	{
