Disable the intermediate back-buffer in MetalRHI and instead render directly into the drawable texture which is much faster, make this texture available for read-back via the various mechanisms used for screenshots and movie capture using the mtlpp auto-released handles so that we don't need to worry about retain/release order.

diff --git a/Engine/Source/Runtime/Apple/MetalRHI/Private/MetalCommands.cpp b/Engine/Source/Runtime/Apple/MetalRHI/Private/MetalCommands.cpp
--- a/Engine/Source/Runtime/Apple/MetalRHI/Private/MetalCommands.cpp
+++ b/Engine/Source/Runtime/Apple/MetalRHI/Private/MetalCommands.cpp
@@ -273,7 +273,15 @@ void FMetalRHICommandContext::RHISetShaderTexture(FVertexShaderRHIParamRef Verte
 	FMetalSurface* Surface = GetMetalSurfaceFromRHITexture(NewTextureRHI);
 	if (Surface != nullptr)
 	{
-		Context->GetCurrentState().SetShaderTexture(SF_Vertex, Surface->Texture, TextureIndex);
+        if (Surface->Texture || !(Surface->Flags & TexCreate_Presentable))
+        {
+            Context->GetCurrentState().SetShaderTexture(SF_Vertex, Surface->Texture, TextureIndex);
+        }
+        else
+        {
+            FMetalTexture Tex = Surface->GetCurrentTexture();
+            Context->GetCurrentState().SetShaderTexture(SF_Vertex, Tex, TextureIndex);
+        }
 	}
 	else
 	{
@@ -287,8 +295,16 @@ void FMetalRHICommandContext::RHISetShaderTexture(FHullShaderRHIParamRef HullSha
 	@autoreleasepool {
 	FMetalSurface* Surface = GetMetalSurfaceFromRHITexture(NewTextureRHI);
 	if (Surface != nullptr)
-	{
-		Context->GetCurrentState().SetShaderTexture(SF_Hull, Surface->Texture, TextureIndex);
+    {
+        if (Surface->Texture || !(Surface->Flags & TexCreate_Presentable))
+        {
+            Context->GetCurrentState().SetShaderTexture(SF_Hull, Surface->Texture, TextureIndex);
+        }
+        else
+        {
+            FMetalTexture Tex = Surface->GetCurrentTexture();
+            Context->GetCurrentState().SetShaderTexture(SF_Hull, Tex, TextureIndex);
+        }
 	}
 	else
 	{
@@ -302,8 +318,16 @@ void FMetalRHICommandContext::RHISetShaderTexture(FDomainShaderRHIParamRef Domai
 	@autoreleasepool {
 	FMetalSurface* Surface = GetMetalSurfaceFromRHITexture(NewTextureRHI);
 	if (Surface != nullptr)
-	{
-		Context->GetCurrentState().SetShaderTexture(SF_Domain, Surface->Texture, TextureIndex);
+    {
+        if (Surface->Texture || !(Surface->Flags & TexCreate_Presentable))
+        {
+            Context->GetCurrentState().SetShaderTexture(SF_Domain, Surface->Texture, TextureIndex);
+        }
+        else
+        {
+            FMetalTexture Tex = Surface->GetCurrentTexture();
+            Context->GetCurrentState().SetShaderTexture(SF_Domain, Tex, TextureIndex);
+        }
 	}
 	else
 	{
@@ -323,8 +347,16 @@ void FMetalRHICommandContext::RHISetShaderTexture(FPixelShaderRHIParamRef PixelS
 	@autoreleasepool {
 	FMetalSurface* Surface = GetMetalSurfaceFromRHITexture(NewTextureRHI);
 	if (Surface != nullptr)
-	{
-		Context->GetCurrentState().SetShaderTexture(SF_Pixel, Surface->Texture, TextureIndex);
+    {
+        if (Surface->Texture || !(Surface->Flags & TexCreate_Presentable))
+        {
+            Context->GetCurrentState().SetShaderTexture(SF_Pixel, Surface->Texture, TextureIndex);
+        }
+        else
+        {
+            FMetalTexture Tex = Surface->GetCurrentTexture();
+            Context->GetCurrentState().SetShaderTexture(SF_Pixel, Tex, TextureIndex);
+        }
 	}
 	else
 	{
@@ -338,8 +370,16 @@ void FMetalRHICommandContext::RHISetShaderTexture(FComputeShaderRHIParamRef Comp
 	@autoreleasepool {
 	FMetalSurface* Surface = GetMetalSurfaceFromRHITexture(NewTextureRHI);
 	if (Surface != nullptr)
-	{
-		Context->GetCurrentState().SetShaderTexture(SF_Compute, Surface->Texture, TextureIndex);
+    {
+        if (Surface->Texture || !(Surface->Flags & TexCreate_Presentable))
+        {
+            Context->GetCurrentState().SetShaderTexture(SF_Compute, Surface->Texture, TextureIndex);
+        }
+        else
+        {
+            FMetalTexture Tex = Surface->GetCurrentTexture();
+            Context->GetCurrentState().SetShaderTexture(SF_Compute, Tex, TextureIndex);
+        }
 	}
 	else
 	{
diff --git a/Engine/Source/Runtime/Apple/MetalRHI/Private/MetalContext.cpp b/Engine/Source/Runtime/Apple/MetalRHI/Private/MetalContext.cpp
--- a/Engine/Source/Runtime/Apple/MetalRHI/Private/MetalContext.cpp
+++ b/Engine/Source/Runtime/Apple/MetalRHI/Private/MetalContext.cpp
@@ -13,11 +13,11 @@
 #include "MetalProfiler.h"
 #include "MetalCommandBuffer.h"
 
-int32 GMetalSupportsIntermediateBackBuffer = PLATFORM_MAC ? 1 : 0;
+int32 GMetalSupportsIntermediateBackBuffer = 0;
 static FAutoConsoleVariableRef CVarMetalSupportsIntermediateBackBuffer(
 	TEXT("rhi.Metal.SupportsIntermediateBackBuffer"),
 	GMetalSupportsIntermediateBackBuffer,
-	TEXT("When enabled (> 0) allocate an intermediate texture to use as the back-buffer & blit from there into the actual device back-buffer, thereby allowing screenshots & video capture that would otherwise be impossible as the texture required has already been released back to the OS as required by Metal's API. (Off by default (0) on iOS/tvOS but enabled (1) on Mac)"), ECVF_ReadOnly);
+	TEXT("When enabled (> 0) allocate an intermediate texture to use as the back-buffer & blit from there into the actual device back-buffer, this is required if we use the experimental separate presentation thread. (Off by default (0))"), ECVF_ReadOnly);
 
 int32 GMetalSeparatePresentThread = 0;
 static FAutoConsoleVariableRef CVarMetalSeparatePresentThread(
@@ -25,7 +25,6 @@ static FAutoConsoleVariableRef CVarMetalSeparatePresentThread(
 	GMetalSeparatePresentThread,
 	TEXT("When enabled (> 0) requires rhi.Metal.SupportsIntermediateBackBuffer be enabled and will cause two intermediate back-buffers be allocated so that the presentation of frames to the screen can be run on a separate thread.\n")
 	TEXT("This option uncouples the Render/RHI thread from calls to -[CAMetalLayer nextDrawable] and will run arbitrarily fast by rendering but not waiting to present all frames. This is equivalent to running without V-Sync, but without the screen tearing.\n")
-	TEXT("On macOS 10.12 this will not be beneficial, but on later macOS versions this is the only way to ensure that we keep the CPU & GPU saturated with commands and don't ever stall waiting for V-Sync.\n")
 	TEXT("On iOS/tvOS this is the only way to run without locking the CPU to V-Sync somewhere - this shouldn't be used in a shipping title without understanding the power/heat implications.\n")
 	TEXT("(Off by default (0))"), ECVF_ReadOnly);
 
diff --git a/Engine/Source/Runtime/Apple/MetalRHI/Private/MetalRenderTarget.cpp b/Engine/Source/Runtime/Apple/MetalRHI/Private/MetalRenderTarget.cpp
--- a/Engine/Source/Runtime/Apple/MetalRHI/Private/MetalRenderTarget.cpp
+++ b/Engine/Source/Runtime/Apple/MetalRHI/Private/MetalRenderTarget.cpp
@@ -592,7 +592,7 @@ void FMetalDynamicRHI::RHIReadSurfaceData(FTextureRHIParamRef TextureRHI, FIntRe
 	FMetalTexture Texture = Surface->Texture;
     if(!Texture && (Surface->Flags & TexCreate_Presentable))
     {
-        Texture = Surface->GetDrawableTexture();
+        Texture = Surface->GetCurrentTexture();
     }
     if(!Texture)
     {
@@ -736,7 +736,7 @@ void FMetalDynamicRHI::RHIReadSurfaceFloatData(FTextureRHIParamRef TextureRHI, F
     FMetalTexture Texture = Surface->Texture;
     if(!Texture && (Surface->Flags & TexCreate_Presentable))
     {
-		Texture = Surface->GetDrawableTexture();
+		Texture = Surface->GetCurrentTexture();
     }
     if(!Texture)
     {
diff --git a/Engine/Source/Runtime/Apple/MetalRHI/Private/MetalTexture.cpp b/Engine/Source/Runtime/Apple/MetalRHI/Private/MetalTexture.cpp
--- a/Engine/Source/Runtime/Apple/MetalRHI/Private/MetalTexture.cpp
+++ b/Engine/Source/Runtime/Apple/MetalRHI/Private/MetalTexture.cpp
@@ -1649,6 +1649,17 @@ FMetalTexture FMetalSurface::GetDrawableTexture()
 	return Texture;
 }
 
+ns::AutoReleased<FMetalTexture> FMetalSurface::GetCurrentTexture()
+{
+	ns::AutoReleased<FMetalTexture> Tex;
+	if (Viewport && (Flags & TexCreate_Presentable))
+	{
+		check(Viewport);
+		Tex = Viewport->GetCurrentTexture(EMetalViewportAccessRHI);
+	}
+	return Tex;
+}
+
 
 /*-----------------------------------------------------------------------------
  Texture allocator support.
diff --git a/Engine/Source/Runtime/Apple/MetalRHI/Private/MetalViewport.cpp b/Engine/Source/Runtime/Apple/MetalRHI/Private/MetalViewport.cpp
--- a/Engine/Source/Runtime/Apple/MetalRHI/Private/MetalViewport.cpp
+++ b/Engine/Source/Runtime/Apple/MetalRHI/Private/MetalViewport.cpp
@@ -364,9 +364,15 @@ FMetalTexture FMetalViewport::GetDrawableTexture(EMetalViewportAccessFlag Access
 		}
 	}
 #endif
+	DrawableTextures[Accessor] = CurrentDrawable.texture;
 	return CurrentDrawable.texture;
 }
 
+ns::AutoReleased<FMetalTexture> FMetalViewport::GetCurrentTexture(EMetalViewportAccessFlag Accessor)
+{
+	return DrawableTextures[Accessor];
+}
+
 void FMetalViewport::ReleaseDrawable()
 {
 	if (!GMetalSeparatePresentThread)
diff --git a/Engine/Source/Runtime/Apple/MetalRHI/Public/MetalResources.h b/Engine/Source/Runtime/Apple/MetalRHI/Public/MetalResources.h
--- a/Engine/Source/Runtime/Apple/MetalRHI/Public/MetalResources.h
+++ b/Engine/Source/Runtime/Apple/MetalRHI/Public/MetalResources.h
@@ -500,6 +500,7 @@ public:
 	
 	/** Gets the drawable texture if this is a back-buffer surface. */
 	FMetalTexture GetDrawableTexture();
+	ns::AutoReleased<FMetalTexture> GetCurrentTexture();
 
 	FMetalTexture Reallocate(FMetalTexture Texture, mtlpp::TextureUsage UsageModifier);
 	void ReplaceTexture(FMetalContext& Context, FMetalTexture OldTexture, FMetalTexture NewTexture);
diff --git a/Engine/Source/Runtime/Apple/MetalRHI/Public/MetalViewport.h b/Engine/Source/Runtime/Apple/MetalRHI/Public/MetalViewport.h
--- a/Engine/Source/Runtime/Apple/MetalRHI/Public/MetalViewport.h
+++ b/Engine/Source/Runtime/Apple/MetalRHI/Public/MetalViewport.h
@@ -39,6 +39,7 @@ public:
 	TRefCountPtr<FMetalTexture2D> GetBackBuffer(EMetalViewportAccessFlag Accessor) const;
 	mtlpp::Drawable GetDrawable(EMetalViewportAccessFlag Accessor);
 	FMetalTexture GetDrawableTexture(EMetalViewportAccessFlag Accessor);
+	ns::AutoReleased<FMetalTexture> GetCurrentTexture(EMetalViewportAccessFlag Accessor);
 	void ReleaseDrawable(void);
 
 	// supports pulling the raw MTLTexture
@@ -67,6 +68,8 @@ private:
 	TRefCountPtr<FMetalTexture2D> BackBuffer[2];
 	mutable FCriticalSection Mutex;
 	
+	ns::AutoReleased<FMetalTexture> DrawableTextures[2];
+	
 	uint32 DisplayID;
 	FMetalViewportPresentHandler Block;
 	volatile int32 FrameAvailable;
