Build fix.

diff --git a/Engine/Source/Runtime/Apple/MetalRHI/Private/MetalProfiler.cpp b/Engine/Source/Runtime/Apple/MetalRHI/Private/MetalProfiler.cpp
--- a/Engine/Source/Runtime/Apple/MetalRHI/Private/MetalProfiler.cpp
+++ b/Engine/Source/Runtime/Apple/MetalRHI/Private/MetalProfiler.cpp
@@ -1538,6 +1538,15 @@ void FMetalProfiler::RemoveCounter(NSString* Counter)
 	CounterTypes.Remove(FString(Counter));
 }
 
+void FMetalProfiler::EncodeFence(FMetalCommandBufferStats* CmdBufStats, const TCHAR* Name, FMetalFence* Fence, EMTLFenceType Type)
+{
+	if (MetalGPUProfilerIsInSafeThread() && Fence && bEnabled && StatisticsAPI && CmdBufStats->ActiveEncoderStats)
+	{
+		FMetalEventStats* Event = new FMetalEventStats(*FString::Printf(TEXT("%s: %s"), Name, *FString(Fence->GetLabel())), 1);
+		CmdBufStats->ActiveEncoderStats->EncodeFence(Event, Type);
+	}
+}
+
 void FMetalProfiler::DumpPipeline(FMetalShaderPipeline* PipelineStat)
 {
 	Pipelines.Add(PipelineStat);
@@ -1578,17 +1587,6 @@ void FMetalProfiler::AddCommandBuffer(FMetalCommandBufferStats *CommandBuffer)
 	}
 }
 
-void FMetalProfiler::EncodeFence(FMetalCommandBufferStats* CmdBufStats, const TCHAR* Name, FMetalFence* Fence, EMTLFenceType Type)
-{
-#if METAL_STATISTICS
-	if (MetalGPUProfilerIsInSafeThread() && Fence && bEnabled && StatisticsAPI && CmdBufStats->ActiveEncoderStats)
-	{
-		FMetalEventStats* Event = new FMetalEventStats(*FString::Printf(TEXT("%s: %s"), Name, *FString(Fence->GetLabel())), 1);
-		CmdBufStats->ActiveEncoderStats->EncodeFence(Event, Type);
-	}
-#endif
-}
-
 void FMetalProfiler::PushEvent(const TCHAR *Name, FColor Color)
 {
 #if METAL_STATISTICS
