More profiler options.

diff --git a/Engine/Source/Runtime/Apple/MetalRHI/Private/MetalProfiler.cpp b/Engine/Source/Runtime/Apple/MetalRHI/Private/MetalProfiler.cpp
--- a/Engine/Source/Runtime/Apple/MetalRHI/Private/MetalProfiler.cpp
+++ b/Engine/Source/Runtime/Apple/MetalRHI/Private/MetalProfiler.cpp
@@ -46,6 +46,22 @@ int64 volatile GMetalGPUWorkTime = 0;
 int64 volatile GMetalGPUIdleTime = 0;
 int64 volatile GMetalPresentTime = 0;
 
+#if METAL_STATISTICS
+int32 GMetalProfilerStatisticsTiming = 1;
+static FAutoConsoleVariableRef CVarMetalProfilerStatisticsTiming(
+	TEXT("rhi.Metal.StatisticsTiming"),
+	GMetalProfilerStatisticsTiming,
+	TEXT("Use MetalStatistics timing rather than command-buffer timing.\n")
+	TEXT("(On by default (1))"));
+
+static int32 GMetalProfilerStatisticsRenderEvents = 1;
+static FAutoConsoleVariableRef CVarMetalProfilerStatisticsRenderEvents(
+	TEXT("rhi.Metal.StatisticsRenderEvents"),
+	GMetalProfilerStatisticsRenderEvents,
+	TEXT("Emit render-events to the Metal Profiler.\n")
+	TEXT("(On by default (1))"));
+#endif
+
 void WriteString(FArchive* OutputFile, const char* String)
 {
 	OutputFile->Serialize((void*)String, sizeof(ANSICHAR)*FCStringAnsi::Strlen(String));
@@ -66,7 +82,7 @@ void FMetalEventNode::StartTiming()
 	EndTime = 0;
 #if METAL_STATISTICS
 	IMetalStatistics* Stats = Context->GetCommandQueue().GetStatistics();
-	if (Stats)
+	if (Stats && GMetalProfilerStatisticsTiming)
 	{
 		id<IMetalStatisticsSamples> StatSample = Stats->GetLastStatisticsSample(Context->GetCurrentCommandBuffer().GetPtr());
 		if (!StatSample)
@@ -76,7 +92,7 @@ void FMetalEventNode::StartTiming()
 		}
 		check(StatSample);
 		[StatSample retain];
-		
+
 		Context->GetCurrentCommandBuffer().AddCompletedHandler(^(const mtlpp::CommandBuffer &) {
 			if (StatSample.Count > 0)
 			{
@@ -114,7 +130,7 @@ void FMetalEventNode::StopTiming()
 {
 #if METAL_STATISTICS
 	IMetalStatistics* Stats = Context->GetCommandQueue().GetStatistics();
-	if (Context->GetCommandQueue().GetStatistics())
+	if (Stats && GMetalProfilerStatisticsTiming)
 	{
 		id<IMetalStatisticsSamples> StatSample = Stats->GetLastStatisticsSample(Context->GetCurrentCommandBuffer().GetPtr());
 		if (!StatSample)
@@ -124,14 +140,14 @@ void FMetalEventNode::StopTiming()
 		}
 		check(StatSample);
 		[StatSample retain];
-		
+
 		Context->GetCurrentCommandBuffer().AddCompletedHandler(^(const mtlpp::CommandBuffer &) {
 			if (StatSample.Count > 0)
 			{
 				EndTime = StatSample.Array[0];
 			}
 			[StatSample release];
-			
+
 			if (bRoot)
 			{
 				// We have a different mechanism for the overall frametime that works even with empty encoders and that doesn't report any GPU idle time between frames, we only use the fallback code below on older OSes.
@@ -140,14 +156,14 @@ void FMetalEventNode::StopTiming()
 					uint32 Time = FMath::TruncToInt( double(GetTiming()) / double(FPlatformTime::GetSecondsPerCycle()) );
 					FPlatformAtomics::AtomicStore_Relaxed((int32*)&GGPUFrameTime, (int32)Time);
 				}
-				
+
 				if(!bFullProfiling)
 				{
 					delete this;
 				}
 			}
 		});
-		
+
 		bool const bWait = Wait();
 		if (bWait)
 		{
@@ -1282,8 +1298,10 @@ void FMetalProfiler::AddDisplayVBlank(uint32 DisplayID, double OutputSeconds, do
 FMetalProfiler::FMetalProfiler(FMetalContext* Context)
 : FMetalGPUProfiler(Context)
 #if METAL_STATISTICS
+, StatsGranularity(EMetalSampleGranularityOperation)
 , NewCounters([NSMutableArray new])
 , StatisticsAPI(Context->GetCommandQueue().GetStatistics())
+, bChangeGranularity(true)
 #endif
 , bEnabled(false)
 {
@@ -1377,11 +1395,12 @@ void FMetalProfiler::BeginFrame()
 		if (bRequestStartCapture && !bEnabled)
 		{
 #if METAL_STATISTICS
-			if (StatisticsAPI && NewCounters)
+			if (StatisticsAPI && (bChangeGranularity || NewCounters))
 			{
 				StatisticsAPI->FinishSamplingStatistics();
-				StatisticsAPI->BeginSamplingStatistics(NewCounters);
+				StatisticsAPI->BeginSamplingStatistics(StatsGranularity, NewCounters);
 				Context->SubmitCommandBufferAndWait();
+				bChangeGranularity = false;
 			}
 #endif
 			
@@ -1538,6 +1557,15 @@ void FMetalProfiler::RemoveCounter(NSString* Counter)
 	CounterTypes.Remove(FString(Counter));
 }
 
+void FMetalProfiler::SetGranularity(EMetalSampleGranularity Sample)
+{
+	if (StatsGranularity != Sample)
+	{
+		StatsGranularity = Sample;
+		bChangeGranularity = true;
+	}
+}
+
 void FMetalProfiler::EncodeFence(FMetalCommandBufferStats* CmdBufStats, const TCHAR* Name, FMetalFence* Fence, EMTLFenceType Type)
 {
 	if (MetalGPUProfilerIsInSafeThread() && Fence && bEnabled && StatisticsAPI && CmdBufStats->ActiveEncoderStats)
@@ -1590,7 +1618,7 @@ void FMetalProfiler::AddCommandBuffer(FMetalCommandBufferStats *CommandBuffer)
 void FMetalProfiler::PushEvent(const TCHAR *Name, FColor Color)
 {
 #if METAL_STATISTICS
-	if (MetalGPUProfilerIsInSafeThread() && bEnabled && StatisticsAPI)
+	if (MetalGPUProfilerIsInSafeThread() && bEnabled && StatisticsAPI && GMetalProfilerStatisticsRenderEvents)
 	{
 		if (!Context->GetCurrentCommandBuffer().GetPtr() || !StatisticsAPI->GetLastStatisticsSample(Context->GetCurrentCommandBuffer().GetPtr()))
 		{
@@ -1608,13 +1636,13 @@ void FMetalProfiler::PushEvent(const TCHAR *Name, FColor Color)
 void FMetalProfiler::PopEvent()
 {
 #if METAL_STATISTICS
-	if (MetalGPUProfilerIsInSafeThread() && bEnabled && StatisticsAPI&& ActiveEvents.Num())
+	if (MetalGPUProfilerIsInSafeThread() && bEnabled && StatisticsAPI&& ActiveEvents.Num() && GMetalProfilerStatisticsRenderEvents)
 	{
 		if (!Context->GetCurrentCommandBuffer().GetPtr() || !StatisticsAPI->GetLastStatisticsSample(Context->GetCurrentCommandBuffer().GetPtr()))
 		{
 			Context->GetCurrentRenderPass().InsertDebugEncoder();
 		}
-		
+
 		FMetalEventStats* Event = ActiveEvents.Pop();
 		Event->End(Context->GetCurrentCommandBuffer());
 		FrameEvents.Add(Event);
@@ -2021,6 +2049,22 @@ static void HandleMetalProfileCommand(const TArray<FString>& Args, UWorld*, FOut
 			}
 		}
 	}
+	else if (Param == TEXT("GRANULARITY"))
+	{
+		IMetalStatistics* Stats = FMetalProfiler::GetStatistics();
+		if (Stats)
+		{
+			FString SamplePos = Args[1];
+			if (SamplePos == TEXT("ENCODER"))
+			{
+				FMetalProfiler::GetProfiler()->SetGranularity(EMetalSampleGranularityEncoder);
+			}
+			else if (SamplePos == TEXT("OPERATION"))
+			{
+				FMetalProfiler::GetProfiler()->SetGranularity(EMetalSampleGranularityOperation);
+			}
+		}
+	}
 #endif
 	else
 	{
diff --git a/Engine/Source/Runtime/Apple/MetalRHI/Private/MetalQuery.cpp b/Engine/Source/Runtime/Apple/MetalRHI/Private/MetalQuery.cpp
--- a/Engine/Source/Runtime/Apple/MetalRHI/Private/MetalQuery.cpp
+++ b/Engine/Source/Runtime/Apple/MetalRHI/Private/MetalQuery.cpp
@@ -9,6 +9,10 @@
 #include "MetalLLM.h"
 #include "MetalCommandBuffer.h"
 
+#if METAL_STATISTICS
+extern int32 GMetalProfilerStatisticsTiming;
+#endif
+
 void FMetalQueryBufferPool::Allocate(FMetalQueryResult& NewQuery)
 {
 	FMetalQueryBuffer* QB = IsValidRef(CurrentBuffer) ? CurrentBuffer.GetReference() : GetCurrentQueryBuffer();
@@ -243,7 +247,7 @@ void FMetalRenderQuery::End(FMetalContext* Context)
 			
 #if METAL_STATISTICS
 			class IMetalStatistics* Stats = Context->GetCommandQueue().GetStatistics();
-			if (Stats)
+			if (Stats && GMetalProfilerStatisticsTiming)
 			{
 				id<IMetalStatisticsSamples> StatSample = Stats->GetLastStatisticsSample(Context->GetCurrentCommandBuffer().GetPtr());
 				if (!StatSample)
