Separate Metal buffer tracking stats from LLM.

diff --git a/Engine/Source/Runtime/Apple/MetalRHI/Private/MetalVertexBuffer.cpp b/Engine/Source/Runtime/Apple/MetalRHI/Private/MetalVertexBuffer.cpp
--- a/Engine/Source/Runtime/Apple/MetalRHI/Private/MetalVertexBuffer.cpp
+++ b/Engine/Source/Runtime/Apple/MetalRHI/Private/MetalVertexBuffer.cpp
@@ -14,7 +14,7 @@
 #include <objc/runtime.h>
 #include "HAL/LowLevelMemTracker.h"
 
-#if ENABLE_LOW_LEVEL_MEM_TRACKER || STATS
+#if ENABLE_LOW_LEVEL_MEM_TRACKER
 #define METAL_LLM_BUFFER_SCOPE(Type) \
 	ELLMTag Tag; \
 	switch(Type)	{ \
@@ -24,6 +24,11 @@
 		default: Tag = ELLMTag::VertexBuffer; break; \
 	} \
 	LLM_SCOPE(Tag)
+#else
+#define METAL_LLM_BUFFER_SCOPE(Type)
+#endif
+
+#if STATS
 #define METAL_INC_DWORD_STAT_BY(Type, Name, Size) \
 	switch(Type)	{ \
 		case RRT_UniformBuffer: INC_DWORD_STAT_BY(STAT_MetalUniform##Name, Size); break; \
@@ -33,7 +38,6 @@
 		default: break; \
 	}
 #else
-#define METAL_LLM_BUFFER_SCOPE(Type)
 #define METAL_INC_DWORD_STAT_BY(Type, Name, Size)
 #endif
 
