And another missing file.

diff --git a/Engine/Source/Runtime/Apple/MetalRHI/Private/MetalProfiler.h b/Engine/Source/Runtime/Apple/MetalRHI/Private/MetalProfiler.h
--- a/Engine/Source/Runtime/Apple/MetalRHI/Private/MetalProfiler.h
+++ b/Engine/Source/Runtime/Apple/MetalRHI/Private/MetalProfiler.h
@@ -434,6 +434,7 @@ public:
 	
 	void AddCounter(NSString* Counter, EMTLCounterType Type);
 	void RemoveCounter(NSString* Counter);
+	void SetGranularity(EMetalSampleGranularity Sample);
 	TMap<FString, EMTLCounterType> const& GetCounterTypes() const { return CounterTypes; }
 	
 	void EncodeFence(FMetalCommandBufferStats* CmdBufStats, const TCHAR* Name, FMetalFence* Fence, EMTLFenceType Type);
@@ -452,12 +453,14 @@ public:
 private:
 	FCriticalSection Mutex;
 #if METAL_STATISTICS
+	EMetalSampleGranularity StatsGranularity;
 	NSMutableArray* NewCounters;
 	TMap<FString, EMTLCounterType> CounterTypes;
 	IMetalStatistics* StatisticsAPI;
 	TArray<FMetalEventStats*> FrameEvents;
 	TArray<FMetalEventStats*> ActiveEvents;
 	TSet<FMetalShaderPipeline*> Pipelines;
+	bool bChangeGranularity;
 #endif
 	
 	TArray<FMetalCommandBufferStats*> TracedBuffers;
