Don't allow prologue work to overlap with previous render passes, it causes a lot of problems.

diff --git a/Engine/Source/Runtime/Apple/MetalRHI/Private/MetalRenderPass.cpp b/Engine/Source/Runtime/Apple/MetalRHI/Private/MetalRenderPass.cpp
--- a/Engine/Source/Runtime/Apple/MetalRHI/Private/MetalRenderPass.cpp
+++ b/Engine/Source/Runtime/Apple/MetalRHI/Private/MetalRenderPass.cpp
@@ -1046,6 +1046,7 @@ FMetalFence* FMetalRenderPass::End(void)
 		CurrentEncoder.WaitForFence(CurrentEncoderFence);
 		CurrentEncoder.WaitForFence(ParallelPassEndFence);
 		CurrentEncoderFence = CurrentEncoder.EndEncoding();
+		LastPrologueEncoderFence = CurrentEncoderFence;
 	}
 	else if (PassStartFence)
 	{
@@ -1054,6 +1055,7 @@ FMetalFence* FMetalRenderPass::End(void)
 		CurrentEncoder.WaitForFence(CurrentEncoderFence);
 		CurrentEncoder.WaitForFence(ParallelPassEndFence);
 		CurrentEncoderFence = CurrentEncoder.EndEncoding();
+		LastPrologueEncoderFence = CurrentEncoderFence;
 	}
 	else if (ParallelPassEndFence)
 	{
@@ -1062,6 +1064,7 @@ FMetalFence* FMetalRenderPass::End(void)
 		CurrentEncoder.WaitForFence(CurrentEncoderFence);
 		CurrentEncoder.WaitForFence(ParallelPassEndFence);
 		CurrentEncoderFence = CurrentEncoder.EndEncoding();
+		LastPrologueEncoderFence = CurrentEncoderFence;
 	}
 	
 	ParallelPassEndFence = nullptr;
