Fence validaiton should be controlled by the same system as all the other validation code.

diff --git a/Engine/Source/Runtime/Apple/MetalRHI/Private/MetalFence.cpp b/Engine/Source/Runtime/Apple/MetalRHI/Private/MetalFence.cpp
--- a/Engine/Source/Runtime/Apple/MetalRHI/Private/MetalFence.cpp
+++ b/Engine/Source/Runtime/Apple/MetalRHI/Private/MetalFence.cpp
@@ -141,11 +141,14 @@ FMetalFence* FMetalFencePool::AllocateFence()
 	if (Fence)
 	{
 		FPlatformAtomics::InterlockedDecrement(&Count);
+#if METAL_DEBUG_OPTIONS
+		if (GMetalRuntimeDebugLevel >= EMetalDebugLevelValidation)
 		{
 			FScopeLock Lock(&Mutex);
 			check(Fences.Contains(Fence));
 			Fences.Remove(Fence);
 		}
+#endif
 	}
 	check(Fence);
 	Fence->Reset();
@@ -171,14 +174,15 @@ void FMetalFencePool::ReleaseFence(FMetalFence* const InFence)
 {
 	if (InFence)
 	{
-		FMetalFence::ValidateUsage(InFence);
-		
+#if METAL_DEBUG_OPTIONS
+		if (GMetalRuntimeDebugLevel >= EMetalDebugLevelValidation)
 		{
 			FScopeLock Lock(&Mutex);
+			FMetalFence::ValidateUsage(InFence);
 			check(!Fences.Contains(InFence));
 			Fences.Add(InFence);
 		}
-		
+#endif
 		FPlatformAtomics::InterlockedIncrement(&Count);
 		check(Count <= FMetalFencePool::NumFences);
 		Lifo.Push(InFence);
diff --git a/Engine/Source/Runtime/Apple/MetalRHI/Private/MetalFence.h b/Engine/Source/Runtime/Apple/MetalRHI/Private/MetalFence.h
--- a/Engine/Source/Runtime/Apple/MetalRHI/Private/MetalFence.h
+++ b/Engine/Source/Runtime/Apple/MetalRHI/Private/MetalFence.h
@@ -198,8 +198,10 @@ public:
 private:
 	int32 Count;
 	mtlpp::Device Device;
+#if METAL_DEBUG_OPTIONS
 	TSet<FMetalFence*> Fences;
 	FCriticalSection Mutex;
+#endif
 	TLockFreePointerListLIFO<FMetalFence> Lifo;
 	static FMetalFencePool sSelf;
 };
diff --git a/Engine/Source/Runtime/Apple/MetalRHI/Private/MetalRenderPass.cpp b/Engine/Source/Runtime/Apple/MetalRHI/Private/MetalRenderPass.cpp
--- a/Engine/Source/Runtime/Apple/MetalRHI/Private/MetalRenderPass.cpp
+++ b/Engine/Source/Runtime/Apple/MetalRHI/Private/MetalRenderPass.cpp
@@ -101,12 +101,12 @@ void FMetalRenderPass::Wait(FMetalFence* Fence)
 		if (PrologueEncoder.IsBlitCommandEncoderActive() || PrologueEncoder.IsComputeCommandEncoderActive())
 		{
 			PrologueEncoder.WaitForFence(Fence);
-			FMetalFence::ValidateUsage(Fence);
+			METAL_DEBUG_LAYER(EMetalDebugLevelValidation, FMetalFence::ValidateUsage(Fence));
 		}
 		else if (CurrentEncoder.IsRenderCommandEncoderActive() || CurrentEncoder.IsBlitCommandEncoderActive() || CurrentEncoder.IsComputeCommandEncoderActive())
 		{
 			CurrentEncoder.WaitForFence(Fence);
-			FMetalFence::ValidateUsage(Fence);
+			METAL_DEBUG_LAYER(EMetalDebugLevelValidation, FMetalFence::ValidateUsage(Fence));
 		}
         else
         {
