More changes to make MTLFence a little bit more CPU efficient but it still adds noticeable overheads once actually encoded.

diff --git a/Engine/Source/Runtime/Apple/MetalRHI/Private/MetalCommandEncoder.cpp b/Engine/Source/Runtime/Apple/MetalRHI/Private/MetalCommandEncoder.cpp
--- a/Engine/Source/Runtime/Apple/MetalRHI/Private/MetalCommandEncoder.cpp
+++ b/Engine/Source/Runtime/Apple/MetalRHI/Private/MetalCommandEncoder.cpp
@@ -721,22 +721,22 @@ void FMetalCommandEncoder::UpdateFence(FMetalFence* Fence)
 		mtlpp::Fence InnerFence = METAL_DEBUG_OPTION(CommandList.GetCommandQueue().GetRuntimeDebuggingLevel() >= EMetalDebugLevelValidation ? mtlpp::Fence(((FMetalDebugFence*)Fence->GetPtr()).Inner) :) *Fence;
 		if (RenderCommandEncoder)
 		{
-//			Fence->Write(mtlpp::RenderStages::Vertex);
-			Fence->Write(mtlpp::RenderStages::Fragment);
+//			METAL_DEBUG_LAYER(EMetalDebugLevelFastValidation, Fence->Write(mtlpp::RenderStages::Vertex));
+			METAL_DEBUG_LAYER(EMetalDebugLevelFastValidation, Fence->Write(mtlpp::RenderStages::Fragment));
 			RenderCommandEncoder.UpdateFence(InnerFence, (mtlpp::RenderStages)(/*mtlpp::RenderStages::Vertex |*/ mtlpp::RenderStages::Fragment));
 			METAL_DEBUG_LAYER(EMetalDebugLevelFastValidation, RenderEncoderDebug.AddUpdateFence(*Fence));
 		}
 		else if (ComputeCommandEncoder)
 		{
-//			Fence->Write(mtlpp::RenderStages::Vertex);
-			Fence->Write(mtlpp::RenderStages::Fragment);
+//			METAL_DEBUG_LAYER(EMetalDebugLevelFastValidation, Fence->Write(mtlpp::RenderStages::Vertex));
+			METAL_DEBUG_LAYER(EMetalDebugLevelFastValidation, Fence->Write(mtlpp::RenderStages::Fragment));
 			ComputeCommandEncoder.UpdateFence(InnerFence);
 			METAL_DEBUG_LAYER(EMetalDebugLevelFastValidation, ComputeEncoderDebug.AddUpdateFence(*Fence));
 		}
 		else if (BlitCommandEncoder)
 		{
-//			Fence->Write(mtlpp::RenderStages::Vertex);
-			Fence->Write(mtlpp::RenderStages::Fragment);
+//			METAL_DEBUG_LAYER(EMetalDebugLevelFastValidation, Fence->Write(mtlpp::RenderStages::Vertex));
+			METAL_DEBUG_LAYER(EMetalDebugLevelFastValidation, Fence->Write(mtlpp::RenderStages::Fragment));
 			BlitCommandEncoder.UpdateFence(InnerFence);
 			METAL_DEBUG_LAYER(EMetalDebugLevelFastValidation, BlitEncoderDebug.AddUpdateFence(*Fence));
 		}
@@ -761,22 +761,22 @@ void FMetalCommandEncoder::WaitForFence(FMetalFence* Fence)
 		mtlpp::Fence InnerFence = METAL_DEBUG_OPTION(CommandList.GetCommandQueue().GetRuntimeDebuggingLevel() >= EMetalDebugLevelValidation ? mtlpp::Fence(((FMetalDebugFence*)Fence->GetPtr()).Inner) :) *Fence;
 		if (RenderCommandEncoder)
 		{
-//			Fence->Wait(mtlpp::RenderStages::Vertex);
-			Fence->Wait(mtlpp::RenderStages::Fragment);
+//			METAL_DEBUG_LAYER(EMetalDebugLevelFastValidation, Fence->Wait(mtlpp::RenderStages::Vertex));
+			METAL_DEBUG_LAYER(EMetalDebugLevelFastValidation, Fence->Wait(mtlpp::RenderStages::Fragment));
 			RenderCommandEncoder.WaitForFence(InnerFence, (mtlpp::RenderStages)(mtlpp::RenderStages::Vertex /*| mtlpp::RenderStages::Fragment*/));
 			METAL_DEBUG_LAYER(EMetalDebugLevelFastValidation, RenderEncoderDebug.AddWaitFence(*Fence));
 		}
 		else if (ComputeCommandEncoder)
 		{
-//			Fence->Wait(mtlpp::RenderStages::Vertex);
-			Fence->Wait(mtlpp::RenderStages::Fragment);
+//			METAL_DEBUG_LAYER(EMetalDebugLevelFastValidation, Fence->Wait(mtlpp::RenderStages::Vertex));
+			METAL_DEBUG_LAYER(EMetalDebugLevelFastValidation, Fence->Wait(mtlpp::RenderStages::Fragment));
 			ComputeCommandEncoder.WaitForFence(InnerFence);
 			METAL_DEBUG_LAYER(EMetalDebugLevelFastValidation, ComputeEncoderDebug.AddWaitFence(*Fence));
 		}
 		else if (BlitCommandEncoder)
 		{
-//			Fence->Wait(mtlpp::RenderStages::Vertex);
-			Fence->Wait(mtlpp::RenderStages::Fragment);
+//			METAL_DEBUG_LAYER(EMetalDebugLevelFastValidation, Fence->Wait(mtlpp::RenderStages::Vertex));
+			METAL_DEBUG_LAYER(EMetalDebugLevelFastValidation, Fence->Wait(mtlpp::RenderStages::Fragment));
 			BlitCommandEncoder.WaitForFence(InnerFence);
 			METAL_DEBUG_LAYER(EMetalDebugLevelFastValidation, BlitEncoderDebug.AddWaitFence(*Fence));
 		}
@@ -803,33 +803,33 @@ void FMetalCommandEncoder::WaitAndUpdateFence(FMetalFence* Fence)
 		mtlpp::Fence InnerFence = METAL_DEBUG_OPTION(CommandList.GetCommandQueue().GetRuntimeDebuggingLevel() >= EMetalDebugLevelValidation ? mtlpp::Fence(((FMetalDebugFence*)Fence->GetPtr()).Inner) :) *Fence;
 		if (RenderCommandEncoder)
 		{
-			//			Fence->Wait(mtlpp::RenderStages::Vertex);
-			Fence->Wait(mtlpp::RenderStages::Fragment);
+			//			METAL_DEBUG_LAYER(EMetalDebugLevelFastValidation, Fence->Wait(mtlpp::RenderStages::Vertex));
+			METAL_DEBUG_LAYER(EMetalDebugLevelFastValidation, Fence->Wait(mtlpp::RenderStages::Fragment));
 			RenderCommandEncoder.WaitForFence(InnerFence, (mtlpp::RenderStages)(mtlpp::RenderStages::Vertex /*| mtlpp::RenderStages::Fragment*/));
 			METAL_DEBUG_LAYER(EMetalDebugLevelFastValidation, RenderEncoderDebug.AddWaitFence(*Fence));
 			RenderCommandEncoder.UpdateFence(InnerFence, (mtlpp::RenderStages)(/*mtlpp::RenderStages::Vertex |*/ mtlpp::RenderStages::Fragment));
 			METAL_DEBUG_LAYER(EMetalDebugLevelFastValidation, RenderEncoderDebug.AddUpdateFence(*Fence));
-			Fence->Write(mtlpp::RenderStages::Fragment);
+			METAL_DEBUG_LAYER(EMetalDebugLevelFastValidation, Fence->Write(mtlpp::RenderStages::Fragment));
 		}
 		else if (ComputeCommandEncoder)
 		{
-			//			Fence->Wait(mtlpp::RenderStages::Vertex);
-			Fence->Wait(mtlpp::RenderStages::Fragment);
+			//			METAL_DEBUG_LAYER(EMetalDebugLevelFastValidation, Fence->Wait(mtlpp::RenderStages::Vertex));
+			METAL_DEBUG_LAYER(EMetalDebugLevelFastValidation, Fence->Wait(mtlpp::RenderStages::Fragment));
 			ComputeCommandEncoder.WaitForFence(InnerFence);
 			METAL_DEBUG_LAYER(EMetalDebugLevelFastValidation, ComputeEncoderDebug.AddWaitFence(*Fence));
 			ComputeCommandEncoder.UpdateFence(InnerFence);
 			METAL_DEBUG_LAYER(EMetalDebugLevelFastValidation, ComputeEncoderDebug.AddUpdateFence(*Fence));
-			Fence->Write(mtlpp::RenderStages::Fragment);
+			METAL_DEBUG_LAYER(EMetalDebugLevelFastValidation, Fence->Write(mtlpp::RenderStages::Fragment));
 		}
 		else if (BlitCommandEncoder)
 		{
-			//			Fence->Wait(mtlpp::RenderStages::Vertex);
-			Fence->Wait(mtlpp::RenderStages::Fragment);
+			//			METAL_DEBUG_LAYER(EMetalDebugLevelFastValidation, Fence->Wait(mtlpp::RenderStages::Vertex));
+			METAL_DEBUG_LAYER(EMetalDebugLevelFastValidation, Fence->Wait(mtlpp::RenderStages::Fragment));
 			BlitCommandEncoder.WaitForFence(InnerFence);
 			METAL_DEBUG_LAYER(EMetalDebugLevelFastValidation, BlitEncoderDebug.AddWaitFence(*Fence));
 			BlitCommandEncoder.UpdateFence(InnerFence);
 			METAL_DEBUG_LAYER(EMetalDebugLevelFastValidation, BlitEncoderDebug.AddUpdateFence(*Fence));
-			Fence->Write(mtlpp::RenderStages::Fragment);
+			METAL_DEBUG_LAYER(EMetalDebugLevelFastValidation, Fence->Write(mtlpp::RenderStages::Fragment));
 		}
 	}
 }
diff --git a/Engine/Source/Runtime/Apple/MetalRHI/Private/MetalCommandQueue.cpp b/Engine/Source/Runtime/Apple/MetalRHI/Private/MetalCommandQueue.cpp
--- a/Engine/Source/Runtime/Apple/MetalRHI/Private/MetalCommandQueue.cpp
+++ b/Engine/Source/Runtime/Apple/MetalRHI/Private/MetalCommandQueue.cpp
@@ -364,7 +364,11 @@ FMetalFence* FMetalCommandQueue::CreateFence(ns::String const& Label) const
 	if ((Features & EMetalFeaturesFences) != 0)
 	{
 		FMetalFence* InternalFence = FMetalFencePool::Get().AllocateFence();
-		NSString* String = [NSString stringWithFormat:@"%p: %@", InternalFence->GetPtr(), Label.GetPtr()];
+		NSString* String = nil;
+		if (GetEmitDrawEvents())
+		{
+			String = [NSString stringWithFormat:@"%p: %@", InternalFence->GetPtr(), Label.GetPtr()];
+		}
 #if METAL_DEBUG_OPTIONS
 		if (RuntimeDebuggingLevel >= EMetalDebugLevelValidation)
 		{
@@ -373,7 +377,7 @@ FMetalFence* FMetalCommandQueue::CreateFence(ns::String const& Label) const
 		}
 		else
 #endif
-		if(InternalFence && Label)
+		if(InternalFence && String)
 		{
 			InternalFence->SetLabel(String);
 		}
diff --git a/Engine/Source/Runtime/Apple/MetalRHI/Private/MetalContext.cpp b/Engine/Source/Runtime/Apple/MetalRHI/Private/MetalContext.cpp
--- a/Engine/Source/Runtime/Apple/MetalRHI/Private/MetalContext.cpp
+++ b/Engine/Source/Runtime/Apple/MetalRHI/Private/MetalContext.cpp
@@ -589,7 +589,16 @@ void FMetalDeviceContext::FlushFreeList(bool const bFlushFences)
 	NewList->ObjectFreeList = ObjectFreeList;
 	if (bFlushFences)
 	{
-		NewList->FenceFreeList = MoveTemp(FenceFreeList);
+		TArray<FMetalFence*> Fences;
+		FenceFreeList.PopAll(Fences);
+		for (FMetalFence* Fence : Fences)
+		{
+			if(!UsedFences.Contains(Fence))
+			{
+				UsedFences.Add(Fence);
+			}
+		}
+		NewList->FenceFreeList = MoveTemp(UsedFences);
 	}
 #if METAL_DEBUG_OPTIONS
 	if (FrameFences.Num())
@@ -702,12 +711,7 @@ void FMetalDeviceContext::ReleaseFence(FMetalFence* Fence)
 	if (GIsMetalInitialized) // @todo zebra: there seems to be some race condition at exit when the framerate is very low
 	{
 		check(Fence);
-		FreeListMutex.Lock();
-		if(!FenceFreeList.Contains(Fence))
-		{
-			FenceFreeList.Add(Fence);
-		}
-		FreeListMutex.Unlock();
+		FenceFreeList.Push(Fence);
 	}
 }
 
diff --git a/Engine/Source/Runtime/Apple/MetalRHI/Private/MetalContext.h b/Engine/Source/Runtime/Apple/MetalRHI/Private/MetalContext.h
--- a/Engine/Source/Runtime/Apple/MetalRHI/Private/MetalContext.h
+++ b/Engine/Source/Runtime/Apple/MetalRHI/Private/MetalContext.h
@@ -238,7 +238,8 @@ private:
 	/** Free lists for releasing objects only once it is safe to do so */
 	TSet<FMetalBuffer> UsedBuffers;
 	TSet<FMetalTexture> UsedTextures;
-	TSet<FMetalFence*> FenceFreeList;
+	TSet<FMetalFence*> UsedFences;
+	TLockFreePointerListLIFO<FMetalFence> FenceFreeList;
 	TSet<id> ObjectFreeList;
 	struct FMetalDelayedFreeList
 	{
diff --git a/Engine/Source/Runtime/Apple/MetalRHI/Private/MetalFence.cpp b/Engine/Source/Runtime/Apple/MetalRHI/Private/MetalFence.cpp
--- a/Engine/Source/Runtime/Apple/MetalRHI/Private/MetalFence.cpp
+++ b/Engine/Source/Runtime/Apple/MetalRHI/Private/MetalFence.cpp
@@ -128,7 +128,12 @@ void FMetalFencePool::Initialise(mtlpp::Device const& InDevice)
 #endif
 		{
 			FMetalFence* F = new FMetalFence(Device.NewFence());
-			Fences.Add(F);
+#if METAL_DEBUG_OPTIONS
+			if (GMetalRuntimeDebugLevel >= EMetalDebugLevelValidation)
+			{
+				Fences.Add(F);
+			}
+#endif
 			Lifo.Push(F);
 		}
 	}
