Use a pool of render pass descriptors and add more scoped timers for inspecting where time goes in MetalRHI.

diff --git a/Engine/Source/Runtime/Apple/MetalRHI/Private/MetalCommandEncoder.cpp b/Engine/Source/Runtime/Apple/MetalRHI/Private/MetalCommandEncoder.cpp
--- a/Engine/Source/Runtime/Apple/MetalRHI/Private/MetalCommandEncoder.cpp
+++ b/Engine/Source/Runtime/Apple/MetalRHI/Private/MetalCommandEncoder.cpp
@@ -68,6 +68,7 @@ FMetalCommandEncoder::~FMetalCommandEncoder(void)
 	check(!IsComputeCommandEncoderActive());
 	check(!IsBlitCommandEncoderActive());
 	
+	SafeReleaseMetalRenderPassDescriptor(RenderPassDesc);
 	RenderPassDesc = nil;
 
 	if(DebugGroups)
@@ -99,6 +100,7 @@ void FMetalCommandEncoder::Reset(void)
 	
 	if(RenderPassDesc)
 	{
+		SafeReleaseMetalRenderPassDescriptor(RenderPassDesc);
 		RenderPassDesc = nil;
 	}
 	
@@ -957,6 +959,7 @@ void FMetalCommandEncoder::SetRenderPassDescriptor(mtlpp::RenderPassDescriptor R
 	
 	if(RenderPass.GetPtr() != RenderPassDesc.GetPtr())
 	{
+		SafeReleaseMetalRenderPassDescriptor(RenderPassDesc);
 		RenderPassDesc = RenderPass;
 		
 		static bool bDeferredStoreActions = CommandList.GetCommandQueue().SupportsFeature(EMetalFeaturesDeferredStoreActions);
diff --git a/Engine/Source/Runtime/Apple/MetalRHI/Private/MetalProfiler.cpp b/Engine/Source/Runtime/Apple/MetalRHI/Private/MetalProfiler.cpp
--- a/Engine/Source/Runtime/Apple/MetalRHI/Private/MetalProfiler.cpp
+++ b/Engine/Source/Runtime/Apple/MetalRHI/Private/MetalProfiler.cpp
@@ -22,6 +22,19 @@ DEFINE_STAT(STAT_MetalTextureMemUpdate);
 DEFINE_STAT(STAT_MetalDrawCallTime);
 DEFINE_STAT(STAT_MetalPipelineStateTime);
 DEFINE_STAT(STAT_MetalPrepareDrawTime);
+
+DEFINE_STAT(STAT_MetalSwitchToRenderTime);
+DEFINE_STAT(STAT_MetalSwitchToTessellationTime);
+DEFINE_STAT(STAT_MetalSwitchToComputeTime);
+DEFINE_STAT(STAT_MetalSwitchToBlitTime);
+DEFINE_STAT(STAT_MetalSwitchToAsyncBlitTime);
+DEFINE_STAT(STAT_MetalPrepareToRenderTime);
+DEFINE_STAT(STAT_MetalPrepareToTessellateTime);
+DEFINE_STAT(STAT_MetalPrepareToDispatchTime);
+DEFINE_STAT(STAT_MetalCommitRenderResourceTablesTime);
+DEFINE_STAT(STAT_MetalSetRenderStateTime);
+DEFINE_STAT(STAT_MetalSetRenderPipelineStateTime);
+
 DEFINE_STAT(STAT_MetalMakeDrawableTime);
 DEFINE_STAT(STAT_MetalBufferPageOffTime);
 DEFINE_STAT(STAT_MetalTexturePageOnTime);
diff --git a/Engine/Source/Runtime/Apple/MetalRHI/Private/MetalProfiler.h b/Engine/Source/Runtime/Apple/MetalRHI/Private/MetalProfiler.h
--- a/Engine/Source/Runtime/Apple/MetalRHI/Private/MetalProfiler.h
+++ b/Engine/Source/Runtime/Apple/MetalRHI/Private/MetalProfiler.h
@@ -14,6 +14,17 @@
 DECLARE_CYCLE_STAT_EXTERN(TEXT("MakeDrawable time"),STAT_MetalMakeDrawableTime,STATGROUP_MetalRHI, );
 DECLARE_CYCLE_STAT_EXTERN(TEXT("Draw call time"),STAT_MetalDrawCallTime,STATGROUP_MetalRHI, );
 DECLARE_CYCLE_STAT_EXTERN(TEXT("PrepareDraw time"),STAT_MetalPrepareDrawTime,STATGROUP_MetalRHI, );
+DECLARE_CYCLE_STAT_EXTERN(TEXT("SwitchToRender time"),STAT_MetalSwitchToRenderTime,STATGROUP_MetalRHI, );
+DECLARE_CYCLE_STAT_EXTERN(TEXT("SwitchToTessellation time"),STAT_MetalSwitchToTessellationTime,STATGROUP_MetalRHI, );
+DECLARE_CYCLE_STAT_EXTERN(TEXT("SwitchToCompute time"),STAT_MetalSwitchToComputeTime,STATGROUP_MetalRHI, );
+DECLARE_CYCLE_STAT_EXTERN(TEXT("SwitchToBlit time"),STAT_MetalSwitchToBlitTime,STATGROUP_MetalRHI, );
+DECLARE_CYCLE_STAT_EXTERN(TEXT("SwitchToAsyncBlit time"),STAT_MetalSwitchToAsyncBlitTime,STATGROUP_MetalRHI, );
+DECLARE_CYCLE_STAT_EXTERN(TEXT("PrepareToRender time"),STAT_MetalPrepareToRenderTime,STATGROUP_MetalRHI, );
+DECLARE_CYCLE_STAT_EXTERN(TEXT("PrepareToTessellate time"),STAT_MetalPrepareToTessellateTime,STATGROUP_MetalRHI, );
+DECLARE_CYCLE_STAT_EXTERN(TEXT("PrepareToDispatch time"),STAT_MetalPrepareToDispatchTime,STATGROUP_MetalRHI, );
+DECLARE_CYCLE_STAT_EXTERN(TEXT("CommitRenderResourceTables time"),STAT_MetalCommitRenderResourceTablesTime,STATGROUP_MetalRHI, );
+DECLARE_CYCLE_STAT_EXTERN(TEXT("SetRenderState time"),STAT_MetalSetRenderStateTime,STATGROUP_MetalRHI, );
+DECLARE_CYCLE_STAT_EXTERN(TEXT("SetRenderPipelineState time"),STAT_MetalSetRenderPipelineStateTime,STATGROUP_MetalRHI, );
 DECLARE_CYCLE_STAT_EXTERN(TEXT("PipelineState time"),STAT_MetalPipelineStateTime,STATGROUP_MetalRHI, );
 DECLARE_CYCLE_STAT_EXTERN(TEXT("Buffer Page-Off time"), STAT_MetalBufferPageOffTime, STATGROUP_MetalRHI, );
 DECLARE_CYCLE_STAT_EXTERN(TEXT("Texture Page-Off time"), STAT_MetalTexturePageOffTime, STATGROUP_MetalRHI, );
diff --git a/Engine/Source/Runtime/Apple/MetalRHI/Private/MetalRHIPrivate.h b/Engine/Source/Runtime/Apple/MetalRHI/Private/MetalRHIPrivate.h
--- a/Engine/Source/Runtime/Apple/MetalRHI/Private/MetalRHIPrivate.h
+++ b/Engine/Source/Runtime/Apple/MetalRHI/Private/MetalRHIPrivate.h
@@ -148,6 +148,9 @@ void SafeReleaseMetalBuffer(FMetalBuffer& Buffer);
 // Safely release a fence, correctly handling cases where fences aren't supported or the debug implementation is used.
 void SafeReleaseMetalFence(class FMetalFence* Object);
 
+// Safely release a render pass descriptor so that it may be reused.
+void SafeReleaseMetalRenderPassDescriptor(mtlpp::RenderPassDescriptor& Desc);
+
 // Access the underlying surface object from any kind of texture
 FMetalSurface* GetMetalSurfaceFromRHITexture(FRHITexture* Texture);
 
diff --git a/Engine/Source/Runtime/Apple/MetalRHI/Private/MetalRenderPass.cpp b/Engine/Source/Runtime/Apple/MetalRHI/Private/MetalRenderPass.cpp
--- a/Engine/Source/Runtime/Apple/MetalRHI/Private/MetalRenderPass.cpp
+++ b/Engine/Source/Runtime/Apple/MetalRHI/Private/MetalRenderPass.cpp
@@ -1173,6 +1173,8 @@ mtlpp::RenderCommandEncoder FMetalRenderPass::GetParallelRenderCommandEncoder(ui
 
 void FMetalRenderPass::ConditionalSwitchToRender(void)
 {
+	SCOPE_CYCLE_COUNTER(STAT_MetalSwitchToRenderTime);
+	
 	check(bWithinRenderPass);
 	check(RenderPassDesc);
 	check(CmdList.IsParallel() || CurrentEncoder.GetCommandBuffer());
@@ -1192,6 +1194,8 @@ void FMetalRenderPass::ConditionalSwitchToRender(void)
 
 void FMetalRenderPass::ConditionalSwitchToTessellation(void)
 {
+	SCOPE_CYCLE_COUNTER(STAT_MetalSwitchToTessellationTime);
+	
 	check(bWithinRenderPass);
 	check(RenderPassDesc);
 	check(CurrentEncoder.GetCommandBuffer());
@@ -1276,6 +1280,8 @@ void FMetalRenderPass::ConditionalSwitchToTessellation(void)
 
 void FMetalRenderPass::ConditionalSwitchToCompute(void)
 {
+	SCOPE_CYCLE_COUNTER(STAT_MetalSwitchToComputeTime);
+	
 	check(CurrentEncoder.GetCommandBuffer());
 	check(!CurrentEncoder.IsParallel());
 	
@@ -1311,6 +1317,8 @@ void FMetalRenderPass::ConditionalSwitchToCompute(void)
 
 void FMetalRenderPass::ConditionalSwitchToBlit(void)
 {
+	SCOPE_CYCLE_COUNTER(STAT_MetalSwitchToBlitTime);
+	
 	check(CurrentEncoder.GetCommandBuffer());
 	check(!CurrentEncoder.IsParallel());
 	
@@ -1345,6 +1353,8 @@ void FMetalRenderPass::ConditionalSwitchToBlit(void)
 
 void FMetalRenderPass::ConditionalSwitchToAsyncBlit(void)
 {
+	SCOPE_CYCLE_COUNTER(STAT_MetalSwitchToAsyncBlitTime);
+	
 	if (PrologueEncoder.IsComputeCommandEncoderActive())
 	{
 		PrologueEncoderFence = PrologueEncoder.EndEncoding();
@@ -1411,6 +1421,8 @@ void FMetalRenderPass::ConditionalSwitchToAsyncBlit(void)
 
 void FMetalRenderPass::CommitRenderResourceTables(void)
 {
+	SCOPE_CYCLE_COUNTER(STAT_MetalCommitRenderResourceTablesTime);
+	
 	State.CommitRenderResources(&CurrentEncoder);
 	
 	State.CommitResourceTable(SF_Vertex, mtlpp::FunctionType::Vertex, CurrentEncoder);
@@ -1467,6 +1479,8 @@ void FMetalRenderPass::CommitDispatchResourceTables(void)
 
 void FMetalRenderPass::PrepareToRender(uint32 PrimitiveType)
 {
+	SCOPE_CYCLE_COUNTER(STAT_MetalPrepareToRenderTime);
+	
 	check(CurrentEncoder.GetCommandBuffer());
 	check(CurrentEncoder.IsRenderCommandEncoderActive());
 	
@@ -1481,6 +1495,8 @@ void FMetalRenderPass::PrepareToRender(uint32 PrimitiveType)
 
 void FMetalRenderPass::PrepareToTessellate(uint32 PrimitiveType)
 {
+	SCOPE_CYCLE_COUNTER(STAT_MetalPrepareToTessellateTime);
+	
 	check(CurrentEncoder.GetCommandBuffer());
 	check(PrologueEncoder.GetCommandBuffer());
 	check(CurrentEncoder.IsRenderCommandEncoderActive());
@@ -1497,6 +1513,8 @@ void FMetalRenderPass::PrepareToTessellate(uint32 PrimitiveType)
 
 void FMetalRenderPass::PrepareToDispatch(void)
 {
+	SCOPE_CYCLE_COUNTER(STAT_MetalPrepareToDispatchTime);
+	
 	check(CurrentEncoder.GetCommandBuffer());
 	check(CurrentEncoder.IsComputeCommandEncoderActive());
 	
diff --git a/Engine/Source/Runtime/Apple/MetalRHI/Private/MetalStateCache.cpp b/Engine/Source/Runtime/Apple/MetalRHI/Private/MetalStateCache.cpp
--- a/Engine/Source/Runtime/Apple/MetalRHI/Private/MetalStateCache.cpp
+++ b/Engine/Source/Runtime/Apple/MetalRHI/Private/MetalStateCache.cpp
@@ -65,6 +65,77 @@ FORCEINLINE mtlpp::StoreAction GetConditionalMetalRTStoreAction(bool bMSAATarget
 	}	
 }
 
+class FMetalRenderPassDescriptorPool
+{
+public:
+	FMetalRenderPassDescriptorPool()
+	{
+		
+	}
+	
+	~FMetalRenderPassDescriptorPool()
+	{
+		
+	}
+	
+	mtlpp::RenderPassDescriptor CreateDescriptor()
+	{
+		MTLRenderPassDescriptor* Desc = Cache.Pop();
+		if (!Desc)
+		{
+			Desc = [MTLRenderPassDescriptor new];
+		}
+		return mtlpp::RenderPassDescriptor(Desc);
+	}
+	
+	void ReleaseDescriptor(mtlpp::RenderPassDescriptor& Desc)
+	{
+		mtlpp::Texture EmptyTex;
+		
+		ns::Array<mtlpp::RenderPassColorAttachmentDescriptor> Attachements = Desc.GetColorAttachments();
+		for (uint32 i = 0; i < MaxSimultaneousRenderTargets; i++)
+		{
+			mtlpp::RenderPassColorAttachmentDescriptor Color = Attachements[i];
+			Color.SetTexture(EmptyTex);
+			Color.SetResolveTexture(EmptyTex);
+		}
+		
+		mtlpp::RenderPassDepthAttachmentDescriptor Depth = Desc.GetDepthAttachment();
+		Depth.SetTexture(EmptyTex);
+		Depth.SetResolveTexture(EmptyTex);
+
+		mtlpp::RenderPassStencilAttachmentDescriptor Stencil = Desc.GetStencilAttachment();
+		Stencil.SetTexture(EmptyTex);
+		Stencil.SetResolveTexture(EmptyTex);
+
+		mtlpp::Buffer Empty;
+		Desc.SetVisibilityResultBuffer(Empty);
+		
+#if PLATFORM_MAC
+		Desc.SetRenderTargetArrayLength(1);
+#endif
+		
+		Cache.Push(Desc.GetPtr());
+	}
+	
+	static FMetalRenderPassDescriptorPool& Get()
+	{
+		static FMetalRenderPassDescriptorPool sSelf;
+		return sSelf;
+	}
+	
+private:
+	TLockFreePointerListLIFO<MTLRenderPassDescriptor> Cache;
+};
+
+void SafeReleaseMetalRenderPassDescriptor(mtlpp::RenderPassDescriptor& Desc)
+{
+	if (Desc.GetPtr())
+	{
+		FMetalRenderPassDescriptorPool::Get().ReleaseDescriptor(Desc);
+	}
+}
+
 FMetalStateCache::FMetalStateCache(bool const bInImmediate)
 : DepthStore(mtlpp::StoreAction::Unknown)
 , StencilStore(mtlpp::StoreAction::Unknown)
@@ -342,7 +413,7 @@ bool FMetalStateCache::SetRenderTargetsInfo(FRHISetRenderTargetsInfo const& InRe
 		RenderTargetsInfo = InRenderTargets;
 		
 		// at this point, we need to fully set up an encoder/command buffer, so make a new one (autoreleased)
-		mtlpp::RenderPassDescriptor RenderPass;
+		mtlpp::RenderPassDescriptor RenderPass = FMetalRenderPassDescriptorPool::Get().CreateDescriptor();
 	
 		// if we need to do queries, write to the supplied query buffer
 		if (IsFeatureLevelSupported(GMaxRHIShaderPlatform, ERHIFeatureLevel::ES3_1))
@@ -1874,6 +1945,8 @@ void FMetalStateCache::FlushVisibilityResults(FMetalCommandEncoder& CommandEncod
 
 void FMetalStateCache::SetRenderState(FMetalCommandEncoder& CommandEncoder, FMetalCommandEncoder* PrologueEncoder)
 {
+	SCOPE_CYCLE_COUNTER(STAT_MetalSetRenderStateTime);
+	
 	if (RasterBits)
 	{
 		if (RasterBits & EMetalRenderFlagViewport)
@@ -1930,6 +2003,8 @@ void FMetalStateCache::SetRenderState(FMetalCommandEncoder& CommandEncoder, FMet
 
 void FMetalStateCache::SetRenderPipelineState(FMetalCommandEncoder& CommandEncoder, FMetalCommandEncoder* PrologueEncoder)
 {
+	SCOPE_CYCLE_COUNTER(STAT_MetalSetRenderPipelineStateTime);
+	
     if ((PipelineBits & EMetalPipelineFlagRasterMask) != 0)
     {
     	// @todo Could optimise it so that we only re-evaluate the buffer hashes if the shader buffer binding mask changes when changing the PSO
