Make sure fence writes & waits are collected even without validation so we enforce matching operations (fixes the Editor) & disable private-buffer zero-filling when using MTLFence not MTLHeap, so it doesn't causes issues with HazardUntracked buffers.

diff --git a/Engine/Source/Runtime/Apple/MetalRHI/Private/MetalCommandEncoder.cpp b/Engine/Source/Runtime/Apple/MetalRHI/Private/MetalCommandEncoder.cpp
--- a/Engine/Source/Runtime/Apple/MetalRHI/Private/MetalCommandEncoder.cpp
+++ b/Engine/Source/Runtime/Apple/MetalRHI/Private/MetalCommandEncoder.cpp
@@ -723,22 +723,22 @@ void FMetalCommandEncoder::UpdateFence(FMetalFence* Fence)
 		mtlpp::Fence InnerFence = METAL_DEBUG_OPTION(CommandList.GetCommandQueue().GetRuntimeDebuggingLevel() >= EMetalDebugLevelValidation ? mtlpp::Fence(((FMetalDebugFence*)Fence->GetPtr()).Inner) :) *Fence;
 		if (RenderCommandEncoder)
 		{
-//			METAL_DEBUG_LAYER(EMetalDebugLevelFastValidation, Fence->Write(mtlpp::RenderStages::Vertex));
-			METAL_DEBUG_LAYER(EMetalDebugLevelFastValidation, Fence->Write(mtlpp::RenderStages::Fragment));
+//			Fence->Write(mtlpp::RenderStages::Vertex);
+			Fence->Write(mtlpp::RenderStages::Fragment);
 			RenderCommandEncoder.UpdateFence(InnerFence, (mtlpp::RenderStages)(/*mtlpp::RenderStages::Vertex |*/ mtlpp::RenderStages::Fragment));
 			METAL_DEBUG_LAYER(EMetalDebugLevelFastValidation, RenderEncoderDebug.AddUpdateFence(*Fence));
 		}
 		else if (ComputeCommandEncoder)
 		{
-//			METAL_DEBUG_LAYER(EMetalDebugLevelFastValidation, Fence->Write(mtlpp::RenderStages::Vertex));
-			METAL_DEBUG_LAYER(EMetalDebugLevelFastValidation, Fence->Write(mtlpp::RenderStages::Fragment));
+//			Fence->Write(mtlpp::RenderStages::Vertex);
+			Fence->Write(mtlpp::RenderStages::Fragment);
 			ComputeCommandEncoder.UpdateFence(InnerFence);
 			METAL_DEBUG_LAYER(EMetalDebugLevelFastValidation, ComputeEncoderDebug.AddUpdateFence(*Fence));
 		}
 		else if (BlitCommandEncoder)
 		{
-//			METAL_DEBUG_LAYER(EMetalDebugLevelFastValidation, Fence->Write(mtlpp::RenderStages::Vertex));
-			METAL_DEBUG_LAYER(EMetalDebugLevelFastValidation, Fence->Write(mtlpp::RenderStages::Fragment));
+//			Fence->Write(mtlpp::RenderStages::Vertex);
+			Fence->Write(mtlpp::RenderStages::Fragment);
 			BlitCommandEncoder.UpdateFence(InnerFence);
 			METAL_DEBUG_LAYER(EMetalDebugLevelFastValidation, BlitEncoderDebug.AddUpdateFence(*Fence));
 		}
@@ -763,22 +763,22 @@ void FMetalCommandEncoder::WaitForFence(FMetalFence* Fence)
 		mtlpp::Fence InnerFence = METAL_DEBUG_OPTION(CommandList.GetCommandQueue().GetRuntimeDebuggingLevel() >= EMetalDebugLevelValidation ? mtlpp::Fence(((FMetalDebugFence*)Fence->GetPtr()).Inner) :) *Fence;
 		if (RenderCommandEncoder)
 		{
-//			METAL_DEBUG_LAYER(EMetalDebugLevelFastValidation, Fence->Wait(mtlpp::RenderStages::Vertex));
-			METAL_DEBUG_LAYER(EMetalDebugLevelFastValidation, Fence->Wait(mtlpp::RenderStages::Fragment));
+//			Fence->Wait(mtlpp::RenderStages::Vertex);
+			Fence->Wait(mtlpp::RenderStages::Fragment);
 			RenderCommandEncoder.WaitForFence(InnerFence, (mtlpp::RenderStages)(mtlpp::RenderStages::Vertex /*| mtlpp::RenderStages::Fragment*/));
 			METAL_DEBUG_LAYER(EMetalDebugLevelFastValidation, RenderEncoderDebug.AddWaitFence(*Fence));
 		}
 		else if (ComputeCommandEncoder)
 		{
-//			METAL_DEBUG_LAYER(EMetalDebugLevelFastValidation, Fence->Wait(mtlpp::RenderStages::Vertex));
-			METAL_DEBUG_LAYER(EMetalDebugLevelFastValidation, Fence->Wait(mtlpp::RenderStages::Fragment));
+//			Fence->Wait(mtlpp::RenderStages::Vertex);
+			Fence->Wait(mtlpp::RenderStages::Fragment);
 			ComputeCommandEncoder.WaitForFence(InnerFence);
 			METAL_DEBUG_LAYER(EMetalDebugLevelFastValidation, ComputeEncoderDebug.AddWaitFence(*Fence));
 		}
 		else if (BlitCommandEncoder)
 		{
-//			METAL_DEBUG_LAYER(EMetalDebugLevelFastValidation, Fence->Wait(mtlpp::RenderStages::Vertex));
-			METAL_DEBUG_LAYER(EMetalDebugLevelFastValidation, Fence->Wait(mtlpp::RenderStages::Fragment));
+//			Fence->Wait(mtlpp::RenderStages::Vertex);
+			Fence->Wait(mtlpp::RenderStages::Fragment);
 			BlitCommandEncoder.WaitForFence(InnerFence);
 			METAL_DEBUG_LAYER(EMetalDebugLevelFastValidation, BlitEncoderDebug.AddWaitFence(*Fence));
 		}
@@ -805,33 +805,33 @@ void FMetalCommandEncoder::WaitAndUpdateFence(FMetalFence* Fence)
 		mtlpp::Fence InnerFence = METAL_DEBUG_OPTION(CommandList.GetCommandQueue().GetRuntimeDebuggingLevel() >= EMetalDebugLevelValidation ? mtlpp::Fence(((FMetalDebugFence*)Fence->GetPtr()).Inner) :) *Fence;
 		if (RenderCommandEncoder)
 		{
-			//			METAL_DEBUG_LAYER(EMetalDebugLevelFastValidation, Fence->Wait(mtlpp::RenderStages::Vertex));
-			METAL_DEBUG_LAYER(EMetalDebugLevelFastValidation, Fence->Wait(mtlpp::RenderStages::Fragment));
+			//			Fence->Wait(mtlpp::RenderStages::Vertex);
+			Fence->Wait(mtlpp::RenderStages::Fragment);
 			RenderCommandEncoder.WaitForFence(InnerFence, (mtlpp::RenderStages)(mtlpp::RenderStages::Vertex /*| mtlpp::RenderStages::Fragment*/));
 			METAL_DEBUG_LAYER(EMetalDebugLevelFastValidation, RenderEncoderDebug.AddWaitFence(*Fence));
 			RenderCommandEncoder.UpdateFence(InnerFence, (mtlpp::RenderStages)(/*mtlpp::RenderStages::Vertex |*/ mtlpp::RenderStages::Fragment));
 			METAL_DEBUG_LAYER(EMetalDebugLevelFastValidation, RenderEncoderDebug.AddUpdateFence(*Fence));
-			METAL_DEBUG_LAYER(EMetalDebugLevelFastValidation, Fence->Write(mtlpp::RenderStages::Fragment));
+			Fence->Write(mtlpp::RenderStages::Fragment);
 		}
 		else if (ComputeCommandEncoder)
 		{
-			//			METAL_DEBUG_LAYER(EMetalDebugLevelFastValidation, Fence->Wait(mtlpp::RenderStages::Vertex));
-			METAL_DEBUG_LAYER(EMetalDebugLevelFastValidation, Fence->Wait(mtlpp::RenderStages::Fragment));
+			//			Fence->Wait(mtlpp::RenderStages::Vertex);
+			Fence->Wait(mtlpp::RenderStages::Fragment);
 			ComputeCommandEncoder.WaitForFence(InnerFence);
 			METAL_DEBUG_LAYER(EMetalDebugLevelFastValidation, ComputeEncoderDebug.AddWaitFence(*Fence));
 			ComputeCommandEncoder.UpdateFence(InnerFence);
 			METAL_DEBUG_LAYER(EMetalDebugLevelFastValidation, ComputeEncoderDebug.AddUpdateFence(*Fence));
-			METAL_DEBUG_LAYER(EMetalDebugLevelFastValidation, Fence->Write(mtlpp::RenderStages::Fragment));
+			Fence->Write(mtlpp::RenderStages::Fragment);
 		}
 		else if (BlitCommandEncoder)
 		{
-			//			METAL_DEBUG_LAYER(EMetalDebugLevelFastValidation, Fence->Wait(mtlpp::RenderStages::Vertex));
-			METAL_DEBUG_LAYER(EMetalDebugLevelFastValidation, Fence->Wait(mtlpp::RenderStages::Fragment));
+			//			Fence->Wait(mtlpp::RenderStages::Vertex);
+			Fence->Wait(mtlpp::RenderStages::Fragment);
 			BlitCommandEncoder.WaitForFence(InnerFence);
 			METAL_DEBUG_LAYER(EMetalDebugLevelFastValidation, BlitEncoderDebug.AddWaitFence(*Fence));
 			BlitCommandEncoder.UpdateFence(InnerFence);
 			METAL_DEBUG_LAYER(EMetalDebugLevelFastValidation, BlitEncoderDebug.AddUpdateFence(*Fence));
-			METAL_DEBUG_LAYER(EMetalDebugLevelFastValidation, Fence->Write(mtlpp::RenderStages::Fragment));
+			Fence->Write(mtlpp::RenderStages::Fragment);
 		}
 	}
 }
diff --git a/Engine/Source/Runtime/Apple/MetalRHI/Private/MetalIndexBuffer.cpp b/Engine/Source/Runtime/Apple/MetalRHI/Private/MetalIndexBuffer.cpp
--- a/Engine/Source/Runtime/Apple/MetalRHI/Private/MetalIndexBuffer.cpp
+++ b/Engine/Source/Runtime/Apple/MetalRHI/Private/MetalIndexBuffer.cpp
@@ -54,7 +54,7 @@ FIndexBufferRHIRef FMetalDynamicRHI::RHICreateIndexBuffer(uint32 Stride,uint32 S
 			IndexBuffer->CPUBuffer = nil;
 		}
 
-		if (GMetalBufferZeroFill && !FMetalCommandQueue::SupportsFeature(EMetalFeaturesHeaps))
+		if (GMetalBufferZeroFill && !FMetalCommandQueue::SupportsFeature(EMetalFeaturesFences))
 		{
 			GetMetalDeviceContext().FillBuffer(IndexBuffer->Buffer, ns::Range(0, IndexBuffer->Buffer.GetLength()), 0);
 		}
@@ -116,7 +116,7 @@ struct FMetalRHICommandInitialiseIndexBuffer : public FRHICommand<FMetalRHIComma
 				Buffer->LastUpdate = GFrameNumberRenderThread;
 			}
 		}
-		else if (GMetalBufferZeroFill && !FMetalCommandQueue::SupportsFeature(EMetalFeaturesHeaps))
+		else if (GMetalBufferZeroFill && !FMetalCommandQueue::SupportsFeature(EMetalFeaturesFences))
 		{
 			GetMetalDeviceContext().FillBuffer(Buffer->Buffer, ns::Range(0, Buffer->Buffer.GetLength()), 0);
 		}
diff --git a/Engine/Source/Runtime/Apple/MetalRHI/Private/MetalStructuredBuffer.cpp b/Engine/Source/Runtime/Apple/MetalRHI/Private/MetalStructuredBuffer.cpp
--- a/Engine/Source/Runtime/Apple/MetalRHI/Private/MetalStructuredBuffer.cpp
+++ b/Engine/Source/Runtime/Apple/MetalRHI/Private/MetalStructuredBuffer.cpp
@@ -40,7 +40,7 @@ FStructuredBufferRHIRef FMetalDynamicRHI::RHICreateStructuredBuffer(uint32 Strid
 			Buffer->CPUBuffer = nil;
 		}
 
-		if (GMetalBufferZeroFill && !FMetalCommandQueue::SupportsFeature(EMetalFeaturesHeaps))
+		if (GMetalBufferZeroFill && !FMetalCommandQueue::SupportsFeature(EMetalFeaturesFences))
 		{
 			GetMetalDeviceContext().FillBuffer(Buffer->Buffer, ns::Range(0, Buffer->Buffer.GetLength()), 0);
 		}
@@ -102,7 +102,7 @@ struct FMetalRHICommandInitialiseStructuredBuffer : public FRHICommand<FMetalRHI
 			}
 		}
 #if PLATFORM_MAC
-		else if (GMetalBufferZeroFill && !FMetalCommandQueue::SupportsFeature(EMetalFeaturesHeaps))
+		else if (GMetalBufferZeroFill && !FMetalCommandQueue::SupportsFeature(EMetalFeaturesFences))
 		{
 			GetMetalDeviceContext().FillBuffer(Buffer->Buffer, ns::Range(0, Buffer->Buffer.GetLength()), 0);
 		}
diff --git a/Engine/Source/Runtime/Apple/MetalRHI/Private/MetalUniformBuffer.cpp b/Engine/Source/Runtime/Apple/MetalRHI/Private/MetalUniformBuffer.cpp
--- a/Engine/Source/Runtime/Apple/MetalRHI/Private/MetalUniformBuffer.cpp
+++ b/Engine/Source/Runtime/Apple/MetalRHI/Private/MetalUniformBuffer.cpp
@@ -31,7 +31,7 @@ struct FMetalRHICommandInitialiseUniformdBuffer : public FRHICommand<FMetalRHICo
 			LLM_SCOPE(ELLMTag::VertexBuffer);
 			SafeReleaseMetalBuffer(Buffer->CPUBuffer);
 		}
-		else if (GMetalBufferZeroFill && !FMetalCommandQueue::SupportsFeature(EMetalFeaturesHeaps))
+		else if (GMetalBufferZeroFill && !FMetalCommandQueue::SupportsFeature(EMetalFeaturesFences))
 		{
 			GetMetalDeviceContext().FillBuffer(Buffer->Buffer, ns::Range(0, Buffer->Buffer.GetLength()), 0);
 		}
diff --git a/Engine/Source/Runtime/Apple/MetalRHI/Private/MetalVertexBuffer.cpp b/Engine/Source/Runtime/Apple/MetalRHI/Private/MetalVertexBuffer.cpp
--- a/Engine/Source/Runtime/Apple/MetalRHI/Private/MetalVertexBuffer.cpp
+++ b/Engine/Source/Runtime/Apple/MetalRHI/Private/MetalVertexBuffer.cpp
@@ -540,7 +540,7 @@ FVertexBufferRHIRef FMetalDynamicRHI::RHICreateVertexBuffer(uint32 Size, uint32
 			VertexBuffer->CPUBuffer = nil;
 		}
 
-		if (GMetalBufferZeroFill && !FMetalCommandQueue::SupportsFeature(EMetalFeaturesHeaps))
+		if (GMetalBufferZeroFill && !FMetalCommandQueue::SupportsFeature(EMetalFeaturesFences))
 		{
 			GetMetalDeviceContext().FillBuffer(VertexBuffer->Buffer, ns::Range(0, VertexBuffer->Buffer.GetLength()), 0);
 		}
@@ -626,7 +626,7 @@ struct FMetalRHICommandInitialiseVertexBuffer : public FRHICommand<FMetalRHIComm
 				Buffer->LastUpdate = GFrameNumberRenderThread;
 			}
 		}
-		else if (GMetalBufferZeroFill && !FMetalCommandQueue::SupportsFeature(EMetalFeaturesHeaps))
+		else if (GMetalBufferZeroFill && !FMetalCommandQueue::SupportsFeature(EMetalFeaturesFences))
 		{
 			GetMetalDeviceContext().FillBuffer(Buffer->Buffer, ns::Range(0, Buffer->Buffer.GetLength()), 0);
 		}
