According to AMD and Apple there should be no advantage to making visibility visibility buffers use Managed memory, so just use Shared on macOS and iOS.

diff --git a/Engine/Source/Runtime/Apple/MetalRHI/Private/MetalQuery.cpp b/Engine/Source/Runtime/Apple/MetalRHI/Private/MetalQuery.cpp
--- a/Engine/Source/Runtime/Apple/MetalRHI/Private/MetalQuery.cpp
+++ b/Engine/Source/Runtime/Apple/MetalRHI/Private/MetalQuery.cpp
@@ -55,16 +55,11 @@ FMetalQueryBuffer* FMetalQueryBufferPool::GetCurrentQueryBuffer()
 		{
 			LLM_SCOPE_METAL(ELLMTagMetal::Buffers);
 			LLM_PLATFORM_SCOPE_METAL(ELLMTagMetal::Buffers);
-#if PLATFORM_MAC
-			METAL_GPUPROFILE(FScopedMetalCPUStats CPUStat(FString::Printf(TEXT("AllocBuffer: %llu, %llu"), EQueryBufferMaxSize, mtlpp::ResourceOptions::StorageModeManaged)));
-			Buffer = FMetalBuffer(MTLPP_VALIDATE(mtlpp::Device, Context->GetDevice(), SafeGetRuntimeDebuggingLevel() >= EMetalDebugLevelValidation, NewBuffer(EQueryBufferMaxSize, GetMetalDeviceContext().GetCommandQueue().GetCompatibleResourceOptions((mtlpp::ResourceOptions)(BUFFER_CACHE_MODE | mtlpp::ResourceOptions::HazardTrackingModeUntracked | mtlpp::ResourceOptions::StorageModeManaged)))), false);
-			FMemory::Memzero((((uint8*)Buffer.GetContents())), EQueryBufferMaxSize);
-			Buffer.DidModify(ns::Range(0, EQueryBufferMaxSize));
-#else
+
 			METAL_GPUPROFILE(FScopedMetalCPUStats CPUStat(FString::Printf(TEXT("AllocBuffer: %llu, %llu"), EQueryBufferMaxSize, mtlpp::ResourceOptions::StorageModeShared)));
 			Buffer = FMetalBuffer(MTLPP_VALIDATE(mtlpp::Device, Context->GetDevice(), SafeGetRuntimeDebuggingLevel() >= EMetalDebugLevelValidation, NewBuffer(EQueryBufferMaxSize, GetMetalDeviceContext().GetCommandQueue().GetCompatibleResourceOptions((mtlpp::ResourceOptions)(BUFFER_CACHE_MODE | mtlpp::ResourceOptions::HazardTrackingModeUntracked | mtlpp::ResourceOptions::StorageModeShared)))), false);
 			FMemory::Memzero((((uint8*)Buffer.GetContents())), EQueryBufferMaxSize);
-#endif
+
 #if STATS || ENABLE_LOW_LEVEL_MEM_TRACKER
 			MetalLLM::LogAllocBuffer(Context->GetDevice(), Buffer);
 #endif
