Initial attempt to track the MTLResourceUsage flags when we bind resources into Metal - disabled unless -metaliabs is passed in on the command-line and requires -metalfence to work properly. First step to getting everything in place for IABs.

diff --git a/Engine/Source/Runtime/Apple/MetalRHI/Private/MetalCommandEncoder.cpp b/Engine/Source/Runtime/Apple/MetalRHI/Private/MetalCommandEncoder.cpp
--- a/Engine/Source/Runtime/Apple/MetalRHI/Private/MetalCommandEncoder.cpp
+++ b/Engine/Source/Runtime/Apple/MetalRHI/Private/MetalCommandEncoder.cpp
@@ -42,6 +42,7 @@ FMetalCommandEncoder::FMetalCommandEncoder(FMetalCommandList& CmdList)
 		FMemory::Memzero(ShaderBuffers[Frequency].Bytes);
 		FMemory::Memzero(ShaderBuffers[Frequency].Offsets);
 		FMemory::Memzero(ShaderBuffers[Frequency].Lengths);
+		FMemory::Memzero(ShaderBuffers[Frequency].Usage);
 		ShaderBuffers[Frequency].SideTable = [[FMetalBufferData alloc] init];
 		ShaderBuffers[Frequency].SideTable->Data = (uint8*)(&ShaderBuffers[Frequency].Lengths[0]);
 		ShaderBuffers[Frequency].SideTable->Len = sizeof(ShaderBuffers[Frequency].Lengths);
@@ -85,6 +86,7 @@ FMetalCommandEncoder::~FMetalCommandEncoder(void)
 		FMemory::Memzero(ShaderBuffers[Frequency].Bytes);
 		FMemory::Memzero(ShaderBuffers[Frequency].Offsets);
 		FMemory::Memzero(ShaderBuffers[Frequency].Lengths);
+		FMemory::Memzero(ShaderBuffers[Frequency].Usage);
 		ShaderBuffers[Frequency].SideTable->Data = nullptr;
 		[ShaderBuffers[Frequency].SideTable release];
 		ShaderBuffers[Frequency].SideTable = nil;
@@ -124,6 +126,7 @@ void FMetalCommandEncoder::Reset(void)
     	FMemory::Memzero(ShaderBuffers[Frequency].Bytes);
 		FMemory::Memzero(ShaderBuffers[Frequency].Offsets);
 		FMemory::Memzero(ShaderBuffers[Frequency].Lengths);
+		FMemory::Memzero(ShaderBuffers[Frequency].Usage);
 		ShaderBuffers[Frequency].Bound = 0;
 	}
 	
@@ -679,6 +682,7 @@ FMetalFence* FMetalCommandEncoder::EndEncoding(void)
 		FMemory::Memzero(ShaderBuffers[Frequency].Bytes);
 		FMemory::Memzero(ShaderBuffers[Frequency].Offsets);
 		FMemory::Memzero(ShaderBuffers[Frequency].Lengths);
+		FMemory::Memzero(ShaderBuffers[Frequency].Usage);
 		ShaderBuffers[Frequency].Bound = 0;
 	}
     return Fence;
@@ -984,6 +988,7 @@ void FMetalCommandEncoder::SetRenderPassDescriptor(mtlpp::RenderPassDescriptor R
 		FMemory::Memzero(ShaderBuffers[Frequency].Bytes);
 		FMemory::Memzero(ShaderBuffers[Frequency].Offsets);
 		FMemory::Memzero(ShaderBuffers[Frequency].Lengths);
+		FMemory::Memzero(ShaderBuffers[Frequency].Usage);
 		ShaderBuffers[Frequency].Bound = 0;
 	}
 }
@@ -1115,14 +1120,17 @@ void FMetalCommandEncoder::SetVisibilityResultMode(mtlpp::VisibilityResultMode c
 	
 #pragma mark - Public Shader Resource Mutators -
 
-void FMetalCommandEncoder::SetShaderBuffer(mtlpp::FunctionType const FunctionType, FMetalBuffer const& Buffer, NSUInteger const Offset, NSUInteger const Length, NSUInteger index, EPixelFormat const Format)
+void FMetalCommandEncoder::SetShaderBuffer(mtlpp::FunctionType const FunctionType, FMetalBuffer const& Buffer, NSUInteger const Offset, NSUInteger const Length, NSUInteger index, mtlpp::ResourceUsage const Usage, EPixelFormat const Format)
 {
 	check(index < ML_MaxBuffers);
     if(GetMetalDeviceContext().SupportsFeature(EMetalFeaturesSetBufferOffset) && Buffer && (ShaderBuffers[uint32(FunctionType)].Bound & (1 << index)) && ShaderBuffers[uint32(FunctionType)].Buffers[index] == Buffer)
     {
 		SetShaderBufferOffset(FunctionType, Offset, Length, index);
 		ShaderBuffers[uint32(FunctionType)].Lengths[index+ML_MaxBuffers] = GMetalBufferFormats[Format].DataFormat;
-    }
+		ShaderBuffers[uint32(FunctionType)].Usage[index] = Usage;
+		
+		UseResource(Buffer, Usage);
+	}
     else
     {
 		if(Buffer)
@@ -1136,6 +1144,7 @@ void FMetalCommandEncoder::SetShaderBuffer(mtlpp::FunctionType const FunctionTyp
 		ShaderBuffers[uint32(FunctionType)].Buffers[index] = Buffer;
 		ShaderBuffers[uint32(FunctionType)].Bytes[index] = nil;
 		ShaderBuffers[uint32(FunctionType)].Offsets[index] = Offset;
+		ShaderBuffers[uint32(FunctionType)].Usage[index] = Usage;
 		ShaderBuffers[uint32(FunctionType)].Lengths[index] = Length;
 		ShaderBuffers[uint32(FunctionType)].Lengths[index+ML_MaxBuffers] = GMetalBufferFormats[Format].DataFormat;
 		
@@ -1155,7 +1164,7 @@ void FMetalCommandEncoder::SetShaderData(mtlpp::FunctionType const FunctionType,
 #if METAL_DEBUG_OPTIONS
 	if (CommandList.GetCommandQueue().GetRuntimeDebuggingLevel() > EMetalDebugLevelResetOnBind)
 	{
-		SetShaderBuffer(FunctionType, nil, 0, 0, Index);
+		SetShaderBuffer(FunctionType, nil, 0, 0, Index, mtlpp::ResourceUsage(0));
 	}
 #endif
 	
@@ -1171,6 +1180,7 @@ void FMetalCommandEncoder::SetShaderData(mtlpp::FunctionType const FunctionType,
 	ShaderBuffers[uint32(FunctionType)].Buffers[Index] = nil;
 	ShaderBuffers[uint32(FunctionType)].Bytes[Index] = Data;
 	ShaderBuffers[uint32(FunctionType)].Offsets[Index] = Offset;
+	ShaderBuffers[uint32(FunctionType)].Usage[Index] = mtlpp::ResourceUsage::Read;
 	ShaderBuffers[uint32(FunctionType)].Lengths[Index] = Data ? (Data->Len - Offset) : 0;
 	ShaderBuffers[uint32(FunctionType)].Lengths[Index+ML_MaxBuffers] = GMetalBufferFormats[Format].DataFormat;
 	
@@ -1184,7 +1194,7 @@ void FMetalCommandEncoder::SetShaderBytes(mtlpp::FunctionType const FunctionType
 #if METAL_DEBUG_OPTIONS
 	if (CommandList.GetCommandQueue().GetRuntimeDebuggingLevel() > EMetalDebugLevelResetOnBind)
 	{
-		SetShaderBuffer(FunctionType, nil, 0, 0, Index);
+		SetShaderBuffer(FunctionType, nil, 0, 0, Index, mtlpp::ResourceUsage(0));
 	}
 #endif
 	
@@ -1226,6 +1236,7 @@ void FMetalCommandEncoder::SetShaderBytes(mtlpp::FunctionType const FunctionType
 		}
 		ShaderBuffers[uint32(FunctionType)].Bytes[Index] = nil;
 		ShaderBuffers[uint32(FunctionType)].Offsets[Index] = 0;
+		ShaderBuffers[uint32(FunctionType)].Usage[Index] = mtlpp::ResourceUsage::Read;
 		ShaderBuffers[uint32(FunctionType)].Lengths[Index] = Length;
 		ShaderBuffers[uint32(FunctionType)].Lengths[Index+ML_MaxBuffers] = GMetalBufferFormats[PF_Unknown].DataFormat;
 	}
@@ -1236,6 +1247,7 @@ void FMetalCommandEncoder::SetShaderBytes(mtlpp::FunctionType const FunctionType
 		ShaderBuffers[uint32(FunctionType)].Buffers[Index] = nil;
 		ShaderBuffers[uint32(FunctionType)].Bytes[Index] = nil;
 		ShaderBuffers[uint32(FunctionType)].Offsets[Index] = 0;
+		ShaderBuffers[uint32(FunctionType)].Usage[Index] = mtlpp::ResourceUsage(0);
 		ShaderBuffers[uint32(FunctionType)].Lengths[Index] = 0;
 		ShaderBuffers[uint32(FunctionType)].Lengths[Index+ML_MaxBuffers] = GMetalBufferFormats[PF_Unknown].DataFormat;
 	}
@@ -1274,7 +1286,7 @@ void FMetalCommandEncoder::SetShaderBufferOffset(mtlpp::FunctionType FunctionTyp
 	}
 }
 
-void FMetalCommandEncoder::SetShaderTexture(mtlpp::FunctionType FunctionType, FMetalTexture const& Texture, NSUInteger index)
+void FMetalCommandEncoder::SetShaderTexture(mtlpp::FunctionType FunctionType, FMetalTexture const& Texture, NSUInteger index, mtlpp::ResourceUsage Usage)
 {
 	check(index < ML_MaxTextures);
 	switch (FunctionType)
@@ -1304,6 +1316,7 @@ void FMetalCommandEncoder::SetShaderTexture(mtlpp::FunctionType FunctionType, FM
 	
 	if (Texture)
 	{
+		UseResource(Texture, Usage);
 		TextureBindingHistory.Add(ns::AutoReleased<FMetalTexture>(Texture));
 	}
 }
@@ -1375,9 +1388,31 @@ bool FMetalCommandEncoder::HasBufferBindingHistory(FMetalBuffer const& Buffer) c
 
 #pragma mark - Private Functions -
 
+void FMetalCommandEncoder::UseResource(mtlpp::Resource const& Resource, mtlpp::ResourceUsage const Usage)
+{
+	static bool UseResourceAvailable = FMetalCommandQueue::SupportsFeature(EMetalFeaturesIABs);
+	if (UseResourceAvailable || SafeGetRuntimeDebuggingLevel() >= EMetalDebugLevelValidation)
+	{
+		mtlpp::ResourceUsage Current = ResourceUsage.FindRef(Resource.GetPtr());
+		if (Current != Usage)
+		{
+			ResourceUsage.Add(Resource.GetPtr(), Usage);
+			if (RenderCommandEncoder)
+			{
+				MTLPP_VALIDATE(mtlpp::RenderCommandEncoder, RenderCommandEncoder, SafeGetRuntimeDebuggingLevel() >= EMetalDebugLevelValidation, UseResource(Resource, Usage));
+			}
+			else if (ComputeCommandEncoder)
+			{
+				MTLPP_VALIDATE(mtlpp::ComputeCommandEncoder, ComputeCommandEncoder, SafeGetRuntimeDebuggingLevel() >= EMetalDebugLevelValidation, UseResource(Resource, Usage));
+			}
+		}
+	}
+}
+
 void FMetalCommandEncoder::SetShaderBufferInternal(mtlpp::FunctionType Function, uint32 Index)
 {
 	NSUInteger Offset = ShaderBuffers[uint32(Function)].Offsets[Index];
+	mtlpp::ResourceUsage Usage = ShaderBuffers[uint32(Function)].Usage[Index];
 	bool bBufferHasBytes = ShaderBuffers[uint32(Function)].Bytes[Index] != nil;
 	if (!ShaderBuffers[uint32(Function)].Buffers[Index] && bBufferHasBytes && !bSupportsMetalFeaturesSetBytes)
 	{
@@ -1428,8 +1463,11 @@ void FMetalCommandEncoder::SetShaderBufferInternal(mtlpp::FunctionType Function,
 				break;
 		}
 		
+		UseResource(Buffer, Usage);
+		
 		if (Buffer.IsSingleUse())
 		{
+			ShaderBuffers[uint32(Function)].Usage[Index] = mtlpp::ResourceUsage(0);
 			ShaderBuffers[uint32(Function)].Offsets[Index] = 0;
 			ShaderBuffers[uint32(Function)].Buffers[Index] = nil;
 			ShaderBuffers[uint32(Function)].Bound &= ~(1 << Index);
diff --git a/Engine/Source/Runtime/Apple/MetalRHI/Private/MetalCommandEncoder.h b/Engine/Source/Runtime/Apple/MetalRHI/Private/MetalCommandEncoder.h
--- a/Engine/Source/Runtime/Apple/MetalRHI/Private/MetalCommandEncoder.h
+++ b/Engine/Source/Runtime/Apple/MetalRHI/Private/MetalCommandEncoder.h
@@ -308,8 +308,10 @@ public:
 	 * @param Offset The offset in the buffer or 0 when Buffer is nil.
 	 * @param Length The data length - caller is responsible for accounting for non-zero Offset.
 	 * @param Index The index to modify.
+	 * @param Usage The resource usage mask.
+	 * @param Format The Pixel format to reinterpret the resource as.
 	 */
-	void SetShaderBuffer(mtlpp::FunctionType const FunctionType, FMetalBuffer const& Buffer, NSUInteger const Offset, NSUInteger const Length, NSUInteger const Index, EPixelFormat const Format = PF_Unknown);
+	void SetShaderBuffer(mtlpp::FunctionType const FunctionType, FMetalBuffer const& Buffer, NSUInteger const Offset, NSUInteger const Length, NSUInteger const Index, mtlpp::ResourceUsage const Usage, EPixelFormat const Format = PF_Unknown);
 	
 	/*
 	 * Set an FMetalBufferData to the specified shader frequency at the given bind point index.
@@ -317,6 +319,7 @@ public:
 	 * @param Data The data to bind or nullptr to clear.
 	 * @param Offset The offset in the buffer or 0 when Buffer is nil.
 	 * @param Index The index to modify.
+	 * @param Format The pixel format to reinterpret the resource as.
 	 */
 	void SetShaderData(mtlpp::FunctionType const FunctionType, FMetalBufferData* Data, NSUInteger const Offset, NSUInteger const Index, EPixelFormat const Format = PF_Unknown);
 	
@@ -334,8 +337,9 @@ public:
 	 * @param FunctionType The shader function to modify.
 	 * @param Texture The texture to bind or nil to clear.
 	 * @param Index The index to modify.
+	 * @param Usage The resource usage mask.
 	 */
-	void SetShaderTexture(mtlpp::FunctionType const FunctionType, FMetalTexture const& Texture, NSUInteger const Index);
+	void SetShaderTexture(mtlpp::FunctionType const FunctionType, FMetalTexture const& Texture, NSUInteger const Index, mtlpp::ResourceUsage const Usage);
 	
 	/*
 	 * Set a global sampler for the specified shader frequency at the given bind point index.
@@ -390,6 +394,8 @@ private:
 	
 	void SetShaderBufferInternal(mtlpp::FunctionType Function, uint32 Index);
 	
+	void UseResource(mtlpp::Resource const& Resource, mtlpp::ResourceUsage const Usage);
+	
 #pragma mark - Private Type Declarations -
 private:
     /** A structure of arrays for the current buffer binding settings. */
@@ -403,6 +409,8 @@ private:
         FMetalBufferData* Bytes[ML_MaxBuffers];
         /** The bound buffer offsets or 0. */
         NSUInteger Offsets[ML_MaxBuffers];
+		/** The usage mask for the bound resource or 0 */
+		mtlpp::ResourceUsage Usage[ML_MaxBuffers];
 		/** The bound buffer lengths */
 		uint32 Lengths[ML_MaxBuffers*2];
         /** A bitmask for which buffers were bound by the application where a bit value of 1 is bound and 0 is unbound. */
@@ -456,5 +464,7 @@ private:
 	TSet<ns::AutoReleased<FMetalBuffer>> BufferBindingHistory;
 	TSet<ns::AutoReleased<FMetalTexture>> TextureBindingHistory;
 	
+	TMap<mtlpp::Resource::Type, mtlpp::ResourceUsage> ResourceUsage;
+	
 	uint32 EncoderNum;
 };
diff --git a/Engine/Source/Runtime/Apple/MetalRHI/Private/MetalCommandQueue.cpp b/Engine/Source/Runtime/Apple/MetalRHI/Private/MetalCommandQueue.cpp
--- a/Engine/Source/Runtime/Apple/MetalRHI/Private/MetalCommandQueue.cpp
+++ b/Engine/Source/Runtime/Apple/MetalRHI/Private/MetalCommandQueue.cpp
@@ -212,6 +212,11 @@ FMetalCommandQueue::FMetalCommandQueue(mtlpp::Device InDevice, uint32 const MaxN
 		{
 			Features |= EMetalFeaturesHeaps;
 		}
+		
+		if (FParse::Param(FCommandLine::Get(),TEXT("metaliabs")))
+		{
+			Features |= EMetalFeaturesIABs;
+		}
     }
 	else
 	// Time query emulation breaks on AMD < 10.13 - disable by default until they can explain why, should work everywhere else.
diff --git a/Engine/Source/Runtime/Apple/MetalRHI/Private/MetalCommandQueue.h b/Engine/Source/Runtime/Apple/MetalRHI/Private/MetalCommandQueue.h
--- a/Engine/Source/Runtime/Apple/MetalRHI/Private/MetalCommandQueue.h
+++ b/Engine/Source/Runtime/Apple/MetalRHI/Private/MetalCommandQueue.h
@@ -99,6 +99,8 @@ typedef NS_OPTIONS(uint64, EMetalFeatures)
 	EMetalFeaturesMaxThreadsPerThreadgroup = 1llu << 39llu,
 	/** Supports parallel render encoders */
 	EMetalFeaturesParallelRenderEncoders = 1llu << 40llu,
+	/** Supports indirect argument buffers */
+	EMetalFeaturesIABs = 1llu << 41llu,
 };
 
 /**
diff --git a/Engine/Source/Runtime/Apple/MetalRHI/Private/MetalCommands.cpp b/Engine/Source/Runtime/Apple/MetalRHI/Private/MetalCommands.cpp
--- a/Engine/Source/Runtime/Apple/MetalRHI/Private/MetalCommands.cpp
+++ b/Engine/Source/Runtime/Apple/MetalRHI/Private/MetalCommands.cpp
@@ -275,17 +275,17 @@ void FMetalRHICommandContext::RHISetShaderTexture(FVertexShaderRHIParamRef Verte
 	{
         if (Surface->Texture || !(Surface->Flags & TexCreate_Presentable))
         {
-            Context->GetCurrentState().SetShaderTexture(SF_Vertex, Surface->Texture, TextureIndex);
+            Context->GetCurrentState().SetShaderTexture(SF_Vertex, Surface->Texture, TextureIndex, (mtlpp::ResourceUsage)(mtlpp::ResourceUsage::Read|mtlpp::ResourceUsage::Sample));
         }
         else
         {
             FMetalTexture Tex = Surface->GetCurrentTexture();
-            Context->GetCurrentState().SetShaderTexture(SF_Vertex, Tex, TextureIndex);
+            Context->GetCurrentState().SetShaderTexture(SF_Vertex, Tex, TextureIndex, (mtlpp::ResourceUsage)(mtlpp::ResourceUsage::Read|mtlpp::ResourceUsage::Sample));
         }
 	}
 	else
 	{
-		Context->GetCurrentState().SetShaderTexture(SF_Vertex, nil, TextureIndex);
+		Context->GetCurrentState().SetShaderTexture(SF_Vertex, nil, TextureIndex, mtlpp::ResourceUsage(0));
 	}
 	}
 }
@@ -298,17 +298,17 @@ void FMetalRHICommandContext::RHISetShaderTexture(FHullShaderRHIParamRef HullSha
     {
         if (Surface->Texture || !(Surface->Flags & TexCreate_Presentable))
         {
-            Context->GetCurrentState().SetShaderTexture(SF_Hull, Surface->Texture, TextureIndex);
+            Context->GetCurrentState().SetShaderTexture(SF_Hull, Surface->Texture, TextureIndex, (mtlpp::ResourceUsage)(mtlpp::ResourceUsage::Read|mtlpp::ResourceUsage::Sample));
         }
         else
         {
             FMetalTexture Tex = Surface->GetCurrentTexture();
-            Context->GetCurrentState().SetShaderTexture(SF_Hull, Tex, TextureIndex);
+            Context->GetCurrentState().SetShaderTexture(SF_Hull, Tex, TextureIndex, (mtlpp::ResourceUsage)(mtlpp::ResourceUsage::Read|mtlpp::ResourceUsage::Sample));
         }
 	}
 	else
 	{
-		Context->GetCurrentState().SetShaderTexture(SF_Hull, nil, TextureIndex);
+		Context->GetCurrentState().SetShaderTexture(SF_Hull, nil, TextureIndex, mtlpp::ResourceUsage(0));
 	}
 	}
 }
@@ -321,17 +321,17 @@ void FMetalRHICommandContext::RHISetShaderTexture(FDomainShaderRHIParamRef Domai
     {
         if (Surface->Texture || !(Surface->Flags & TexCreate_Presentable))
         {
-            Context->GetCurrentState().SetShaderTexture(SF_Domain, Surface->Texture, TextureIndex);
+            Context->GetCurrentState().SetShaderTexture(SF_Domain, Surface->Texture, TextureIndex, (mtlpp::ResourceUsage)(mtlpp::ResourceUsage::Read|mtlpp::ResourceUsage::Sample));
         }
         else
         {
             FMetalTexture Tex = Surface->GetCurrentTexture();
-            Context->GetCurrentState().SetShaderTexture(SF_Domain, Tex, TextureIndex);
+            Context->GetCurrentState().SetShaderTexture(SF_Domain, Tex, TextureIndex, (mtlpp::ResourceUsage)(mtlpp::ResourceUsage::Read|mtlpp::ResourceUsage::Sample));
         }
 	}
 	else
 	{
-		Context->GetCurrentState().SetShaderTexture(SF_Domain, nil, TextureIndex);
+		Context->GetCurrentState().SetShaderTexture(SF_Domain, nil, TextureIndex, mtlpp::ResourceUsage(0));
 	}
 	}
 }
@@ -350,17 +350,17 @@ void FMetalRHICommandContext::RHISetShaderTexture(FPixelShaderRHIParamRef PixelS
     {
         if (Surface->Texture || !(Surface->Flags & TexCreate_Presentable))
         {
-            Context->GetCurrentState().SetShaderTexture(SF_Pixel, Surface->Texture, TextureIndex);
+            Context->GetCurrentState().SetShaderTexture(SF_Pixel, Surface->Texture, TextureIndex, (mtlpp::ResourceUsage)(mtlpp::ResourceUsage::Read|mtlpp::ResourceUsage::Sample));
         }
         else
         {
             FMetalTexture Tex = Surface->GetCurrentTexture();
-            Context->GetCurrentState().SetShaderTexture(SF_Pixel, Tex, TextureIndex);
+            Context->GetCurrentState().SetShaderTexture(SF_Pixel, Tex, TextureIndex, (mtlpp::ResourceUsage)(mtlpp::ResourceUsage::Read|mtlpp::ResourceUsage::Sample));
         }
 	}
 	else
 	{
-		Context->GetCurrentState().SetShaderTexture(SF_Pixel, nil, TextureIndex);
+		Context->GetCurrentState().SetShaderTexture(SF_Pixel, nil, TextureIndex, mtlpp::ResourceUsage(0));
 	}
 	}
 }
@@ -373,17 +373,17 @@ void FMetalRHICommandContext::RHISetShaderTexture(FComputeShaderRHIParamRef Comp
     {
         if (Surface->Texture || !(Surface->Flags & TexCreate_Presentable))
         {
-            Context->GetCurrentState().SetShaderTexture(SF_Compute, Surface->Texture, TextureIndex);
+            Context->GetCurrentState().SetShaderTexture(SF_Compute, Surface->Texture, TextureIndex, (mtlpp::ResourceUsage)(mtlpp::ResourceUsage::Read|mtlpp::ResourceUsage::Sample));
         }
         else
         {
             FMetalTexture Tex = Surface->GetCurrentTexture();
-            Context->GetCurrentState().SetShaderTexture(SF_Compute, Tex, TextureIndex);
+            Context->GetCurrentState().SetShaderTexture(SF_Compute, Tex, TextureIndex, (mtlpp::ResourceUsage)(mtlpp::ResourceUsage::Read|mtlpp::ResourceUsage::Sample));
         }
 	}
 	else
 	{
-		Context->GetCurrentState().SetShaderTexture(SF_Compute, nil, TextureIndex);
+		Context->GetCurrentState().SetShaderTexture(SF_Compute, nil, TextureIndex, mtlpp::ResourceUsage(0));
 	}
 	}
 }
@@ -536,7 +536,7 @@ void FMetalRHICommandContext::RHISetShaderUniformBuffer(FVertexShaderRHIParamRef
 	if (Bindings.ConstantBuffers & 1 << BufferIndex)
 	{
 		auto* UB = (FMetalUniformBuffer*)BufferRHI;
-		Context->GetCurrentState().SetShaderBuffer(SF_Vertex, UB->Buffer, UB->Data, 0, UB->GetSize(), BufferIndex);
+		Context->GetCurrentState().SetShaderBuffer(SF_Vertex, UB->Buffer, UB->Data, 0, UB->GetSize(), BufferIndex, mtlpp::ResourceUsage::Read);
 	}
 	}
 }
@@ -552,7 +552,7 @@ void FMetalRHICommandContext::RHISetShaderUniformBuffer(FHullShaderRHIParamRef H
 	if (Bindings.ConstantBuffers & 1 << BufferIndex)
 	{
 		auto* UB = (FMetalUniformBuffer*)BufferRHI;
-		Context->GetCurrentState().SetShaderBuffer(SF_Hull, UB->Buffer, UB->Data, 0, UB->GetSize(), BufferIndex);
+		Context->GetCurrentState().SetShaderBuffer(SF_Hull, UB->Buffer, UB->Data, 0, UB->GetSize(), BufferIndex, mtlpp::ResourceUsage::Read);
 	}
 	}
 }
@@ -568,7 +568,7 @@ void FMetalRHICommandContext::RHISetShaderUniformBuffer(FDomainShaderRHIParamRef
 	if (Bindings.ConstantBuffers & 1 << BufferIndex)
 	{
 		auto* UB = (FMetalUniformBuffer*)BufferRHI;
-		Context->GetCurrentState().SetShaderBuffer(SF_Domain, UB->Buffer, UB->Data, 0, UB->GetSize(), BufferIndex);
+		Context->GetCurrentState().SetShaderBuffer(SF_Domain, UB->Buffer, UB->Data, 0, UB->GetSize(), BufferIndex, mtlpp::ResourceUsage::Read);
 	}
 	}
 }
@@ -589,7 +589,7 @@ void FMetalRHICommandContext::RHISetShaderUniformBuffer(FPixelShaderRHIParamRef
 	if (Bindings.ConstantBuffers & 1 << BufferIndex)
 	{
 		auto* UB = (FMetalUniformBuffer*)BufferRHI;
-		Context->GetCurrentState().SetShaderBuffer(SF_Pixel, UB->Buffer, UB->Data, 0, UB->GetSize(), BufferIndex);
+		Context->GetCurrentState().SetShaderBuffer(SF_Pixel, UB->Buffer, UB->Data, 0, UB->GetSize(), BufferIndex, mtlpp::ResourceUsage::Read);
 	}
 	}
 }
@@ -605,7 +605,7 @@ void FMetalRHICommandContext::RHISetShaderUniformBuffer(FComputeShaderRHIParamRe
 	if (Bindings.ConstantBuffers & 1 << BufferIndex)
 	{
 		auto* UB = (FMetalUniformBuffer*)BufferRHI;
-		Context->GetCurrentState().SetShaderBuffer(SF_Compute, UB->Buffer, UB->Data, 0, UB->GetSize(), BufferIndex);
+		Context->GetCurrentState().SetShaderBuffer(SF_Compute, UB->Buffer, UB->Data, 0, UB->GetSize(), BufferIndex, mtlpp::ResourceUsage::Read);
 	}
 	}
 }
@@ -845,8 +845,8 @@ void FMetalRHICommandContext::RHIEndDrawPrimitiveUP()
 	Context->GetCurrentState().SetVertexStream(0, &PendingVertexBuffer, nil, 0, PendingVertexDataStride * NumVertices);
 	if(Context->GetCurrentState().GetUsingTessellation())
 	{
-		Context->GetCurrentState().SetShaderBuffer(SF_Hull, PendingVertexBuffer, nil, 0, PendingVertexDataStride * NumVertices, UNREAL_TO_METAL_BUFFER_INDEX(0));
-		Context->GetCurrentState().SetShaderBuffer(SF_Domain, PendingVertexBuffer, nil, 0, PendingVertexDataStride * NumVertices, UNREAL_TO_METAL_BUFFER_INDEX(0));
+		Context->GetCurrentState().SetShaderBuffer(SF_Hull, PendingVertexBuffer, nil, 0, PendingVertexDataStride * NumVertices, UNREAL_TO_METAL_BUFFER_INDEX(0), mtlpp::ResourceUsage::Read);
+		Context->GetCurrentState().SetShaderBuffer(SF_Domain, PendingVertexBuffer, nil, 0, PendingVertexDataStride * NumVertices, UNREAL_TO_METAL_BUFFER_INDEX(0), mtlpp::ResourceUsage::Read);
 	}
 	
 	Context->DrawPrimitive(PendingPrimitiveType, 0, PendingNumPrimitives, Context->GetCurrentState().GetRenderTargetArraySize());
diff --git a/Engine/Source/Runtime/Apple/MetalRHI/Private/MetalRenderPass.cpp b/Engine/Source/Runtime/Apple/MetalRHI/Private/MetalRenderPass.cpp
--- a/Engine/Source/Runtime/Apple/MetalRHI/Private/MetalRenderPass.cpp
+++ b/Engine/Source/Runtime/Apple/MetalRHI/Private/MetalRenderPass.cpp
@@ -616,61 +616,61 @@ void FMetalRenderPass::DrawPatches(uint32 PrimitiveType,FMetalBuffer const& Inde
 		// Per-draw call bindings should *not* be cached in the StateCache - causes absolute chaos.
 		if(IndexBuffer && Pipeline->TessellationPipelineDesc.TessellationControlPointIndexBufferIndex != UINT_MAX)
 		{
-			PrologueEncoder.SetShaderBuffer(mtlpp::FunctionType::Kernel, IndexBuffer, StartIndex * IndexBufferStride, IndexBuffer.GetLength() - (StartIndex * IndexBufferStride), Pipeline->TessellationPipelineDesc.TessellationControlPointIndexBufferIndex);
-			PrologueEncoder.SetShaderBuffer(mtlpp::FunctionType::Kernel, IndexBuffer, StartIndex * IndexBufferStride, IndexBuffer.GetLength() - (StartIndex * IndexBufferStride), Pipeline->TessellationPipelineDesc.TessellationIndexBufferIndex);
-			State.SetShaderBuffer(SF_Vertex, nil, nil, 0, 0, Pipeline->TessellationPipelineDesc.TessellationControlPointIndexBufferIndex);
-			State.SetShaderBuffer(SF_Vertex, nil, nil, 0, 0, Pipeline->TessellationPipelineDesc.TessellationIndexBufferIndex);
+			PrologueEncoder.SetShaderBuffer(mtlpp::FunctionType::Kernel, IndexBuffer, StartIndex * IndexBufferStride, IndexBuffer.GetLength() - (StartIndex * IndexBufferStride), Pipeline->TessellationPipelineDesc.TessellationControlPointIndexBufferIndex, mtlpp::ResourceUsage::Read);
+			PrologueEncoder.SetShaderBuffer(mtlpp::FunctionType::Kernel, IndexBuffer, StartIndex * IndexBufferStride, IndexBuffer.GetLength() - (StartIndex * IndexBufferStride), Pipeline->TessellationPipelineDesc.TessellationIndexBufferIndex, mtlpp::ResourceUsage::Read);
+			State.SetShaderBuffer(SF_Vertex, nil, nil, 0, 0, Pipeline->TessellationPipelineDesc.TessellationControlPointIndexBufferIndex, mtlpp::ResourceUsage(0));
+			State.SetShaderBuffer(SF_Vertex, nil, nil, 0, 0, Pipeline->TessellationPipelineDesc.TessellationIndexBufferIndex, mtlpp::ResourceUsage(0));
 		}
 		
 		if(Pipeline->TessellationPipelineDesc.TessellationOutputControlPointBufferIndex != UINT_MAX) //TessellationOutputControlPointBufferIndex -> hullShaderOutputBuffer
 		{
-			PrologueEncoder.SetShaderBuffer(mtlpp::FunctionType::Kernel, hullShaderOutputBuffer, hullShaderOutputOffset, hullShaderOutputBuffer.GetLength() - hullShaderOutputOffset, Pipeline->TessellationPipelineDesc.TessellationOutputControlPointBufferIndex);
-			State.SetShaderBuffer(SF_Vertex, nil, nil, 0, 0, Pipeline->TessellationPipelineDesc.TessellationOutputControlPointBufferIndex);
+			PrologueEncoder.SetShaderBuffer(mtlpp::FunctionType::Kernel, hullShaderOutputBuffer, hullShaderOutputOffset, hullShaderOutputBuffer.GetLength() - hullShaderOutputOffset, Pipeline->TessellationPipelineDesc.TessellationOutputControlPointBufferIndex, mtlpp::ResourceUsage(0));
+			State.SetShaderBuffer(SF_Vertex, nil, nil, 0, 0, Pipeline->TessellationPipelineDesc.TessellationOutputControlPointBufferIndex, mtlpp::ResourceUsage(0));
 		}
 		
 		if(Pipeline->TessellationPipelineDesc.TessellationPatchConstBufferIndex != UINT_MAX) //TessellationPatchConstBufferIndex -> hullConstShaderOutputBuffer
 		{
-			PrologueEncoder.SetShaderBuffer(mtlpp::FunctionType::Kernel, hullConstShaderOutputBuffer, hullConstShaderOutputOffset, hullConstShaderOutputBuffer.GetLength() - hullConstShaderOutputOffset, Pipeline->TessellationPipelineDesc.TessellationPatchConstBufferIndex);
-			State.SetShaderBuffer(SF_Vertex, nil, nil, 0, 0, Pipeline->TessellationPipelineDesc.TessellationPatchConstBufferIndex);
+			PrologueEncoder.SetShaderBuffer(mtlpp::FunctionType::Kernel, hullConstShaderOutputBuffer, hullConstShaderOutputOffset, hullConstShaderOutputBuffer.GetLength() - hullConstShaderOutputOffset, Pipeline->TessellationPipelineDesc.TessellationPatchConstBufferIndex, mtlpp::ResourceUsage(0));
+			State.SetShaderBuffer(SF_Vertex, nil, nil, 0, 0, Pipeline->TessellationPipelineDesc.TessellationPatchConstBufferIndex, mtlpp::ResourceUsage(0));
 		}
 		
 		if(Pipeline->TessellationPipelineDesc.TessellationFactorBufferIndex != UINT_MAX) // TessellationFactorBufferIndex->tessellationFactorBuffer
 		{
-			PrologueEncoder.SetShaderBuffer(mtlpp::FunctionType::Kernel, tessellationFactorBuffer, tessellationFactorsOffset, tessellationFactorBuffer.GetLength() - tessellationFactorsOffset, Pipeline->TessellationPipelineDesc.TessellationFactorBufferIndex);
-			State.SetShaderBuffer(SF_Vertex, nil, nil, 0, 0, Pipeline->TessellationPipelineDesc.TessellationFactorBufferIndex);
+			PrologueEncoder.SetShaderBuffer(mtlpp::FunctionType::Kernel, tessellationFactorBuffer, tessellationFactorsOffset, tessellationFactorBuffer.GetLength() - tessellationFactorsOffset, Pipeline->TessellationPipelineDesc.TessellationFactorBufferIndex, mtlpp::ResourceUsage(0));
+			State.SetShaderBuffer(SF_Vertex, nil, nil, 0, 0, Pipeline->TessellationPipelineDesc.TessellationFactorBufferIndex, mtlpp::ResourceUsage(0));
 		}
 		
 		if(Pipeline->TessellationPipelineDesc.TessellationInputControlPointBufferIndex != UINT_MAX) //TessellationInputControlPointBufferIndex->hullShaderOutputBuffer
 		{
-			CurrentEncoder.SetShaderBuffer(mtlpp::FunctionType::Vertex, hullShaderOutputBuffer, hullShaderOutputOffset, hullShaderOutputBuffer.GetLength() - hullShaderOutputOffset, Pipeline->TessellationPipelineDesc.TessellationInputControlPointBufferIndex);
-			State.SetShaderBuffer(SF_Domain, nil, nil, 0, 0, Pipeline->TessellationPipelineDesc.TessellationInputControlPointBufferIndex);
+			CurrentEncoder.SetShaderBuffer(mtlpp::FunctionType::Vertex, hullShaderOutputBuffer, hullShaderOutputOffset, hullShaderOutputBuffer.GetLength() - hullShaderOutputOffset, Pipeline->TessellationPipelineDesc.TessellationInputControlPointBufferIndex, mtlpp::ResourceUsage(0));
+			State.SetShaderBuffer(SF_Domain, nil, nil, 0, 0, Pipeline->TessellationPipelineDesc.TessellationInputControlPointBufferIndex, mtlpp::ResourceUsage(0));
 		}
 		if(Pipeline->TessellationPipelineDesc.TessellationInputPatchConstBufferIndex != UINT_MAX) //TessellationInputPatchConstBufferIndex->hullConstShaderOutputBuffer
 		{
-			CurrentEncoder.SetShaderBuffer(mtlpp::FunctionType::Vertex, hullConstShaderOutputBuffer, hullConstShaderOutputOffset, hullConstShaderOutputBuffer.GetLength() - hullConstShaderOutputOffset, Pipeline->TessellationPipelineDesc.TessellationInputPatchConstBufferIndex);
-			State.SetShaderBuffer(SF_Domain, nil, nil, 0, 0, Pipeline->TessellationPipelineDesc.TessellationInputPatchConstBufferIndex);
+			CurrentEncoder.SetShaderBuffer(mtlpp::FunctionType::Vertex, hullConstShaderOutputBuffer, hullConstShaderOutputOffset, hullConstShaderOutputBuffer.GetLength() - hullConstShaderOutputOffset, Pipeline->TessellationPipelineDesc.TessellationInputPatchConstBufferIndex, mtlpp::ResourceUsage(0));
+			State.SetShaderBuffer(SF_Domain, nil, nil, 0, 0, Pipeline->TessellationPipelineDesc.TessellationInputPatchConstBufferIndex, mtlpp::ResourceUsage(0));
 		}
 		
 		// set the patchCount
 		PrologueEncoder.SetShaderBytes(mtlpp::FunctionType::Kernel, (const uint8*)&NumPrimitives, sizeof(NumPrimitives), Pipeline->TessellationPipelineDesc.TessellationPatchCountBufferIndex);
-		State.SetShaderBuffer(SF_Vertex, nil, nil, 0, 0, Pipeline->TessellationPipelineDesc.TessellationPatchCountBufferIndex);
+		State.SetShaderBuffer(SF_Vertex, nil, nil, 0, 0, Pipeline->TessellationPipelineDesc.TessellationPatchCountBufferIndex, mtlpp::ResourceUsage(0));
 		
 		if (boundShaderState->VertexShader->SideTableBinding >= 0)
 		{
 			PrologueEncoder.SetShaderSideTable(mtlpp::FunctionType::Kernel, boundShaderState->VertexShader->SideTableBinding);
-			State.SetShaderBuffer(SF_Vertex, nil, nil, 0, 0, boundShaderState->VertexShader->SideTableBinding);
+			State.SetShaderBuffer(SF_Vertex, nil, nil, 0, 0, boundShaderState->VertexShader->SideTableBinding, mtlpp::ResourceUsage(0));
 		}
 		
 		if (boundShaderState->DomainShader->SideTableBinding >= 0)
 		{
 			CurrentEncoder.SetShaderSideTable(mtlpp::FunctionType::Vertex, boundShaderState->DomainShader->SideTableBinding);
-			State.SetShaderBuffer(SF_Domain, nil, nil, 0, 0, boundShaderState->DomainShader->SideTableBinding);
+			State.SetShaderBuffer(SF_Domain, nil, nil, 0, 0, boundShaderState->DomainShader->SideTableBinding, mtlpp::ResourceUsage(0));
 		}
 		
 		if (IsValidRef(boundShaderState->PixelShader) && boundShaderState->PixelShader->SideTableBinding >= 0)
 		{
 			CurrentEncoder.SetShaderSideTable(mtlpp::FunctionType::Fragment, boundShaderState->PixelShader->SideTableBinding);
-			State.SetShaderBuffer(SF_Pixel, nil, nil, 0, 0, boundShaderState->PixelShader->SideTableBinding);
+			State.SetShaderBuffer(SF_Pixel, nil, nil, 0, 0, boundShaderState->PixelShader->SideTableBinding, mtlpp::ResourceUsage(0));
 		}
 		
 		auto patchesPerThreadGroup = boundShaderState->VertexShader->TessellationPatchesPerThreadGroup;
@@ -1432,7 +1432,7 @@ void FMetalRenderPass::CommitRenderResourceTables(void)
 	if (BoundShaderState->VertexShader->SideTableBinding >= 0)
 	{
 		CurrentEncoder.SetShaderSideTable(mtlpp::FunctionType::Vertex, BoundShaderState->VertexShader->SideTableBinding);
-		State.SetShaderBuffer(SF_Vertex, nil, nil, 0, 0, BoundShaderState->VertexShader->SideTableBinding);
+		State.SetShaderBuffer(SF_Vertex, nil, nil, 0, 0, BoundShaderState->VertexShader->SideTableBinding, mtlpp::ResourceUsage(0));
 	}
 	
 	if (IsValidRef(BoundShaderState->PixelShader))
@@ -1441,7 +1441,7 @@ void FMetalRenderPass::CommitRenderResourceTables(void)
 		if (BoundShaderState->PixelShader->SideTableBinding >= 0)
 		{
 			CurrentEncoder.SetShaderSideTable(mtlpp::FunctionType::Fragment, BoundShaderState->PixelShader->SideTableBinding);
-			State.SetShaderBuffer(SF_Pixel, nil, nil, 0, 0, BoundShaderState->PixelShader->SideTableBinding);
+			State.SetShaderBuffer(SF_Pixel, nil, nil, 0, 0, BoundShaderState->PixelShader->SideTableBinding, mtlpp::ResourceUsage(0));
 		}
 	}
 }
@@ -1473,7 +1473,7 @@ void FMetalRenderPass::CommitDispatchResourceTables(void)
 	if (ComputeShader->SideTableBinding >= 0)
 	{
 		CurrentEncoder.SetShaderSideTable(mtlpp::FunctionType::Kernel, ComputeShader->SideTableBinding);
-		State.SetShaderBuffer(SF_Compute, nil, nil, 0, 0, ComputeShader->SideTableBinding);
+		State.SetShaderBuffer(SF_Compute, nil, nil, 0, 0, ComputeShader->SideTableBinding, mtlpp::ResourceUsage(0));
 	}
 }
 
diff --git a/Engine/Source/Runtime/Apple/MetalRHI/Private/MetalShaders.cpp b/Engine/Source/Runtime/Apple/MetalRHI/Private/MetalShaders.cpp
--- a/Engine/Source/Runtime/Apple/MetalRHI/Private/MetalShaders.cpp
+++ b/Engine/Source/Runtime/Apple/MetalRHI/Private/MetalShaders.cpp
@@ -1409,13 +1409,13 @@ void FMetalShaderParameterCache::CommitPackedGlobals(FMetalStateCache* Cache, FM
 				uint8 const* Bytes = PackedGlobalUniforms[Index]->Data;
 				ns::AutoReleased<FMetalBuffer> Buffer(Encoder->GetRingBuffer().NewBuffer(Size, 0));
 				FMemory::Memcpy((uint8*)Buffer.GetContents(), Bytes, Size);
-				Cache->SetShaderBuffer(Frequency, Buffer, nil, 0, Size, UniformBufferIndex);
+				Cache->SetShaderBuffer(Frequency, Buffer, nil, 0, Size, UniformBufferIndex, mtlpp::ResourceUsage::Read);
 			}
 			else
 			{
 				PackedGlobalUniforms[Index]->Len = Size;
-				Cache->SetShaderBuffer(Frequency, nil, nil, 0, 0, UniformBufferIndex);
-				Cache->SetShaderBuffer(Frequency, nil, PackedGlobalUniforms[Index], 0, Size, UniformBufferIndex);
+				Cache->SetShaderBuffer(Frequency, nil, nil, 0, 0, UniformBufferIndex, mtlpp::ResourceUsage(0));
+				Cache->SetShaderBuffer(Frequency, nil, PackedGlobalUniforms[Index], 0, Size, UniformBufferIndex, mtlpp::ResourceUsage::Read);
 			}
 
 			// mark as clean
diff --git a/Engine/Source/Runtime/Apple/MetalRHI/Private/MetalStateCache.cpp b/Engine/Source/Runtime/Apple/MetalRHI/Private/MetalStateCache.cpp
--- a/Engine/Source/Runtime/Apple/MetalRHI/Private/MetalStateCache.cpp
+++ b/Engine/Source/Runtime/Apple/MetalRHI/Private/MetalStateCache.cpp
@@ -201,6 +201,7 @@ FMetalStateCache::~FMetalStateCache()
 			ShaderBuffers[Frequency].Buffers[i].Bytes = nil;
 			ShaderBuffers[Frequency].Buffers[i].Length = 0;
 			ShaderBuffers[Frequency].Buffers[i].Offset = 0;
+			ShaderBuffers[Frequency].Buffers[i].Usage = mtlpp::ResourceUsage(0);
 			ShaderBuffers[Frequency].Formats[i] = PF_Unknown;
 		}
 		ShaderBuffers[Frequency].Bound = 0;
@@ -208,6 +209,7 @@ FMetalStateCache::~FMetalStateCache()
 		for (uint32 i = 0; i < ML_MaxTextures; i++)
 		{
 			ShaderTextures[Frequency].Textures[i] = nil;
+			ShaderTextures[Frequency].Usage[i] = mtlpp::ResourceUsage(0);
 		}
 		ShaderTextures[Frequency].Bound = 0;
 	}
@@ -1005,7 +1007,7 @@ void FMetalStateCache::SetVertexStream(uint32 const Index, FMetalBuffer* Buffer,
 	VertexBuffers[Index].Bytes = Bytes;
 	VertexBuffers[Index].Length = Length;
 	
-	SetShaderBuffer(SF_Vertex, VertexBuffers[Index].Buffer, Bytes, Offset, Length, UNREAL_TO_METAL_BUFFER_INDEX(Index));
+	SetShaderBuffer(SF_Vertex, VertexBuffers[Index].Buffer, Bytes, Offset, Length, UNREAL_TO_METAL_BUFFER_INDEX(Index), mtlpp::ResourceUsage::Read);
 }
 
 uint32 FMetalStateCache::GetVertexBufferSize(uint32 const Index)
@@ -1058,11 +1060,13 @@ void FMetalStateCache::SetGraphicsPipelineState(FMetalGraphicsPipelineState* Sta
 				ShaderBuffers[SF_Hull].Buffers[i].Bytes = nil;
 				ShaderBuffers[SF_Hull].Buffers[i].Length = 0;
 				ShaderBuffers[SF_Hull].Buffers[i].Offset = 0;
+				ShaderBuffers[SF_Hull].Buffers[i].Usage = mtlpp::ResourceUsage(0);
 				ShaderBuffers[SF_Hull].Formats[i] = PF_Unknown;
 			}
 			for (uint32 i = 0; i < ML_MaxTextures; i++)
 			{
 				ShaderTextures[SF_Hull].Textures[i] = nil;
+				ShaderTextures[SF_Hull].Usage[i] = mtlpp::ResourceUsage(0);
 			}
 			
 			for (uint32 i = 0; i < ML_MaxSamplers; i++)
@@ -1300,7 +1304,7 @@ static uint8 GMetalShaderFreqFormat[SF_NumFrequencies] =
 	EMetalPipelineFlagComputeBuffers
 };
 
-void FMetalStateCache::SetShaderBuffer(EShaderFrequency const Frequency, FMetalBuffer const& Buffer, FMetalBufferData* const Bytes, NSUInteger const Offset, NSUInteger const Length, NSUInteger const Index, EPixelFormat const Format)
+void FMetalStateCache::SetShaderBuffer(EShaderFrequency const Frequency, FMetalBuffer const& Buffer, FMetalBufferData* const Bytes, NSUInteger const Offset, NSUInteger const Length, NSUInteger const Index, mtlpp::ResourceUsage const Usage, EPixelFormat const Format)
 {
 	check(Frequency < SF_NumFrequencies);
 	check(Index < ML_MaxBuffers);
@@ -1309,12 +1313,14 @@ void FMetalStateCache::SetShaderBuffer(EShaderFrequency const Frequency, FMetalB
 		ShaderBuffers[Frequency].Buffers[Index].Bytes != Bytes ||
 		ShaderBuffers[Frequency].Buffers[Index].Offset != Offset ||
 		ShaderBuffers[Frequency].Buffers[Index].Length != Length ||
+		!(ShaderBuffers[Frequency].Buffers[Index].Usage & Usage) ||
 		ShaderBuffers[Frequency].Formats[Index] != Format)
 	{
 		ShaderBuffers[Frequency].Buffers[Index].Buffer = Buffer;
 		ShaderBuffers[Frequency].Buffers[Index].Bytes = Bytes;
 		ShaderBuffers[Frequency].Buffers[Index].Offset = Offset;
 		ShaderBuffers[Frequency].Buffers[Index].Length = Length;
+		ShaderBuffers[Frequency].Buffers[Index].Usage = Usage;
 		
 		PipelineBits |= (ShaderBuffers[Frequency].Formats[Index] != Format) ? GMetalShaderFreqFormat[Frequency] : 0;
 		ShaderBuffers[Frequency].Formats[Index] = Format;
@@ -1330,14 +1336,16 @@ void FMetalStateCache::SetShaderBuffer(EShaderFrequency const Frequency, FMetalB
 	}
 }
 
-void FMetalStateCache::SetShaderTexture(EShaderFrequency const Frequency, FMetalTexture const& Texture, NSUInteger const Index)
+void FMetalStateCache::SetShaderTexture(EShaderFrequency const Frequency, FMetalTexture const& Texture, NSUInteger const Index, mtlpp::ResourceUsage const Usage)
 {
 	check(Frequency < SF_NumFrequencies);
 	check(Index < ML_MaxTextures);
 	
-	if (ShaderTextures[Frequency].Textures[Index] != Texture)
+	if (ShaderTextures[Frequency].Textures[Index] != Texture
+		|| ShaderTextures[Frequency].Usage[Index] != Usage)
 	{
 		ShaderTextures[Frequency].Textures[Index] = Texture;
+		ShaderTextures[Frequency].Usage[Index] = Usage;
 		
 		if (Texture)
 		{
@@ -1378,32 +1386,34 @@ void FMetalStateCache::SetResource(uint32 ShaderStage, uint32 BindIndex, FRHITex
 {
 	FMetalSurface* Surface = GetMetalSurfaceFromRHITexture(TextureRHI);
 	ns::AutoReleased<FMetalTexture> Texture;
+	mtlpp::ResourceUsage Usage = (mtlpp::ResourceUsage)0;
 	if (Surface != nullptr)
 	{
 		TextureRHI->SetLastRenderTime(CurrentTime);
 		Texture = Surface->Texture;
+		Usage = mtlpp::ResourceUsage(mtlpp::ResourceUsage::Read|mtlpp::ResourceUsage::Sample);
 	}
 	
 	switch (ShaderStage)
 	{
 		case CrossCompiler::SHADER_STAGE_PIXEL:
-			SetShaderTexture(SF_Pixel, Texture, BindIndex);
+			SetShaderTexture(SF_Pixel, Texture, BindIndex, Usage);
 			break;
 			
 		case CrossCompiler::SHADER_STAGE_VERTEX:
-			SetShaderTexture(SF_Vertex, Texture, BindIndex);
+			SetShaderTexture(SF_Vertex, Texture, BindIndex, Usage);
 			break;
 			
 		case CrossCompiler::SHADER_STAGE_COMPUTE:
-			SetShaderTexture(SF_Compute, Texture, BindIndex);
+			SetShaderTexture(SF_Compute, Texture, BindIndex, Usage);
 			break;
 			
 		case CrossCompiler::SHADER_STAGE_HULL:
-			SetShaderTexture(SF_Hull, Texture, BindIndex);
+			SetShaderTexture(SF_Hull, Texture, BindIndex, Usage);
 			break;
 			
 		case CrossCompiler::SHADER_STAGE_DOMAIN:
-			SetShaderTexture(SF_Domain, Texture, BindIndex);
+			SetShaderTexture(SF_Domain, Texture, BindIndex, Usage);
 			break;
 			
 		default:
@@ -1425,11 +1435,11 @@ void FMetalStateCache::SetShaderResourceView(FMetalContext* Context, EShaderFreq
 			FMetalSurface* Surface = SRV->TextureView;
 			if (Surface != nullptr)
 			{
-				SetShaderTexture(ShaderStage, Surface->Texture, BindIndex);
+				SetShaderTexture(ShaderStage, Surface->Texture, BindIndex, mtlpp::ResourceUsage(mtlpp::ResourceUsage::Read|mtlpp::ResourceUsage::Sample));
 			}
 			else
 			{
-				SetShaderTexture(ShaderStage, nil, BindIndex);
+				SetShaderTexture(ShaderStage, nil, BindIndex, mtlpp::ResourceUsage(0));
 			}
 		}
 		else if (IsLinearBuffer(ShaderStage, BindIndex) && SRV->GetLinearTexture(false))
@@ -1438,29 +1448,29 @@ void FMetalStateCache::SetShaderResourceView(FMetalContext* Context, EShaderFreq
 			Tex = SRV->GetLinearTexture(false);
 			
 			uint32 PackedLen = (Tex.GetWidth() | (Tex.GetHeight() << 16));
-			SetShaderTexture(ShaderStage, Tex, BindIndex);
+			SetShaderTexture(ShaderStage, Tex, BindIndex, mtlpp::ResourceUsage(mtlpp::ResourceUsage::Read|mtlpp::ResourceUsage::Sample));
 			if (VB)
             {
-                SetShaderBuffer(ShaderStage, VB->Buffer, VB->Data, 0, PackedLen, BindIndex, (EPixelFormat)SRV->Format);
+                SetShaderBuffer(ShaderStage, VB->Buffer, VB->Data, 0, PackedLen, BindIndex, mtlpp::ResourceUsage::Read, (EPixelFormat)SRV->Format);
             }
             else if (IB)
             {
-                SetShaderBuffer(ShaderStage, IB->Buffer, nil, 0, PackedLen, BindIndex, (EPixelFormat)SRV->Format);
+                SetShaderBuffer(ShaderStage, IB->Buffer, nil, 0, PackedLen, BindIndex, mtlpp::ResourceUsage::Read, (EPixelFormat)SRV->Format);
             }
 		}
 		else if (VB)
 		{
 			checkf(ValidateBufferFormat(ShaderStage, BindIndex, (EPixelFormat)SRV->Format), TEXT("Invalid buffer format %d for index %d, shader %d"), SRV->Format, BindIndex, ShaderStage);
-			SetShaderBuffer(ShaderStage, VB->Buffer, VB->Data, 0, VB->GetSize(), BindIndex, (EPixelFormat)SRV->Format);
+			SetShaderBuffer(ShaderStage, VB->Buffer, VB->Data, 0, VB->GetSize(), BindIndex, mtlpp::ResourceUsage::Read, (EPixelFormat)SRV->Format);
 		}
 		else if (IB)
 		{
 			checkf(ValidateBufferFormat(ShaderStage, BindIndex, (EPixelFormat)SRV->Format), TEXT("Invalid buffer format %d for index %d, shader %d"), SRV->Format, BindIndex, ShaderStage);
-			SetShaderBuffer(ShaderStage, IB->Buffer, nil, 0, IB->GetSize(), BindIndex, (EPixelFormat)SRV->Format);
+			SetShaderBuffer(ShaderStage, IB->Buffer, nil, 0, IB->GetSize(), BindIndex, mtlpp::ResourceUsage::Read, (EPixelFormat)SRV->Format);
 		}
 		else if (SB)
 		{
-			SetShaderBuffer(ShaderStage, SB->Buffer, nil, 0, SB->GetSize(), BindIndex);
+			SetShaderBuffer(ShaderStage, SB->Buffer, nil, 0, SB->GetSize(), BindIndex, mtlpp::ResourceUsage::Read);
 		}
 	}
 }
@@ -1549,7 +1559,7 @@ void FMetalStateCache::SetShaderUnorderedAccessView(EShaderFrequency ShaderStage
 		FMetalSurface* Surface = UAV->SourceView->TextureView;
 		if (StructuredBuffer)
 		{
-			SetShaderBuffer(ShaderStage, StructuredBuffer->Buffer, nil, 0, StructuredBuffer->GetSize(), BindIndex);
+			SetShaderBuffer(ShaderStage, StructuredBuffer->Buffer, nil, 0, StructuredBuffer->GetSize(), BindIndex, mtlpp::ResourceUsage(mtlpp::ResourceUsage::Read | mtlpp::ResourceUsage::Write));
 		}
 		else if (VertexBuffer)
 		{
@@ -1558,15 +1568,15 @@ void FMetalStateCache::SetShaderUnorderedAccessView(EShaderFrequency ShaderStage
 			{
 				ns::AutoReleased<FMetalTexture> Tex;
 				Tex = UAV->SourceView->GetLinearTexture(true);
-				SetShaderTexture(ShaderStage, Tex, BindIndex);
+				SetShaderTexture(ShaderStage, Tex, BindIndex, mtlpp::ResourceUsage(mtlpp::ResourceUsage::Read | mtlpp::ResourceUsage::Write));
                 
                 uint32 PackedLen = (Tex.GetWidth() | (Tex.GetHeight() << 16));
-                SetShaderBuffer(ShaderStage, VertexBuffer->Buffer, VertexBuffer->Data, 0, PackedLen, BindIndex, (EPixelFormat)UAV->SourceView->Format);
+                SetShaderBuffer(ShaderStage, VertexBuffer->Buffer, VertexBuffer->Data, 0, PackedLen, BindIndex, mtlpp::ResourceUsage(mtlpp::ResourceUsage::Read | mtlpp::ResourceUsage::Write), (EPixelFormat)UAV->SourceView->Format);
 			}
 			else
 			{
 				checkf(ValidateBufferFormat(ShaderStage, BindIndex, (EPixelFormat)UAV->SourceView->Format), TEXT("Invalid buffer format %d for index %d, shader %d"), UAV->SourceView->Format, BindIndex, ShaderStage);
-				SetShaderBuffer(ShaderStage, VertexBuffer->Buffer, VertexBuffer->Data, 0, VertexBuffer->GetSize(), BindIndex, (EPixelFormat)UAV->SourceView->Format);
+				SetShaderBuffer(ShaderStage, VertexBuffer->Buffer, VertexBuffer->Data, 0, VertexBuffer->GetSize(), BindIndex, mtlpp::ResourceUsage(mtlpp::ResourceUsage::Read | mtlpp::ResourceUsage::Write), (EPixelFormat)UAV->SourceView->Format);
 			}
 		}
 		else if (IndexBuffer)
@@ -1576,15 +1586,15 @@ void FMetalStateCache::SetShaderUnorderedAccessView(EShaderFrequency ShaderStage
 			{
 				ns::AutoReleased<FMetalTexture> Tex;
 				Tex = UAV->SourceView->GetLinearTexture(true);
-				SetShaderTexture(ShaderStage, Tex, BindIndex);
+				SetShaderTexture(ShaderStage, Tex, BindIndex, mtlpp::ResourceUsage(mtlpp::ResourceUsage::Read | mtlpp::ResourceUsage::Write));
 				
 				uint32 PackedLen = (Tex.GetWidth() | (Tex.GetHeight() << 16));
-				SetShaderBuffer(ShaderStage, IndexBuffer->Buffer, nullptr, 0, PackedLen, BindIndex, (EPixelFormat)UAV->SourceView->Format);
+				SetShaderBuffer(ShaderStage, IndexBuffer->Buffer, nullptr, 0, PackedLen, BindIndex, mtlpp::ResourceUsage(mtlpp::ResourceUsage::Read | mtlpp::ResourceUsage::Write), (EPixelFormat)UAV->SourceView->Format);
 			}
 			else
 			{
 				checkf(ValidateBufferFormat(ShaderStage, BindIndex, (EPixelFormat)UAV->SourceView->Format), TEXT("Invalid buffer format %d for index %d, shader %d"), UAV->SourceView->Format, BindIndex, ShaderStage);
-				SetShaderBuffer(ShaderStage, IndexBuffer->Buffer, nullptr, 0, IndexBuffer->GetSize(), BindIndex, (EPixelFormat)UAV->SourceView->Format);
+				SetShaderBuffer(ShaderStage, IndexBuffer->Buffer, nullptr, 0, IndexBuffer->GetSize(), BindIndex, mtlpp::ResourceUsage(mtlpp::ResourceUsage::Read | mtlpp::ResourceUsage::Write), (EPixelFormat)UAV->SourceView->Format);
 			}
 		}
 		else if (Texture)
@@ -1600,11 +1610,11 @@ void FMetalStateCache::SetShaderUnorderedAccessView(EShaderFrequency ShaderStage
 				FPlatformAtomics::InterlockedExchange(&Surface->Written, 1);
 				FPlatformAtomics::InterlockedExchange(&Source->Written, 1);
 				
-				SetShaderTexture(ShaderStage, Surface->Texture, BindIndex);
+				SetShaderTexture(ShaderStage, Surface->Texture, BindIndex, mtlpp::ResourceUsage(mtlpp::ResourceUsage::Read | mtlpp::ResourceUsage::Write));
 			}
 			else
 			{
-				SetShaderTexture(ShaderStage, nil, BindIndex);
+				SetShaderTexture(ShaderStage, nil, BindIndex, mtlpp::ResourceUsage(0));
 			}
 		}
 	}
@@ -2083,7 +2093,7 @@ void FMetalStateCache::CommitResourceTable(EShaderFrequency const Frequency, mtl
 			FMetalBufferBinding& Binding = BufferBindings.Buffers[Index];
 			if (Binding.Buffer)
 			{
-				CommandEncoder.SetShaderBuffer(Type, Binding.Buffer, Binding.Offset, Binding.Length, Index, BufferBindings.Formats[Index]);
+				CommandEncoder.SetShaderBuffer(Type, Binding.Buffer, Binding.Offset, Binding.Length, Index, Binding.Usage, BufferBindings.Formats[Index]);
 				
 				if (Binding.Buffer.IsSingleUse())
 				{
@@ -2107,7 +2117,7 @@ void FMetalStateCache::CommitResourceTable(EShaderFrequency const Frequency, mtl
 		
 		if (Index < ML_MaxTextures && TextureBindings.Textures[Index])
 		{
-			CommandEncoder.SetShaderTexture(Type, TextureBindings.Textures[Index], Index);
+			CommandEncoder.SetShaderTexture(Type, TextureBindings.Textures[Index], Index, TextureBindings.Usage[Index]);
 		}
 	}
 	
@@ -2119,7 +2129,7 @@ void FMetalStateCache::CommitResourceTable(EShaderFrequency const Frequency, mtl
 		
 		if (Index < ML_MaxTextures && TextureBindings.Textures[Index])
 		{
-			CommandEncoder.SetShaderTexture(Type, TextureBindings.Textures[Index], Index + 64);
+			CommandEncoder.SetShaderTexture(Type, TextureBindings.Textures[Index], Index + 64, TextureBindings.Usage[Index]);
 		}
 	}
 	
@@ -2133,7 +2143,7 @@ void FMetalStateCache::CommitResourceTable(EShaderFrequency const Frequency, mtl
 		
 		if (Index < ML_MaxTextures && TextureBindings.Textures[Index])
 		{
-			CommandEncoder.SetShaderTexture(Type, TextureBindings.Textures[Index], Index);
+			CommandEncoder.SetShaderTexture(Type, TextureBindings.Textures[Index], Index, TextureBindings.Usage[Index]);
 		}
 	}
 #endif
diff --git a/Engine/Source/Runtime/Apple/MetalRHI/Private/MetalStateCache.h b/Engine/Source/Runtime/Apple/MetalRHI/Private/MetalStateCache.h
--- a/Engine/Source/Runtime/Apple/MetalRHI/Private/MetalStateCache.h
+++ b/Engine/Source/Runtime/Apple/MetalRHI/Private/MetalStateCache.h
@@ -74,17 +74,19 @@ public:
 	 * @param Offset The offset in the buffer or 0 when Buffer is nil.
 	 * @param Offset The length of data (caller accounts for Offset) in the buffer or 0 when Buffer is nil.
 	 * @param Index The index to modify.
+	 * @param Usage The resource usage flags.
 	 * @param Format The UAV pixel format.
 	 */
-	void SetShaderBuffer(EShaderFrequency const Frequency, FMetalBuffer const& Buffer, FMetalBufferData* const Bytes, NSUInteger const Offset, NSUInteger const Length, NSUInteger const Index, EPixelFormat const Format = PF_Unknown);
+	void SetShaderBuffer(EShaderFrequency const Frequency, FMetalBuffer const& Buffer, FMetalBufferData* const Bytes, NSUInteger const Offset, NSUInteger const Length, NSUInteger const Index, mtlpp::ResourceUsage const Usage, EPixelFormat const Format = PF_Unknown);
 	
 	/*
 	 * Set a global texture for the specified shader frequency at the given bind point index.
 	 * @param Frequency The shader frequency to modify.
 	 * @param Texture The texture to bind or nil to clear.
 	 * @param Index The index to modify.
+	 * @param Usage The resource usage flags.
 	 */
-	void SetShaderTexture(EShaderFrequency const Frequency, FMetalTexture const& Texture, NSUInteger const Index);
+	void SetShaderTexture(EShaderFrequency const Frequency, FMetalTexture const& Texture, NSUInteger const Index, mtlpp::ResourceUsage const Usage);
 	
 	/*
 	 * Set a global sampler for the specified shader frequency at the given bind point index.
@@ -181,7 +183,7 @@ private:
 #pragma mark - Private Type Declarations -
 	struct FMetalBufferBinding
 	{
-		FMetalBufferBinding() : Bytes(nil), Offset(0), Length(0) {}
+		FMetalBufferBinding() : Bytes(nil), Offset(0), Length(0), Usage((mtlpp::ResourceUsage)0) {}
 		/** The bound buffers or nil. */
 		ns::AutoReleased<FMetalBuffer> Buffer;
 		/** Optional bytes buffer used instead of an FMetalBuffer */
@@ -190,6 +192,8 @@ private:
 		NSUInteger Offset;
 		/** The bound buffer lengths or 0. */
 		NSUInteger Length;
+		/** The bound buffer usage or 0 */
+		mtlpp::ResourceUsage Usage;
 	};
 	
 	/** A structure of arrays for the current buffer binding settings. */
@@ -209,9 +213,11 @@ private:
 	/** A structure of arrays for the current texture binding settings. */
 	struct FMetalTextureBindings
 	{
-		FMetalTextureBindings() : Bound(0) {}
+		FMetalTextureBindings() : Bound(0) { FMemory::Memzero(Usage); }
 		/** The bound textures or nil. */
 		ns::AutoReleased<FMetalTexture> Textures[ML_MaxTextures];
+		/** The bound texture usage or 0 */
+		mtlpp::ResourceUsage Usage[ML_MaxTextures];
 		/** A bitmask for which textures were bound by the application where a bit value of 1 is bound and 0 is unbound. */
 		FMetalTextureMask Bound;
 	};
