On macOS 10.14 we can now enable MTLFence by default, this lets us turn on HazardUntracked (unless using validation) and MTLHeap (except on Intel) as well as calling UseResource. This means that all future testing will find bugs in this code and the drivers before it becomes critical. These features are not enabled on iOS/tvOS as thaat would require per-stage (Vertex vs. Fragment) fencing in order to run efficiently on the GPU.

diff --git a/Engine/Source/Runtime/Apple/MetalRHI/Private/MetalCommandQueue.cpp b/Engine/Source/Runtime/Apple/MetalRHI/Private/MetalCommandQueue.cpp
--- a/Engine/Source/Runtime/Apple/MetalRHI/Private/MetalCommandQueue.cpp
+++ b/Engine/Source/Runtime/Apple/MetalRHI/Private/MetalCommandQueue.cpp
@@ -186,11 +186,16 @@ FMetalCommandQueue::FMetalCommandQueue(mtlpp::Device InDevice, uint32 const MaxN
 		// Turn on Texture Buffers! These are faster on the GPU as we don't need to do out-of-bounds tests but require Metal 2.1 and macOS 10.14
 		if (FPlatformMisc::MacOSXVersionCompare(10,14,0) >= 0)
 		{
-			Features |= EMetalFeaturesMaxThreadsPerThreadgroup;
+			Features |= EMetalFeaturesMaxThreadsPerThreadgroup | EMetalFeaturesFences | EMetalFeaturesIABs;
 			if (MaxShaderVersion >= 4)
 			{
 				Features |= EMetalFeaturesTextureBuffers;
 			}
+			
+			if ([Device.GetName().GetPtr() rangeOfString:@"Intel" options:NSCaseInsensitiveSearch].location == NSNotFound)
+			{
+				Features |= EMetalFeaturesHeaps;
+			}
 		}
     }
     else if ([Device.GetName().GetPtr() rangeOfString:@"Nvidia" options:NSCaseInsensitiveSearch].location != NSNotFound)
@@ -279,7 +284,8 @@ FMetalCommandQueue::FMetalCommandQueue(mtlpp::Device InDevice, uint32 const MaxN
 			PermittedOptions |= mtlpp::ResourceOptions::StorageModeMemoryless;
 		}
 #endif
-		if ((Features & EMetalFeaturesFences) && FParse::Param(FCommandLine::Get(),TEXT("metaluntracktextures")))
+		// You can't use HazardUntracked under the validation layer due to bugs in the layer when trying to create linear-textures/texture-buffers
+		if ((Features & EMetalFeaturesFences) && !(Features & EMetalFeaturesValidation))
 		{
 			PermittedOptions |= mtlpp::ResourceOptions::HazardTrackingModeUntracked;
 		}
