Override all forms of RHITransitionResources in MetalRHI.

diff --git a/Engine/Source/Runtime/Apple/MetalRHI/Private/MetalUAV.cpp b/Engine/Source/Runtime/Apple/MetalRHI/Private/MetalUAV.cpp
--- a/Engine/Source/Runtime/Apple/MetalRHI/Private/MetalUAV.cpp
+++ b/Engine/Source/Runtime/Apple/MetalRHI/Private/MetalUAV.cpp
@@ -807,6 +807,18 @@ void FMetalRHICommandContext::RHITransitionResources(EResourceTransitionAccess T
 	}
 }
 
+void FMetalRHICommandContext::RHITransitionResources(EResourceTransitionAccess TransitionType, FTextureRHIParamRef* InTextures, int32 NumTextures)
+{
+	if (TransitionType == EResourceTransitionAccess::EReadable)
+	{
+		const FResolveParams ResolveParams;
+		for (int32 i = 0; i < NumTextures; ++i)
+		{
+			RHICopyToResolveTarget(InTextures[i], InTextures[i], ResolveParams);
+		}
+	}
+}
+
 void FMetalRHICommandContext::RHIWaitComputeFence(FComputeFenceRHIParamRef InFence)
 {
 	@autoreleasepool {
diff --git a/Engine/Source/Runtime/Apple/MetalRHI/Public/MetalRHIContext.h b/Engine/Source/Runtime/Apple/MetalRHI/Public/MetalRHIContext.h
--- a/Engine/Source/Runtime/Apple/MetalRHI/Public/MetalRHIContext.h
+++ b/Engine/Source/Runtime/Apple/MetalRHI/Public/MetalRHIContext.h
@@ -320,6 +320,17 @@ public:
 	 */
 	virtual void RHITransitionResources(EResourceTransitionAccess TransitionType, EResourceTransitionPipeline TransitionPipeline, FUnorderedAccessViewRHIParamRef* InUAVs, int32 NumUAVs, FComputeFenceRHIParamRef WriteComputeFence) final override;
 	
+	/**
+	 * Explicitly transition a texture resource from readable -> writable by the GPU or vice versa.
+	 * We know rendertargets are only used as rendered targets on the Gfx pipeline, so these transitions are assumed to be implemented such
+	 * Gfx->Gfx and Gfx->Compute pipeline transitions are both handled by this call by the RHI implementation.  Hence, no pipeline parameter on this call.
+	 *
+	 * @param TransitionType - direction of the transition
+	 * @param InTextures - array of texture objects to transition
+	 * @param NumTextures - number of textures to transition
+	 */
+	virtual void RHITransitionResources(EResourceTransitionAccess TransitionType, FTextureRHIParamRef* InTextures, int32 NumTextures) final override;
+	
 	/**
 	 * Compute queue will wait for the fence to be written before continuing.
 	 */
