Separate vertex & fragment stage fences in MetalRHI so that we can attempt to implement per-stage fencing between render encoders for the benefit of iOS. This will undoubtedly take further refinement to perfect but it does appear to work acceptably whereas earlier iterations would cause unnecessary stalls on the GPU.

diff --git a/Engine/Source/Runtime/Apple/MetalRHI/Private/MetalCommandEncoder.cpp b/Engine/Source/Runtime/Apple/MetalRHI/Private/MetalCommandEncoder.cpp
--- a/Engine/Source/Runtime/Apple/MetalRHI/Private/MetalCommandEncoder.cpp
+++ b/Engine/Source/Runtime/Apple/MetalRHI/Private/MetalCommandEncoder.cpp
@@ -35,6 +35,7 @@ FMetalCommandEncoder::FMetalCommandEncoder(FMetalCommandList& CmdList)
 , UpdateCount(0)
 #endif
 , DebugGroups([NSMutableArray new])
+, FenceStage(mtlpp::RenderStages::Fragment)
 , EncoderNum(0)
 {
 	for (uint32 Frequency = 0; Frequency < uint32(mtlpp::FunctionType::Kernel)+1; Frequency++)
@@ -309,6 +310,8 @@ void FMetalCommandEncoder::BeginParallelRenderCommandEncoding(uint32 NumChildren
 	check(CommandBuffer);
 	check(IsRenderCommandEncoderActive() == false && IsComputeCommandEncoderActive() == false && IsBlitCommandEncoderActive() == false);
 	
+	FenceResources.Append(TransitionedResources);
+	
 	ParallelRenderCommandEncoder = MTLPP_VALIDATE(mtlpp::CommandBuffer, CommandBuffer, SafeGetRuntimeDebuggingLevel() >= EMetalDebugLevelValidation, ParallelRenderCommandEncoder(RenderPassDesc));
 	METAL_DEBUG_LAYER(EMetalDebugLevelFastValidation, ParallelEncoderDebug = FMetalParallelRenderCommandEncoderDebugging(ParallelRenderCommandEncoder, RenderPassDesc, CommandBufferDebug));
 	
@@ -351,6 +354,8 @@ void FMetalCommandEncoder::BeginRenderCommandEncoding(void)
 	check(CommandList.IsParallel() || CommandBuffer);
 	check(IsRenderCommandEncoderActive() == false && IsComputeCommandEncoderActive() == false && IsBlitCommandEncoderActive() == false);
 	
+	FenceResources.Append(TransitionedResources);
+	
 	if (!CommandList.IsParallel())
 	{
 		RenderCommandEncoder = MTLPP_VALIDATE(mtlpp::CommandBuffer, CommandBuffer, SafeGetRuntimeDebuggingLevel() >= EMetalDebugLevelValidation, RenderCommandEncoder(RenderPassDesc));
@@ -404,6 +409,9 @@ void FMetalCommandEncoder::BeginComputeCommandEncoding(mtlpp::DispatchType Type)
 	check(CommandBuffer);
 	check(IsRenderCommandEncoderActive() == false && IsComputeCommandEncoderActive() == false && IsBlitCommandEncoderActive() == false);
 	
+	FenceResources.Append(TransitionedResources);
+	TransitionedResources.Empty();
+	
 	if (Type == mtlpp::DispatchType::Serial)
 	{
 		ComputeCommandEncoder = MTLPP_VALIDATE(mtlpp::CommandBuffer, CommandBuffer, SafeGetRuntimeDebuggingLevel() >= EMetalDebugLevelValidation, ComputeCommandEncoder());
@@ -446,6 +454,9 @@ void FMetalCommandEncoder::BeginBlitCommandEncoding(void)
 	check(CommandBuffer);
 	check(IsRenderCommandEncoderActive() == false && IsComputeCommandEncoderActive() == false && IsBlitCommandEncoderActive() == false);
 	
+	FenceResources.Append(TransitionedResources);
+	TransitionedResources.Empty();
+	
 	BlitCommandEncoder = MTLPP_VALIDATE(mtlpp::CommandBuffer, CommandBuffer, SafeGetRuntimeDebuggingLevel() >= EMetalDebugLevelValidation, BlitCommandEncoder());
 	METAL_DEBUG_LAYER(EMetalDebugLevelFastValidation, BlitEncoderDebug = FMetalBlitCommandEncoderDebugging(BlitCommandEncoder, CommandBufferDebug));
 	
@@ -485,7 +496,6 @@ FMetalFence* FMetalCommandEncoder::EndEncoding(void)
 	{
 		if(IsRenderCommandEncoderActive())
 		{
-
 			if (RenderCommandEncoder)
 			{
 				check(!bSupportsFences || EncoderFence || !CommandList.IsImmediate());
@@ -518,6 +528,26 @@ FMetalFence* FMetalCommandEncoder::EndEncoding(void)
 					}
 				}
 				
+				for (FMetalFence* FragFence : FragmentFences)
+				{
+					if (FragFence->NeedsWait(mtlpp::RenderStages::Fragment))
+					{
+						mtlpp::Fence FragmentFence = FragFence->Get(mtlpp::RenderStages::Fragment);
+						mtlpp::Fence FragInnerFence = METAL_DEBUG_OPTION(CommandList.GetCommandQueue().GetRuntimeDebuggingLevel() >= EMetalDebugLevelValidation ? mtlpp::Fence(((FMetalDebugFence*)FragmentFence.GetPtr()).Inner) :) FragmentFence;
+						
+						RenderCommandEncoder.WaitForFence(FragInnerFence, FenceStage);
+						METAL_DEBUG_LAYER(EMetalDebugLevelFastValidation, RenderEncoderDebug.AddWaitFence(FragmentFence));
+						FragFence->Wait(mtlpp::RenderStages::Fragment);
+					}
+				}
+				FragmentFences.Empty();
+				
+				if (FenceStage == mtlpp::RenderStages::Vertex)
+				{
+					FenceResources.Empty();
+					FenceStage = mtlpp::RenderStages::Fragment;
+				}
+				
 				if (EncoderFence && EncoderFence->NeedsWrite(mtlpp::RenderStages::Fragment))
 				{
 					Fence = EncoderFence;
@@ -622,6 +652,23 @@ FMetalFence* FMetalCommandEncoder::EndEncoding(void)
 		else if(IsComputeCommandEncoderActive())
 		{
 			check(!bSupportsFences || EncoderFence);
+			
+			for (FMetalFence* FragFence : FragmentFences)
+			{
+				if (FragFence->NeedsWait(mtlpp::RenderStages::Fragment))
+				{
+					mtlpp::Fence FragmentFence = FragFence->Get(mtlpp::RenderStages::Fragment);
+					mtlpp::Fence FragInnerFence = METAL_DEBUG_OPTION(CommandList.GetCommandQueue().GetRuntimeDebuggingLevel() >= EMetalDebugLevelValidation ? mtlpp::Fence(((FMetalDebugFence*)FragmentFence.GetPtr()).Inner) :) FragmentFence;
+					
+					ComputeCommandEncoder.WaitForFence(FragInnerFence);
+					METAL_DEBUG_LAYER(EMetalDebugLevelFastValidation, ComputeEncoderDebug.AddWaitFence(FragmentFence));
+					FragFence->Wait(mtlpp::RenderStages::Fragment);
+				}
+			}
+			FragmentFences.Empty();
+			FenceResources.Empty();
+			FenceStage = mtlpp::RenderStages::Fragment;
+			
 			if (EncoderFence && EncoderFence->NeedsWrite(mtlpp::RenderStages::Fragment))
 			{
 				Fence = EncoderFence;
@@ -648,6 +695,23 @@ FMetalFence* FMetalCommandEncoder::EndEncoding(void)
 		else if(IsBlitCommandEncoderActive())
 		{
 			// check(!bSupportsFences || EncoderFence);
+			
+			for (FMetalFence* FragFence : FragmentFences)
+			{
+				if (FragFence->NeedsWait(mtlpp::RenderStages::Fragment))
+				{
+					mtlpp::Fence FragmentFence = FragFence->Get(mtlpp::RenderStages::Fragment);
+					mtlpp::Fence FragInnerFence = METAL_DEBUG_OPTION(CommandList.GetCommandQueue().GetRuntimeDebuggingLevel() >= EMetalDebugLevelValidation ? mtlpp::Fence(((FMetalDebugFence*)FragmentFence.GetPtr()).Inner) :) FragmentFence;
+					
+					BlitCommandEncoder.WaitForFence(FragInnerFence);
+					METAL_DEBUG_LAYER(EMetalDebugLevelFastValidation, BlitEncoderDebug.AddWaitFence(FragmentFence));
+					FragFence->Wait(mtlpp::RenderStages::Fragment);
+				}
+			}
+			FragmentFences.Empty();
+			FenceResources.Empty();
+			FenceStage = mtlpp::RenderStages::Fragment;
+			
 			if (EncoderFence && EncoderFence->NeedsWrite(mtlpp::RenderStages::Fragment))
 			{
 				Fence = EncoderFence;
@@ -713,38 +777,76 @@ void FMetalCommandEncoder::UpdateFence(FMetalFence* Fence)
 {
 	check(IsRenderCommandEncoderActive() || IsComputeCommandEncoderActive() || IsBlitCommandEncoderActive());
 	static bool bSupportsFences = CommandList.GetCommandQueue().SupportsFeature(EMetalFeaturesFences);
-	if ((bSupportsFences METAL_DEBUG_OPTION(|| CommandList.GetCommandQueue().GetRuntimeDebuggingLevel() >= EMetalDebugLevelValidation)) && Fence && Fence->NeedsWrite(mtlpp::RenderStages::Fragment))
+	if ((bSupportsFences METAL_DEBUG_OPTION(|| CommandList.GetCommandQueue().GetRuntimeDebuggingLevel() >= EMetalDebugLevelValidation)) && Fence)
 	{
-#if ENABLE_METAL_GPUPROFILE && METAL_STATISTICS
-		FMetalProfiler* Profiler = FMetalProfiler::GetProfiler();
-		if (Profiler)
-		{
-			Profiler->EncodeFence(GetCommandBufferStats(), TEXT("UpdateFence"), Fence, EMTLFenceTypeUpdate);
-		}
-#endif
-		METAL_DEBUG_LAYER(EMetalDebugLevelFastValidation, UpdateCount++);
-		
-		mtlpp::Fence InnerFence = METAL_DEBUG_OPTION(CommandList.GetCommandQueue().GetRuntimeDebuggingLevel() >= EMetalDebugLevelValidation ? mtlpp::Fence(((FMetalDebugFence*)Fence->GetPtr()).Inner) :) *Fence;
+		mtlpp::Fence VertexFence = Fence->Get(mtlpp::RenderStages::Vertex);
+		mtlpp::Fence InnerFence = METAL_DEBUG_OPTION(CommandList.GetCommandQueue().GetRuntimeDebuggingLevel() >= EMetalDebugLevelValidation ? mtlpp::Fence(((FMetalDebugFence*)VertexFence.GetPtr()).Inner) :) VertexFence;
 		if (RenderCommandEncoder)
 		{
-//			Fence->Write(mtlpp::RenderStages::Vertex);
-			Fence->Write(mtlpp::RenderStages::Fragment);
-			RenderCommandEncoder.UpdateFence(InnerFence, (mtlpp::RenderStages)(/*mtlpp::RenderStages::Vertex |*/ mtlpp::RenderStages::Fragment));
-			METAL_DEBUG_LAYER(EMetalDebugLevelFastValidation, RenderEncoderDebug.AddUpdateFence(*Fence));
+			mtlpp::Fence FragmentFence = Fence->Get(mtlpp::RenderStages::Fragment);
+			mtlpp::Fence FragInnerFence = METAL_DEBUG_OPTION(CommandList.GetCommandQueue().GetRuntimeDebuggingLevel() >= EMetalDebugLevelValidation ? mtlpp::Fence(((FMetalDebugFence*)FragmentFence.GetPtr()).Inner) :) FragmentFence;
+			
+			if (Fence->NeedsWrite(mtlpp::RenderStages::Vertex))
+			{
+				RenderCommandEncoder.UpdateFence(InnerFence, (mtlpp::RenderStages)(mtlpp::RenderStages::Vertex));
+				METAL_DEBUG_LAYER(EMetalDebugLevelFastValidation, RenderEncoderDebug.AddUpdateFence(VertexFence));
+				Fence->Write(mtlpp::RenderStages::Vertex);
+				METAL_DEBUG_LAYER(EMetalDebugLevelFastValidation, UpdateCount++);
+				
+#if ENABLE_METAL_GPUPROFILE && METAL_STATISTICS
+				FMetalProfiler* Profiler = FMetalProfiler::GetProfiler();
+				if (Profiler)
+				{
+					Profiler->EncodeFence(GetCommandBufferStats(), TEXT("UpdateFence"), Fence, EMTLFenceTypeUpdate);
+				}
+#endif
+			}
+
+			if (Fence->NeedsWrite(mtlpp::RenderStages::Fragment))
+			{
+				RenderCommandEncoder.UpdateFence(FragInnerFence, (mtlpp::RenderStages)(mtlpp::RenderStages::Fragment));
+				METAL_DEBUG_LAYER(EMetalDebugLevelFastValidation, RenderEncoderDebug.AddUpdateFence(FragmentFence));
+				Fence->Write(mtlpp::RenderStages::Fragment);
+				METAL_DEBUG_LAYER(EMetalDebugLevelFastValidation, UpdateCount++);
+				
+#if ENABLE_METAL_GPUPROFILE && METAL_STATISTICS
+				FMetalProfiler* Profiler = FMetalProfiler::GetProfiler();
+				if (Profiler)
+				{
+					Profiler->EncodeFence(GetCommandBufferStats(), TEXT("UpdateFence"), Fence, EMTLFenceTypeUpdate);
+				}
+#endif
+			}
 		}
-		else if (ComputeCommandEncoder)
+		else if (ComputeCommandEncoder && Fence->NeedsWrite(mtlpp::RenderStages::Vertex))
 		{
-//			Fence->Write(mtlpp::RenderStages::Vertex);
-			Fence->Write(mtlpp::RenderStages::Fragment);
 			ComputeCommandEncoder.UpdateFence(InnerFence);
-			METAL_DEBUG_LAYER(EMetalDebugLevelFastValidation, ComputeEncoderDebug.AddUpdateFence(*Fence));
+			METAL_DEBUG_LAYER(EMetalDebugLevelFastValidation, ComputeEncoderDebug.AddUpdateFence(VertexFence));
+			Fence->Write(mtlpp::RenderStages::Vertex);
+			METAL_DEBUG_LAYER(EMetalDebugLevelFastValidation, UpdateCount++);
+			
+#if ENABLE_METAL_GPUPROFILE && METAL_STATISTICS
+			FMetalProfiler* Profiler = FMetalProfiler::GetProfiler();
+			if (Profiler)
+			{
+				Profiler->EncodeFence(GetCommandBufferStats(), TEXT("UpdateFence"), Fence, EMTLFenceTypeUpdate);
+			}
+#endif
 		}
-		else if (BlitCommandEncoder)
+		else if (BlitCommandEncoder && Fence->NeedsWrite(mtlpp::RenderStages::Vertex))
 		{
-//			Fence->Write(mtlpp::RenderStages::Vertex);
-			Fence->Write(mtlpp::RenderStages::Fragment);
 			BlitCommandEncoder.UpdateFence(InnerFence);
-			METAL_DEBUG_LAYER(EMetalDebugLevelFastValidation, BlitEncoderDebug.AddUpdateFence(*Fence));
+			METAL_DEBUG_LAYER(EMetalDebugLevelFastValidation, BlitEncoderDebug.AddUpdateFence(VertexFence));
+			Fence->Write(mtlpp::RenderStages::Vertex);
+			METAL_DEBUG_LAYER(EMetalDebugLevelFastValidation, UpdateCount++);
+			
+#if ENABLE_METAL_GPUPROFILE && METAL_STATISTICS
+			FMetalProfiler* Profiler = FMetalProfiler::GetProfiler();
+			if (Profiler)
+			{
+				Profiler->EncodeFence(GetCommandBufferStats(), TEXT("UpdateFence"), Fence, EMTLFenceTypeUpdate);
+			}
+#endif
 		}
 	}
 }
@@ -753,38 +855,51 @@ void FMetalCommandEncoder::WaitForFence(FMetalFence* Fence)
 {
 	check(IsRenderCommandEncoderActive() || IsComputeCommandEncoderActive() || IsBlitCommandEncoderActive());
 	static bool bSupportsFences = CommandList.GetCommandQueue().SupportsFeature(EMetalFeaturesFences);
-	if ((bSupportsFences METAL_DEBUG_OPTION(|| CommandList.GetCommandQueue().GetRuntimeDebuggingLevel() >= EMetalDebugLevelValidation)) && Fence && Fence->NeedsWait(mtlpp::RenderStages::Fragment))
+	if ((bSupportsFences METAL_DEBUG_OPTION(|| CommandList.GetCommandQueue().GetRuntimeDebuggingLevel() >= EMetalDebugLevelValidation)) && Fence)
 	{
-#if ENABLE_METAL_GPUPROFILE && METAL_STATISTICS
-		FMetalProfiler* Profiler = FMetalProfiler::GetProfiler();
-		if (Profiler)
+		if (Fence->NeedsWait(mtlpp::RenderStages::Vertex))
 		{
-			Profiler->EncodeFence(GetCommandBufferStats(), TEXT("WaitForFence"), Fence, EMTLFenceTypeWait);
-		}
+#if ENABLE_METAL_GPUPROFILE && METAL_STATISTICS
+			FMetalProfiler* Profiler = FMetalProfiler::GetProfiler();
+			if (Profiler)
+			{
+				Profiler->EncodeFence(GetCommandBufferStats(), TEXT("WaitForFence"), Fence, EMTLFenceTypeWait);
+			}
 #endif
-		METAL_DEBUG_LAYER(EMetalDebugLevelFastValidation, WaitCount++);
-		
-		mtlpp::Fence InnerFence = METAL_DEBUG_OPTION(CommandList.GetCommandQueue().GetRuntimeDebuggingLevel() >= EMetalDebugLevelValidation ? mtlpp::Fence(((FMetalDebugFence*)Fence->GetPtr()).Inner) :) *Fence;
-		if (RenderCommandEncoder)
-		{
-//			Fence->Wait(mtlpp::RenderStages::Vertex);
-			Fence->Wait(mtlpp::RenderStages::Fragment);
-			RenderCommandEncoder.WaitForFence(InnerFence, (mtlpp::RenderStages)(mtlpp::RenderStages::Vertex /*| mtlpp::RenderStages::Fragment*/));
-			METAL_DEBUG_LAYER(EMetalDebugLevelFastValidation, RenderEncoderDebug.AddWaitFence(*Fence));
-		}
-		else if (ComputeCommandEncoder)
-		{
-//			Fence->Wait(mtlpp::RenderStages::Vertex);
-			Fence->Wait(mtlpp::RenderStages::Fragment);
-			ComputeCommandEncoder.WaitForFence(InnerFence);
-			METAL_DEBUG_LAYER(EMetalDebugLevelFastValidation, ComputeEncoderDebug.AddWaitFence(*Fence));
+			METAL_DEBUG_LAYER(EMetalDebugLevelFastValidation, WaitCount++);
+			
+			mtlpp::Fence VertexFence = Fence->Get(mtlpp::RenderStages::Vertex);
+			mtlpp::Fence InnerFence = METAL_DEBUG_OPTION(CommandList.GetCommandQueue().GetRuntimeDebuggingLevel() >= EMetalDebugLevelValidation ? mtlpp::Fence(((FMetalDebugFence*)VertexFence.GetPtr()).Inner) :) VertexFence;
+			if (RenderCommandEncoder)
+			{
+				RenderCommandEncoder.WaitForFence(InnerFence, (mtlpp::RenderStages)(mtlpp::RenderStages::Vertex));
+				METAL_DEBUG_LAYER(EMetalDebugLevelFastValidation, RenderEncoderDebug.AddWaitFence(VertexFence));
+				Fence->Wait(mtlpp::RenderStages::Vertex);
+			}
+			else if (ComputeCommandEncoder)
+			{
+				ComputeCommandEncoder.WaitForFence(InnerFence);
+				METAL_DEBUG_LAYER(EMetalDebugLevelFastValidation, ComputeEncoderDebug.AddWaitFence(VertexFence));
+				Fence->Wait(mtlpp::RenderStages::Vertex);
+			}
+			else if (BlitCommandEncoder)
+			{
+				BlitCommandEncoder.WaitForFence(InnerFence);
+				METAL_DEBUG_LAYER(EMetalDebugLevelFastValidation, BlitEncoderDebug.AddWaitFence(VertexFence));
+				Fence->Wait(mtlpp::RenderStages::Vertex);
+			}
 		}
-		else if (BlitCommandEncoder)
+		if (Fence->NeedsWait(mtlpp::RenderStages::Fragment))
 		{
-//			Fence->Wait(mtlpp::RenderStages::Vertex);
-			Fence->Wait(mtlpp::RenderStages::Fragment);
-			BlitCommandEncoder.WaitForFence(InnerFence);
-			METAL_DEBUG_LAYER(EMetalDebugLevelFastValidation, BlitEncoderDebug.AddWaitFence(*Fence));
+#if ENABLE_METAL_GPUPROFILE && METAL_STATISTICS
+			FMetalProfiler* Profiler = FMetalProfiler::GetProfiler();
+			if (Profiler)
+			{
+				Profiler->EncodeFence(GetCommandBufferStats(), TEXT("WaitForFence"), Fence, EMTLFenceTypeWait);
+			}
+#endif
+			METAL_DEBUG_LAYER(EMetalDebugLevelFastValidation, WaitCount++);
+			FragmentFences.Add(Fence);
 		}
 	}
 }
@@ -806,36 +921,51 @@ void FMetalCommandEncoder::WaitAndUpdateFence(FMetalFence* Fence)
 		METAL_DEBUG_LAYER(EMetalDebugLevelFastValidation, WaitCount++);
 		METAL_DEBUG_LAYER(EMetalDebugLevelFastValidation, UpdateCount++);
 		
-		mtlpp::Fence InnerFence = METAL_DEBUG_OPTION(CommandList.GetCommandQueue().GetRuntimeDebuggingLevel() >= EMetalDebugLevelValidation ? mtlpp::Fence(((FMetalDebugFence*)Fence->GetPtr()).Inner) :) *Fence;
+		mtlpp::Fence VertexFence = Fence->Get(mtlpp::RenderStages::Vertex);
+		mtlpp::Fence InnerFence = METAL_DEBUG_OPTION(CommandList.GetCommandQueue().GetRuntimeDebuggingLevel() >= EMetalDebugLevelValidation ? mtlpp::Fence(((FMetalDebugFence*)VertexFence.GetPtr()).Inner) :) VertexFence;
 		if (RenderCommandEncoder)
 		{
-			//			Fence->Wait(mtlpp::RenderStages::Vertex);
+			mtlpp::Fence FragmentFence = Fence->Get(mtlpp::RenderStages::Fragment);
+			mtlpp::Fence FragInnerFence = METAL_DEBUG_OPTION(CommandList.GetCommandQueue().GetRuntimeDebuggingLevel() >= EMetalDebugLevelValidation ? mtlpp::Fence(((FMetalDebugFence*)FragmentFence.GetPtr()).Inner) :) FragmentFence;
+			
+			METAL_DEBUG_LAYER(EMetalDebugLevelFastValidation, WaitCount++);
+			METAL_DEBUG_LAYER(EMetalDebugLevelFastValidation, UpdateCount++);
+			
+			RenderCommandEncoder.WaitForFence(InnerFence, (mtlpp::RenderStages)(mtlpp::RenderStages::Vertex));
+			METAL_DEBUG_LAYER(EMetalDebugLevelFastValidation, RenderEncoderDebug.AddWaitFence(VertexFence));
+			Fence->Wait(mtlpp::RenderStages::Vertex);
+			
+			RenderCommandEncoder.WaitForFence(FragInnerFence, (mtlpp::RenderStages)(mtlpp::RenderStages::Fragment));
+			METAL_DEBUG_LAYER(EMetalDebugLevelFastValidation, RenderEncoderDebug.AddWaitFence(FragmentFence));
 			Fence->Wait(mtlpp::RenderStages::Fragment);
-			RenderCommandEncoder.WaitForFence(InnerFence, (mtlpp::RenderStages)(mtlpp::RenderStages::Vertex /*| mtlpp::RenderStages::Fragment*/));
-			METAL_DEBUG_LAYER(EMetalDebugLevelFastValidation, RenderEncoderDebug.AddWaitFence(*Fence));
-			RenderCommandEncoder.UpdateFence(InnerFence, (mtlpp::RenderStages)(/*mtlpp::RenderStages::Vertex |*/ mtlpp::RenderStages::Fragment));
-			METAL_DEBUG_LAYER(EMetalDebugLevelFastValidation, RenderEncoderDebug.AddUpdateFence(*Fence));
+			
+			RenderCommandEncoder.UpdateFence(InnerFence, (mtlpp::RenderStages)(mtlpp::RenderStages::Fragment));
+			METAL_DEBUG_LAYER(EMetalDebugLevelFastValidation, RenderEncoderDebug.AddUpdateFence(VertexFence));
+			Fence->Write(mtlpp::RenderStages::Vertex);
+			
+			RenderCommandEncoder.UpdateFence(FragInnerFence, (mtlpp::RenderStages)(mtlpp::RenderStages::Vertex));
+			METAL_DEBUG_LAYER(EMetalDebugLevelFastValidation, RenderEncoderDebug.AddUpdateFence(FragmentFence));
 			Fence->Write(mtlpp::RenderStages::Fragment);
 		}
 		else if (ComputeCommandEncoder)
 		{
-			//			Fence->Wait(mtlpp::RenderStages::Vertex);
-			Fence->Wait(mtlpp::RenderStages::Fragment);
 			ComputeCommandEncoder.WaitForFence(InnerFence);
-			METAL_DEBUG_LAYER(EMetalDebugLevelFastValidation, ComputeEncoderDebug.AddWaitFence(*Fence));
+			METAL_DEBUG_LAYER(EMetalDebugLevelFastValidation, ComputeEncoderDebug.AddWaitFence(VertexFence));
+			Fence->Wait(mtlpp::RenderStages::Vertex);
+			
 			ComputeCommandEncoder.UpdateFence(InnerFence);
-			METAL_DEBUG_LAYER(EMetalDebugLevelFastValidation, ComputeEncoderDebug.AddUpdateFence(*Fence));
-			Fence->Write(mtlpp::RenderStages::Fragment);
+			METAL_DEBUG_LAYER(EMetalDebugLevelFastValidation, ComputeEncoderDebug.AddUpdateFence(VertexFence));
+			Fence->Write(mtlpp::RenderStages::Vertex);
 		}
 		else if (BlitCommandEncoder)
 		{
-			//			Fence->Wait(mtlpp::RenderStages::Vertex);
-			Fence->Wait(mtlpp::RenderStages::Fragment);
 			BlitCommandEncoder.WaitForFence(InnerFence);
-			METAL_DEBUG_LAYER(EMetalDebugLevelFastValidation, BlitEncoderDebug.AddWaitFence(*Fence));
+			METAL_DEBUG_LAYER(EMetalDebugLevelFastValidation, BlitEncoderDebug.AddWaitFence(VertexFence));
+			Fence->Wait(mtlpp::RenderStages::Vertex);
+			
 			BlitCommandEncoder.UpdateFence(InnerFence);
-			METAL_DEBUG_LAYER(EMetalDebugLevelFastValidation, BlitEncoderDebug.AddUpdateFence(*Fence));
-			Fence->Write(mtlpp::RenderStages::Fragment);
+			METAL_DEBUG_LAYER(EMetalDebugLevelFastValidation, BlitEncoderDebug.AddUpdateFence(VertexFence));
+			Fence->Write(mtlpp::RenderStages::Vertex);
 		}
 	}
 }
@@ -1130,6 +1260,10 @@ void FMetalCommandEncoder::SetShaderBuffer(mtlpp::FunctionType const FunctionTyp
 		ShaderBuffers[uint32(FunctionType)].Usage[index] = Usage;
 		
 		UseResource(Buffer, Usage);
+		if (FunctionType == mtlpp::FunctionType::Vertex)
+		{
+			FenceResource(Buffer);
+		}
 	}
     else
     {
@@ -1296,6 +1430,7 @@ void FMetalCommandEncoder::SetShaderTexture(mtlpp::FunctionType FunctionType, FM
 			// MTLPP_VALIDATE(mtlpp::RenderCommandEncoder, RenderCommandEncoder, SafeGetRuntimeDebuggingLevel() >= EMetalDebugLevelValidation, UseResource(Texture, mtlpp::ResourceUsage::Read));
 			METAL_DEBUG_LAYER(EMetalDebugLevelFastValidation, RenderEncoderDebug.SetTexture(EMetalShaderVertex, Texture, index));
 			RenderCommandEncoder.SetVertexTexture(Texture, index);
+			FenceResource(Texture);
 			break;
 		case mtlpp::FunctionType::Fragment:
 			check(RenderCommandEncoder);
@@ -1355,6 +1490,11 @@ void FMetalCommandEncoder::SetShaderSideTable(mtlpp::FunctionType const Function
 	}
 }
 
+void FMetalCommandEncoder::TransitionResources(mtlpp::Resource const& Resource)
+{
+	TransitionedResources.Add(Resource.GetPtr());
+}
+
 #pragma mark - Public Compute State Mutators -
 
 void FMetalCommandEncoder::SetComputePipelineState(FMetalShaderPipeline* State)
@@ -1388,6 +1528,28 @@ bool FMetalCommandEncoder::HasBufferBindingHistory(FMetalBuffer const& Buffer) c
 
 #pragma mark - Private Functions -
 
+void FMetalCommandEncoder::FenceResource(mtlpp::Texture const& Resource)
+{
+	mtlpp::Resource::Type Res = Resource.GetPtr();
+	ns::AutoReleased<mtlpp::Texture> Parent = Resource.GetParentTexture();
+	ns::AutoReleased<mtlpp::Buffer> Buffer = Resource.GetBuffer();
+	if (Parent)
+	{
+		Res = Parent.GetPtr();
+	}
+	else if (Buffer)
+	{
+		Res = Buffer.GetPtr();
+	}
+	FenceStage = FenceResources.Contains(Res) ? mtlpp::RenderStages::Vertex : mtlpp::RenderStages::Fragment;
+}
+
+void FMetalCommandEncoder::FenceResource(mtlpp::Buffer const& Resource)
+{
+	mtlpp::Resource::Type Res = Resource.GetPtr();
+	FenceStage = FenceResources.Contains(Res) ? mtlpp::RenderStages::Vertex : mtlpp::RenderStages::Fragment;
+}
+
 void FMetalCommandEncoder::UseResource(mtlpp::Resource const& Resource, mtlpp::ResourceUsage const Usage)
 {
 	static bool UseResourceAvailable = FMetalCommandQueue::SupportsFeature(EMetalFeaturesIABs);
@@ -1443,6 +1605,7 @@ void FMetalCommandEncoder::SetShaderBufferInternal(mtlpp::FunctionType Function,
 				// MTLPP_VALIDATE(mtlpp::RenderCommandEncoder, RenderCommandEncoder, SafeGetRuntimeDebuggingLevel() >= EMetalDebugLevelValidation, UseResource(Buffer, mtlpp::ResourceUsage::Read));
 				METAL_DEBUG_LAYER(EMetalDebugLevelFastValidation, RenderEncoderDebug.SetBuffer(EMetalShaderVertex, Buffer, Offset, Index));
 				RenderCommandEncoder.SetVertexBuffer(Buffer, Offset, Index);
+				FenceResource(Buffer);
 				break;
 			case mtlpp::FunctionType::Fragment:
 				ShaderBuffers[uint32(Function)].Bound |= (1 << Index);
diff --git a/Engine/Source/Runtime/Apple/MetalRHI/Private/MetalCommandEncoder.h b/Engine/Source/Runtime/Apple/MetalRHI/Private/MetalCommandEncoder.h
--- a/Engine/Source/Runtime/Apple/MetalRHI/Private/MetalCommandEncoder.h
+++ b/Engine/Source/Runtime/Apple/MetalRHI/Private/MetalCommandEncoder.h
@@ -356,6 +356,12 @@ public:
 	 */
 	void SetShaderSideTable(mtlpp::FunctionType const FunctionType, NSUInteger const Index);
 	
+	/*
+	 * Transition resource so that we can barrier fragment->vertex stages.
+	 * @param Resource The resource we are going to make readable having been written.
+	 */
+	void TransitionResources(mtlpp::Resource const& Resource);
+	
 #pragma mark - Public Compute State Mutators -
 	
 	/*
@@ -394,6 +400,9 @@ private:
 	
 	void SetShaderBufferInternal(mtlpp::FunctionType Function, uint32 Index);
 	
+	void FenceResource(mtlpp::Texture const& Resource);
+	void FenceResource(mtlpp::Buffer const& Resource);
+
 	void UseResource(mtlpp::Resource const& Resource, mtlpp::ResourceUsage const Usage);
 	
 #pragma mark - Private Type Declarations -
@@ -466,5 +475,11 @@ private:
 	
 	TMap<mtlpp::Resource::Type, mtlpp::ResourceUsage> ResourceUsage;
 	
+	TSet<mtlpp::Resource::Type> TransitionedResources;
+	TSet<mtlpp::Resource::Type> FenceResources;
+	
+	TSet<FMetalFence*> FragmentFences;
+	
+	mtlpp::RenderStages FenceStage;
 	uint32 EncoderNum;
 };
diff --git a/Engine/Source/Runtime/Apple/MetalRHI/Private/MetalCommandQueue.cpp b/Engine/Source/Runtime/Apple/MetalRHI/Private/MetalCommandQueue.cpp
--- a/Engine/Source/Runtime/Apple/MetalRHI/Private/MetalCommandQueue.cpp
+++ b/Engine/Source/Runtime/Apple/MetalRHI/Private/MetalCommandQueue.cpp
@@ -375,22 +375,26 @@ FMetalFence* FMetalCommandQueue::CreateFence(ns::String const& Label) const
 	if ((Features & EMetalFeaturesFences) != 0)
 	{
 		FMetalFence* InternalFence = FMetalFencePool::Get().AllocateFence();
-		NSString* String = nil;
-		if (GetEmitDrawEvents())
+		for (uint32 i = mtlpp::RenderStages::Vertex; InternalFence && i <= mtlpp::RenderStages::Fragment; i++)
 		{
-			String = [NSString stringWithFormat:@"%p: %@", InternalFence->GetPtr(), Label.GetPtr()];
-		}
-#if METAL_DEBUG_OPTIONS
-		if (RuntimeDebuggingLevel >= EMetalDebugLevelValidation)
-		{
-			FMetalDebugFence* Fence = (FMetalDebugFence*)InternalFence->GetPtr();
-			Fence.label = String;
-		}
-		else
-#endif
-		if(InternalFence && String)
-		{
-			InternalFence->SetLabel(String);
+			mtlpp::Fence InnerFence = InternalFence->Get((mtlpp::RenderStages)i);
+			NSString* String = nil;
+			if (GetEmitDrawEvents())
+			{
+				String = [NSString stringWithFormat:@"%u %p: %@", i, InnerFence.GetPtr(), Label.GetPtr()];
+			}
+	#if METAL_DEBUG_OPTIONS
+			if (RuntimeDebuggingLevel >= EMetalDebugLevelValidation)
+			{
+				FMetalDebugFence* Fence = (FMetalDebugFence*)InnerFence.GetPtr();
+				Fence.label = String;
+			}
+			else
+	#endif
+			if(InnerFence && String)
+			{
+				InnerFence.SetLabel(String);
+			}
 		}
 		return InternalFence;
 	}
diff --git a/Engine/Source/Runtime/Apple/MetalRHI/Private/MetalContext.cpp b/Engine/Source/Runtime/Apple/MetalRHI/Private/MetalContext.cpp
--- a/Engine/Source/Runtime/Apple/MetalRHI/Private/MetalContext.cpp
+++ b/Engine/Source/Runtime/Apple/MetalRHI/Private/MetalContext.cpp
@@ -1060,6 +1060,85 @@ void FMetalContext::FinishFrame()
 #endif
 }
 
+void FMetalContext::TransitionResources(FUnorderedAccessViewRHIParamRef* InUAVs, int32 NumUAVs)
+{
+	for (uint32 i = 0; i < NumUAVs; i++)
+	{
+		FMetalUnorderedAccessView* UAV = (FMetalUnorderedAccessView*)InUAVs[i];
+		if (UAV)
+		{
+			ns::AutoReleased<mtlpp::Resource> Resource;
+			
+			// figure out which one of the resources we need to set
+			FMetalStructuredBuffer* StructuredBuffer = UAV->SourceView->SourceStructuredBuffer.GetReference();
+			FMetalVertexBuffer* VertexBuffer = UAV->SourceView->SourceVertexBuffer.GetReference();
+			FMetalIndexBuffer* IndexBuffer = UAV->SourceView->SourceIndexBuffer.GetReference();
+			FRHITexture* Texture = UAV->SourceView->SourceTexture.GetReference();
+			FMetalSurface* Surface = UAV->SourceView->TextureView;
+			if (StructuredBuffer)
+			{
+				check(StructuredBuffer->Buffer);
+				RenderPass.TransitionResources(StructuredBuffer->Buffer);
+			}
+			else if (VertexBuffer && VertexBuffer->Buffer)
+			{
+				RenderPass.TransitionResources(VertexBuffer->Buffer);
+			}
+			else if (IndexBuffer)
+			{
+				check(IndexBuffer->Buffer);
+				RenderPass.TransitionResources(IndexBuffer->Buffer);
+			}
+			else if (Surface)
+			{
+				RenderPass.TransitionResources(Surface->Texture.GetParentTexture());
+			}
+			else if (Texture)
+			{
+				if (!Surface)
+				{
+					Surface = GetMetalSurfaceFromRHITexture(Texture);
+				}
+				if (Surface != nullptr && Surface->Texture)
+				{
+					RenderPass.TransitionResources(Surface->Texture);
+					if (Surface->StencilTexture)
+					{
+						RenderPass.TransitionResources(Surface->StencilTexture);
+					}
+					if (Surface->MSAATexture)
+					{
+						RenderPass.TransitionResources(Surface->MSAATexture);
+					}
+				}
+			}
+		}
+	}
+}
+
+void FMetalContext::TransitionResources(FTextureRHIParamRef* InTextures, int32 NumTextures)
+{
+	for (uint32 i = 0; i < NumTextures; i++)
+	{
+		if (InTextures[i])
+		{
+			FMetalSurface* Surface = GetMetalSurfaceFromRHITexture(InTextures[i]);
+			if (Surface != nullptr && Surface->Texture)
+			{
+				RenderPass.TransitionResources(Surface->Texture);
+				if (Surface->StencilTexture)
+				{
+					RenderPass.TransitionResources(Surface->StencilTexture);
+				}
+				if (Surface->MSAATexture)
+				{
+					RenderPass.TransitionResources(Surface->MSAATexture);
+				}
+			}
+		}
+	}
+}
+
 void FMetalContext::SubmitCommandsHint(uint32 const Flags)
 {
 	// When the command-buffer is submitted for a reason other than a break of a logical command-buffer (where one high-level command-sequence becomes more than one command-buffer).
diff --git a/Engine/Source/Runtime/Apple/MetalRHI/Private/MetalContext.h b/Engine/Source/Runtime/Apple/MetalRHI/Private/MetalContext.h
--- a/Engine/Source/Runtime/Apple/MetalRHI/Private/MetalContext.h
+++ b/Engine/Source/Runtime/Apple/MetalRHI/Private/MetalContext.h
@@ -123,6 +123,10 @@ public:
 	void InitFrame(bool const bImmediateContext, uint32 Index, uint32 Num);
 	void FinishFrame();
 
+	// Track Write->Read transitions for TBDR Fragment->Verex fencing
+	void TransitionResources(FUnorderedAccessViewRHIParamRef* InUAVs, int32 NumUAVs);
+	void TransitionResources(FTextureRHIParamRef* InTextures, int32 NumTextures);
+	
 protected:
 	/** The underlying Metal device */
 	mtlpp::Device Device;
diff --git a/Engine/Source/Runtime/Apple/MetalRHI/Private/MetalFence.cpp b/Engine/Source/Runtime/Apple/MetalRHI/Private/MetalFence.cpp
--- a/Engine/Source/Runtime/Apple/MetalRHI/Private/MetalFence.cpp
+++ b/Engine/Source/Runtime/Apple/MetalRHI/Private/MetalFence.cpp
@@ -98,9 +98,13 @@ APPLE_PLATFORM_OBJECT_ALLOC_OVERRIDES(FMetalDebugFence)
 #if METAL_DEBUG_OPTIONS
 void FMetalFence::Validate(void) const
 {
-	if (GetMetalDeviceContext().GetCommandQueue().GetRuntimeDebuggingLevel() >= EMetalDebugLevelValidation && GetPtr())
+	if (GetMetalDeviceContext().GetCommandQueue().GetRuntimeDebuggingLevel() >= EMetalDebugLevelValidation && Get(mtlpp::RenderStages::Vertex))
 	{
-		[(FMetalDebugFence*)GetPtr() validate];
+		[(FMetalDebugFence*)Get(mtlpp::RenderStages::Vertex).GetPtr() validate];
+		if (Get(mtlpp::RenderStages::Fragment))
+		{
+			[(FMetalDebugFence*)Get(mtlpp::RenderStages::Fragment).GetPtr() validate];
+		}
 	}
 }
 #endif
@@ -118,16 +122,22 @@ void FMetalFencePool::Initialise(mtlpp::Device const& InDevice)
 #if METAL_DEBUG_OPTIONS
 		if (GMetalRuntimeDebugLevel >= EMetalDebugLevelValidation)
 		{
-			FMetalDebugFence* Fence = [[FMetalDebugFence new] autorelease];
-			Fence.Inner = Device.NewFence();
-			FMetalFence* F = new FMetalFence(Fence);
+			FMetalDebugFence* VertexFence = [[FMetalDebugFence new] autorelease];
+			VertexFence.Inner = Device.NewFence();
+			FMetalDebugFence* FragmentFence = [[FMetalDebugFence new] autorelease];
+			FragmentFence.Inner = Device.NewFence();
+			FMetalFence* F = new FMetalFence;
+			F->Set(mtlpp::RenderStages::Vertex, VertexFence);
+			F->Set(mtlpp::RenderStages::Fragment, FragmentFence);
 			Fences.Add(F);
 			Lifo.Push(F);
 		}
 		else
 #endif
 		{
-			FMetalFence* F = new FMetalFence(Device.NewFence());
+			FMetalFence* F = new FMetalFence;
+			F->Set(mtlpp::RenderStages::Vertex, Device.NewFence());
+			F->Set(mtlpp::RenderStages::Fragment, Device.NewFence());
 #if METAL_DEBUG_OPTIONS
 			if (GMetalRuntimeDebugLevel >= EMetalDebugLevelValidation)
 			{
@@ -166,11 +176,11 @@ void FMetalFence::ValidateUsage(FMetalFence* InFence)
 	{
 		if (InFence->NumWrites(mtlpp::RenderStages::Vertex) != InFence->NumWaits(mtlpp::RenderStages::Vertex))
 		{
-			UE_LOG(LogMetal, Warning, TEXT("%p (%s) writes %d waits %d"), InFence, *FString(InFence->GetLabel()), (uint32)InFence->NumWrites(mtlpp::RenderStages::Vertex), (uint32)InFence->NumWaits(mtlpp::RenderStages::Vertex));
+			UE_LOG(LogMetal, Warning, TEXT("%p (%s) writes %d waits %d"), InFence, *FString(InFence->Get(mtlpp::RenderStages::Vertex).GetLabel()), (uint32)InFence->NumWrites(mtlpp::RenderStages::Vertex), (uint32)InFence->NumWaits(mtlpp::RenderStages::Vertex));
 		}
 		if (InFence->NumWrites(mtlpp::RenderStages::Fragment) != InFence->NumWaits(mtlpp::RenderStages::Fragment))
 		{
-			UE_LOG(LogMetal, Warning, TEXT("%p (%s) writes %d waits %d"), InFence, *FString(InFence->GetLabel()), (uint32)InFence->NumWrites(mtlpp::RenderStages::Fragment), (uint32)InFence->NumWaits(mtlpp::RenderStages::Fragment));
+			UE_LOG(LogMetal, Warning, TEXT("%p (%s) writes %d waits %d"), InFence, *FString(InFence->Get(mtlpp::RenderStages::Fragment).GetLabel()), (uint32)InFence->NumWrites(mtlpp::RenderStages::Fragment), (uint32)InFence->NumWaits(mtlpp::RenderStages::Fragment));
 		}
 	}
 }
diff --git a/Engine/Source/Runtime/Apple/MetalRHI/Private/MetalFence.h b/Engine/Source/Runtime/Apple/MetalRHI/Private/MetalFence.h
--- a/Engine/Source/Runtime/Apple/MetalRHI/Private/MetalFence.h
+++ b/Engine/Source/Runtime/Apple/MetalRHI/Private/MetalFence.h
@@ -61,7 +61,7 @@
 -(void) waitForFence:(id <MTLFence>)fence beforeStages:(MTLRenderStages)stages;
 @end
 
-class FMetalFence : public mtlpp::Fence
+class FMetalFence
 {
 	enum
 	{
@@ -69,17 +69,6 @@ class FMetalFence : public mtlpp::Fence
 	};
 public:
 	FMetalFence()
-	: mtlpp::Fence(nil)
-	{
-		for (uint32 i = 0; i < NumFenceStages; i++)
-		{
-			Writes[i] = 0;
-			Waits[i] = 0;
-		}
-	}
-	
-	explicit FMetalFence(mtlpp::Fence Obj)
-	: mtlpp::Fence(Obj)
 	{
 		for (uint32 i = 0; i < NumFenceStages; i++)
 		{
@@ -101,27 +90,10 @@ public:
 	{
 		if (&Other != this)
 		{
-			mtlpp::Fence::operator=(Other);
-		}
-		return *this;
-	}
-	
-	FMetalFence& operator=(mtlpp::Fence const& Other)
-	{
-		if (Other.GetPtr() != GetPtr())
-		{
-			METAL_DEBUG_OPTION(Validate());
-			mtlpp::Fence::operator=(Other);
-		}
-		return *this;
-	}
-	
-	FMetalFence& operator=(mtlpp::Fence&& Other)
-	{
-		if (Other.GetPtr() != GetPtr())
-		{
-			METAL_DEBUG_OPTION(Validate());
-			mtlpp::Fence::operator=(Other);
+			for (uint32 i = 0; i < NumFenceStages; i++)
+			{
+				Fences[i] = Other.Fences[i];
+			}
 		}
 		return *this;
 	}
@@ -169,9 +141,20 @@ public:
 		return Waits[(uint32)Stage - 1u] == 0 || (Writes[(uint32)Stage - 1u] > Waits[(uint32)Stage - 1u]);
 	}
 	
+	mtlpp::Fence Get(mtlpp::RenderStages Stage) const
+	{
+		return Fences[(uint32)Stage - 1u];
+	}
+	
+	void Set(mtlpp::RenderStages Stage, mtlpp::Fence Fence)
+	{
+		Fences[(uint32)Stage - 1u] = Fence;
+	}
+	
 	static void ValidateUsage(FMetalFence* Fence);
 	
 private:
+	mtlpp::Fence Fences[NumFenceStages];
 	int8 Writes[NumFenceStages];
 	int8 Waits[NumFenceStages];
 };
diff --git a/Engine/Source/Runtime/Apple/MetalRHI/Private/MetalProfiler.cpp b/Engine/Source/Runtime/Apple/MetalRHI/Private/MetalProfiler.cpp
--- a/Engine/Source/Runtime/Apple/MetalRHI/Private/MetalProfiler.cpp
+++ b/Engine/Source/Runtime/Apple/MetalRHI/Private/MetalProfiler.cpp
@@ -1583,7 +1583,7 @@ void FMetalProfiler::EncodeFence(FMetalCommandBufferStats* CmdBufStats, const TC
 {
 	if (MetalGPUProfilerIsInSafeThread() && Fence && bEnabled && StatisticsAPI && CmdBufStats->ActiveEncoderStats)
 	{
-		FMetalEventStats* Event = new FMetalEventStats(*FString::Printf(TEXT("%s: %s"), Name, *FString(Fence->GetLabel())), 1);
+		FMetalEventStats* Event = new FMetalEventStats(*FString::Printf(TEXT("%s: %s"), Name, *FString(Fence->Get(mtlpp::RenderStages::Vertex).GetLabel())), 1);
 		CmdBufStats->ActiveEncoderStats->EncodeFence(Event, Type);
 	}
 }
diff --git a/Engine/Source/Runtime/Apple/MetalRHI/Private/MetalRenderPass.cpp b/Engine/Source/Runtime/Apple/MetalRHI/Private/MetalRenderPass.cpp
--- a/Engine/Source/Runtime/Apple/MetalRHI/Private/MetalRenderPass.cpp
+++ b/Engine/Source/Runtime/Apple/MetalRHI/Private/MetalRenderPass.cpp
@@ -1101,6 +1101,12 @@ void FMetalRenderPass::AddAsyncCommandBufferHandlers(mtlpp::CommandBufferHandler
 	}
 }
 
+void FMetalRenderPass::TransitionResources(mtlpp::Resource const& Resource)
+{
+	PrologueEncoder.TransitionResources(Resource);
+	CurrentEncoder.TransitionResources(Resource);
+}
+
 #pragma mark - Public Debug Support -
 
 void FMetalRenderPass::InsertDebugEncoder()
@@ -1252,13 +1258,16 @@ void FMetalRenderPass::ConditionalSwitchToTessellation(void)
 #if METAL_DEBUG_OPTIONS
 		if (GetEmitDrawEvents() && PrologueEncoderFence)
 		{
-			if (CmdList.GetCommandQueue().GetRuntimeDebuggingLevel() >= EMetalDebugLevelValidation)
-			{
-				(*PrologueEncoderFence).GetPtr().label = [NSString stringWithFormat:@"Prologue %@", (*PrologueEncoderFence).GetLabel().GetPtr()];
-			}
-			else
+			for (uint32 i = mtlpp::RenderStages::Vertex; i <= mtlpp::RenderStages::Fragment; i++)
 			{
-				(*PrologueEncoderFence).SetLabel([NSString stringWithFormat:@"Prologue %@", (*PrologueEncoderFence).GetLabel().GetPtr()]);
+				if (CmdList.GetCommandQueue().GetRuntimeDebuggingLevel() >= EMetalDebugLevelValidation)
+				{
+					PrologueEncoderFence->Get((mtlpp::RenderStages)i).GetPtr().label = [NSString stringWithFormat:@"Prologue %@", PrologueEncoderFence->Get((mtlpp::RenderStages)i).GetLabel().GetPtr()];
+				}
+				else
+				{
+					PrologueEncoderFence->Get((mtlpp::RenderStages)i).SetLabel([NSString stringWithFormat:@"Prologue %@", PrologueEncoderFence->Get((mtlpp::RenderStages)i).GetLabel().GetPtr()]);
+				}
 			}
 		}
 #endif
@@ -1397,13 +1406,16 @@ void FMetalRenderPass::ConditionalSwitchToAsyncBlit(void)
 #if METAL_DEBUG_OPTIONS
 		if (GetEmitDrawEvents() && PrologueEncoderFence)
 		{
-			if (CmdList.GetCommandQueue().GetRuntimeDebuggingLevel() >= EMetalDebugLevelValidation)
+			for (uint32 i = mtlpp::RenderStages::Vertex; i <= mtlpp::RenderStages::Fragment; i++)
 			{
-				(*PrologueEncoderFence).GetPtr().label = [NSString stringWithFormat:@"Prologue %@", (*PrologueEncoderFence).GetLabel().GetPtr()];
-			}
-			else
-			{
-				(*PrologueEncoderFence).SetLabel([NSString stringWithFormat:@"Prologue %@", (*PrologueEncoderFence).GetLabel().GetPtr()]);
+				if (CmdList.GetCommandQueue().GetRuntimeDebuggingLevel() >= EMetalDebugLevelValidation)
+				{
+					PrologueEncoderFence->Get((mtlpp::RenderStages)i).GetPtr().label = [NSString stringWithFormat:@"Prologue %@", PrologueEncoderFence->Get((mtlpp::RenderStages)i).GetLabel().GetPtr()];
+				}
+				else
+				{
+					PrologueEncoderFence->Get((mtlpp::RenderStages)i).SetLabel([NSString stringWithFormat:@"Prologue %@", PrologueEncoderFence->Get((mtlpp::RenderStages)i).GetLabel().GetPtr()]);
+				}
 			}
 		}
 #endif
diff --git a/Engine/Source/Runtime/Apple/MetalRHI/Private/MetalRenderPass.h b/Engine/Source/Runtime/Apple/MetalRHI/Private/MetalRenderPass.h
--- a/Engine/Source/Runtime/Apple/MetalRHI/Private/MetalRenderPass.h
+++ b/Engine/Source/Runtime/Apple/MetalRHI/Private/MetalRenderPass.h
@@ -89,6 +89,8 @@ public:
 	void AddCompletionHandler(mtlpp::CommandBufferHandler Handler);
 	
 	void AddAsyncCommandBufferHandlers(mtlpp::CommandBufferHandler Scheduled, mtlpp::CommandBufferHandler Completion);
+	
+	void TransitionResources(mtlpp::Resource const& Resource);
 
 #pragma mark - Public Debug Support -
 	
diff --git a/Engine/Source/Runtime/Apple/MetalRHI/Private/MetalUAV.cpp b/Engine/Source/Runtime/Apple/MetalRHI/Private/MetalUAV.cpp
--- a/Engine/Source/Runtime/Apple/MetalRHI/Private/MetalUAV.cpp
+++ b/Engine/Source/Runtime/Apple/MetalRHI/Private/MetalUAV.cpp
@@ -795,6 +795,10 @@ void FMetalRHICommandContext::RHITransitionResources(EResourceTransitionAccess T
 {
 	@autoreleasepool
 	{
+		if (TransitionType != EResourceTransitionAccess::EMetaData)
+		{
+			Context->TransitionResources(InUAVs, NumUAVs);
+		}
 		if (WriteComputeFence)
 		{
 			FMetalComputeFence* Fence = ResourceCast(WriteComputeFence);
@@ -809,12 +813,19 @@ void FMetalRHICommandContext::RHITransitionResources(EResourceTransitionAccess T
 
 void FMetalRHICommandContext::RHITransitionResources(EResourceTransitionAccess TransitionType, FTextureRHIParamRef* InTextures, int32 NumTextures)
 {
-	if (TransitionType == EResourceTransitionAccess::EReadable)
+	@autoreleasepool
 	{
-		const FResolveParams ResolveParams;
-		for (int32 i = 0; i < NumTextures; ++i)
+		if (TransitionType != EResourceTransitionAccess::EMetaData)
 		{
-			RHICopyToResolveTarget(InTextures[i], InTextures[i], ResolveParams);
+			Context->TransitionResources(InTextures, NumTextures);
+		}
+		if (TransitionType == EResourceTransitionAccess::EReadable)
+		{
+			const FResolveParams ResolveParams;
+			for (int32 i = 0; i < NumTextures; ++i)
+			{
+				RHICopyToResolveTarget(InTextures[i], InTextures[i], ResolveParams);
+			}
 		}
 	}
 }
