On iOS 12 enable Fences, Heaps and UseResource calls on by default to match macOS.

diff --git a/Engine/Source/Runtime/Apple/MetalRHI/Private/MetalCommandQueue.cpp b/Engine/Source/Runtime/Apple/MetalRHI/Private/MetalCommandQueue.cpp
--- a/Engine/Source/Runtime/Apple/MetalRHI/Private/MetalCommandQueue.cpp
+++ b/Engine/Source/Runtime/Apple/MetalRHI/Private/MetalCommandQueue.cpp
@@ -114,6 +114,10 @@ FMetalCommandQueue::FMetalCommandQueue(mtlpp::Device InDevice, uint32 const MaxN
 			if (Vers.majorVersion >= 12)
 			{
 				Features |= EMetalFeaturesMaxThreadsPerThreadgroup;
+				Features |= EMetalFeaturesFences;
+				Features |= EMetalFeaturesHeaps;
+				Features |= EMetalFeaturesIABs;
+				
 				if (MaxShaderVersion >= 4)
 				{
 					Features |= EMetalFeaturesTextureBuffers;
