Fix memory leak in Metal's RHIReadSurfaceData

diff --git a/Engine/Source/Runtime/Apple/MetalRHI/Private/MetalRenderTarget.cpp b/Engine/Source/Runtime/Apple/MetalRHI/Private/MetalRenderTarget.cpp
--- a/Engine/Source/Runtime/Apple/MetalRHI/Private/MetalRenderTarget.cpp
+++ b/Engine/Source/Runtime/Apple/MetalRHI/Private/MetalRenderTarget.cpp
@@ -630,7 +630,7 @@ void FMetalDynamicRHI::RHIReadSurfaceData(FTextureRHIParamRef TextureRHI, FIntRe
 			Desc.SetStorageMode(StorageMode);
 			Desc.SetUsage(Texture.GetUsage());
 			
-			TempTexture = [GetMetalDeviceContext().GetDevice() newTextureWithDescriptor:Desc];
+			TempTexture = GetMetalDeviceContext().GetDevice().NewTexture(Desc);
 			
 			ImmediateContext.Context->CopyFromTextureToTexture(Texture, 0, 0, mtlpp::Origin(Region.origin), mtlpp::Size(Region.size), TempTexture, 0, 0, mtlpp::Origin(0, 0, 0));
 			
@@ -654,7 +654,7 @@ void FMetalDynamicRHI::RHIReadSurfaceData(FTextureRHIParamRef TextureRHI, FIntRe
 		TArray<uint8> Data;
 		Data.AddUninitialized(BytesPerImage);
 		
-		mtlpp::Texture(Texture).GetBytes(Data.GetData(), Stride, BytesPerImage, Region, 0, 0);
+		Texture.GetBytes(Data.GetData(), Stride, BytesPerImage, Region, 0, 0);
 		
 		ConvertSurfaceDataToFColor(Surface->PixelFormat, SizeX, SizeY, (uint8*)Data.GetData(), Stride, OutDataPtr, InFlags);
 		
