Export all the caps bits for more recent tvOS versions.

diff --git a/Engine/Source/Runtime/Apple/MetalRHI/Private/MetalCommandQueue.cpp b/Engine/Source/Runtime/Apple/MetalRHI/Private/MetalCommandQueue.cpp
--- a/Engine/Source/Runtime/Apple/MetalRHI/Private/MetalCommandQueue.cpp
+++ b/Engine/Source/Runtime/Apple/MetalRHI/Private/MetalCommandQueue.cpp
@@ -62,7 +62,42 @@ FMetalCommandQueue::FMetalCommandQueue(mtlpp::Device InDevice, uint32 const MaxN
         Features &= ~(EMetalFeaturesSetBytes);
 		if(Device.SupportsFeatureSet(mtlpp::FeatureSet::tvOS_GPUFamily1_v2))
 		{
-			Features |= EMetalFeaturesStencilView | EMetalFeaturesGraphicsUAVs;
+			Features |= EMetalFeaturesStencilView | EMetalFeaturesGraphicsUAVs | EMetalFeaturesFunctionConstants | EMetalFeaturesMemoryLessResources;
+		}
+		
+		if(Device.SupportsFeatureSet(mtlpp::FeatureSet::tvOS_GPUFamily2_v1))
+		{
+			Features |= EMetalFeaturesCountingQueries | EMetalFeaturesBaseVertexInstance | EMetalFeaturesIndirectBuffer | EMetalFeaturesMSAADepthResolve | EMetalFeaturesTessellation | EMetalFeaturesMSAAStoreAndResolve;
+		}
+		
+		if(Vers.majorVersion > 10)
+		{
+			Features |= EMetalFeaturesGPUCommandBufferTimes;
+			Features |= EMetalFeaturesLinearTextures;
+			Features |= EMetalFeaturesPrivateBufferSubAllocation;
+			Features |= EMetalFeaturesDeferredStoreActions | EMetalFeaturesCombinedDepthStencil;
+			
+			if(Vers.majorVersion >= 11)
+			{
+				Features |= EMetalFeaturesGPUCaptureManager | EMetalFeaturesBufferSubAllocation | EMetalFeaturesParallelRenderEncoders;
+				
+				if (MaxShaderVersion >= 3)
+				{
+					Features |= EMetalFeaturesLinearTextureUAVs;
+				}
+				if (Vers.majorVersion >= 12)
+				{
+					Features |= EMetalFeaturesMaxThreadsPerThreadgroup;
+					Features |= EMetalFeaturesFences;
+					Features |= EMetalFeaturesHeaps;
+					Features |= EMetalFeaturesIABs;
+					
+					if (MaxShaderVersion >= 4)
+					{
+						Features |= EMetalFeaturesTextureBuffers;
+					}
+				}
+			}
 		}
 #else
 		if (Device.SupportsFeatureSet(mtlpp::FeatureSet::iOS_GPUFamily3_v1))
