Use Memoryless render-targets for MSAA textures when using the mobile renderer on iOS where we never sample from the MSAA texture itself. When a texture does not ask for SRV or UAV support we can also use memoryless.

diff --git a/Engine/Source/Runtime/Apple/MetalRHI/Private/MetalLLM.cpp b/Engine/Source/Runtime/Apple/MetalRHI/Private/MetalLLM.cpp
--- a/Engine/Source/Runtime/Apple/MetalRHI/Private/MetalLLM.cpp
+++ b/Engine/Source/Runtime/Apple/MetalRHI/Private/MetalLLM.cpp
@@ -199,7 +199,13 @@ void MetalLLM::LogAllocTexture(mtlpp::Device& Device, mtlpp::TextureDescriptor c
 	void* Ptr = (void*)Texture.GetPtr();
 	uint64 Size = SizeAlign.Size;
 	
-	INC_MEMORY_STAT_BY(STAT_MetalTextureMemory, Size);
+#if PLATFORM_IOS
+	bool bMemoryless = (Texture.GetStorageMode() == mtlpp::StorageMode::Memoryless);
+	if (!bMemoryless)
+#endif
+	{
+		INC_MEMORY_STAT_BY(STAT_MetalTextureMemory, Size);
+	}
 	INC_DWORD_STAT(STAT_MetalTextureCount);
 	
 	LLM(FLowLevelMemTracker::Get().OnLowLevelAlloc(ELLMTracker::Platform, Ptr, Size, ELLMTag::Untagged, ELLMAllocType::System));
@@ -213,7 +219,12 @@ void MetalLLM::LogAllocTexture(mtlpp::Device& Device, mtlpp::TextureDescriptor c
 			
 			LLM(FLowLevelMemTracker::Get().OnLowLevelFree(ELLMTracker::Platform, Ptr, ELLMAllocType::System));
 			
-			DEC_MEMORY_STAT_BY(STAT_MetalTextureMemory, Size);
+#if PLATFORM_IOS
+			if (!bMemoryless)
+#endif
+			{
+				DEC_MEMORY_STAT_BY(STAT_MetalTextureMemory, Size);
+			}
 			DEC_DWORD_STAT(STAT_MetalTextureCount);
 		}] autorelease],
 		OBJC_ASSOCIATION_RETAIN);
diff --git a/Engine/Source/Runtime/Apple/MetalRHI/Private/MetalStateCache.cpp b/Engine/Source/Runtime/Apple/MetalRHI/Private/MetalStateCache.cpp
--- a/Engine/Source/Runtime/Apple/MetalRHI/Private/MetalStateCache.cpp
+++ b/Engine/Source/Runtime/Apple/MetalRHI/Private/MetalStateCache.cpp
@@ -98,15 +98,18 @@ public:
 			mtlpp::RenderPassColorAttachmentDescriptor Color = Attachements[i];
 			Color.SetTexture(EmptyTex);
 			Color.SetResolveTexture(EmptyTex);
+			Color.SetStoreAction(mtlpp::StoreAction::Store);
 		}
 		
 		mtlpp::RenderPassDepthAttachmentDescriptor Depth = Desc.GetDepthAttachment();
 		Depth.SetTexture(EmptyTex);
 		Depth.SetResolveTexture(EmptyTex);
+		Depth.SetStoreAction(mtlpp::StoreAction::Store);
 
 		mtlpp::RenderPassStencilAttachmentDescriptor Stencil = Desc.GetStencilAttachment();
 		Stencil.SetTexture(EmptyTex);
 		Stencil.SetResolveTexture(EmptyTex);
+		Stencil.SetStoreAction(mtlpp::StoreAction::Store);
 
 		mtlpp::Buffer Empty;
 		Desc.SetVisibilityResultBuffer(Empty);
@@ -528,6 +531,9 @@ bool FMetalStateCache::SetRenderTargetsInfo(FRHISetRenderTargetsInfo const& InRe
 	
 				mtlpp::RenderPassColorAttachmentDescriptor ColorAttachment = Attachements[RenderTargetIndex];
 	
+				ERenderTargetStoreAction HighLevelStoreAction = RenderTargetView.StoreAction;
+				ERenderTargetLoadAction HighLevelLoadAction = RenderTargetView.LoadAction;
+				
 				if (Surface.MSAATexture)
 				{
 					// set up an MSAA attachment
@@ -536,15 +542,27 @@ bool FMetalStateCache::SetRenderTargetsInfo(FRHISetRenderTargetsInfo const& InRe
 					ColorAttachment.SetStoreAction(bSupportsDeferredStore && GRHIDeviceId > 2 ? mtlpp::StoreAction::Unknown : NewColorStore[RenderTargetIndex]);
 					ColorAttachment.SetResolveTexture(Surface.MSAAResolveTexture ? Surface.MSAAResolveTexture : Surface.Texture);
 					SampleCount = Surface.MSAATexture.GetSampleCount();
-                    
+#if PLATFORM_IOS
+					if (Surface.MSAATexture.GetStorageMode() == mtlpp::StorageMode::Memoryless)
+					{
+						HighLevelLoadAction = ERenderTargetLoadAction::EClear;
+					}
+#endif
 					// only allow one MRT with msaa
 					checkf(RenderTargetsInfo.NumColorRenderTargets == 1, TEXT("Only expected one MRT when using MSAA"));
 				}
 				else
 				{
+#if PLATFORM_IOS
+					if (Surface.Texture.GetStorageMode() == mtlpp::StorageMode::Memoryless)
+					{
+						HighLevelStoreAction = ERenderTargetStoreAction::ENoAction;
+						HighLevelLoadAction = ERenderTargetLoadAction::EClear;
+					}
+#endif
 					// set up non-MSAA attachment
 					ColorAttachment.SetTexture(Surface.Texture);
-					NewColorStore[RenderTargetIndex] = GetMetalRTStoreAction(RenderTargetView.StoreAction);
+					NewColorStore[RenderTargetIndex] = GetMetalRTStoreAction(HighLevelStoreAction);
 					ColorAttachment.SetStoreAction(bSupportsDeferredStore ? mtlpp::StoreAction::Unknown : NewColorStore[RenderTargetIndex]);
                     SampleCount = 1;
 				}
@@ -559,7 +577,7 @@ bool FMetalStateCache::SetRenderTargetsInfo(FRHISetRenderTargetsInfo const& InRe
 					ColorAttachment.SetSlice(ArraySliceIndex);
 				}
 				
-				ColorAttachment.SetLoadAction((Surface.Written || !bImmediate || bRestart) ? GetMetalRTLoadAction(RenderTargetView.LoadAction) : mtlpp::LoadAction::Clear);
+				ColorAttachment.SetLoadAction((Surface.Written || !bImmediate || bRestart) ? GetMetalRTLoadAction(HighLevelLoadAction) : mtlpp::LoadAction::Clear);
 				FPlatformAtomics::InterlockedExchange(&Surface.Written, 1);
 				
 				bNeedsClear |= (ColorAttachment.GetLoadAction() == mtlpp::LoadAction::Clear);
@@ -767,9 +785,26 @@ bool FMetalStateCache::SetRenderTargetsInfo(FRHISetRenderTargetsInfo const& InRe
 					}
 				}
 				
+				const bool bSupportsMSAADepthResolve = GetMetalDeviceContext().SupportsFeature(EMetalFeaturesMSAADepthResolve);
+				bool bDepthTextureMemoryless = false;
+#if PLATFORM_IOS
+				bDepthTextureMemoryless = DepthTexture.GetStorageMode() == mtlpp::StorageMode::Memoryless;
+				if (bDepthTextureMemoryless)
+				{
+					DepthAttachment.SetLoadAction(mtlpp::LoadAction::Clear);
+					
+					if (bSupportsMSAADepthResolve && Surface.MSAATexture)
+					{
+						HighLevelStoreAction = ERenderTargetStoreAction::EMultisampleResolve;
+					}
+					else
+					{
+						HighLevelStoreAction = ERenderTargetStoreAction::ENoAction;
+					}
+				}
+#endif
                 //needed to quiet the metal validation that runs when you end renderpass. (it requires some kind of 'resolve' for an msaa target)
 				//But with deferredstore we don't set the real one until submit time.
-				const bool bSupportsMSAADepthResolve = GetMetalDeviceContext().SupportsFeature(EMetalFeaturesMSAADepthResolve);
 				NewDepthStore = !Surface.MSAATexture || bSupportsMSAADepthResolve ? GetMetalRTStoreAction(HighLevelStoreAction) : mtlpp::StoreAction::DontCare;
 				DepthAttachment.SetStoreAction(bSupportsDeferredStore && Surface.MSAATexture && GRHIDeviceId > 2 ? mtlpp::StoreAction::Unknown : NewDepthStore);
 				DepthAttachment.SetClearDepth(DepthClearValue);
@@ -827,6 +862,14 @@ bool FMetalStateCache::SetRenderTargetsInfo(FRHISetRenderTargetsInfo const& InRe
 						HighLevelStoreAction = ERenderTargetStoreAction::EStore;
 					}
 					
+#if PLATFORM_IOS
+					if (StencilTexture.GetStorageMode() == mtlpp::StorageMode::Memoryless)
+					{
+						HighLevelStoreAction = ERenderTargetStoreAction::ENoAction;
+						StencilAttachment.SetLoadAction(mtlpp::LoadAction::Clear);
+					}
+#endif
+					
 					// For the case where Depth+Stencil is MSAA we can't Resolve depth and Store stencil - we can only Resolve + DontCare or StoreResolve + Store (on newer H/W and iOS).
 					// We only allow use of StoreResolve in the Desktop renderers as the mobile renderer does not and should not assume hardware support for it.
 					NewStencilStore = (StencilTexture.GetSampleCount() == 1  || GetMetalRTStoreAction(ERenderTargetStoreAction::EMultisampleResolve) == mtlpp::StoreAction::StoreAndMultisampleResolve) ? GetMetalRTStoreAction(HighLevelStoreAction) : mtlpp::StoreAction::DontCare;
@@ -1250,6 +1293,14 @@ bool FMetalStateCache::NeedsToSetRenderTarget(const FRHISetRenderTargetsInfo& In
 
 						ERenderTargetStoreAction HighLevelStoreAction = (Surface.MSAATexture && !bDepthStencilSampleCountMismatchFixup) ? ERenderTargetStoreAction::EMultisampleResolve : RenderTargetsInfo.DepthStencilRenderTarget.DepthStoreAction;
 						
+#if PLATFORM_IOS
+						FMetalTexture& Tex = Surface.MSAATexture ? Surface.MSAATexture : Surface.Texture;
+						if (Tex.GetStorageMode() == mtlpp::StorageMode::Memoryless)
+						{
+							HighLevelStoreAction = ERenderTargetStoreAction::ENoAction;
+						}
+#endif
+						
 						NewDepthStore = GetMetalRTStoreAction(HighLevelStoreAction);
 					}
 					else
@@ -1263,6 +1314,12 @@ bool FMetalStateCache::NeedsToSetRenderTarget(const FRHISetRenderTargetsInfo& In
 					if (RenderPassDesc.GetStencilAttachment().GetTexture())
 					{
 						NewStencilStore = GetMetalRTStoreAction(RenderTargetsInfo.DepthStencilRenderTarget.GetStencilStoreAction());
+#if PLATFORM_IOS
+						if (RenderPassDesc.GetStencilAttachment().GetTexture().GetStorageMode() == mtlpp::StorageMode::Memoryless)
+						{
+							NewStencilStore = GetMetalRTStoreAction(ERenderTargetStoreAction::ENoAction);
+						}
+#endif
 					}
 					else
 					{
@@ -1344,10 +1401,15 @@ void FMetalStateCache::SetShaderTexture(EShaderFrequency const Frequency, FMetal
 	if (ShaderTextures[Frequency].Textures[Index] != Texture
 		|| ShaderTextures[Frequency].Usage[Index] != Usage)
 	{
-		ShaderTextures[Frequency].Textures[Index] = Texture;
+		bool bMemoryLess = false;
+#if PLATFORM_IOS
+		bMemoryLess = (Texture.GetStorageMode() == mtlpp::StorageMode::Memoryless);
+#endif
+		
+		ShaderTextures[Frequency].Textures[Index] = !bMemoryLess ? Texture : FMetalTexture();
 		ShaderTextures[Frequency].Usage[Index] = Usage;
 		
-		if (Texture)
+		if (Texture && !bMemoryLess)
 		{
 			ShaderTextures[Frequency].Bound |= (1 << Index);
 		}
diff --git a/Engine/Source/Runtime/Apple/MetalRHI/Private/MetalTexture.cpp b/Engine/Source/Runtime/Apple/MetalRHI/Private/MetalTexture.cpp
--- a/Engine/Source/Runtime/Apple/MetalRHI/Private/MetalTexture.cpp
+++ b/Engine/Source/Runtime/Apple/MetalRHI/Private/MetalTexture.cpp
@@ -214,7 +214,11 @@ void FMetalSurface::PrepareTextureView()
 {
 	// Recreate the texture to enable MTLTextureUsagePixelFormatView which must be off unless we definitely use this feature or we are throwing ~4% performance vs. Windows on the floor.
 	mtlpp::TextureUsage Usage = (mtlpp::TextureUsage)Texture.GetUsage();
-	if(!(Usage & mtlpp::TextureUsage::PixelFormatView))
+	bool bMemoryLess = false;
+#if PLATFORM_IOS
+	bMemoryLess = (Texture.GetStorageMode() == mtlpp::StorageMode::Memoryless);
+#endif
+	if(!(Usage & mtlpp::TextureUsage::PixelFormatView) && !bMemoryLess)
 	{
 		check(!bTextureView);
 		check(ImageSurfaceRef == nullptr);
@@ -462,6 +466,14 @@ void FMetalSurface::MakeUnAliasable(void)
 
 void FMetalSurface::Init(FMetalSurface& Source, NSRange MipRange)
 {
+#if PLATFORM_IOS
+	// Mmeory;ess targets can't have texture views (SRVs or UAVs)
+	if (Source.Texture.GetStorageMode() == mtlpp::StorageMode::Memoryless)
+	{
+		return;
+	}
+#endif
+	
 	mtlpp::PixelFormat MetalFormat = (mtlpp::PixelFormat)GPixelFormats[PixelFormat].PlatformFormat;
 	
 	bool const bUseSourceTex = (Source.PixelFormat != PF_DepthStencil) && MipRange.location == 0 && MipRange.length == Source.Texture.GetMipmapLevelCount();
@@ -495,6 +507,13 @@ void FMetalSurface::Init(FMetalSurface& Source, NSRange MipRange)
 void FMetalSurface::Init(FMetalSurface& Source, NSRange MipRange, EPixelFormat Format)
 {
 	check(!Source.MSAATexture || Format == PF_X24_G8);
+#if PLATFORM_IOS
+	// Mmeory;ess targets can't have texture views (SRVs or UAVs)
+	if (Source.Texture.GetStorageMode() == mtlpp::StorageMode::Memoryless)
+	{
+		return;
+	}
+#endif
 	
 	mtlpp::PixelFormat MetalFormat = (mtlpp::PixelFormat)GPixelFormats[PixelFormat].PlatformFormat;
 	
@@ -855,9 +874,19 @@ FMetalSurface::FMetalSurface(ERHIResourceType ResourceType, EPixelFormat Format,
 		else if (Flags & (TexCreate_RenderTargetable|TexCreate_DepthStencilTargetable))
 		{
 			check(!(Flags & TexCreate_CPUReadback));
-			Desc.SetCpuCacheMode(mtlpp::CpuCacheMode::DefaultCache);
-			Desc.SetStorageMode(mtlpp::StorageMode::Private);
-			Desc.SetResourceOptions((mtlpp::ResourceOptions)(mtlpp::ResourceOptions::CpuCacheModeDefaultCache|mtlpp::ResourceOptions::StorageModePrivate));
+#if PLATFORM_IOS
+			if (FMetalCommandQueue::SupportsFeature(EMetalFeaturesMemoryLessResources) && !(Flags & (TexCreate_ShaderResource|TexCreate_UAV)))
+			{
+				Desc.SetStorageMode(mtlpp::StorageMode::Memoryless);
+				Desc.SetResourceOptions(mtlpp::ResourceOptions::StorageModeMemoryless);
+			}
+			else
+#endif
+			{
+				Desc.SetCpuCacheMode(mtlpp::CpuCacheMode::DefaultCache);
+				Desc.SetStorageMode(mtlpp::StorageMode::Private);
+				Desc.SetResourceOptions((mtlpp::ResourceOptions)(mtlpp::ResourceOptions::CpuCacheModeDefaultCache|mtlpp::ResourceOptions::StorageModePrivate));
+			}
 		}
 		else
 		{
@@ -969,6 +998,14 @@ FMetalSurface::FMetalSurface(ERHIResourceType ResourceType, EPixelFormat Format,
 			FParse::Value(FCommandLine::Get(), TEXT("msaa="), NumSamples);
 			Desc.SetSampleCount(NumSamples);
 			
+#if PLATFORM_IOS
+			if (FMetalCommandQueue::SupportsFeature(EMetalFeaturesMemoryLessResources) && GMaxRHIShaderPlatform == SP_METAL)
+			{
+				Desc.SetStorageMode(mtlpp::StorageMode::Memoryless);
+				Desc.SetResourceOptions(mtlpp::ResourceOptions::StorageModeMemoryless);
+			}
+#endif
+			
 			MSAATexture = GetMetalDeviceContext().CreateTexture(this, Desc);
 			
 			//device doesn't support HW depth resolve.  This case only valid on mobile renderer or
