A couple of small bug fixes for 4461341 to avoid validation errors when using memoryless targets. You can't avoid the resolve when a resolve-target is specified and can only demote StoreAndResolve to Resolve or the validation layer asserts.

diff --git a/Engine/Source/Runtime/Apple/MetalRHI/Private/MetalStateCache.cpp b/Engine/Source/Runtime/Apple/MetalRHI/Private/MetalStateCache.cpp
--- a/Engine/Source/Runtime/Apple/MetalRHI/Private/MetalStateCache.cpp
+++ b/Engine/Source/Runtime/Apple/MetalRHI/Private/MetalStateCache.cpp
@@ -2242,7 +2242,18 @@ void FMetalStateCache::DiscardRenderTargets(bool Depth, bool Stencil, uint32 Col
 {
 	if (Depth)
 	{
-		DepthStore = mtlpp::StoreAction::DontCare;
+		switch (DepthStore)
+		{
+			case mtlpp::StoreAction::Unknown:
+			case mtlpp::StoreAction::Store:
+				DepthStore = mtlpp::StoreAction::DontCare;
+				break;
+			case mtlpp::StoreAction::StoreAndMultisampleResolve:
+				DepthStore = mtlpp::StoreAction::MultisampleResolve;
+				break;
+			default:
+				break;
+		}
 	}
 
 	if (Stencil)
@@ -2254,7 +2265,18 @@ void FMetalStateCache::DiscardRenderTargets(bool Depth, bool Stencil, uint32 Col
 	{
 		if ((ColorBitMask & (1u << Index)) != 0)
 		{
-			ColorStore[Index] = mtlpp::StoreAction::DontCare;
+			switch (ColorStore[Index])
+			{
+				case mtlpp::StoreAction::Unknown:
+				case mtlpp::StoreAction::Store:
+					ColorStore[Index] = mtlpp::StoreAction::DontCare;
+					break;
+				case mtlpp::StoreAction::StoreAndMultisampleResolve:
+					ColorStore[Index] = mtlpp::StoreAction::MultisampleResolve;
+					break;
+				default:
+					break;
+			}
 		}
 	}
 }
diff --git a/Engine/Source/Runtime/Apple/MetalRHI/Private/MetalTexture.cpp b/Engine/Source/Runtime/Apple/MetalRHI/Private/MetalTexture.cpp
--- a/Engine/Source/Runtime/Apple/MetalRHI/Private/MetalTexture.cpp
+++ b/Engine/Source/Runtime/Apple/MetalRHI/Private/MetalTexture.cpp
@@ -998,9 +998,11 @@ FMetalSurface::FMetalSurface(ERHIResourceType ResourceType, EPixelFormat Format,
 			FParse::Value(FCommandLine::Get(), TEXT("msaa="), NumSamples);
 			Desc.SetSampleCount(NumSamples);
 			
+			bool bMemoryless = false;
 #if PLATFORM_IOS
 			if (FMetalCommandQueue::SupportsFeature(EMetalFeaturesMemoryLessResources) && GMaxRHIShaderPlatform == SP_METAL)
 			{
+				bMemoryless = true;
 				Desc.SetStorageMode(mtlpp::StorageMode::Memoryless);
 				Desc.SetResourceOptions(mtlpp::ResourceOptions::StorageModeMemoryless);
 			}
@@ -1018,7 +1020,7 @@ FMetalSurface::FMetalSurface(ERHIResourceType ResourceType, EPixelFormat Format,
 				// we don't have the resolve texture, so we just update the memory size with the MSAA size
 				TotalTextureSize = TotalTextureSize * NumSamples;
 			}
-			else
+			else if (!bMemoryless)
 			{
 				// an MSAA render target takes NumSamples more space, in addition to the resolve texture
 				TotalTextureSize += TotalTextureSize * NumSamples;
