Reimplement support for Metal resource aliasing, enabling it by default for vertex & structured buffers. Textures and render-targets can now be allocated from heaps if the command-line -metaltextureheaps is provided but it is less efficient than allocating from the device directly for the moment so is off by default. Due to the vaguaries of where fences are inserted and wait on the GPU render-targets cannot be aliased properly on macOS at present so render-targets will not come from heaps on macOS at all - iOS doesn't have this problem.

diff --git a/Engine/Source/Runtime/Apple/MetalRHI/Private/MetalBuffer.cpp b/Engine/Source/Runtime/Apple/MetalRHI/Private/MetalBuffer.cpp
--- a/Engine/Source/Runtime/Apple/MetalRHI/Private/MetalBuffer.cpp
+++ b/Engine/Source/Runtime/Apple/MetalRHI/Private/MetalBuffer.cpp
@@ -137,6 +137,7 @@ void FMetalBuffer::Release()
 
 FMetalSubBufferHeap::FMetalSubBufferHeap(NSUInteger Size, NSUInteger Alignment, mtlpp::ResourceOptions Options, FCriticalSection& InPoolMutex)
 : PoolMutex(InPoolMutex)
+, OutstandingAllocs(0)
 , MinAlign(Alignment)
 , UsedSize(0)
 {
@@ -188,6 +189,7 @@ FMetalSubBufferHeap::~FMetalSubBufferHeap()
 
 void FMetalSubBufferHeap::FreeRange(ns::Range const& Range)
 {
+	FPlatformAtomics::InterlockedDecrement(&OutstandingAllocs);
 	if (ParentHeap)
 	{
 		SET_MEMORY_STAT(STAT_MetalBufferUnusedMemory, ParentHeap.GetSize() - ParentHeap.GetUsedSize());
@@ -329,6 +331,11 @@ NSUInteger FMetalSubBufferHeap::GetUsedSize() const
 	}
 }
 
+int64 FMetalSubBufferHeap::NumCurrentAllocations() const
+{
+	return OutstandingAllocs;
+}
+
 void FMetalSubBufferHeap::SetLabel(const ns::String& label)
 {
 	if (ParentHeap)
@@ -400,6 +407,7 @@ FMetalBuffer FMetalSubBufferHeap::NewBuffer(NSUInteger length)
 					if (Range.Length > Size)
 					{
 						ns::Range Split = ns::Range(Range.Location + Size, Range.Length - Size);
+						FPlatformAtomics::InterlockedIncrement(&OutstandingAllocs);
 						FreeRange(Split);
 						
 						Range.Length = Size;
@@ -421,6 +429,7 @@ FMetalBuffer FMetalSubBufferHeap::NewBuffer(NSUInteger length)
 			}
 		}
 	}
+	FPlatformAtomics::InterlockedIncrement(&OutstandingAllocs);
 	check(Result && Result.GetPtr());
 	return Result;
 }
@@ -569,6 +578,7 @@ mtlpp::PurgeableState FMetalSubBufferLinear::SetPurgeableState(mtlpp::PurgeableS
 
 FMetalSubBufferMagazine::FMetalSubBufferMagazine(NSUInteger Size, NSUInteger ChunkSize, mtlpp::ResourceOptions Options)
 : MinAlign(ChunkSize)
+, OutstandingAllocs(0)
 , UsedSize(0)
 {
 	NSUInteger FullSize = Align(Size, ChunkSize);
@@ -624,6 +634,7 @@ FMetalSubBufferMagazine::~FMetalSubBufferMagazine()
 
 void FMetalSubBufferMagazine::FreeRange(ns::Range const& Range)
 {
+	FPlatformAtomics::InterlockedDecrement(&OutstandingAllocs);
 	if (ParentHeap)
 	{
 		SET_MEMORY_STAT(STAT_MetalBufferUnusedMemory, ParentHeap.GetSize() - ParentHeap.GetUsedSize());
@@ -725,6 +736,11 @@ NSUInteger FMetalSubBufferMagazine::GetFreeSize() const
 	return GetSize() - GetUsedSize();
 }
 
+int64 FMetalSubBufferMagazine::NumCurrentAllocations() const
+{
+	return OutstandingAllocs;
+}
+
 void FMetalSubBufferMagazine::SetLabel(const ns::String& label)
 {
 	if (ParentHeap)
@@ -770,6 +786,7 @@ FMetalBuffer FMetalSubBufferMagazine::NewBuffer()
 		}
 	}
 
+	FPlatformAtomics::InterlockedIncrement(&OutstandingAllocs);
 	check(Result && Result.GetPtr());
 	return Result;
 }
@@ -1246,6 +1263,80 @@ uint32 FMetalResourceHeap::GetHeapIndex(uint32 Size)
 	return Lower;
 }
 
+FMetalResourceHeap::TextureHeapSize FMetalResourceHeap::TextureSizeToIndex(uint32 Size)
+{
+	unsigned long Lower = 0;
+	unsigned long Upper = NumTextureHeapSizes;
+	unsigned long Middle;
+	
+	do
+	{
+		Middle = ( Upper + Lower ) >> 1;
+		if( Size <= (HeapTextureHeapSizes[Middle-1] / MinTexturesPerHeap) )
+		{
+			Upper = Middle;
+		}
+		else
+		{
+			Lower = Middle;
+		}
+	}
+	while( Upper - Lower > 1 );
+	
+	check( Size <= (HeapTextureHeapSizes[Lower] / MinTexturesPerHeap) );
+	check( (Lower == 0 ) || ( Size > (HeapTextureHeapSizes[Lower-1] / MinTexturesPerHeap) ) );
+	
+	return (TextureHeapSize)Lower;
+}
+
+mtlpp::Heap FMetalResourceHeap::GetTextureHeap(mtlpp::TextureDescriptor Desc, mtlpp::SizeAndAlign Size)
+{
+	mtlpp::Heap Result;
+	static bool bTextureHeaps = FParse::Param(FCommandLine::Get(),TEXT("metaltextureheaps"));
+	if (FMetalCommandQueue::SupportsFeature(EMetalFeaturesHeaps) && bTextureHeaps && Size.Size <= HeapTextureHeapSizes[MaxTextureSize])
+	{
+		FMetalResourceHeap::TextureHeapSize HeapIndex = TextureSizeToIndex(Size.Size);
+
+		EMetalHeapTextureUsage UsageIndex = EMetalHeapTextureUsageNum;
+		mtlpp::StorageMode StorageMode = Desc.GetStorageMode();
+		mtlpp::CpuCacheMode CPUMode = Desc.GetCpuCacheMode();
+		if ((Desc.GetUsage() & mtlpp::TextureUsage::RenderTarget) && StorageMode == mtlpp::StorageMode::Private && CPUMode == mtlpp::CpuCacheMode::DefaultCache)
+		{
+			UsageIndex = PLATFORM_MAC ? EMetalHeapTextureUsageNum : EMetalHeapTextureUsageRenderTarget;
+		}
+		else if (StorageMode == mtlpp::StorageMode::Private && CPUMode == mtlpp::CpuCacheMode::WriteCombined)
+		{
+			UsageIndex = EMetalHeapTextureUsageResource;
+		}
+		
+		if (UsageIndex < EMetalHeapTextureUsageNum)
+		{
+			for (mtlpp::Heap& Heap : TextureHeaps[UsageIndex][HeapIndex])
+			{
+				if (Heap.MaxAvailableSizeWithAlignment(Size.Align) >= Size.Size)
+				{
+					Result = Heap;
+					break;
+				}
+			}
+			if (!Result)
+			{
+				mtlpp::HeapDescriptor HeapDesc;
+				HeapDesc.SetSize(HeapTextureHeapSizes[HeapIndex]);
+				HeapDesc.SetStorageMode(Desc.GetStorageMode());
+				HeapDesc.SetCpuCacheMode(Desc.GetCpuCacheMode());
+				Result = Queue->GetDevice().NewHeap(HeapDesc);
+#if STATS || ENABLE_LOW_LEVEL_MEM_TRACKER
+				MetalLLM::LogAllocHeap(Queue->GetDevice(), Result);
+#endif
+				TextureHeaps[UsageIndex][HeapIndex].Add(Result);
+			}
+			check(Result);
+		}
+	}
+	return Result;
+}
+
 FMetalBuffer FMetalResourceHeap::CreateBuffer(uint32 Size, uint32 Alignment, mtlpp::ResourceOptions Options, bool bForceUnique)
 {
 	LLM_SCOPE_METAL(ELLMTagMetal::Buffers);
@@ -1440,7 +1531,18 @@ FMetalTexture FMetalResourceHeap::CreateTexture(mtlpp::TextureDescriptor Desc, F
 	LLM_SCOPE_METAL(ELLMTagMetal::Textures);
 	LLM_PLATFORM_SCOPE_METAL(ELLMTagMetal::Textures);
 	
-	if (Desc.GetUsage() & mtlpp::TextureUsage::RenderTarget)
+	mtlpp::SizeAndAlign Res = Queue->GetDevice().HeapTextureSizeAndAlign(Desc);
+	mtlpp::Heap Heap = GetTextureHeap(Desc, Res);
+	if (Heap)
+	{
+		METAL_GPUPROFILE(FScopedMetalCPUStats CPUStat(FString::Printf(TEXT("AllocTexture: %s"), TEXT("")/**FString([Desc.GetPtr() description])*/)));
+		FMetalTexture Texture = Heap.NewTexture(Desc);
+#if STATS || ENABLE_LOW_LEVEL_MEM_TRACKER
+		MetalLLM::LogAllocTexture(Queue->GetDevice(), Desc, Texture);
+#endif
+		return Texture;
+	}
+	else if (Desc.GetUsage() & mtlpp::TextureUsage::RenderTarget)
 	{
 		return TargetPool.CreateTexture(Queue->GetDevice(), Desc);
 	}
@@ -1452,15 +1554,15 @@ FMetalTexture FMetalResourceHeap::CreateTexture(mtlpp::TextureDescriptor Desc, F
 
 void FMetalResourceHeap::ReleaseTexture(FMetalSurface* Surface, FMetalTexture& Texture)
 {
-	if (!Texture.GetBuffer() && !Texture.GetParentTexture())
+	if (!Texture.GetBuffer() && !Texture.GetParentTexture() && !Texture.GetHeap())
 	{
         if (Texture.GetUsage() & mtlpp::TextureUsage::RenderTarget)
         {
-            return TargetPool.ReleaseTexture(Texture);
+           	TargetPool.ReleaseTexture(Texture);
         }
         else
         {
-            return TexturePool.ReleaseTexture(Texture);
+            TexturePool.ReleaseTexture(Texture);
         }
 	}
 }
@@ -1475,7 +1577,7 @@ void FMetalResourceHeap::Compact(bool const bForce)
 			for (auto It = SmallBuffers[t][i].CreateIterator(); It; ++It)
 			{
 				FMetalSubBufferMagazine* Data = *It;
-				if (Data->GetUsedSize() == 0 || bForce)
+				if (Data->NumCurrentAllocations() == 0 || bForce)
 				{
 					It.RemoveCurrent();
 					delete Data;
@@ -1488,7 +1590,7 @@ void FMetalResourceHeap::Compact(bool const bForce)
 			for (auto It = BufferHeaps[t][i].CreateIterator(); It; ++It)
 			{
 				FMetalSubBufferHeap* Data = *It;
-				if (Data->GetUsedSize() == 0 || bForce)
+				if (Data->NumCurrentAllocations() == 0 || bForce)
 				{
 					It.RemoveCurrent();
 					delete Data;
@@ -1496,6 +1598,20 @@ void FMetalResourceHeap::Compact(bool const bForce)
 			}
 		}
 	}
+	
+	for (uint32 i = 0; i < EMetalHeapTextureUsageNum; i++)
+	{
+		for (uint32 j = 0; j < NumTextureHeapSizes; j++)
+		{
+			for (auto It = TextureHeaps[i][j].CreateIterator(); It; ++It)
+			{
+				if (It->GetUsedSize() == 0 || bForce)
+				{
+					It.RemoveCurrent();
+				}
+			}
+		}
+	}
 
 	Buffers[AllocShared].DrainPool(bForce);
 	Buffers[AllocPrivate].DrainPool(bForce);
@@ -1587,3 +1703,13 @@ uint32 FMetalResourceHeap::HeapAllocSizes[FMetalResourceHeap::NumHeapSizes] = {
 	4194304,
 	4194304,
 };
+
+uint32 FMetalResourceHeap::HeapTextureHeapSizes[FMetalResourceHeap::NumTextureHeapSizes] = {
+	4194304,
+	8388608,
+	16777216,
+	33554432,
+	67108864,
+	134217728,
+	268435456
+};
diff --git a/Engine/Source/Runtime/Apple/MetalRHI/Private/MetalBuffer.h b/Engine/Source/Runtime/Apple/MetalRHI/Private/MetalBuffer.h
--- a/Engine/Source/Runtime/Apple/MetalRHI/Private/MetalBuffer.h
+++ b/Engine/Source/Runtime/Apple/MetalRHI/Private/MetalBuffer.h
@@ -38,6 +38,7 @@ public:
     NSUInteger     GetSize() const;
     NSUInteger     GetUsedSize() const;
 	NSUInteger	 MaxAvailableSize() const;
+	int64     NumCurrentAllocations() const;
 
     void SetLabel(const ns::String& label);
 	
@@ -47,6 +48,7 @@ public:
 
 private:
 	FCriticalSection& PoolMutex;
+	int64 volatile OutstandingAllocs;
 	NSUInteger MinAlign;
 	NSUInteger UsedSize;
 	mtlpp::Buffer ParentBuffer;
@@ -67,7 +69,7 @@ public:
 	NSUInteger     GetSize() const;
 	NSUInteger     GetUsedSize() const;
 	bool	 CanAllocateSize(NSUInteger Size) const;
-	
+
 	void SetLabel(const ns::String& label);
 	
 	FMetalBuffer NewBuffer(NSUInteger length);
@@ -96,6 +98,7 @@ public:
     NSUInteger     GetSize() const;
     NSUInteger     GetUsedSize() const;
 	NSUInteger	 GetFreeSize() const;
+	int64     NumCurrentAllocations() const;
 
     void SetLabel(const ns::String& label);
 	void FreeRange(ns::Range const& Range);
@@ -105,6 +108,7 @@ public:
 
 private:
 	NSUInteger MinAlign;
+	int64 volatile OutstandingAllocs;
 	int64 volatile UsedSize;
 	mtlpp::Buffer ParentBuffer;
 	mutable mtlpp::Heap ParentHeap;
@@ -356,6 +360,20 @@ class FMetalResourceHeap
 		Size2Mb,
 		NumHeapSizes
 	};
+
+	enum TextureHeapSize
+	{
+		Size4Mb,
+		Size8Mb,
+		Size16Mb,
+		Size32Mb,
+		Size64Mb,
+		Size128Mb,
+		Size256Mb,
+		NumTextureHeapSizes,
+		MinTexturesPerHeap = 4,
+		MaxTextureSize = Size64Mb,
+	};
 	
 	enum AllocTypes
 	{
@@ -363,6 +381,16 @@ class FMetalResourceHeap
 		AllocPrivate,
 		NumAllocTypes = 2
 	};
+	
+	enum EMetalHeapTextureUsage
+	{
+		/** Regular texture resource */
+		EMetalHeapTextureUsageResource = 0,
+		/** Render target or UAV that can be aliased */
+		EMetalHeapTextureUsageRenderTarget = 1,
+		/** Number of texture usage types */
+		EMetalHeapTextureUsageNum = 2
+	};
 
 public:
 	FMetalResourceHeap(void);
@@ -381,12 +409,16 @@ public:
 private:
 	uint32 GetMagazineIndex(uint32 Size);
 	uint32 GetHeapIndex(uint32 Size);
+	TextureHeapSize TextureSizeToIndex(uint32 Size);
+	
+	mtlpp::Heap GetTextureHeap(mtlpp::TextureDescriptor Desc, mtlpp::SizeAndAlign Size);
 	
 private:
 	static uint32 MagazineSizes[NumMagazineSizes];
 	static uint32 HeapSizes[NumHeapSizes];
 	static uint32 MagazineAllocSizes[NumMagazineSizes];
 	static uint32 HeapAllocSizes[NumHeapSizes];
+	static uint32 HeapTextureHeapSizes[NumTextureHeapSizes];
 
 	FCriticalSection Mutex;
 	FMetalCommandQueue* Queue;
@@ -409,4 +441,6 @@ private:
 	/** We can reuse texture allocations as well, to minimize their performance impact */
 	FMetalTexturePool TexturePool;
 	FMetalTexturePool TargetPool;
+	
+	TArray<mtlpp::Heap> TextureHeaps[EMetalHeapTextureUsageNum][NumTextureHeapSizes];
 };
diff --git a/Engine/Source/Runtime/Apple/MetalRHI/Private/MetalRHI.cpp b/Engine/Source/Runtime/Apple/MetalRHI/Private/MetalRHI.cpp
--- a/Engine/Source/Runtime/Apple/MetalRHI/Private/MetalRHI.cpp
+++ b/Engine/Source/Runtime/Apple/MetalRHI/Private/MetalRHI.cpp
@@ -475,7 +475,7 @@ FMetalDynamicRHI::FMetalDynamicRHI(ERHIFeatureLevel::Type RequestedFeatureLevel)
 	GSupportsRenderTargetFormat_PF_G8 = false;
 	GRHISupportsTextureStreaming = true;
 	GSupportsWideMRT = bCanUseWideMRTs;
-
+	GSupportsTransientResourceAliasing = FMetalCommandQueue::SupportsFeature(EMetalFeaturesHeaps) && FMetalCommandQueue::SupportsFeature(EMetalFeaturesFences);
 	GSupportsSeparateRenderTargetBlendState = (GMaxRHIFeatureLevel >= ERHIFeatureLevel::SM4);
 
 #if PLATFORM_MAC
diff --git a/Engine/Source/Runtime/Apple/MetalRHI/Private/MetalTexture.cpp b/Engine/Source/Runtime/Apple/MetalRHI/Private/MetalTexture.cpp
--- a/Engine/Source/Runtime/Apple/MetalRHI/Private/MetalTexture.cpp
+++ b/Engine/Source/Runtime/Apple/MetalRHI/Private/MetalTexture.cpp
@@ -86,7 +86,7 @@ FMetalSurface* GetMetalSurfaceFromRHITexture(FRHITexture* Texture)
 
 static bool IsRenderTarget(uint32 Flags)
 {
-	return (Flags & (TexCreate_RenderTargetable | TexCreate_ResolveTargetable | TexCreate_DepthStencilTargetable)) != 0;
+	return (Flags & (TexCreate_RenderTargetable | TexCreate_ResolveTargetable | TexCreate_DepthStencilTargetable | TexCreate_DepthStencilResolveTarget)) != 0;
 }
 
 static mtlpp::TextureUsage ConvertFlagsToUsage(uint32 Flags)
@@ -115,7 +115,7 @@ static mtlpp::TextureUsage ConvertFlagsToUsage(uint32 Flags)
 	//are likely to be used in a manual shader resolve by the high level and must be bindable as rendertargets.
 	const bool bSeparateResolveTargets = FMetalCommandQueue::SupportsSeparateMSAAAndResolveTarget();
 	const bool bResolveTarget = (Flags & TexCreate_ResolveTargetable);
-	if ((Flags & (TexCreate_RenderTargetable|TexCreate_DepthStencilTargetable)) || (bResolveTarget && bSeparateResolveTargets))
+	if ((Flags & (TexCreate_RenderTargetable|TexCreate_DepthStencilTargetable|TexCreate_DepthStencilResolveTarget)) || (bResolveTarget && bSeparateResolveTargets))
 	{
 		Usage |= mtlpp::TextureUsage::RenderTarget;
 		Usage |= mtlpp::TextureUsage::ShaderRead;
@@ -415,7 +415,7 @@ void FMetalSurface::MakeUnAliasable(void)
 	check(ImageSurfaceRef == nullptr);
 	
 	static bool bSupportsHeaps = GetMetalDeviceContext().SupportsFeature(EMetalFeaturesHeaps);
-	if (bSupportsHeaps && Texture.GetStorageMode() == mtlpp::StorageMode::Private && Texture.GetHeap())
+	if (bSupportsHeaps && Texture.GetStorageMode() == mtlpp::StorageMode::Private && Texture.GetHeap() && Texture.IsAliasable())
 	{
 		FMetalTexture OldTexture = Texture;
 		Texture = Reallocate(Texture, mtlpp::TextureUsage::Unknown);
@@ -871,7 +871,7 @@ FMetalSurface::FMetalSurface(ERHIResourceType ResourceType, EPixelFormat Format,
 			Desc.SetResourceOptions((mtlpp::ResourceOptions)(mtlpp::ResourceOptions::CpuCacheModeDefaultCache|mtlpp::ResourceOptions::StorageModeShared));
 #endif
 		}
-		else if (Flags & (TexCreate_RenderTargetable|TexCreate_DepthStencilTargetable))
+		else if (Flags & (TexCreate_RenderTargetable|TexCreate_DepthStencilTargetable|TexCreate_ResolveTargetable|TexCreate_DepthStencilResolveTarget))
 		{
 			check(!(Flags & TexCreate_CPUReadback));
 #if PLATFORM_IOS
@@ -1879,6 +1879,9 @@ struct FMetalRHICommandAsyncReallocateTexture2D final : public FRHICommand<FMeta
 		
 		// Like D3D mark this as complete immediately.
 		RequestStatus->Decrement();
+		
+		FMetalSurface* Source = GetMetalSurfaceFromRHITexture(OldTexture);
+		Source->MakeAliasable();
 	}
 };
 
@@ -2536,14 +2539,133 @@ struct FMetalRHICommandUnaliasTextures final : public FRHICommand<FMetalRHIComma
 	
 	void Execute(FRHICommandListBase& CmdList)
 	{
+		@autoreleasepool {
 		for (int32 i = 0; i < Textures.Num(); ++i)
 		{
 			FMetalSurface* Source = GetMetalSurfaceFromRHITexture(Textures[i]);
 			Source->MakeUnAliasable();
 		}
+		}
+	}
+};
+
+void FMetalDynamicRHI::RHIAcquireTransientResource_RenderThread(FTextureRHIParamRef Texture)
+{
+	@autoreleasepool {
+	FRHICommandListImmediate& RHICmdList = FRHICommandListExecutor::GetImmediateCommandList();
+	if (RHICmdList.Bypass() || !IsRunningRHIInSeparateThread())
+	{
+		FMetalSurface* Source = GetMetalSurfaceFromRHITexture(Texture);
+		Source->MakeUnAliasable();
+	}
+	else
+	{
+		new (RHICmdList.AllocCommand<FMetalRHICommandUnaliasTextures>()) FMetalRHICommandUnaliasTextures(&Texture, 1);
+	}
+	}
+}
+
+struct FMetalRHICommandAliasTextures final : public FRHICommand<FMetalRHICommandAliasTextures>
+{
+	TArray<FTextureRHIParamRef> Textures;
+	
+	FORCEINLINE_DEBUGGABLE FMetalRHICommandAliasTextures(FTextureRHIParamRef* InTextures, int32 NumTextures)
+	{
+		check(InTextures && NumTextures);
+		Textures.Append(InTextures, NumTextures);
+	}
+	
+	void Execute(FRHICommandListBase& CmdList)
+	{
+		@autoreleasepool {
+		for (int32 i = 0; i < Textures.Num(); ++i)
+		{
+			FMetalSurface* Source = GetMetalSurfaceFromRHITexture(Textures[i]);
+			Source->MakeAliasable();
+		}
+		}
 	}
 };
 
+void FMetalDynamicRHI::RHIDiscardTransientResource_RenderThread(FTextureRHIParamRef Texture)
+{
+	@autoreleasepool {
+	FRHICommandListImmediate& RHICmdList = FRHICommandListExecutor::GetImmediateCommandList();
+	if (RHICmdList.Bypass() || !IsRunningRHIInSeparateThread())
+	{
+		FMetalSurface* Source = GetMetalSurfaceFromRHITexture(Texture);
+		Source->MakeAliasable();
+	}
+	else
+	{
+		new (RHICmdList.AllocCommand<FMetalRHICommandAliasTextures>()) FMetalRHICommandAliasTextures(&Texture, 1);
+	}
+	}
+}
+
+struct FMetalRHICommandAliasBuffer final : public FRHICommand<FMetalRHICommandAliasBuffer>
+{
+	FMetalRHIBuffer* Buffer;
+	
+	FORCEINLINE_DEBUGGABLE FMetalRHICommandAliasBuffer(FMetalRHIBuffer* InBuffer)
+	{
+		check(InBuffer);
+		Buffer = InBuffer;
+	}
+	
+	void Execute(FRHICommandListBase& CmdList)
+	{
+		@autoreleasepool {
+		Buffer->Unalias();
+		}
+	}
+};
+
+void FMetalDynamicRHI::RHIAcquireTransientResource_RenderThread(FVertexBufferRHIParamRef Buffer)
+{
+	@autoreleasepool {
+	FMetalVertexBuffer* MetalBuffer = ResourceCast(Buffer);
+	FRHICommandListImmediate& RHICmdList = FRHICommandListExecutor::GetImmediateCommandList();
+	if (RHICmdList.Bypass() || !IsRunningRHIInSeparateThread())
+	{
+		MetalBuffer->Unalias();
+	}
+	else
+	{
+		new (RHICmdList.AllocCommand<FMetalRHICommandAliasBuffer>()) FMetalRHICommandAliasBuffer(MetalBuffer);
+	}
+	}
+}
+void FMetalDynamicRHI::RHIDiscardTransientResource_RenderThread(FVertexBufferRHIParamRef Buffer)
+{
+	@autoreleasepool {
+	FMetalVertexBuffer* MetalBuffer = ResourceCast(Buffer);
+	MetalBuffer->Alias();
+	}
+}
+void FMetalDynamicRHI::RHIAcquireTransientResource_RenderThread(FStructuredBufferRHIParamRef Buffer)
+{
+	@autoreleasepool {
+	FMetalStructuredBuffer* MetalBuffer = ResourceCast(Buffer);
+	FRHICommandListImmediate& RHICmdList = FRHICommandListExecutor::GetImmediateCommandList();
+	if (RHICmdList.Bypass() || !IsRunningRHIInSeparateThread())
+	{
+		MetalBuffer->Unalias();
+	}
+	else
+	{
+		new (RHICmdList.AllocCommand<FMetalRHICommandAliasBuffer>()) FMetalRHICommandAliasBuffer(MetalBuffer);
+	}
+	}
+}
+void FMetalDynamicRHI::RHIDiscardTransientResource_RenderThread(FStructuredBufferRHIParamRef Buffer)
+{
+	@autoreleasepool {
+	FMetalStructuredBuffer* MetalBuffer = ResourceCast(Buffer);
+	MetalBuffer->Alias();
+	}
+}
+
 void FMetalDynamicRHI::RHISetResourceAliasability_RenderThread(class FRHICommandListImmediate& RHICmdList, EResourceAliasability AliasMode, FTextureRHIParamRef* InTextures, int32 NumTextures)
 {
 	check(InTextures && NumTextures);
diff --git a/Engine/Source/Runtime/Apple/MetalRHI/Private/MetalVertexBuffer.cpp b/Engine/Source/Runtime/Apple/MetalRHI/Private/MetalVertexBuffer.cpp
--- a/Engine/Source/Runtime/Apple/MetalRHI/Private/MetalVertexBuffer.cpp
+++ b/Engine/Source/Runtime/Apple/MetalRHI/Private/MetalVertexBuffer.cpp
@@ -202,7 +202,28 @@ FMetalRHIBuffer::~FMetalRHIBuffer()
 		SafeReleaseMetalObject(Data);
 	}
 }
-	
+
+void FMetalRHIBuffer::Alias()
+{
+	if (Buffer.GetStorageMode() == mtlpp::StorageMode::Private && Buffer.GetHeap() && !Buffer.IsAliasable())
+	{
+		Buffer.MakeAliasable();
+	}
+}
+
+void FMetalRHIBuffer::Unalias()
+{
+	if (Buffer.GetStorageMode() == mtlpp::StorageMode::Private && Buffer.GetHeap() && Buffer.IsAliasable())
+	{
+		uint32 Len = Buffer.GetLength();
+		METAL_INC_DWORD_STAT_BY(Type, MemFreed, Len);
+		SafeReleaseMetalBuffer(Buffer);
+		Buffer = nil;
+		
+		Alloc(Len, RLM_WriteOnly);
+	}
+}
+
 void FMetalRHIBuffer::Alloc(uint32 InSize, EResourceLockMode LockMode)
 {
 	METAL_LLM_BUFFER_SCOPE(Type);
@@ -389,6 +410,8 @@ void* FMetalRHIBuffer::Lock(EResourceLockMode LockMode, uint32 Offset, uint32 In
 		return ((uint8*)Data->Data) + Offset;
 	}
 	
+	check(!Buffer.IsAliasable());
+	
     uint32 Len = Buffer.GetLength();
     
 	// In order to properly synchronise the buffer access, when a dynamic buffer is locked for writing, discard the old buffer & create a new one. This prevents writing to a buffer while it is being read by the GPU & thus causing corruption. This matches the logic of other RHIs.
diff --git a/Engine/Source/Runtime/Apple/MetalRHI/Public/MetalRHI.h b/Engine/Source/Runtime/Apple/MetalRHI/Public/MetalRHI.h
--- a/Engine/Source/Runtime/Apple/MetalRHI/Public/MetalRHI.h
+++ b/Engine/Source/Runtime/Apple/MetalRHI/Public/MetalRHI.h
@@ -196,6 +196,12 @@ public:
 	virtual void UnlockStagingBuffer_RenderThread(class FRHICommandListImmediate& RHICmdList, FStagingBufferRHIParamRef StagingBuffer) final override;	
 	virtual void RHIEnqueueStagedRead(FStagingBufferRHIParamRef StagingBuffer, FGPUFenceRHIParamRef Fence, uint32 Offset, uint32 NumBytes) final override;
 	
+	virtual void RHIAcquireTransientResource_RenderThread(FTextureRHIParamRef Texture) final override;
+	virtual void RHIDiscardTransientResource_RenderThread(FTextureRHIParamRef Texture) final override;
+	virtual void RHIAcquireTransientResource_RenderThread(FVertexBufferRHIParamRef Buffer) final override;
+	virtual void RHIDiscardTransientResource_RenderThread(FVertexBufferRHIParamRef Buffer) final override;
+	virtual void RHIAcquireTransientResource_RenderThread(FStructuredBufferRHIParamRef Buffer) final override;
+	virtual void RHIDiscardTransientResource_RenderThread(FStructuredBufferRHIParamRef Buffer) final override;
 	virtual void RHISetResourceAliasability_RenderThread(class FRHICommandListImmediate& RHICmdList, EResourceAliasability AliasMode, FTextureRHIParamRef* InTextures, int32 NumTextures) final override;
 	
 	/* Copy the source box pixels in the destination box texture, return true if implemented for the current platform*/
diff --git a/Engine/Source/Runtime/Apple/MetalRHI/Public/MetalResources.h b/Engine/Source/Runtime/Apple/MetalRHI/Public/MetalResources.h
--- a/Engine/Source/Runtime/Apple/MetalRHI/Public/MetalResources.h
+++ b/Engine/Source/Runtime/Apple/MetalRHI/Public/MetalResources.h
@@ -746,6 +746,16 @@ public:
 	FMetalRHIBuffer(uint32 InSize, uint32 InUsage, ERHIResourceType InType);
 	virtual ~FMetalRHIBuffer();
 	
+	/**
+	 * Alias the buffer backing store allowing the memory to be reused by another resource.
+	 */
+	void Alias();
+
+	/**
+	 * Unalias the buffer backing store forcing the memory to be reserved for use by this resource.
+	 */
+	void Unalias();
+	
 	/**
 	 * Allocate the index buffer backing store.
 	 */
