For the purposes of memory reporting treat Metal resource aliasing as deallocation.

diff --git a/Engine/Source/Runtime/Apple/MetalRHI/Private/MetalLLM.cpp b/Engine/Source/Runtime/Apple/MetalRHI/Private/MetalLLM.cpp
--- a/Engine/Source/Runtime/Apple/MetalRHI/Private/MetalLLM.cpp
+++ b/Engine/Source/Runtime/Apple/MetalRHI/Private/MetalLLM.cpp
@@ -284,3 +284,14 @@ void MetalLLM::LogAllocHeap(mtlpp::Device& Device, mtlpp::Heap const& Heap)
 	}
 }
 
+void MetalLLM::LogAliasTexture(mtlpp::Texture const& Texture)
+{
+	objc_setAssociatedObject(Texture.GetPtr(), (void*)&MetalLLM::LogAllocTexture, nil, OBJC_ASSOCIATION_RETAIN);
+}
+
+void MetalLLM::LogAliasBuffer(mtlpp::Buffer const& Buffer)
+{
+	objc_setAssociatedObject(Buffer.GetPtr(), (void*)&MetalLLM::LogAllocBuffer, nil, OBJC_ASSOCIATION_RETAIN);
+}
+
+
diff --git a/Engine/Source/Runtime/Apple/MetalRHI/Private/MetalLLM.h b/Engine/Source/Runtime/Apple/MetalRHI/Private/MetalLLM.h
--- a/Engine/Source/Runtime/Apple/MetalRHI/Private/MetalLLM.h
+++ b/Engine/Source/Runtime/Apple/MetalRHI/Private/MetalLLM.h
@@ -47,6 +47,9 @@ namespace MetalLLM
 	void LogAllocTexture(mtlpp::Device& Device, mtlpp::TextureDescriptor const& Desc, mtlpp::Texture const& Texture);
 	void LogAllocBuffer(mtlpp::Device& Device, mtlpp::Buffer const& Buffer);
 	void LogAllocHeap(mtlpp::Device& Device, mtlpp::Heap const& Heap);
+	
+	void LogAliasTexture(mtlpp::Texture const& Texture);
+	void LogAliasBuffer(mtlpp::Buffer const& Buffer);
 }
 
 
diff --git a/Engine/Source/Runtime/Apple/MetalRHI/Private/MetalTexture.cpp b/Engine/Source/Runtime/Apple/MetalRHI/Private/MetalTexture.cpp
--- a/Engine/Source/Runtime/Apple/MetalRHI/Private/MetalTexture.cpp
+++ b/Engine/Source/Runtime/Apple/MetalRHI/Private/MetalTexture.cpp
@@ -10,6 +10,7 @@
 #include "RenderUtils.h"
 #include "Containers/ResourceArray.h"
 #include "Misc/ScopeRWLock.h"
+#include "MetalLLM.h"
 
 volatile int64 FMetalSurface::ActiveUploads = 0;
 
@@ -397,14 +398,23 @@ void FMetalSurface::MakeAliasable(void)
 		if (StencilTexture && (StencilTexture != Texture) && !StencilTexture.IsAliasable())
 		{
 			StencilTexture.MakeAliasable();
+#if STATS || ENABLE_LOW_LEVEL_MEM_TRACKER
+			MetalLLM::LogAliasTexture(StencilTexture);
+#endif
 		}
 		if (MSAATexture && (MSAATexture != Texture) && !MSAATexture.IsAliasable())
 		{
 			MSAATexture.MakeAliasable();
+#if STATS || ENABLE_LOW_LEVEL_MEM_TRACKER
+			MetalLLM::LogAliasTexture(MSAATexture);
+#endif
 		}
 		if (!Texture.IsAliasable())
 		{
 			Texture.MakeAliasable();
+#if STATS || ENABLE_LOW_LEVEL_MEM_TRACKER
+			MetalLLM::LogAliasTexture(Texture);
+#endif
 		}
 	}
 }
diff --git a/Engine/Source/Runtime/Apple/MetalRHI/Private/MetalVertexBuffer.cpp b/Engine/Source/Runtime/Apple/MetalRHI/Private/MetalVertexBuffer.cpp
--- a/Engine/Source/Runtime/Apple/MetalRHI/Private/MetalVertexBuffer.cpp
+++ b/Engine/Source/Runtime/Apple/MetalRHI/Private/MetalVertexBuffer.cpp
@@ -208,6 +208,9 @@ void FMetalRHIBuffer::Alias()
 	if (Buffer.GetStorageMode() == mtlpp::StorageMode::Private && Buffer.GetHeap() && !Buffer.IsAliasable())
 	{
 		Buffer.MakeAliasable();
+#if STATS || ENABLE_LOW_LEVEL_MEM_TRACKER
+		MetalLLM::LogAliasBuffer(Buffer);
+#endif
 	}
 }
 
