Fix mismerge.

diff --git a/Engine/Source/Runtime/Apple/MetalRHI/Private/MetalRHI.cpp b/Engine/Source/Runtime/Apple/MetalRHI/Private/MetalRHI.cpp
--- a/Engine/Source/Runtime/Apple/MetalRHI/Private/MetalRHI.cpp
+++ b/Engine/Source/Runtime/Apple/MetalRHI/Private/MetalRHI.cpp
@@ -632,7 +632,7 @@ FMetalDynamicRHI::FMetalDynamicRHI(ERHIFeatureLevel::Type RequestedFeatureLevel)
 	GPixelFormats[PF_PLATFORM_HDR_0		].Supported			= GRHISupportsHDROutput;
 		
 #if PLATFORM_TVOS
-	if (![Device supportsFeatureSet:MTLFeatureSet_tvOS_GPUFamily2_v1])
+    if (!Device.SupportsFeatureSet(mtlpp::FeatureSet::tvOS_GPUFamily2_v1))
 #else
 	if (!Device.SupportsFeatureSet(mtlpp::FeatureSet::iOS_GPUFamily3_v2))
 	{
@@ -649,6 +649,7 @@ FMetalDynamicRHI::FMetalDynamicRHI(ERHIFeatureLevel::Type RequestedFeatureLevel)
 		GPixelFormats[PF_FloatR11G11B10		].BlockBytes		= 4;
 		GPixelFormats[PF_FloatR11G11B10		].Supported			= true;
 	}
+#endif
 	
 	if (FMetalCommandQueue::SupportsFeature(EMetalFeaturesStencilView) && FMetalCommandQueue::SupportsFeature(EMetalFeaturesCombinedDepthStencil) && !FParse::Param(FCommandLine::Get(),TEXT("metalforceseparatedepthstencil")))
 	{
