NV feel that uniform buffers would be better off in Managed memory or via set*Bytes than using manual Shared->Private blits - there's no disadvantage to this AFAICT.

diff --git a/Engine/Source/Runtime/Apple/MetalRHI/Private/MetalUniformBuffer.cpp b/Engine/Source/Runtime/Apple/MetalRHI/Private/MetalUniformBuffer.cpp
--- a/Engine/Source/Runtime/Apple/MetalRHI/Private/MetalUniformBuffer.cpp
+++ b/Engine/Source/Runtime/Apple/MetalRHI/Private/MetalUniformBuffer.cpp
@@ -38,15 +38,9 @@ struct FMetalRHICommandInitialiseUniformdBuffer : public FRHICommand<FMetalRHICo
 	}
 };
 
-static EBufferUsageFlags MetalUniformBufferUsage(uint32 Size, EUniformBufferUsage InUsage)
-{
-	EBufferUsageFlags Result = (!(InUsage & UniformBuffer_SingleDraw) && Size >= MetalBufferBytesSize && (IsInRHIThread() || IsInRenderingThread()) ? BUF_Static : BUF_Volatile);
-	return Result;
-}
-
 FMetalUniformBuffer::FMetalUniformBuffer(const void* Contents, const FRHIUniformBufferLayout& Layout, EUniformBufferUsage InUsage)
 	: FRHIUniformBuffer(Layout)
-	, FMetalRHIBuffer(Layout.ConstantBufferSize, MetalUniformBufferUsage(Layout.ConstantBufferSize, InUsage), RRT_UniformBuffer)
+	, FMetalRHIBuffer(Layout.ConstantBufferSize, BUF_Volatile, RRT_UniformBuffer)
 {
 	if (Layout.ConstantBufferSize > 0)
 	{
