Add command-line args "nometalparallelencoder" to disable parallel encoders, "nometalfence" and "nometalheap" to disable MTLFence and MTLHeap respectively.

diff --git a/Engine/Source/Runtime/Apple/MetalRHI/Private/MetalCommandQueue.cpp b/Engine/Source/Runtime/Apple/MetalRHI/Private/MetalCommandQueue.cpp
--- a/Engine/Source/Runtime/Apple/MetalRHI/Private/MetalCommandQueue.cpp
+++ b/Engine/Source/Runtime/Apple/MetalRHI/Private/MetalCommandQueue.cpp
@@ -210,7 +210,7 @@ FMetalCommandQueue::FMetalCommandQueue(mtlpp::Device InDevice, uint32 const MaxN
 		if (FPlatformMisc::MacOSXVersionCompare(10,13,5) >= 0)
 		{
 			// Except on Nvidia for the moment
-			if ([Device.GetName().GetPtr() rangeOfString:@"Nvidia" options:NSCaseInsensitiveSearch].location == NSNotFound)
+			if ([Device.GetName().GetPtr() rangeOfString:@"Nvidia" options:NSCaseInsensitiveSearch].location == NSNotFound && !FParse::Param(FCommandLine::Get(),TEXT("nometalparallelencoder")))
 			{
 				Features |= EMetalFeaturesParallelRenderEncoders;
 			}
@@ -225,13 +225,18 @@ FMetalCommandQueue::FMetalCommandQueue(mtlpp::Device InDevice, uint32 const MaxN
 		// Turn on Texture Buffers! These are faster on the GPU as we don't need to do out-of-bounds tests but require Metal 2.1 and macOS 10.14
 		if (FPlatformMisc::MacOSXVersionCompare(10,14,0) >= 0)
 		{
-			Features |= EMetalFeaturesMaxThreadsPerThreadgroup | EMetalFeaturesFences | EMetalFeaturesIABs;
+			Features |= EMetalFeaturesMaxThreadsPerThreadgroup | EMetalFeaturesIABs;
 			if (MaxShaderVersion >= 4)
 			{
 				Features |= EMetalFeaturesTextureBuffers;
 			}
 			
-			if ([Device.GetName().GetPtr() rangeOfString:@"Intel" options:NSCaseInsensitiveSearch].location == NSNotFound)
+			if (!FParse::Param(FCommandLine::Get(),TEXT("nometalfence")))
+			{
+				Features |= EMetalFeaturesFences;
+			}
+			
+			if (!FParse::Param(FCommandLine::Get(),TEXT("nometalheap")) && [Device.GetName().GetPtr() rangeOfString:@"Intel" options:NSCaseInsensitiveSearch].location == NSNotFound)
 			{
 				Features |= EMetalFeaturesHeaps;
 			}
