UE-66189 MSAA 8X + DOF crashes the app in iOS

#iOS
#4.22

Limit MSAA to the maximum returned by the supportsTextureSampleCount function.

diff --git a/Engine/Source/Runtime/Apple/MetalRHI/Private/MetalRHI.cpp b/Engine/Source/Runtime/Apple/MetalRHI/Private/MetalRHI.cpp
--- a/Engine/Source/Runtime/Apple/MetalRHI/Private/MetalRHI.cpp
+++ b/Engine/Source/Runtime/Apple/MetalRHI/Private/MetalRHI.cpp
@@ -18,6 +18,7 @@
 #include "GenericPlatform/GenericPlatformDriver.h"
 #include "MetalShaderResources.h"
 #include "MetalLLM.h"
+#include "Engine/RendererSettings.h"
 
 DEFINE_LOG_CATEGORY(LogMetal)
 
@@ -1154,3 +1155,26 @@ void* FMetalDynamicRHI::RHIGetNativeDevice()
 {
 	return (void*)ImmediateContext.Context->GetDevice().GetPtr();
 }
+
+uint16 FMetalDynamicRHI::RHIGetPlatformTextureMaxSampleCount()
+{
+	TArray<EMobileMSAASampleCount::Type> SamplesArray{ EMobileMSAASampleCount::Type::One, EMobileMSAASampleCount::Type::Two, EMobileMSAASampleCount::Type::Four, EMobileMSAASampleCount::Type::Eight };
+
+	uint16 PlatformMaxSampleCount = EMobileMSAASampleCount::Type::One;
+	for (auto sampleIt = SamplesArray.CreateConstIterator(); sampleIt; ++sampleIt)
+	{
+		int sample = *sampleIt;
+
+#if PLATFORM_IOS || PLATFORM_MAC
+		id<MTLDevice> Device = (id<MTLDevice>)RHIGetNativeDevice();
+		check(Device);
+
+		if (![Device supportsTextureSampleCount : sample])
+		{
+			break;
+		}
+		PlatformMaxSampleCount = sample;
+#endif
+	}
+	return PlatformMaxSampleCount;
+}
diff --git a/Engine/Source/Runtime/Apple/MetalRHI/Public/MetalRHI.h b/Engine/Source/Runtime/Apple/MetalRHI/Public/MetalRHI.h
--- a/Engine/Source/Runtime/Apple/MetalRHI/Public/MetalRHI.h
+++ b/Engine/Source/Runtime/Apple/MetalRHI/Public/MetalRHI.h
@@ -218,6 +218,9 @@ public:
 	virtual FRHIShaderLibraryRef RHICreateShaderLibrary_RenderThread(class FRHICommandListImmediate& RHICmdList, EShaderPlatform Platform, FString FilePath, FString Name) final override;
 
 	virtual void RHICalibrateTimers() final override;
+
+	virtual uint16 RHIGetPlatformTextureMaxSampleCount() override;
+
 private:
 	FTextureMemoryStats MemoryStats;
 	FMetalRHIImmediateCommandContext ImmediateContext;
diff --git a/Engine/Source/Runtime/RHI/Public/DynamicRHI.h b/Engine/Source/Runtime/RHI/Public/DynamicRHI.h
--- a/Engine/Source/Runtime/RHI/Public/DynamicRHI.h
+++ b/Engine/Source/Runtime/RHI/Public/DynamicRHI.h
@@ -1055,6 +1055,9 @@ public:
 
 	virtual bool RHIIsTypedUAVLoadSupported(EPixelFormat PixelFormat) { return true; }
 
+	virtual uint16 RHIGetPlatformTextureMaxSampleCount() { return 8; };
+
+
 protected:
 	TArray<uint32> PixelFormatBlockBytes;
 };
diff --git a/Engine/Source/Runtime/Renderer/Private/PostProcess/SceneRenderTargets.cpp b/Engine/Source/Runtime/Renderer/Private/PostProcess/SceneRenderTargets.cpp
--- a/Engine/Source/Runtime/Renderer/Private/PostProcess/SceneRenderTargets.cpp
+++ b/Engine/Source/Runtime/Renderer/Private/PostProcess/SceneRenderTargets.cpp
@@ -459,6 +459,10 @@ uint16 FSceneRenderTargets::GetNumSceneColorMSAASamples(ERHIFeatureLevel::Type I
 	else
 	{
 		NumSamples = CVarMobileMSAA.GetValueOnRenderThread();
+
+		static uint16 PlatformMaxSampleCount = GDynamicRHI->RHIGetPlatformTextureMaxSampleCount();
+		NumSamples = FMath::Min(NumSamples, PlatformMaxSampleCount);
+		
 		if (NumSamples != 1 && NumSamples != 2 && NumSamples != 4 && NumSamples != 8)
 		{
 			UE_LOG(LogRenderer, Warning, TEXT("Requested %d samples for MSAA, but this is not supported; falling back to 1 sample"), NumSamples);
