Merging arkit fix using ARKit20_to_Dev-VR

diff --git a/Engine/Source/Runtime/Apple/MetalRHI/Public/MetalResources.h b/Engine/Source/Runtime/Apple/MetalRHI/Public/MetalResources.h
--- a/Engine/Source/Runtime/Apple/MetalRHI/Public/MetalResources.h
+++ b/Engine/Source/Runtime/Apple/MetalRHI/Public/MetalResources.h
@@ -645,6 +645,11 @@ public:
 	{
 		return &Surface;
 	}
+
+	virtual void* GetNativeResource() const override final
+	{
+		return Surface.Texture;
+	}
 };
 
 struct FMetalCommandBufferFence
