Refactoring some awkward stuff in MetalBackend and ue4_stdlib.
- Only buffers and textures should be allocated a resource index - samplers, input structures etc. should not.
- Make the buffer-side table a simple uint array with interleaved buffer sizes and formats, rather than pushing the formats on the end, this makes future work much simpler and decouples the definition from the number of buffers.

diff --git a/Engine/Source/Developer/Apple/MetalShaderFormat/Private/MetalBackend.cpp b/Engine/Source/Developer/Apple/MetalShaderFormat/Private/MetalBackend.cpp
--- a/Engine/Source/Developer/Apple/MetalShaderFormat/Private/MetalBackend.cpp
+++ b/Engine/Source/Developer/Apple/MetalShaderFormat/Private/MetalBackend.cpp
@@ -2114,7 +2114,7 @@ protected:
 					ralloc_asprintf_append(buffer, ", ");
 					tex->coordinate->accept(this);
 					if (bInsertSideTable)
-						ralloc_asprintf_append(buffer, ", GMetalTypedBufferFormat%d, (constant buffer_meta_table*)BufferSizes)", Index);
+						ralloc_asprintf_append(buffer, ", GMetalTypedBufferFormat%d, BufferSizes)", Index);
 					else
 						ralloc_asprintf_append(buffer, ", GMetalTypedBufferFormat%d)", Index);
 				}
@@ -2131,7 +2131,7 @@ protected:
                         ralloc_asprintf_append(buffer, ", ");
                         tex->coordinate->accept(this);
 						if (bInsertSideTable)
-	                        ralloc_asprintf_append(buffer, ", (constant buffer_meta_table*)BufferSizes)");
+	                        ralloc_asprintf_append(buffer, ", BufferSizes)");
 					}
                     else
                     {
@@ -2476,7 +2476,7 @@ protected:
 						ralloc_asprintf_append(buffer, ", ");
 						deref->image_index->accept(this);
 						if (bInsertSideTable)
-							ralloc_asprintf_append(buffer, ", GMetalTypedBufferFormat%d, (constant buffer_meta_table*)BufferSizes)", Index);
+							ralloc_asprintf_append(buffer, ", GMetalTypedBufferFormat%d, BufferSizes)", Index);
 						else
 							ralloc_asprintf_append(buffer, ", GMetalTypedBufferFormat%d)", Index);
 					}
@@ -2492,7 +2492,7 @@ protected:
                             ralloc_asprintf_append(buffer, ", ");
                             deref->image_index->accept(this);
 							if (bInsertSideTable)
-	                            ralloc_asprintf_append(buffer, ", (constant buffer_meta_table*)BufferSizes)");
+	                            ralloc_asprintf_append(buffer, ", BufferSizes)");
                         }
                         else
                         {
@@ -2636,7 +2636,7 @@ protected:
 						ralloc_asprintf_append(buffer, ", ");
 						deref->image_index->accept(this);
 						if (bInsertSideTable)
-							ralloc_asprintf_append(buffer, ", GMetalTypedBufferFormat%d, (constant buffer_meta_table*)BufferSizes, ", Index);
+							ralloc_asprintf_append(buffer, ", GMetalTypedBufferFormat%d, BufferSizes, ", Index);
 						else
 							ralloc_asprintf_append(buffer, ", GMetalTypedBufferFormat%d, ", Index);
 						src->accept(this);
@@ -2654,7 +2654,7 @@ protected:
                         ralloc_asprintf_append(buffer, ", ");
                         deref->image_index->accept(this);
 						if (bInsertSideTable)
-	                        ralloc_asprintf_append(buffer, ", (constant buffer_meta_table*)BufferSizes, ");
+	                        ralloc_asprintf_append(buffer, ", BufferSizes, ");
 						else
 	                        ralloc_asprintf_append(buffer, ", ");
                         src->accept(this);
diff --git a/Engine/Source/Developer/Apple/MetalShaderFormat/Private/MetalUtils.cpp b/Engine/Source/Developer/Apple/MetalShaderFormat/Private/MetalUtils.cpp
--- a/Engine/Source/Developer/Apple/MetalShaderFormat/Private/MetalUtils.cpp
+++ b/Engine/Source/Developer/Apple/MetalShaderFormat/Private/MetalUtils.cpp
@@ -1249,13 +1249,21 @@ void FMetalCodeBackend::MovePackedUniformsToMain(exec_list* ir, _mesa_glsl_parse
 					check(false);
 					break;
 			}
-			if (bIsBuffer)
+			if (Var->type->is_sampler() || Var->type->is_image() || Var->type->is_record() || Var->mode == ir_var_uniform)
 			{
-				OutBuffers.AddBuffer(Var);
+				if (bIsBuffer)
+				{
+					OutBuffers.AddBuffer(Var);
+				}
+				else
+				{
+					OutBuffers.AddTexture(Var);
+				}
 			}
 			else
 			{
-				OutBuffers.AddTexture(Var);
+				Var->remove();
+				MainSig->parameters.push_tail(Var);
 			}
 		}
 	}
diff --git a/Engine/Source/Runtime/Apple/MetalRHI/Private/MetalCommandEncoder.cpp b/Engine/Source/Runtime/Apple/MetalRHI/Private/MetalCommandEncoder.cpp
--- a/Engine/Source/Runtime/Apple/MetalRHI/Private/MetalCommandEncoder.cpp
+++ b/Engine/Source/Runtime/Apple/MetalRHI/Private/MetalCommandEncoder.cpp
@@ -1256,7 +1256,7 @@ void FMetalCommandEncoder::SetShaderBuffer(mtlpp::FunctionType const FunctionTyp
     if(GetMetalDeviceContext().SupportsFeature(EMetalFeaturesSetBufferOffset) && Buffer && (ShaderBuffers[uint32(FunctionType)].Bound & (1 << index)) && ShaderBuffers[uint32(FunctionType)].Buffers[index] == Buffer)
     {
 		SetShaderBufferOffset(FunctionType, Offset, Length, index);
-		ShaderBuffers[uint32(FunctionType)].Lengths[index+ML_MaxBuffers] = GMetalBufferFormats[Format].DataFormat;
+		ShaderBuffers[uint32(FunctionType)].Lengths[(index*2)+1] = GMetalBufferFormats[Format].DataFormat;
 		ShaderBuffers[uint32(FunctionType)].Usage[index] = Usage;
 		
 		UseResource(Buffer, Usage);
@@ -1279,8 +1279,8 @@ void FMetalCommandEncoder::SetShaderBuffer(mtlpp::FunctionType const FunctionTyp
 		ShaderBuffers[uint32(FunctionType)].Bytes[index] = nil;
 		ShaderBuffers[uint32(FunctionType)].Offsets[index] = Offset;
 		ShaderBuffers[uint32(FunctionType)].Usage[index] = Usage;
-		ShaderBuffers[uint32(FunctionType)].Lengths[index] = Length;
-		ShaderBuffers[uint32(FunctionType)].Lengths[index+ML_MaxBuffers] = GMetalBufferFormats[Format].DataFormat;
+		ShaderBuffers[uint32(FunctionType)].Lengths[index*2] = Length;
+		ShaderBuffers[uint32(FunctionType)].Lengths[(index*2)+1] = GMetalBufferFormats[Format].DataFormat;
 		
 		SetShaderBufferInternal(FunctionType, index);
     }
@@ -1315,8 +1315,8 @@ void FMetalCommandEncoder::SetShaderData(mtlpp::FunctionType const FunctionType,
 	ShaderBuffers[uint32(FunctionType)].Bytes[Index] = Data;
 	ShaderBuffers[uint32(FunctionType)].Offsets[Index] = Offset;
 	ShaderBuffers[uint32(FunctionType)].Usage[Index] = mtlpp::ResourceUsage::Read;
-	ShaderBuffers[uint32(FunctionType)].Lengths[Index] = Data ? (Data->Len - Offset) : 0;
-	ShaderBuffers[uint32(FunctionType)].Lengths[Index+ML_MaxBuffers] = GMetalBufferFormats[Format].DataFormat;
+	ShaderBuffers[uint32(FunctionType)].Lengths[Index*2] = Data ? (Data->Len - Offset) : 0;
+	ShaderBuffers[uint32(FunctionType)].Lengths[(Index*2)+1] = GMetalBufferFormats[Format].DataFormat;
 	
 	SetShaderBufferInternal(FunctionType, Index);
 }
@@ -1371,8 +1371,8 @@ void FMetalCommandEncoder::SetShaderBytes(mtlpp::FunctionType const FunctionType
 		ShaderBuffers[uint32(FunctionType)].Bytes[Index] = nil;
 		ShaderBuffers[uint32(FunctionType)].Offsets[Index] = 0;
 		ShaderBuffers[uint32(FunctionType)].Usage[Index] = mtlpp::ResourceUsage::Read;
-		ShaderBuffers[uint32(FunctionType)].Lengths[Index] = Length;
-		ShaderBuffers[uint32(FunctionType)].Lengths[Index+ML_MaxBuffers] = GMetalBufferFormats[PF_Unknown].DataFormat;
+		ShaderBuffers[uint32(FunctionType)].Lengths[Index*2] = Length;
+		ShaderBuffers[uint32(FunctionType)].Lengths[(Index*2)+1] = GMetalBufferFormats[PF_Unknown].DataFormat;
 	}
 	else
 	{
@@ -1382,8 +1382,8 @@ void FMetalCommandEncoder::SetShaderBytes(mtlpp::FunctionType const FunctionType
 		ShaderBuffers[uint32(FunctionType)].Bytes[Index] = nil;
 		ShaderBuffers[uint32(FunctionType)].Offsets[Index] = 0;
 		ShaderBuffers[uint32(FunctionType)].Usage[Index] = mtlpp::ResourceUsage(0);
-		ShaderBuffers[uint32(FunctionType)].Lengths[Index] = 0;
-		ShaderBuffers[uint32(FunctionType)].Lengths[Index+ML_MaxBuffers] = GMetalBufferFormats[PF_Unknown].DataFormat;
+		ShaderBuffers[uint32(FunctionType)].Lengths[Index*2] = 0;
+		ShaderBuffers[uint32(FunctionType)].Lengths[(Index*2)+1] = GMetalBufferFormats[PF_Unknown].DataFormat;
 	}
 	
 	SetShaderBufferInternal(FunctionType, Index);
@@ -1395,8 +1395,8 @@ void FMetalCommandEncoder::SetShaderBufferOffset(mtlpp::FunctionType FunctionTyp
     checkf(ShaderBuffers[uint32(FunctionType)].Buffers[index] && (ShaderBuffers[uint32(FunctionType)].Bound & (1 << index)), TEXT("Buffer must already be bound"));
 	check(GetMetalDeviceContext().SupportsFeature(EMetalFeaturesSetBufferOffset));
 	ShaderBuffers[uint32(FunctionType)].Offsets[index] = Offset;
-	ShaderBuffers[uint32(FunctionType)].Lengths[index] = Length;
-	ShaderBuffers[uint32(FunctionType)].Lengths[index+ML_MaxBuffers] = GMetalBufferFormats[PF_Unknown].DataFormat;
+	ShaderBuffers[uint32(FunctionType)].Lengths[index*2] = Length;
+	ShaderBuffers[uint32(FunctionType)].Lengths[(index*2)+1] = GMetalBufferFormats[PF_Unknown].DataFormat;
 	switch (FunctionType)
 	{
 		case mtlpp::FunctionType::Vertex:
