Add missing calls to DidModify otherwise the IAB won't ever upload to the GPU...

diff --git a/Engine/Source/Runtime/Apple/MetalRHI/Private/MetalUniformBuffer.cpp b/Engine/Source/Runtime/Apple/MetalRHI/Private/MetalUniformBuffer.cpp
--- a/Engine/Source/Runtime/Apple/MetalRHI/Private/MetalUniformBuffer.cpp
+++ b/Engine/Source/Runtime/Apple/MetalRHI/Private/MetalUniformBuffer.cpp
@@ -350,6 +350,13 @@ FMetalUniformBuffer::FMetalUniformBuffer(const void* Contents, const FRHIUniform
 				
 				FMemory::Memcpy(IndirectArgumentBufferSideTable.GetContents(), BufferSizes.GetData(), BufferSizes.Num() * sizeof(uint32));
 				
+#if PLATFORM_MAC
+				if(IndirectArgumentBufferSideTable.GetStorageMode() == mtlpp::StorageMode::Managed)
+				{
+					MTLPP_VALIDATE(mtlpp::Buffer, IndirectArgumentBufferSideTable, SafeGetRuntimeDebuggingLevel() >= EMetalDebugLevelValidation, DidModify(ns::Range(0, BufferSizes.Num() * sizeof(uint32))));
+				}
+#endif
+				
 				IndirectArgumentResources.Add(Argument(IndirectArgumentBufferSideTable, mtlpp::ResourceUsage::Read));
 				Pointers.Add(IndirectArgumentBufferSideTable.GetPtr());
 			}
@@ -398,6 +405,13 @@ FMetalUniformBuffer::FMetalUniformBuffer(const void* Contents, const FRHIUniform
 						break;
 				}
 			}
+			
+#if PLATFORM_MAC
+			if(IndirectArgumentBuffer.GetStorageMode() == mtlpp::StorageMode::Managed)
+			{
+				MTLPP_VALIDATE(mtlpp::Buffer, IndirectArgumentBuffer, SafeGetRuntimeDebuggingLevel() >= EMetalDebugLevelValidation, DidModify(ns::Range(0, Encoder.GetEncodedLength())));
+			}
+#endif
 		}
 		else
 		{
