Don't even attempt to bind IAB resources directly to the encoder.

diff --git a/Engine/Source/Runtime/Apple/MetalRHI/Private/MetalStateCache.cpp b/Engine/Source/Runtime/Apple/MetalRHI/Private/MetalStateCache.cpp
--- a/Engine/Source/Runtime/Apple/MetalRHI/Private/MetalStateCache.cpp
+++ b/Engine/Source/Runtime/Apple/MetalRHI/Private/MetalStateCache.cpp
@@ -1842,7 +1842,7 @@ void FMetalStateCache::SetResourcesFromTables(ShaderType Shader, uint32 ShaderSt
 		const int32 BufferIndex = FMath::FloorLog2(LowestBitMask); // todo: This has a branch on zero, we know it could never be zero...
 		DirtyBits ^= LowestBitMask;
 		FMetalUniformBuffer* Buffer = (FMetalUniformBuffer*)GetBoundUniformBuffers(Frequency)[BufferIndex];
-		if (Buffer)
+		if (Buffer && (Shader->Bindings.ConstantBuffers & (1 << BufferIndex)) && !(Shader->Bindings.ArgumentBuffers & (1 << BufferIndex)))
 		{
 			check(BufferIndex < Shader->Bindings.ShaderResourceTable.ResourceTableLayoutHashes.Num());
 			check(Buffer->GetLayout().GetHash() == Shader->Bindings.ShaderResourceTable.ResourceTableLayoutHashes[BufferIndex]);
