Make sure that resources aren't prematurely disposed of and that we bind buffers at the correct offsets when creating IABs.

diff --git a/Engine/Source/Runtime/Apple/MetalRHI/Private/MetalUniformBuffer.cpp b/Engine/Source/Runtime/Apple/MetalRHI/Private/MetalUniformBuffer.cpp
--- a/Engine/Source/Runtime/Apple/MetalRHI/Private/MetalUniformBuffer.cpp
+++ b/Engine/Source/Runtime/Apple/MetalRHI/Private/MetalUniformBuffer.cpp
@@ -133,7 +133,6 @@ FMetalUniformBuffer::FMetalUniformBuffer(const void* Contents, const FRHIUniform
 	{
 		if (FMetalCommandQueue::SupportsFeature(EMetalFeaturesIABs))
 		{
-			TArray<id> Pointers;
 			TArray<uint32> BufferSizes;
 			NSMutableArray<mtlpp::ArgumentDescriptor::Type>* Arguments = [[NSMutableArray new] autorelease];
 			
@@ -177,7 +176,6 @@ FMetalUniformBuffer::FMetalUniformBuffer(const void* Contents, const FRHIUniform
 							
 							check(!Surface->Texture.IsAliasable());
 							IndirectArgumentResources.Add(Argument(Surface->Texture, (mtlpp::ResourceUsage)(mtlpp::ResourceUsage::Read|mtlpp::ResourceUsage::Sample)));
-							Pointers.Add(Surface->Texture.GetPtr());
 						}
 						else
 						{
@@ -185,7 +183,6 @@ FMetalUniformBuffer::FMetalUniformBuffer(const void* Contents, const FRHIUniform
 							ns::AutoReleased<FMetalTexture> Tex = SRV->GetLinearTexture(false);
 							Desc.SetDataType(mtlpp::DataType::Texture);
 							Desc.SetTextureType(Tex.GetTextureType());
-							Pointers.Add(Tex.GetPtr());
 							IndirectArgumentResources.Add(Argument(Tex, (mtlpp::ResourceUsage)(mtlpp::ResourceUsage::Read|mtlpp::ResourceUsage::Sample)));
 							
 							mtlpp::ArgumentDescriptor BufferDesc;
@@ -196,7 +193,6 @@ FMetalUniformBuffer::FMetalUniformBuffer(const void* Contents, const FRHIUniform
 							if (VB)
 							{
 								BufferDesc.SetDataType(mtlpp::DataType::Pointer);
-								Pointers.Add(VB->Buffer.GetPtr());
 								IndirectArgumentResources.Add(Argument(VB->Buffer, (mtlpp::ResourceUsage)(mtlpp::ResourceUsage::Read)));
 								
 								check(VB->Buffer.GetStorageMode() == mtlpp::StorageMode::Private);
@@ -206,7 +202,6 @@ FMetalUniformBuffer::FMetalUniformBuffer(const void* Contents, const FRHIUniform
 							else if (IB)
 							{
 								BufferDesc.SetDataType(mtlpp::DataType::Pointer);
-								Pointers.Add(IB->Buffer.GetPtr());
 								IndirectArgumentResources.Add(Argument(IB->Buffer, (mtlpp::ResourceUsage)(mtlpp::ResourceUsage::Read)));
 								
 								check(IB->Buffer.GetStorageMode() == mtlpp::StorageMode::Private);
@@ -216,7 +211,6 @@ FMetalUniformBuffer::FMetalUniformBuffer(const void* Contents, const FRHIUniform
 							else if (SB)
 							{
 								BufferDesc.SetDataType(mtlpp::DataType::Pointer);
-								Pointers.Add(SB->Buffer.GetPtr());
 								IndirectArgumentResources.Add(Argument(SB->Buffer, (mtlpp::ResourceUsage)(mtlpp::ResourceUsage::Read)));
 								
 								BufferSizes.Add(SB->GetSize());
@@ -251,7 +245,6 @@ FMetalUniformBuffer::FMetalUniformBuffer(const void* Contents, const FRHIUniform
 							Desc.SetTextureType(Surface->Texture.GetTextureType());
 							
 							check(!Surface->Texture.IsAliasable());
-							Pointers.Add(Surface->Texture.GetPtr());
 							IndirectArgumentResources.Add(Argument(Surface->Texture, (mtlpp::ResourceUsage)(mtlpp::ResourceUsage::Read|mtlpp::ResourceUsage::Write)));
 						}
 						else
@@ -260,7 +253,6 @@ FMetalUniformBuffer::FMetalUniformBuffer(const void* Contents, const FRHIUniform
 							ns::AutoReleased<FMetalTexture> Tex = SRV->GetLinearTexture(false);
 							Desc.SetDataType(mtlpp::DataType::Texture);
 							Desc.SetTextureType(Tex.GetTextureType());
-							Pointers.Add(Tex.GetPtr());
 							IndirectArgumentResources.Add(Argument(Tex, (mtlpp::ResourceUsage)(mtlpp::ResourceUsage::Read|mtlpp::ResourceUsage::Write)));
 							
 							mtlpp::ArgumentDescriptor BufferDesc;
@@ -271,7 +263,6 @@ FMetalUniformBuffer::FMetalUniformBuffer(const void* Contents, const FRHIUniform
 							if (VB)
 							{
 								BufferDesc.SetDataType(mtlpp::DataType::Pointer);
-								Pointers.Add(VB->Buffer.GetPtr());
 								IndirectArgumentResources.Add(Argument(VB->Buffer, (mtlpp::ResourceUsage)(mtlpp::ResourceUsage::Read|mtlpp::ResourceUsage::Write)));
 								
 								check(VB->Buffer.GetStorageMode() == mtlpp::StorageMode::Private);
@@ -281,7 +272,6 @@ FMetalUniformBuffer::FMetalUniformBuffer(const void* Contents, const FRHIUniform
 							else if (IB)
 							{
 								BufferDesc.SetDataType(mtlpp::DataType::Pointer);
-								Pointers.Add(IB->Buffer.GetPtr());
 								IndirectArgumentResources.Add(Argument(IB->Buffer, (mtlpp::ResourceUsage)(mtlpp::ResourceUsage::Read|mtlpp::ResourceUsage::Write)));
 								
 								check(IB->Buffer.GetStorageMode() == mtlpp::StorageMode::Private);
@@ -291,7 +281,6 @@ FMetalUniformBuffer::FMetalUniformBuffer(const void* Contents, const FRHIUniform
 							else if (SB)
 							{
 								BufferDesc.SetDataType(mtlpp::DataType::Pointer);
-								Pointers.Add(SB->Buffer.GetPtr());
 								IndirectArgumentResources.Add(Argument(SB->Buffer, (mtlpp::ResourceUsage)(mtlpp::ResourceUsage::Read|mtlpp::ResourceUsage::Write)));
 								
 								BufferSizes.Add(SB->GetSize());
@@ -309,8 +298,7 @@ FMetalUniformBuffer::FMetalUniformBuffer(const void* Contents, const FRHIUniform
 						
 						FMetalSamplerState* Sampler = (FMetalSamplerState*)Resource;
 						Desc.SetDataType(mtlpp::DataType::Sampler);
-						Pointers.Add(Sampler->State.GetPtr());
-						IndirectArgumentResources.Add(Argument());
+						IndirectArgumentResources.Add(Argument(Sampler->State));
 						break;
 					}
 					case UBMT_GRAPH_TRACKED_TEXTURE:
@@ -328,7 +316,6 @@ FMetalUniformBuffer::FMetalUniformBuffer(const void* Contents, const FRHIUniform
 						Desc.SetTextureType(Surface->Texture.GetTextureType());
 						
 						check(!Surface->Texture.IsAliasable());
-						Pointers.Add(Surface->Texture.GetPtr());
 						IndirectArgumentResources.Add(Argument(Surface->Texture, (mtlpp::ResourceUsage)(mtlpp::ResourceUsage::Read|mtlpp::ResourceUsage::Sample)));
 						break;
 					}
@@ -358,7 +345,6 @@ FMetalUniformBuffer::FMetalUniformBuffer(const void* Contents, const FRHIUniform
 #endif
 				
 				IndirectArgumentResources.Add(Argument(IndirectArgumentBufferSideTable, mtlpp::ResourceUsage::Read));
-				Pointers.Add(IndirectArgumentBufferSideTable.GetPtr());
 			}
 			
 			if (Layout.ConstantBufferSize > 0)
@@ -368,16 +354,15 @@ FMetalUniformBuffer::FMetalUniformBuffer(const void* Contents, const FRHIUniform
                 Desc.SetAccess(mtlpp::ArgumentAccess::ReadOnly);
                 Desc.SetDataType(mtlpp::DataType::Pointer);
                 [Arguments addObject:Desc.GetPtr()];
-                
-                Pointers.Add(Buffer.GetPtr());
+				
                 IndirectArgumentResources.Add(Argument(Buffer, mtlpp::ResourceUsage::Read));
 			}
 			
 			ns::Array<mtlpp::ArgumentDescriptor> ArgArray(Arguments);
 			mtlpp::ArgumentEncoder Encoder = FMetalArgumentEncoderCache::Get().CreateEncoder(ArgArray);
 			
-			FMetalPooledBufferArgs Args(GetMetalDeviceContext().GetDevice(), Encoder.GetEncodedLength(), BUFFER_STORAGE_MODE);
-			IndirectArgumentBuffer = GetMetalDeviceContext().CreatePooledBuffer(Args);
+			IndirectArgumentBuffer = GetMetalDeviceContext().GetResourceHeap().CreateBuffer(Encoder.GetEncodedLength(), 16, mtlpp::ResourceOptions(BUFFER_CACHE_MODE | mtlpp::ResourceOptions::HazardTrackingModeUntracked | ((NSUInteger)BUFFER_STORAGE_MODE << mtlpp::ResourceStorageModeShift)), true);
+			
 			Encoder.SetArgumentBuffer(IndirectArgumentBuffer, 0);
 			
 			for (MTLArgumentDescriptor* Arg in Arguments)
@@ -387,13 +372,13 @@ FMetalUniformBuffer::FMetalUniformBuffer(const void* Contents, const FRHIUniform
 				switch(Argument.GetDataType())
 				{
 					case mtlpp::DataType::Pointer:
-						Encoder.SetBuffer((id<MTLBuffer>)Pointers[NewIndex], 0, NewIndex);
+						Encoder.SetBuffer(IndirectArgumentResources[NewIndex].Buffer, 0, NewIndex);
 						break;
 					case mtlpp::DataType::Texture:
-						Encoder.SetTexture((id<MTLTexture>)Pointers[NewIndex], NewIndex);
+						Encoder.SetTexture(IndirectArgumentResources[NewIndex].Texture, NewIndex);
 						break;
 					case mtlpp::DataType::Sampler:
-						Encoder.SetSamplerState((id<MTLSamplerState>)Pointers[NewIndex], NewIndex);
+						Encoder.SetSamplerState(IndirectArgumentResources[NewIndex].Sampler, NewIndex);
 						break;
 					case mtlpp::DataType::Struct:
 					{
diff --git a/Engine/Source/Runtime/Apple/MetalRHI/Public/MetalResources.h b/Engine/Source/Runtime/Apple/MetalRHI/Public/MetalResources.h
--- a/Engine/Source/Runtime/Apple/MetalRHI/Public/MetalResources.h
+++ b/Engine/Source/Runtime/Apple/MetalRHI/Public/MetalResources.h
@@ -865,9 +865,11 @@ public:
 		Argument() {}
 		Argument(FMetalBuffer const& InBuffer, mtlpp::ResourceUsage const InUsage) : Buffer(InBuffer), Usage(InUsage) {}
 		Argument(FMetalTexture const& InTexture, mtlpp::ResourceUsage const InUsage) : Texture(InTexture), Usage(InUsage) {}
-		
-		ns::AutoReleased<FMetalBuffer> Buffer;
-		ns::AutoReleased<FMetalTexture> Texture;
+		Argument(FMetalSampler const& InSampler) : Sampler(InSampler), Usage(mtlpp::ResourceUsage::Read) {}
+
+		FMetalBuffer Buffer;
+		FMetalTexture Texture;
+		FMetalSampler Sampler;
 		mtlpp::ResourceUsage Usage;
 	};
 	
