Don't force MetalRHI to reallocate render-targets or depth-stencil surfaces that were marked as ShaderResource by disabling the PixelFormatView usage flag, it makes no sense to do so.

diff --git a/Engine/Source/Runtime/Apple/MetalRHI/Private/MetalTexture.cpp b/Engine/Source/Runtime/Apple/MetalRHI/Private/MetalTexture.cpp
--- a/Engine/Source/Runtime/Apple/MetalRHI/Private/MetalTexture.cpp
+++ b/Engine/Source/Runtime/Apple/MetalRHI/Private/MetalTexture.cpp
@@ -120,7 +120,10 @@ static mtlpp::TextureUsage ConvertFlagsToUsage(uint32 Flags)
 	{
 		Usage |= mtlpp::TextureUsage::RenderTarget;
 		Usage |= mtlpp::TextureUsage::ShaderRead;
-		Usage &= ~(mtlpp::TextureUsage::PixelFormatView);
+		if (!(Flags & (TexCreate_ShaderResource)))
+		{
+			Usage &= ~(mtlpp::TextureUsage::PixelFormatView);
+		}
 	}
 	return (mtlpp::TextureUsage)Usage;
 }
