Previous change incorrectly only allowed binding of uniform-buffer resource tables when the CB itself was referenced in the shader, that is incorrect, we always want to bind the resources without IABs and never bind them when IABs are available.

diff --git a/Engine/Source/Runtime/Apple/MetalRHI/Private/MetalStateCache.cpp b/Engine/Source/Runtime/Apple/MetalRHI/Private/MetalStateCache.cpp
--- a/Engine/Source/Runtime/Apple/MetalRHI/Private/MetalStateCache.cpp
+++ b/Engine/Source/Runtime/Apple/MetalRHI/Private/MetalStateCache.cpp
@@ -1842,7 +1842,7 @@ void FMetalStateCache::SetResourcesFromTables(ShaderType Shader, uint32 ShaderSt
 		const int32 BufferIndex = FMath::FloorLog2(LowestBitMask); // todo: This has a branch on zero, we know it could never be zero...
 		DirtyBits ^= LowestBitMask;
 		FMetalUniformBuffer* Buffer = (FMetalUniformBuffer*)GetBoundUniformBuffers(Frequency)[BufferIndex];
-		if (Buffer && (Shader->Bindings.ConstantBuffers & (1 << BufferIndex)) && !(Shader->Bindings.ArgumentBuffers & (1 << BufferIndex)))
+		if (Buffer && !FMetalCommandQueue::SupportsFeature(EMetalFeaturesIABs))
 		{
 			check(BufferIndex < Shader->Bindings.ShaderResourceTable.ResourceTableLayoutHashes.Num());
 			check(Buffer->GetLayout().GetHash() == Shader->Bindings.ShaderResourceTable.ResourceTableLayoutHashes[BufferIndex]);
