Make Indirect Argument Buffers work with the RHI thread and parallel contexts, requires a bit of trick to defer initialisation for single-draw UBs that are created on parallel threads as we can't dispatch them to the RHI thread.
While we are here we'll add our own C-style argument descriptor array to reduce the number of retain/release calls.

diff --git a/Engine/Source/Runtime/Apple/MetalRHI/Private/MetalCommands.cpp b/Engine/Source/Runtime/Apple/MetalRHI/Private/MetalCommands.cpp
--- a/Engine/Source/Runtime/Apple/MetalRHI/Private/MetalCommands.cpp
+++ b/Engine/Source/Runtime/Apple/MetalRHI/Private/MetalCommands.cpp
@@ -536,7 +536,7 @@ void FMetalRHICommandContext::RHISetShaderUniformBuffer(FVertexShaderRHIParamRef
 	if (Bindings.ConstantBuffers & 1 << BufferIndex)
 	{
 		auto* UB = (FMetalUniformBuffer*)BufferRHI;
-		Context->GetCurrentState().SetShaderBuffer(SF_Vertex, (Bindings.ArgumentBuffers & 1 << BufferIndex) ? UB->IndirectArgumentBuffer : UB->Buffer, UB->Data, 0, UB->GetSize(), BufferIndex, mtlpp::ResourceUsage::Read);
+		Context->GetCurrentState().SetShaderBuffer(SF_Vertex, (Bindings.ArgumentBuffers & 1 << BufferIndex) ? UB->GetIAB().IndirectArgumentBuffer : UB->Buffer, UB->Data, 0, UB->GetSize(), BufferIndex, mtlpp::ResourceUsage::Read);
 	}
 	}
 }
@@ -552,7 +552,7 @@ void FMetalRHICommandContext::RHISetShaderUniformBuffer(FHullShaderRHIParamRef H
 	if (Bindings.ConstantBuffers & 1 << BufferIndex)
 	{
 		auto* UB = (FMetalUniformBuffer*)BufferRHI;
-		Context->GetCurrentState().SetShaderBuffer(SF_Hull, (Bindings.ArgumentBuffers & 1 << BufferIndex) ? UB->IndirectArgumentBuffer : UB->Buffer, UB->Data, 0, UB->GetSize(), BufferIndex, mtlpp::ResourceUsage::Read);
+		Context->GetCurrentState().SetShaderBuffer(SF_Hull, (Bindings.ArgumentBuffers & 1 << BufferIndex) ? UB->GetIAB().IndirectArgumentBuffer : UB->Buffer, UB->Data, 0, UB->GetSize(), BufferIndex, mtlpp::ResourceUsage::Read);
 	}
 	}
 }
@@ -568,7 +568,7 @@ void FMetalRHICommandContext::RHISetShaderUniformBuffer(FDomainShaderRHIParamRef
 	if (Bindings.ConstantBuffers & 1 << BufferIndex)
 	{
 		auto* UB = (FMetalUniformBuffer*)BufferRHI;
-		Context->GetCurrentState().SetShaderBuffer(SF_Domain, (Bindings.ArgumentBuffers & 1 << BufferIndex) ? UB->IndirectArgumentBuffer : UB->Buffer, UB->Data, 0, UB->GetSize(), BufferIndex, mtlpp::ResourceUsage::Read);
+		Context->GetCurrentState().SetShaderBuffer(SF_Domain, (Bindings.ArgumentBuffers & 1 << BufferIndex) ? UB->GetIAB().IndirectArgumentBuffer : UB->Buffer, UB->Data, 0, UB->GetSize(), BufferIndex, mtlpp::ResourceUsage::Read);
 	}
 	}
 }
@@ -589,7 +589,7 @@ void FMetalRHICommandContext::RHISetShaderUniformBuffer(FPixelShaderRHIParamRef
 	if (Bindings.ConstantBuffers & 1 << BufferIndex)
 	{
 		auto* UB = (FMetalUniformBuffer*)BufferRHI;
-		Context->GetCurrentState().SetShaderBuffer(SF_Pixel, (Bindings.ArgumentBuffers & 1 << BufferIndex) ? UB->IndirectArgumentBuffer : UB->Buffer, UB->Data, 0, UB->GetSize(), BufferIndex, mtlpp::ResourceUsage::Read);
+		Context->GetCurrentState().SetShaderBuffer(SF_Pixel, (Bindings.ArgumentBuffers & 1 << BufferIndex) ? UB->GetIAB().IndirectArgumentBuffer : UB->Buffer, UB->Data, 0, UB->GetSize(), BufferIndex, mtlpp::ResourceUsage::Read);
 	}
 	}
 }
@@ -605,7 +605,7 @@ void FMetalRHICommandContext::RHISetShaderUniformBuffer(FComputeShaderRHIParamRe
 	if (Bindings.ConstantBuffers & 1 << BufferIndex)
 	{
 		auto* UB = (FMetalUniformBuffer*)BufferRHI;
-		Context->GetCurrentState().SetShaderBuffer(SF_Compute, (Bindings.ArgumentBuffers & 1 << BufferIndex) ? UB->IndirectArgumentBuffer : UB->Buffer, UB->Data, 0, UB->GetSize(), BufferIndex, mtlpp::ResourceUsage::Read);
+		Context->GetCurrentState().SetShaderBuffer(SF_Compute, (Bindings.ArgumentBuffers & 1 << BufferIndex) ? UB->GetIAB().IndirectArgumentBuffer : UB->Buffer, UB->Data, 0, UB->GetSize(), BufferIndex, mtlpp::ResourceUsage::Read);
 	}
 	}
 }
diff --git a/Engine/Source/Runtime/Apple/MetalRHI/Private/MetalStateCache.cpp b/Engine/Source/Runtime/Apple/MetalRHI/Private/MetalStateCache.cpp
--- a/Engine/Source/Runtime/Apple/MetalRHI/Private/MetalStateCache.cpp
+++ b/Engine/Source/Runtime/Apple/MetalRHI/Private/MetalStateCache.cpp
@@ -2258,10 +2258,12 @@ void FMetalStateCache::CommitResourceTable(EShaderFrequency const Frequency, mtl
 		FMetalUniformBuffer* Buffer = (FMetalUniformBuffer*)BoundUniformBuffers[Frequency][Index];
 		check(Buffer);
 		
+		FMetalUniformBuffer::FMetalIndirectArgumentBuffer& IAB = Buffer->GetIAB();
+		
 		TArray<uint8>& Mask = ArgumentBufferMasks.FindChecked((uint8)Index);
 		for (uint8 Entry : Mask)
 		{
-			FMetalUniformBuffer::Argument const& Resource = Buffer->IndirectArgumentResources[Entry];
+			FMetalUniformBuffer::Argument const& Resource = IAB.IndirectArgumentResources[Entry];
 			if (Resource.Buffer)
 			{
 				CommandEncoder.UseIndirectArgumentResource(Resource.Buffer, Resource.Usage);
@@ -2272,9 +2274,9 @@ void FMetalStateCache::CommitResourceTable(EShaderFrequency const Frequency, mtl
 			}
 		}
 		
-		if (Buffer->IndirectArgumentBufferSideTable)
+		if (IAB.IndirectArgumentBufferSideTable)
 		{
-			CommandEncoder.UseIndirectArgumentResource(Buffer->IndirectArgumentBufferSideTable, mtlpp::ResourceUsage::Read);
+			CommandEncoder.UseIndirectArgumentResource(IAB.IndirectArgumentBufferSideTable, mtlpp::ResourceUsage::Read);
 		}
 		
 		if (Buffer->Buffer)
diff --git a/Engine/Source/Runtime/Apple/MetalRHI/Private/MetalUniformBuffer.cpp b/Engine/Source/Runtime/Apple/MetalRHI/Private/MetalUniformBuffer.cpp
--- a/Engine/Source/Runtime/Apple/MetalRHI/Private/MetalUniformBuffer.cpp
+++ b/Engine/Source/Runtime/Apple/MetalRHI/Private/MetalUniformBuffer.cpp
@@ -39,15 +39,111 @@ struct FMetalRHICommandInitialiseUniformdBuffer : public FRHICommand<FMetalRHICo
 	}
 };
 
-struct FMetalArgumentEncoderMapFuncs : TDefaultMapKeyFuncs<ns::Array<mtlpp::ArgumentDescriptor>, mtlpp::ArgumentEncoder, /*bInAllowDuplicateKeys*/false>
+struct FMetalRHICommandInitialiseUniformBufferIAB : public FRHICommand<FMetalRHICommandInitialiseUniformBufferIAB>
+{
+	TRefCountPtr<FMetalUniformBuffer> Buffer;
+	
+	FORCEINLINE_DEBUGGABLE FMetalRHICommandInitialiseUniformBufferIAB(FMetalUniformBuffer* InBuffer)
+	: Buffer(InBuffer)
+	{
+	}
+	
+	virtual ~FMetalRHICommandInitialiseUniformBufferIAB()
+	{
+	}
+	
+	void Execute(FRHICommandListBase& CmdList)
+	{
+		Buffer->InitIAB();
+	}
+};
+
+struct FMetalArgumentDesc
+{
+	FMetalArgumentDesc()
+	: DataType((mtlpp::DataType)0)
+	, Index(0)
+	, ArrayLength(0)
+	, Access(mtlpp::ArgumentAccess::ReadOnly)
+	, TextureType((mtlpp::TextureType)0)
+	, ConstantBlockAlignment(0)
+	{
+		
+	}
+	
+	mtlpp::DataType DataType;
+	NSUInteger Index;
+	NSUInteger ArrayLength;
+	mtlpp::ArgumentAccess Access;
+	mtlpp::TextureType TextureType;
+	NSUInteger ConstantBlockAlignment;
+	
+	void FillDescriptor(mtlpp::ArgumentDescriptor& Desc) const
+	{
+		Desc.SetDataType(DataType);
+		Desc.SetIndex(Index);
+		Desc.SetArrayLength(ArrayLength);
+		Desc.SetAccess(Access);
+		Desc.SetTextureType(TextureType);
+		Desc.SetConstantBlockAlignment(ConstantBlockAlignment);
+	}
+	
+	void SetDataType(mtlpp::DataType Type)
+	{
+		DataType = Type;
+	}
+	
+	void SetIndex(NSUInteger i)
+	{
+		Index = i;
+	}
+	
+	void SetArrayLength(NSUInteger Len)
+	{
+		ArrayLength = Len;
+	}
+	
+	void SetAccess(mtlpp::ArgumentAccess InAccess)
+	{
+		Access = InAccess;
+	}
+	
+	void SetTextureType(mtlpp::TextureType Type)
+	{
+		TextureType = Type;
+	}
+	
+	void SetConstantBlockAlignment(NSUInteger Align)
+	{
+		ConstantBlockAlignment = Align;
+	}
+	
+	bool operator==(FMetalArgumentDesc const& Other) const
+	{
+		if (this != &Other)
+		{
+			return (DataType == Other.DataType && Index == Other.Index && ArrayLength == Other.ArrayLength && Access == Other.Access && TextureType == Other.TextureType && ConstantBlockAlignment == Other.ConstantBlockAlignment);
+		}
+		return true;
+	}
+	
+	friend uint32 GetTypeHash(FMetalArgumentDesc const& Desc)
+	{
+		uint32 Hash = 0;
+		Hash ^= ((uint32)Desc.DataType * (uint32)Desc.TextureType * (uint32)Desc.Access * Desc.ArrayLength) << Desc.Index;
+		return Hash;
+	}
+};
+
+struct FMetalArgumentEncoderMapFuncs : TDefaultMapKeyFuncs<TArray<FMetalArgumentDesc>, mtlpp::ArgumentEncoder, /*bInAllowDuplicateKeys*/false>
 {
 	static FORCEINLINE uint32 GetKeyHash(KeyInitType Key)
 	{
 		uint32 Hash = 0;
 		
-		for (ns::AutoReleased<mtlpp::ArgumentDescriptor> const& Desc : Key)
+		for (FMetalArgumentDesc const& Desc : Key)
 		{
-			Hash ^= ((uint32)Desc.GetDataType() * (uint32)Desc.GetTextureType() * (uint32)Desc.GetAccess() * Desc.GetArrayLength()) << Desc.GetIndex();
+			HashCombine(Hash, GetTypeHash(Desc));
 		}
 		
 		return Hash;
@@ -71,7 +167,7 @@ public:
 		return sSelf;
 	}
 	
-	mtlpp::ArgumentEncoder CreateEncoder(ns::Array<mtlpp::ArgumentDescriptor> const& Desc)
+	mtlpp::ArgumentEncoder CreateEncoder(TArray<FMetalArgumentDesc> const& Desc)
 	{
 		mtlpp::ArgumentEncoder Encoder;
 		
@@ -80,11 +176,19 @@ public:
 		
 		if (!Encoder)
 		{
-			Encoder = GetMetalDeviceContext().GetDevice().NewArgumentEncoderWithArguments(Desc);
+			NSMutableArray<mtlpp::ArgumentDescriptor::Type>* Arguments = [[NSMutableArray new] autorelease];
+			for (FMetalArgumentDesc const& Args : Desc)
+			{
+				mtlpp::ArgumentDescriptor Arg;
+				Args.FillDescriptor(Arg);
+				[Arguments addObject:Arg.GetPtr()];
+			}
+			
+			Encoder = GetMetalDeviceContext().GetDevice().NewArgumentEncoderWithArguments(Arguments);
 			
 			// Now we are a writer as we want to create & add the new pipeline
 			Lock.ReleaseReadOnlyLockAndAcquireWriteLock_USE_WITH_CAUTION();
-			
+
 			if (!Encoders.Find(Desc))
 			{
 				Encoders.Add(Desc, Encoder);
@@ -95,335 +199,333 @@ public:
 	}
 private:
 	FRWLock Mutex;
-	TMap<ns::Array<mtlpp::ArgumentDescriptor>, mtlpp::ArgumentEncoder, FDefaultSetAllocator, FMetalArgumentEncoderMapFuncs> Encoders;
+	TMap<TArray<FMetalArgumentDesc>, mtlpp::ArgumentEncoder, FDefaultSetAllocator, FMetalArgumentEncoderMapFuncs> Encoders;
 };
 
 FMetalUniformBuffer::FMetalUniformBuffer(const void* Contents, const FRHIUniformBufferLayout& Layout, EUniformBufferUsage InUsage)
 	: FRHIUniformBuffer(Layout)
     , FMetalRHIBuffer(Layout.ConstantBufferSize, (FMetalCommandQueue::SupportsFeature(EMetalFeaturesIABs) && Layout.Resources.Num() ? (EMetalBufferUsage_GPUOnly|BUF_Volatile) : BUF_Volatile), RRT_UniformBuffer)
+	, IAB(nullptr)
 {
 	if (Layout.ConstantBufferSize > 0)
 	{
 		UE_CLOG(Layout.ConstantBufferSize > 65536, LogMetal, Fatal, TEXT("Trying to allocated a uniform layout of size %d that is greater than the maximum permitted 64k."), Layout.ConstantBufferSize);
 		
-		FRHICommandListImmediate& RHICmdList = FRHICommandListExecutor::GetImmediateCommandList();
-		if (!(InUsage & UniformBuffer_SingleDraw) && CPUBuffer && IsRunningRHIInSeparateThread() && !RHICmdList.Bypass() && !IsInRHIThread())
-		{
-			FMemory::Memcpy(CPUBuffer.GetContents(), Contents, Layout.ConstantBufferSize);
+		void* Data = Lock(RLM_WriteOnly, 0);
+		FMemory::Memcpy(Data, Contents, Layout.ConstantBufferSize);
+		Unlock();
+	}
 
-#if PLATFORM_MAC
-			if(CPUBuffer.GetStorageMode() == mtlpp::StorageMode::Managed)
+	// set up an SRT-style uniform buffer
+	if (Layout.Resources.Num())
+	{
+		int32 NumResources = Layout.Resources.Num();
+		ResourceTable.Empty(NumResources);
+		ResourceTable.AddZeroed(NumResources);
+		for (int32 i = 0; i < NumResources; ++i)
+		{
+			FRHIResource* Resource = *(FRHIResource**)((uint8*)Contents + Layout.ResourceOffsets[i]);
+			
+			// Allow null SRV's in uniform buffers for feature levels that don't support SRV's in shaders
+			if (!(GMaxRHIFeatureLevel <= ERHIFeatureLevel::ES3_1 && Layout.Resources[i] == UBMT_SRV))
 			{
-				MTLPP_VALIDATE(mtlpp::Buffer, CPUBuffer, SafeGetRuntimeDebuggingLevel() >= EMetalDebugLevelValidation, DidModify(ns::Range(0, GMetalBufferZeroFill ? CPUBuffer.GetLength() : Layout.ConstantBufferSize)));
+				check(Resource);
 			}
-#endif
-
-			new (RHICmdList.AllocCommand<FMetalRHICommandInitialiseUniformdBuffer>()) FMetalRHICommandInitialiseUniformdBuffer(this);
+			ResourceTable[i] = Resource;
 		}
-		else
+		
+		FRHICommandListImmediate& RHICmdList = FRHICommandListExecutor::GetImmediateCommandList();
+		if (!(InUsage & UniformBuffer_SingleDraw) && IsRunningRHIInSeparateThread() && !RHICmdList.Bypass() && IsInRenderingThread())
 		{
-			void* Data = Lock(RLM_WriteOnly, 0);
-			FMemory::Memcpy(Data, Contents, Layout.ConstantBufferSize);
-			Unlock();
+			new (RHICmdList.AllocCommand<FMetalRHICommandInitialiseUniformBufferIAB>()) FMetalRHICommandInitialiseUniformBufferIAB(this);
 		}
 	}
+}
 
-	// set up an SRT-style uniform buffer
-	if (Layout.Resources.Num())
+FMetalUniformBuffer::~FMetalUniformBuffer()
+{
+	if (IAB)
 	{
-		if (FMetalCommandQueue::SupportsFeature(EMetalFeaturesIABs))
+		delete IAB;
+		IAB = nullptr;
+	}
+}
+
+FMetalUniformBuffer::FMetalIndirectArgumentBuffer::FMetalIndirectArgumentBuffer()
+{
+	
+}
+
+FMetalUniformBuffer::FMetalIndirectArgumentBuffer::~FMetalIndirectArgumentBuffer()
+{
+	SafeReleaseMetalBuffer(IndirectArgumentBuffer);
+	SafeReleaseMetalBuffer(IndirectArgumentBufferSideTable);
+}
+
+FMetalUniformBuffer::FMetalIndirectArgumentBuffer& FMetalUniformBuffer::GetIAB()
+{
+	check(ResourceTable.Num() && FMetalCommandQueue::SupportsFeature(EMetalFeaturesIABs));
+	InitIAB();
+	check(IAB);
+	return *IAB;
+}
+
+void FMetalUniformBuffer::InitIAB()
+{
+	int32 NumResources = ResourceTable.Num();
+	if (NumResources && FMetalCommandQueue::SupportsFeature(EMetalFeaturesIABs) && !IAB)
+	{
+		FMetalIndirectArgumentBuffer* NewIAB = new FMetalIndirectArgumentBuffer;
+		
+		TArray<uint32> BufferSizes;
+		TArray<FMetalArgumentDesc> Arguments;
+		
+		int32 Index = 0;
+		for (int32 i = 0; i < NumResources; ++i, ++Index)
 		{
-			TArray<uint32> BufferSizes;
-			NSMutableArray<mtlpp::ArgumentDescriptor::Type>* Arguments = [[NSMutableArray new] autorelease];
+			FRHIResource* Resource = ResourceTable[i].GetReference();
 			
-			int32 NumResources = Layout.Resources.Num();
-			ResourceTable.Empty(NumResources);
-			ResourceTable.AddZeroed(NumResources);
-			int32 Index = 0;
-			for (int32 i = 0; i < NumResources; ++i, ++Index)
+			switch(GetLayout().Resources[i])
 			{
-				FRHIResource* Resource = *(FRHIResource**)((uint8*)Contents + Layout.ResourceOffsets[i]);
-				
-				// Allow null SRV's in uniform buffers for feature levels that don't support SRV's in shaders
-				if (!(GMaxRHIFeatureLevel <= ERHIFeatureLevel::ES3_1 && Layout.Resources[i] == UBMT_SRV))
+				case UBMT_GRAPH_TRACKED_SRV:
+				case UBMT_GRAPH_TRACKED_BUFFER_SRV:
+				case UBMT_SRV:
 				{
-					check(Resource);
-				}
-				ResourceTable[i] = Resource;
-				
-				switch(Layout.Resources[i])
-				{
-					case UBMT_GRAPH_TRACKED_SRV:
-					case UBMT_GRAPH_TRACKED_BUFFER_SRV:
-					case UBMT_SRV:
+					FMetalArgumentDesc& Desc = Arguments.Emplace_GetRef();
+					Desc.SetIndex(Index);
+					Desc.SetAccess(mtlpp::ArgumentAccess::ReadOnly);
+					
+					FMetalShaderResourceView* SRV = (FMetalShaderResourceView*)Resource;
+					FRHITexture* Texture = SRV->SourceTexture.GetReference();
+					FMetalVertexBuffer* VB = SRV->SourceVertexBuffer.GetReference();
+					FMetalIndexBuffer* IB = SRV->SourceIndexBuffer.GetReference();
+					FMetalStructuredBuffer* SB = SRV->SourceStructuredBuffer.GetReference();
+					if (Texture)
 					{
-						mtlpp::ArgumentDescriptor Desc;
-						Desc.SetIndex(Index);
-						Desc.SetAccess(mtlpp::ArgumentAccess::ReadOnly);
-						[Arguments addObject:Desc.GetPtr()];
+						FMetalSurface* Surface = SRV->TextureView;
+						check (Surface != nullptr);
+						Desc.SetDataType(mtlpp::DataType::Texture);
+						Desc.SetTextureType(Surface->Texture.GetTextureType());
 						
-						FMetalShaderResourceView* SRV = (FMetalShaderResourceView*)Resource;
-						FRHITexture* Texture = SRV->SourceTexture.GetReference();
-						FMetalVertexBuffer* VB = SRV->SourceVertexBuffer.GetReference();
-						FMetalIndexBuffer* IB = SRV->SourceIndexBuffer.GetReference();
-						FMetalStructuredBuffer* SB = SRV->SourceStructuredBuffer.GetReference();
-						if (Texture)
-						{
-							FMetalSurface* Surface = SRV->TextureView;
-							check (Surface != nullptr);
-							Desc.SetDataType(mtlpp::DataType::Texture);
-							Desc.SetTextureType(Surface->Texture.GetTextureType());
-							
-							check(!Surface->Texture.IsAliasable());
-							IndirectArgumentResources.Add(Argument(Surface->Texture, (mtlpp::ResourceUsage)(mtlpp::ResourceUsage::Read|mtlpp::ResourceUsage::Sample)));
-						}
-						else
-						{
-							check(VB || IB || SB);
-							ns::AutoReleased<FMetalTexture> Tex = SRV->GetLinearTexture(false);
-							Desc.SetDataType(mtlpp::DataType::Texture);
-							Desc.SetTextureType(Tex.GetTextureType());
-							IndirectArgumentResources.Add(Argument(Tex, (mtlpp::ResourceUsage)(mtlpp::ResourceUsage::Read|mtlpp::ResourceUsage::Sample)));
-							
-							mtlpp::ArgumentDescriptor BufferDesc;
-							BufferDesc.SetIndex(++Index);
-							BufferDesc.SetAccess(mtlpp::ArgumentAccess::ReadOnly);
-							[Arguments addObject:BufferDesc.GetPtr()];
-							
-							if (VB)
-							{
-								BufferDesc.SetDataType(mtlpp::DataType::Pointer);
-								IndirectArgumentResources.Add(Argument(VB->Buffer, (mtlpp::ResourceUsage)(mtlpp::ResourceUsage::Read)));
-								
-								check(VB->Buffer.GetStorageMode() == mtlpp::StorageMode::Private);
-								BufferSizes.Add(VB->GetSize());
-								BufferSizes.Add(GMetalBufferFormats[SRV->Format].DataFormat);
-							}
-							else if (IB)
-							{
-								BufferDesc.SetDataType(mtlpp::DataType::Pointer);
-								IndirectArgumentResources.Add(Argument(IB->Buffer, (mtlpp::ResourceUsage)(mtlpp::ResourceUsage::Read)));
-								
-								check(IB->Buffer.GetStorageMode() == mtlpp::StorageMode::Private);
-								BufferSizes.Add(IB->GetSize());
-								BufferSizes.Add(GMetalBufferFormats[SRV->Format].DataFormat);
-							}
-							else if (SB)
-							{
-								BufferDesc.SetDataType(mtlpp::DataType::Pointer);
-								IndirectArgumentResources.Add(Argument(SB->Buffer, (mtlpp::ResourceUsage)(mtlpp::ResourceUsage::Read)));
-								
-								BufferSizes.Add(SB->GetSize());
-								BufferSizes.Add(GMetalBufferFormats[SRV->Format].DataFormat);
-							}
-						}
-						break;
+						check(!Surface->Texture.IsAliasable());
+						NewIAB->IndirectArgumentResources.Add(Argument(Surface->Texture, (mtlpp::ResourceUsage)(mtlpp::ResourceUsage::Read|mtlpp::ResourceUsage::Sample)));
 					}
-					case UBMT_GRAPH_TRACKED_UAV:
-					case UBMT_GRAPH_TRACKED_BUFFER_UAV:
+					else
 					{
-						mtlpp::ArgumentDescriptor Desc;
-						Desc.SetIndex(Index);
-						Desc.SetAccess(mtlpp::ArgumentAccess::ReadWrite);
-						[Arguments addObject:Desc.GetPtr()];
+						check(VB || IB || SB);
+						ns::AutoReleased<FMetalTexture> Tex = SRV->GetLinearTexture(false);
+						Desc.SetDataType(mtlpp::DataType::Texture);
+						Desc.SetTextureType(Tex.GetTextureType());
+						NewIAB->IndirectArgumentResources.Add(Argument(Tex, (mtlpp::ResourceUsage)(mtlpp::ResourceUsage::Read|mtlpp::ResourceUsage::Sample)));
+						
+						FMetalArgumentDesc& BufferDesc = Arguments.Emplace_GetRef();
+						BufferDesc.SetIndex(++Index);
+						BufferDesc.SetAccess(mtlpp::ArgumentAccess::ReadOnly);
 						
-						FMetalUnorderedAccessView* UAV = (FMetalUnorderedAccessView*)Resource;
-						FMetalShaderResourceView* SRV = (FMetalShaderResourceView*)UAV->SourceView;
-						FMetalStructuredBuffer* SB = UAV->SourceView->SourceStructuredBuffer.GetReference();
-						FMetalVertexBuffer* VB = UAV->SourceView->SourceVertexBuffer.GetReference();
-						FMetalIndexBuffer* IB = UAV->SourceView->SourceIndexBuffer.GetReference();
-						FRHITexture* Texture = UAV->SourceView->SourceTexture.GetReference();
-						FMetalSurface* Surface = UAV->SourceView->TextureView;
-						if (Texture)
+						if (VB)
 						{
-							if (!Surface)
-							{
-								Surface = GetMetalSurfaceFromRHITexture(Texture);
-							}
-							check (Surface != nullptr);
-							Desc.SetDataType(mtlpp::DataType::Texture);
-							Desc.SetTextureType(Surface->Texture.GetTextureType());
+							BufferDesc.SetDataType(mtlpp::DataType::Pointer);
+							NewIAB->IndirectArgumentResources.Add(Argument(VB->Buffer, (mtlpp::ResourceUsage)(mtlpp::ResourceUsage::Read)));
 							
-							check(!Surface->Texture.IsAliasable());
-							IndirectArgumentResources.Add(Argument(Surface->Texture, (mtlpp::ResourceUsage)(mtlpp::ResourceUsage::Read|mtlpp::ResourceUsage::Write)));
+							check(VB->Buffer.GetStorageMode() == mtlpp::StorageMode::Private);
+							BufferSizes.Add(VB->GetSize());
+							BufferSizes.Add(GMetalBufferFormats[SRV->Format].DataFormat);
 						}
-						else
+						else if (IB)
 						{
-							check(VB || IB || SB);
-							ns::AutoReleased<FMetalTexture> Tex = SRV->GetLinearTexture(false);
-							Desc.SetDataType(mtlpp::DataType::Texture);
-							Desc.SetTextureType(Tex.GetTextureType());
-							IndirectArgumentResources.Add(Argument(Tex, (mtlpp::ResourceUsage)(mtlpp::ResourceUsage::Read|mtlpp::ResourceUsage::Write)));
+							BufferDesc.SetDataType(mtlpp::DataType::Pointer);
+							NewIAB->IndirectArgumentResources.Add(Argument(IB->Buffer, (mtlpp::ResourceUsage)(mtlpp::ResourceUsage::Read)));
 							
-							mtlpp::ArgumentDescriptor BufferDesc;
-							BufferDesc.SetIndex(++Index);
-							BufferDesc.SetAccess(mtlpp::ArgumentAccess::ReadWrite);
-							[Arguments addObject:BufferDesc.GetPtr()];
+							check(IB->Buffer.GetStorageMode() == mtlpp::StorageMode::Private);
+							BufferSizes.Add(IB->GetSize());
+							BufferSizes.Add(GMetalBufferFormats[SRV->Format].DataFormat);
+						}
+						else if (SB)
+						{
+							BufferDesc.SetDataType(mtlpp::DataType::Pointer);
+							NewIAB->IndirectArgumentResources.Add(Argument(SB->Buffer, (mtlpp::ResourceUsage)(mtlpp::ResourceUsage::Read)));
 							
-							if (VB)
-							{
-								BufferDesc.SetDataType(mtlpp::DataType::Pointer);
-								IndirectArgumentResources.Add(Argument(VB->Buffer, (mtlpp::ResourceUsage)(mtlpp::ResourceUsage::Read|mtlpp::ResourceUsage::Write)));
-								
-								check(VB->Buffer.GetStorageMode() == mtlpp::StorageMode::Private);
-								BufferSizes.Add(VB->GetSize());
-								BufferSizes.Add(GMetalBufferFormats[SRV->Format].DataFormat);
-							}
-							else if (IB)
-							{
-								BufferDesc.SetDataType(mtlpp::DataType::Pointer);
-								IndirectArgumentResources.Add(Argument(IB->Buffer, (mtlpp::ResourceUsage)(mtlpp::ResourceUsage::Read|mtlpp::ResourceUsage::Write)));
-								
-								check(IB->Buffer.GetStorageMode() == mtlpp::StorageMode::Private);
-								BufferSizes.Add(IB->GetSize());
-								BufferSizes.Add(GMetalBufferFormats[SRV->Format].DataFormat);
-							}
-							else if (SB)
-							{
-								BufferDesc.SetDataType(mtlpp::DataType::Pointer);
-								IndirectArgumentResources.Add(Argument(SB->Buffer, (mtlpp::ResourceUsage)(mtlpp::ResourceUsage::Read|mtlpp::ResourceUsage::Write)));
-								
-								BufferSizes.Add(SB->GetSize());
-								BufferSizes.Add(GMetalBufferFormats[SRV->Format].DataFormat);
-							}
+							BufferSizes.Add(SB->GetSize());
+							BufferSizes.Add(GMetalBufferFormats[SRV->Format].DataFormat);
 						}
-						break;
 					}
-					case UBMT_SAMPLER:
-					{
-						mtlpp::ArgumentDescriptor Desc;
-						Desc.SetIndex(Index);
-						Desc.SetAccess(mtlpp::ArgumentAccess::ReadOnly);
-						[Arguments addObject:Desc.GetPtr()];
-						
-						FMetalSamplerState* Sampler = (FMetalSamplerState*)Resource;
-						Desc.SetDataType(mtlpp::DataType::Sampler);
-						IndirectArgumentResources.Add(Argument(Sampler->State));
-						break;
-					}
-					case UBMT_GRAPH_TRACKED_TEXTURE:
-					case UBMT_TEXTURE:
+					break;
+				}
+				case UBMT_GRAPH_TRACKED_UAV:
+				case UBMT_GRAPH_TRACKED_BUFFER_UAV:
+				{
+					FMetalArgumentDesc& Desc = Arguments.Emplace_GetRef();
+					Desc.SetIndex(Index);
+					Desc.SetAccess(mtlpp::ArgumentAccess::ReadWrite);
+					
+					FMetalUnorderedAccessView* UAV = (FMetalUnorderedAccessView*)Resource;
+					FMetalShaderResourceView* SRV = (FMetalShaderResourceView*)UAV->SourceView;
+					FMetalStructuredBuffer* SB = UAV->SourceView->SourceStructuredBuffer.GetReference();
+					FMetalVertexBuffer* VB = UAV->SourceView->SourceVertexBuffer.GetReference();
+					FMetalIndexBuffer* IB = UAV->SourceView->SourceIndexBuffer.GetReference();
+					FRHITexture* Texture = UAV->SourceView->SourceTexture.GetReference();
+					FMetalSurface* Surface = UAV->SourceView->TextureView;
+					if (Texture)
 					{
-						mtlpp::ArgumentDescriptor Desc;
-						Desc.SetIndex(Index);
-						Desc.SetAccess(mtlpp::ArgumentAccess::ReadOnly);
-						[Arguments addObject:Desc.GetPtr()];
-						
-						FRHITexture* Texture = (FRHITexture*)Resource;
-						FMetalSurface* Surface = GetMetalSurfaceFromRHITexture(Texture);
+						if (!Surface)
+						{
+							Surface = GetMetalSurfaceFromRHITexture(Texture);
+						}
 						check (Surface != nullptr);
 						Desc.SetDataType(mtlpp::DataType::Texture);
 						Desc.SetTextureType(Surface->Texture.GetTextureType());
 						
 						check(!Surface->Texture.IsAliasable());
-						IndirectArgumentResources.Add(Argument(Surface->Texture, (mtlpp::ResourceUsage)(mtlpp::ResourceUsage::Read|mtlpp::ResourceUsage::Sample)));
-						break;
+						NewIAB->IndirectArgumentResources.Add(Argument(Surface->Texture, (mtlpp::ResourceUsage)(mtlpp::ResourceUsage::Read|mtlpp::ResourceUsage::Write)));
+					}
+					else
+					{
+						check(VB || IB || SB);
+						ns::AutoReleased<FMetalTexture> Tex = SRV->GetLinearTexture(false);
+						Desc.SetDataType(mtlpp::DataType::Texture);
+						Desc.SetTextureType(Tex.GetTextureType());
+						NewIAB->IndirectArgumentResources.Add(Argument(Tex, (mtlpp::ResourceUsage)(mtlpp::ResourceUsage::Read|mtlpp::ResourceUsage::Write)));
+						
+						FMetalArgumentDesc& BufferDesc = Arguments.Emplace_GetRef();
+						BufferDesc.SetIndex(++Index);
+						BufferDesc.SetAccess(mtlpp::ArgumentAccess::ReadWrite);
+						
+						if (VB)
+						{
+							BufferDesc.SetDataType(mtlpp::DataType::Pointer);
+							NewIAB->IndirectArgumentResources.Add(Argument(VB->Buffer, (mtlpp::ResourceUsage)(mtlpp::ResourceUsage::Read|mtlpp::ResourceUsage::Write)));
+							
+							check(VB->Buffer.GetStorageMode() == mtlpp::StorageMode::Private);
+							BufferSizes.Add(VB->GetSize());
+							BufferSizes.Add(GMetalBufferFormats[SRV->Format].DataFormat);
+						}
+						else if (IB)
+						{
+							BufferDesc.SetDataType(mtlpp::DataType::Pointer);
+							NewIAB->IndirectArgumentResources.Add(Argument(IB->Buffer, (mtlpp::ResourceUsage)(mtlpp::ResourceUsage::Read|mtlpp::ResourceUsage::Write)));
+							
+							check(IB->Buffer.GetStorageMode() == mtlpp::StorageMode::Private);
+							BufferSizes.Add(IB->GetSize());
+							BufferSizes.Add(GMetalBufferFormats[SRV->Format].DataFormat);
+						}
+						else if (SB)
+						{
+							BufferDesc.SetDataType(mtlpp::DataType::Pointer);
+							NewIAB->IndirectArgumentResources.Add(Argument(SB->Buffer, (mtlpp::ResourceUsage)(mtlpp::ResourceUsage::Read|mtlpp::ResourceUsage::Write)));
+							
+							BufferSizes.Add(SB->GetSize());
+							BufferSizes.Add(GMetalBufferFormats[SRV->Format].DataFormat);
+						}
 					}
-					default:
-						break;
+					break;
 				}
-			}
-			
-			if (BufferSizes.Num() > 0)
-			{
-				mtlpp::ArgumentDescriptor Desc;
-				Desc.SetIndex(Index++);
-				Desc.SetAccess(mtlpp::ArgumentAccess::ReadOnly);
-				Desc.SetDataType(mtlpp::DataType::Pointer);
-				[Arguments addObject:Desc.GetPtr()];
-				
-				FMetalPooledBufferArgs Args(GetMetalDeviceContext().GetDevice(), BufferSizes.Num() * sizeof(uint32), BUFFER_STORAGE_MODE);
-				IndirectArgumentBufferSideTable = GetMetalDeviceContext().CreatePooledBuffer(Args);
-				
-				FMemory::Memcpy(IndirectArgumentBufferSideTable.GetContents(), BufferSizes.GetData(), BufferSizes.Num() * sizeof(uint32));
-				
-#if PLATFORM_MAC
-				if(IndirectArgumentBufferSideTable.GetStorageMode() == mtlpp::StorageMode::Managed)
+				case UBMT_SAMPLER:
 				{
-					MTLPP_VALIDATE(mtlpp::Buffer, IndirectArgumentBufferSideTable, SafeGetRuntimeDebuggingLevel() >= EMetalDebugLevelValidation, DidModify(ns::Range(0, BufferSizes.Num() * sizeof(uint32))));
+					FMetalArgumentDesc& Desc = Arguments.Emplace_GetRef();
+					Desc.SetIndex(Index);
+					Desc.SetAccess(mtlpp::ArgumentAccess::ReadOnly);
+					
+					FMetalSamplerState* Sampler = (FMetalSamplerState*)Resource;
+					Desc.SetDataType(mtlpp::DataType::Sampler);
+					NewIAB->IndirectArgumentResources.Add(Argument(Sampler->State));
+					break;
 				}
-#endif
-				
-				IndirectArgumentResources.Add(Argument(IndirectArgumentBufferSideTable, mtlpp::ResourceUsage::Read));
-			}
-			
-			if (Layout.ConstantBufferSize > 0)
-            {
-                mtlpp::ArgumentDescriptor Desc;
-                Desc.SetIndex(Index++);
-                Desc.SetAccess(mtlpp::ArgumentAccess::ReadOnly);
-                Desc.SetDataType(mtlpp::DataType::Pointer);
-                [Arguments addObject:Desc.GetPtr()];
-				
-                IndirectArgumentResources.Add(Argument(Buffer, mtlpp::ResourceUsage::Read));
-			}
-			
-			ns::Array<mtlpp::ArgumentDescriptor> ArgArray(Arguments);
-			mtlpp::ArgumentEncoder Encoder = FMetalArgumentEncoderCache::Get().CreateEncoder(ArgArray);
-			
-			IndirectArgumentBuffer = GetMetalDeviceContext().GetResourceHeap().CreateBuffer(Encoder.GetEncodedLength(), 16, mtlpp::ResourceOptions(BUFFER_CACHE_MODE | mtlpp::ResourceOptions::HazardTrackingModeUntracked | ((NSUInteger)BUFFER_STORAGE_MODE << mtlpp::ResourceStorageModeShift)), true);
-			
-			Encoder.SetArgumentBuffer(IndirectArgumentBuffer, 0);
-			
-			for (MTLArgumentDescriptor* Arg in Arguments)
-			{
-				mtlpp::ArgumentDescriptor Argument(Arg);
-				NSUInteger NewIndex = Argument.GetIndex();
-				switch(Argument.GetDataType())
+				case UBMT_GRAPH_TRACKED_TEXTURE:
+				case UBMT_TEXTURE:
 				{
-					case mtlpp::DataType::Pointer:
-						Encoder.SetBuffer(IndirectArgumentResources[NewIndex].Buffer, 0, NewIndex);
-						break;
-					case mtlpp::DataType::Texture:
-						Encoder.SetTexture(IndirectArgumentResources[NewIndex].Texture, NewIndex);
-						break;
-					case mtlpp::DataType::Sampler:
-						Encoder.SetSamplerState(IndirectArgumentResources[NewIndex].Sampler, NewIndex);
-						break;
-					case mtlpp::DataType::Struct:
-					{
-						void* Data = Encoder.GetConstantDataAtIndex(NewIndex);
-						FMemory::Memcpy(Data, Contents, Layout.ConstantBufferSize);
-						break;
-					}
-					default:
-						break;
+					FMetalArgumentDesc& Desc = Arguments.Emplace_GetRef();
+					Desc.SetIndex(Index);
+					Desc.SetAccess(mtlpp::ArgumentAccess::ReadOnly);
+					
+					FRHITexture* Texture = (FRHITexture*)Resource;
+					FMetalSurface* Surface = GetMetalSurfaceFromRHITexture(Texture);
+					check (Surface != nullptr);
+					Desc.SetDataType(mtlpp::DataType::Texture);
+					Desc.SetTextureType(Surface->Texture.GetTextureType());
+					
+					check(!Surface->Texture.IsAliasable());
+					NewIAB->IndirectArgumentResources.Add(Argument(Surface->Texture, (mtlpp::ResourceUsage)(mtlpp::ResourceUsage::Read|mtlpp::ResourceUsage::Sample)));
+					break;
 				}
+				default:
+					break;
 			}
+		}
+		
+		if (BufferSizes.Num() > 0)
+		{
+			FMetalArgumentDesc& Desc = Arguments.Emplace_GetRef();
+			Desc.SetIndex(Index++);
+			Desc.SetAccess(mtlpp::ArgumentAccess::ReadOnly);
+			Desc.SetDataType(mtlpp::DataType::Pointer);
+			
+			FMetalPooledBufferArgs Args(GetMetalDeviceContext().GetDevice(), BufferSizes.Num() * sizeof(uint32), BUFFER_STORAGE_MODE);
+			NewIAB->IndirectArgumentBufferSideTable = GetMetalDeviceContext().CreatePooledBuffer(Args);
+			
+			FMemory::Memcpy(NewIAB->IndirectArgumentBufferSideTable.GetContents(), BufferSizes.GetData(), BufferSizes.Num() * sizeof(uint32));
 			
 #if PLATFORM_MAC
-			if(IndirectArgumentBuffer.GetStorageMode() == mtlpp::StorageMode::Managed)
+			if(NewIAB->IndirectArgumentBufferSideTable.GetStorageMode() == mtlpp::StorageMode::Managed)
 			{
-				MTLPP_VALIDATE(mtlpp::Buffer, IndirectArgumentBuffer, SafeGetRuntimeDebuggingLevel() >= EMetalDebugLevelValidation, DidModify(ns::Range(0, Encoder.GetEncodedLength())));
+				MTLPP_VALIDATE(mtlpp::Buffer, NewIAB->IndirectArgumentBufferSideTable, SafeGetRuntimeDebuggingLevel() >= EMetalDebugLevelValidation, DidModify(ns::Range(0, BufferSizes.Num() * sizeof(uint32))));
 			}
 #endif
+			
+			NewIAB->IndirectArgumentResources.Add(Argument(NewIAB->IndirectArgumentBufferSideTable, mtlpp::ResourceUsage::Read));
 		}
-		else
+		
+		if (GetLayout().ConstantBufferSize > 0)
 		{
-			int32 NumResources = Layout.Resources.Num();
-			ResourceTable.Empty(NumResources);
-			ResourceTable.AddZeroed(NumResources);
-			for (int32 i = 0; i < NumResources; ++i)
+			FMetalArgumentDesc& Desc = Arguments.Emplace_GetRef();
+			Desc.SetIndex(Index++);
+			Desc.SetAccess(mtlpp::ArgumentAccess::ReadOnly);
+			Desc.SetDataType(mtlpp::DataType::Pointer);
+			
+			NewIAB->IndirectArgumentResources.Add(Argument(Buffer, mtlpp::ResourceUsage::Read));
+		}
+		
+		mtlpp::ArgumentEncoder Encoder = FMetalArgumentEncoderCache::Get().CreateEncoder(Arguments);
+		
+		NewIAB->IndirectArgumentBuffer = GetMetalDeviceContext().GetResourceHeap().CreateBuffer(Encoder.GetEncodedLength(), 16, mtlpp::ResourceOptions(BUFFER_CACHE_MODE | ((NSUInteger)BUFFER_STORAGE_MODE << mtlpp::ResourceStorageModeShift)), true);
+		
+		Encoder.SetArgumentBuffer(NewIAB->IndirectArgumentBuffer, 0);
+		
+		for (FMetalArgumentDesc const& Arg : Arguments)
+		{
+			NSUInteger NewIndex = Arg.Index;
+			switch(Arg.DataType)
 			{
-				FRHIResource* Resource = *(FRHIResource**)((uint8*)Contents + Layout.ResourceOffsets[i]);
-				
-				// Allow null SRV's in uniform buffers for feature levels that don't support SRV's in shaders
-				if (!(GMaxRHIFeatureLevel <= ERHIFeatureLevel::ES3_1 && Layout.Resources[i] == UBMT_SRV))
-				{
-					check(Resource);
-				}
-				ResourceTable[i] = Resource;
+				case mtlpp::DataType::Pointer:
+					Encoder.SetBuffer(NewIAB->IndirectArgumentResources[NewIndex].Buffer, 0, NewIndex);
+					break;
+				case mtlpp::DataType::Texture:
+					Encoder.SetTexture(NewIAB->IndirectArgumentResources[NewIndex].Texture, NewIndex);
+					break;
+				case mtlpp::DataType::Sampler:
+					Encoder.SetSamplerState(NewIAB->IndirectArgumentResources[NewIndex].Sampler, NewIndex);
+					break;
+				default:
+					break;
 			}
 		}
+		
+#if PLATFORM_MAC
+		if(NewIAB->IndirectArgumentBuffer.GetStorageMode() == mtlpp::StorageMode::Managed)
+		{
+			MTLPP_VALIDATE(mtlpp::Buffer, NewIAB->IndirectArgumentBuffer, SafeGetRuntimeDebuggingLevel() >= EMetalDebugLevelValidation, DidModify(ns::Range(0, Encoder.GetEncodedLength())));
+		}
+#endif
+		
+		// Atomically swap so that we don't explode if multiple threads attempt to initialise at the same time.
+		void* Result = FPlatformAtomics::InterlockedCompareExchangePointer((void**)&IAB, (void*)NewIAB, nullptr);
+		if (Result != nullptr)
+		{
+			check(Result != NewIAB);
+			delete NewIAB;
+		}
 	}
 }
 
-FMetalUniformBuffer::~FMetalUniformBuffer()
-{
-	SafeReleaseMetalBuffer(IndirectArgumentBuffer);
-	SafeReleaseMetalBuffer(IndirectArgumentBufferSideTable);
-}
-
 void const* FMetalUniformBuffer::GetData()
 {
 	if (Data)
diff --git a/Engine/Source/Runtime/Apple/MetalRHI/Public/MetalResources.h b/Engine/Source/Runtime/Apple/MetalRHI/Public/MetalResources.h
--- a/Engine/Source/Runtime/Apple/MetalRHI/Public/MetalResources.h
+++ b/Engine/Source/Runtime/Apple/MetalRHI/Public/MetalResources.h
@@ -856,6 +856,8 @@ public:
 	virtual ~FMetalUniformBuffer();
 	
 	void const* GetData();
+
+	void InitIAB();
 	
 	/** Resource table containing RHI references. */
 	TArray<TRefCountPtr<FRHIResource> > ResourceTable;
@@ -873,9 +875,18 @@ public:
 		mtlpp::ResourceUsage Usage;
 	};
 	
-	TArray<Argument> IndirectArgumentResources;
-	FMetalBuffer IndirectArgumentBuffer;
-	FMetalBuffer IndirectArgumentBufferSideTable;
+	struct FMetalIndirectArgumentBuffer
+	{
+		FMetalIndirectArgumentBuffer();
+		~FMetalIndirectArgumentBuffer();
+		
+		TArray<Argument> IndirectArgumentResources;
+		FMetalBuffer IndirectArgumentBuffer;
+		FMetalBuffer IndirectArgumentBufferSideTable;
+	};
+	
+	FMetalIndirectArgumentBuffer& GetIAB();
+	FMetalIndirectArgumentBuffer* IAB;
 };
 
 
