Handle invalidating IABs when a contained texture reference is updated, this implementation is a bit simple-minded and might need future revisions.

diff --git a/Engine/Source/Runtime/Apple/MetalRHI/Private/MetalContext.cpp b/Engine/Source/Runtime/Apple/MetalRHI/Private/MetalContext.cpp
--- a/Engine/Source/Runtime/Apple/MetalRHI/Private/MetalContext.cpp
+++ b/Engine/Source/Runtime/Apple/MetalRHI/Private/MetalContext.cpp
@@ -715,6 +715,35 @@ void FMetalDeviceContext::ReleaseFence(FMetalFence* Fence)
 	}
 }
 
+void FMetalDeviceContext::RegisterUB(FMetalUniformBuffer* UB)
+{
+	UniformBuffers.Add(UB);
+}
+
+void FMetalDeviceContext::UpdateIABs(FTextureReferenceRHIParamRef ModifiedRef)
+{
+	if(GIsMetalInitialized)
+	{
+		for (FMetalUniformBuffer* UB : UniformBuffers)
+		{
+			if (UB->IAB && UB->TextureReferences.Contains(ModifiedRef))
+			{
+				delete UB->IAB;
+				UB->IAB = nullptr;
+				UB->InitIAB();
+			}
+		}
+	}
+}
+
+void FMetalDeviceContext::UnregisterUB(FMetalUniformBuffer* UB)
+{
+	if(GIsMetalInitialized)
+	{
+		UniformBuffers.Remove(UB);
+	}
+}
+
 FMetalTexture FMetalDeviceContext::CreateTexture(FMetalSurface* Surface, mtlpp::TextureDescriptor Descriptor)
 {
 	FMetalTexture Tex = Heap.CreateTexture(Descriptor, Surface);
diff --git a/Engine/Source/Runtime/Apple/MetalRHI/Private/MetalContext.h b/Engine/Source/Runtime/Apple/MetalRHI/Private/MetalContext.h
--- a/Engine/Source/Runtime/Apple/MetalRHI/Private/MetalContext.h
+++ b/Engine/Source/Runtime/Apple/MetalRHI/Private/MetalContext.h
@@ -187,6 +187,9 @@ public:
 	void ReleaseTexture(FMetalSurface* Surface, FMetalTexture& Texture);
 	void ReleaseTexture(FMetalTexture& Texture);
 	void ReleaseFence(FMetalFence* Fence);
+	void RegisterUB(FMetalUniformBuffer* UB);
+	void UpdateIABs(FTextureReferenceRHIParamRef ModifiedRef);
+	void UnregisterUB(FMetalUniformBuffer* UB);
 	
 	void BeginFrame();
 	void FlushFreeList(bool const bFlushFences = true);
@@ -259,6 +262,8 @@ private:
 	};
 	TArray<FMetalDelayedFreeList*> DelayedFreeLists;
 	
+	TSet<FMetalUniformBuffer*> UniformBuffers;
+	
 #if METAL_DEBUG_OPTIONS
 	/** The list of fences for the current frame */
 	TArray<FMetalFence*> FrameFences;
diff --git a/Engine/Source/Runtime/Apple/MetalRHI/Private/MetalTexture.cpp b/Engine/Source/Runtime/Apple/MetalRHI/Private/MetalTexture.cpp
--- a/Engine/Source/Runtime/Apple/MetalRHI/Private/MetalTexture.cpp
+++ b/Engine/Source/Runtime/Apple/MetalRHI/Private/MetalTexture.cpp
@@ -22,21 +22,6 @@ FAutoConsoleVariableRef CVarMetalMaxOutstandingAsyncTexUploads(
 															   ECVF_ReadOnly|ECVF_RenderThreadSafe
 															   );
 
-enum EMetalTextureCacheMode
-{
-	EMetalTextureCacheModeOff = 0,
-	EMetalTextureCacheModeInFrame = 1,
-	EMetalTextureCacheModeAlways = 2
-};
-
-int32 GMetalTextureCacheMode = 0;
-FAutoConsoleVariableRef CVarMetalTextureCacheMode(
-												  TEXT("rhi.Metal.TextureCacheMode"),
-												  GMetalTextureCacheMode,
-												  TEXT("Set the internal texture cache mode to use in Metal.\n\t0: Off.\n\t1: Mark as volatile during streaming & either reuse within the frame or delete at the end.\n\t2: Always cache the texture object but if not reused within the frame, mark the backing store as empty to clear from VRAM. Default is 1."),
-												  ECVF_ReadOnly|ECVF_RenderThreadSafe
-												  );
-
 
 /** Texture reference class. */
 class FMetalTextureReference : public FRHITextureReference
@@ -2499,6 +2484,10 @@ void FMetalRHICommandContext::RHIUpdateTextureReference(FTextureReferenceRHIPara
 		if (TextureRef)
 		{
 			TextureRef->SetReferencedTexture(NewTextureRHI);
+			if (NewTextureRHI)
+			{
+				GetMetalDeviceContext().UpdateIABs(TextureRefRHI);
+			}
 		}
 	}
 }
diff --git a/Engine/Source/Runtime/Apple/MetalRHI/Private/MetalUniformBuffer.cpp b/Engine/Source/Runtime/Apple/MetalRHI/Private/MetalUniformBuffer.cpp
--- a/Engine/Source/Runtime/Apple/MetalRHI/Private/MetalUniformBuffer.cpp
+++ b/Engine/Source/Runtime/Apple/MetalRHI/Private/MetalUniformBuffer.cpp
@@ -232,6 +232,50 @@ FMetalUniformBuffer::FMetalUniformBuffer(const void* Contents, const FRHIUniform
 				check(Resource);
 			}
 			ResourceTable[i] = Resource;
+			
+			if (FMetalCommandQueue::SupportsFeature(EMetalFeaturesIABs))
+			{
+				switch(Layout.Resources[i])
+				{
+					case UBMT_GRAPH_TRACKED_SRV:
+					case UBMT_GRAPH_TRACKED_BUFFER_SRV:
+					case UBMT_SRV:
+					{
+						FMetalShaderResourceView* SRV = (FMetalShaderResourceView*)Resource;
+						FRHITexture* Texture = SRV->SourceTexture.GetReference();
+						if (Texture && Texture->GetTextureReference())
+						{
+							TextureReferences.Add(Texture->GetTextureReference());
+						}
+						break;
+					}
+					case UBMT_GRAPH_TRACKED_UAV:
+					case UBMT_GRAPH_TRACKED_BUFFER_UAV:
+					{
+						FMetalUnorderedAccessView* UAV = (FMetalUnorderedAccessView*)Resource;
+						FRHITexture* Texture = UAV->SourceView->SourceTexture.GetReference();
+						if (Texture && Texture->GetTextureReference())
+						{
+							TextureReferences.Add(Texture->GetTextureReference());
+						}
+						break;
+					}
+					case UBMT_GRAPH_TRACKED_TEXTURE:
+					case UBMT_TEXTURE:
+					{
+						FRHITexture* Texture = (FRHITexture*)Resource;
+						if (Texture && Texture->GetTextureReference())
+						{
+							TextureReferences.Add(Texture->GetTextureReference());
+						}
+						break;
+					}
+					default:
+					{
+						break;
+					}
+				}
+			}
 		}
 		
 		FRHICommandListImmediate& RHICmdList = FRHICommandListExecutor::GetImmediateCommandList();
@@ -246,6 +290,12 @@ FMetalUniformBuffer::~FMetalUniformBuffer()
 {
 	if (IAB)
 	{
+		// Single draw calls don't need to permit texture-reference updates as that wouldn't make sense.
+		if (!(Usage & UniformBuffer_SingleDraw))
+		{
+			GetMetalDeviceContext().UnregisterUB(this);
+		}
+		
 		delete IAB;
 		IAB = nullptr;
 	}
@@ -523,6 +573,12 @@ void FMetalUniformBuffer::InitIAB()
 			check(Result != NewIAB);
 			delete NewIAB;
 		}
+		
+		// Single draw calls don't need to permit texture-reference updates as that wouldn't make sense.
+		if (!(Usage & UniformBuffer_SingleDraw))
+		{
+			GetMetalDeviceContext().RegisterUB(this);
+		}
 	}
 }
 
diff --git a/Engine/Source/Runtime/Apple/MetalRHI/Public/MetalResources.h b/Engine/Source/Runtime/Apple/MetalRHI/Public/MetalResources.h
--- a/Engine/Source/Runtime/Apple/MetalRHI/Public/MetalResources.h
+++ b/Engine/Source/Runtime/Apple/MetalRHI/Public/MetalResources.h
@@ -861,6 +861,8 @@ public:
 	
 	/** Resource table containing RHI references. */
 	TArray<TRefCountPtr<FRHIResource> > ResourceTable;
+	
+	TSet<FTextureReferenceRHIParamRef> TextureReferences;
 
 	struct Argument
 	{
