Use the present convenience functions, it seems that there's no longer any difference between this and the previous method that I can see and this is simpler.

diff --git a/Engine/Source/Runtime/Apple/MetalRHI/Private/MetalViewport.cpp b/Engine/Source/Runtime/Apple/MetalRHI/Private/MetalViewport.cpp
--- a/Engine/Source/Runtime/Apple/MetalRHI/Private/MetalViewport.cpp
+++ b/Engine/Source/Runtime/Apple/MetalRHI/Private/MetalViewport.cpp
@@ -510,23 +510,18 @@ void FMetalViewport::Present(FMetalCommandQueue& CommandQueue, bool bLockToVsync
 #endif
 						};
 						
-#if !PLATFORM_IOS
-						mtlpp::CommandBufferHandler H = [LocalDrawable](mtlpp::CommandBuffer const&) {
-#else
-						mtlpp::CommandBufferHandler H = [LocalDrawable, MinPresentDuration, FramePace](mtlpp::CommandBuffer const&) {
-							if (MinPresentDuration && GEnablePresentPacing)
-							{
-								[LocalDrawable presentAfterMinimumDuration:1.0f/(float)FramePace];
-							}
-							else
+#if PLATFORM_IOS
+						if (MinPresentDuration && GEnablePresentPacing)
+						{
+							CurrentCommandBuffer.PresentAfterMinimumDuration(LocalDrawable, 1.0f/(float)FramePace);
+						}
+						else
 #endif
-							{
-								[LocalDrawable present];
-							};
-						};
+						{
+							CurrentCommandBuffer.Present(LocalDrawable);
+						}
 						
 						CurrentCommandBuffer.AddCompletedHandler(C);
-						CurrentCommandBuffer.AddScheduledHandler(H);
 						
 						METAL_GPUPROFILE(Stats->End(CurrentCommandBuffer));
 						CommandQueue.CommitCommandBuffer(CurrentCommandBuffer);
