More IAB bug fixes for threading violations and incorrect re-entrancy while iterating over a set.

diff --git a/Engine/Source/Runtime/Apple/MetalRHI/Private/MetalContext.cpp b/Engine/Source/Runtime/Apple/MetalRHI/Private/MetalContext.cpp
--- a/Engine/Source/Runtime/Apple/MetalRHI/Private/MetalContext.cpp
+++ b/Engine/Source/Runtime/Apple/MetalRHI/Private/MetalContext.cpp
@@ -717,6 +717,7 @@ void FMetalDeviceContext::ReleaseFence(FMetalFence* Fence)
 
 void FMetalDeviceContext::RegisterUB(FMetalUniformBuffer* UB)
 {
+	FScopeLock Lock(&FreeListMutex);
 	UniformBuffers.Add(UB);
 }
 
@@ -724,13 +725,15 @@ void FMetalDeviceContext::UpdateIABs(FTextureReferenceRHIParamRef ModifiedRef)
 {
 	if(GIsMetalInitialized)
 	{
+		FScopeLock Lock(&FreeListMutex);
 		for (FMetalUniformBuffer* UB : UniformBuffers)
 		{
-			if (UB->IAB && UB->TextureReferences.Contains(ModifiedRef))
+			if (UB && UB->IAB && UB->TextureReferences.Contains(ModifiedRef))
 			{
-				delete UB->IAB;
+				FMetalUniformBuffer::FMetalIndirectArgumentBuffer* IAB = UB->IAB;
 				UB->IAB = nullptr;
 				UB->InitIAB();
+				delete IAB;
 			}
 		}
 	}
@@ -740,6 +743,7 @@ void FMetalDeviceContext::UnregisterUB(FMetalUniformBuffer* UB)
 {
 	if(GIsMetalInitialized)
 	{
+		FScopeLock Lock(&FreeListMutex);
 		UniformBuffers.Remove(UB);
 	}
 }
diff --git a/Engine/Source/Runtime/Apple/MetalRHI/Private/MetalUniformBuffer.cpp b/Engine/Source/Runtime/Apple/MetalRHI/Private/MetalUniformBuffer.cpp
--- a/Engine/Source/Runtime/Apple/MetalRHI/Private/MetalUniformBuffer.cpp
+++ b/Engine/Source/Runtime/Apple/MetalRHI/Private/MetalUniformBuffer.cpp
@@ -205,6 +205,7 @@ private:
 FMetalUniformBuffer::FMetalUniformBuffer(const void* Contents, const FRHIUniformBufferLayout& Layout, EUniformBufferUsage InUsage)
 	: FRHIUniformBuffer(Layout)
     , FMetalRHIBuffer(Layout.ConstantBufferSize, (FMetalCommandQueue::SupportsFeature(EMetalFeaturesIABs) && Layout.Resources.Num() ? (EMetalBufferUsage_GPUOnly|BUF_Volatile) : BUF_Volatile), RRT_UniformBuffer)
+	, UniformUsage(InUsage)
 	, IAB(nullptr)
 {
 	if (Layout.ConstantBufferSize > 0)
@@ -278,8 +279,10 @@ FMetalUniformBuffer::FMetalUniformBuffer(const void* Contents, const FRHIUniform
 			}
 		}
 		
+		GetMetalDeviceContext().RegisterUB(this);
+		
 		FRHICommandListImmediate& RHICmdList = FRHICommandListExecutor::GetImmediateCommandList();
-		if (!(InUsage & UniformBuffer_SingleDraw) && IsRunningRHIInSeparateThread() && !RHICmdList.Bypass() && IsInRenderingThread())
+		if (!(UniformUsage & UniformBuffer_SingleDraw) && IsRunningRHIInSeparateThread() && !RHICmdList.Bypass() && IsInRenderingThread())
 		{
 			new (RHICmdList.AllocCommand<FMetalRHICommandInitialiseUniformBufferIAB>()) FMetalRHICommandInitialiseUniformBufferIAB(this);
 		}
@@ -290,11 +293,7 @@ FMetalUniformBuffer::~FMetalUniformBuffer()
 {
 	if (IAB)
 	{
-		// Single draw calls don't need to permit texture-reference updates as that wouldn't make sense.
-		if (!(Usage & UniformBuffer_SingleDraw))
-		{
-			GetMetalDeviceContext().UnregisterUB(this);
-		}
+		GetMetalDeviceContext().UnregisterUB(this);
 		
 		delete IAB;
 		IAB = nullptr;
@@ -315,8 +314,10 @@ FMetalUniformBuffer::FMetalIndirectArgumentBuffer::~FMetalIndirectArgumentBuffer
 FMetalUniformBuffer::FMetalIndirectArgumentBuffer& FMetalUniformBuffer::GetIAB()
 {
 	check(ResourceTable.Num() && FMetalCommandQueue::SupportsFeature(EMetalFeaturesIABs));
+
 	InitIAB();
 	check(IAB);
+	
 	return *IAB;
 }
 
@@ -573,12 +574,6 @@ void FMetalUniformBuffer::InitIAB()
 			check(Result != NewIAB);
 			delete NewIAB;
 		}
-		
-		// Single draw calls don't need to permit texture-reference updates as that wouldn't make sense.
-		if (!(Usage & UniformBuffer_SingleDraw))
-		{
-			GetMetalDeviceContext().RegisterUB(this);
-		}
 	}
 }
 
diff --git a/Engine/Source/Runtime/Apple/MetalRHI/Public/MetalResources.h b/Engine/Source/Runtime/Apple/MetalRHI/Public/MetalResources.h
--- a/Engine/Source/Runtime/Apple/MetalRHI/Public/MetalResources.h
+++ b/Engine/Source/Runtime/Apple/MetalRHI/Public/MetalResources.h
@@ -887,6 +887,7 @@ public:
 		FMetalBuffer IndirectArgumentBufferSideTable;
 	};
 	
+	EUniformBufferUsage UniformUsage;
 	FMetalIndirectArgumentBuffer& GetIAB();
 	FMetalIndirectArgumentBuffer* IAB;
 };
