Fix another IAB related bug where uniform-buffers weren't unregistering themselves properly.

diff --git a/Engine/Source/Runtime/Apple/MetalRHI/Private/MetalUniformBuffer.cpp b/Engine/Source/Runtime/Apple/MetalRHI/Private/MetalUniformBuffer.cpp
--- a/Engine/Source/Runtime/Apple/MetalRHI/Private/MetalUniformBuffer.cpp
+++ b/Engine/Source/Runtime/Apple/MetalRHI/Private/MetalUniformBuffer.cpp
@@ -291,10 +291,13 @@ FMetalUniformBuffer::FMetalUniformBuffer(const void* Contents, const FRHIUniform
 
 FMetalUniformBuffer::~FMetalUniformBuffer()
 {
-	if (IAB)
+	if (ResourceTable.Num() && FMetalCommandQueue::SupportsFeature(EMetalFeaturesIABs))
 	{
 		GetMetalDeviceContext().UnregisterUB(this);
-		
+	}
+	
+	if (IAB)
+	{
 		delete IAB;
 		IAB = nullptr;
 	}
