When in the Editor use the old present mechanism as the conveniance functions appear to cause some problems with the completion handler when there are multiple windows.

diff --git a/Engine/Source/Runtime/Apple/MetalRHI/Private/MetalViewport.cpp b/Engine/Source/Runtime/Apple/MetalRHI/Private/MetalViewport.cpp
--- a/Engine/Source/Runtime/Apple/MetalRHI/Private/MetalViewport.cpp
+++ b/Engine/Source/Runtime/Apple/MetalRHI/Private/MetalViewport.cpp
@@ -510,18 +510,45 @@ void FMetalViewport::Present(FMetalCommandQueue& CommandQueue, bool bLockToVsync
 #endif
 						};
 						
-#if PLATFORM_IOS
-						if (MinPresentDuration && GEnablePresentPacing)
+#if WITH_EDITOR			// The Editor needs the older way to present otherwise we end up with bad behaviour of the completion handlers that causes GPU timeouts.
+						if (GIsEditor)
 						{
-							CurrentCommandBuffer.PresentAfterMinimumDuration(LocalDrawable, 1.0f/(float)FramePace);
+#if !PLATFORM_IOS
+							mtlpp::CommandBufferHandler H = [LocalDrawable](mtlpp::CommandBuffer const&)
+							{
+#else
+							mtlpp::CommandBufferHandler H = [LocalDrawable, MinPresentDuration, FramePace](mtlpp::CommandBuffer const&)
+							{
+								if (MinPresentDuration && GEnablePresentPacing)
+								{
+									[LocalDrawable presentAfterMinimumDuration:1.0f/(float)FramePace];
+								}
+								else
+#endif
+								{
+									[LocalDrawable present];
+								};
+							};
+								
+							CurrentCommandBuffer.AddCompletedHandler(C);
+							CurrentCommandBuffer.AddScheduledHandler(H);
 						}
 						else
 #endif
 						{
-							CurrentCommandBuffer.Present(LocalDrawable);
+							CurrentCommandBuffer.AddCompletedHandler(C);
+#if PLATFORM_IOS
+							if (MinPresentDuration && GEnablePresentPacing)
+							{
+								CurrentCommandBuffer.PresentAfterMinimumDuration(LocalDrawable, 1.0f/(float)FramePace);
+							}
+							else
+#endif
+							{
+								CurrentCommandBuffer.Present(LocalDrawable);
+							}
 						}
 						
-						CurrentCommandBuffer.AddCompletedHandler(C);
 						
 						METAL_GPUPROFILE(Stats->End(CurrentCommandBuffer));
 						CommandQueue.CommitCommandBuffer(CurrentCommandBuffer);
