Buffer mutability in the pipeline state - Uniform Buffers should be, by definition, constant whereas the other buffers aren't.

diff --git a/Engine/Source/Runtime/Apple/MetalRHI/Private/MetalPipeline.cpp b/Engine/Source/Runtime/Apple/MetalRHI/Private/MetalPipeline.cpp
--- a/Engine/Source/Runtime/Apple/MetalRHI/Private/MetalPipeline.cpp
+++ b/Engine/Source/Runtime/Apple/MetalRHI/Private/MetalPipeline.cpp
@@ -388,6 +388,43 @@ static FMetalShaderPipeline* CreateMTLRenderPipeline(bool const bSync, FMetalGra
 		mtlpp::RenderPipelineDescriptor RenderPipelineDesc;
 		mtlpp::ComputePipelineDescriptor ComputePipelineDesc(nil);
 		
+		if (FMetalCommandQueue::SupportsFeature(EMetalFeaturesGPUCommandBufferTimes))
+		{
+			ns::AutoReleased<ns::Array<mtlpp::PipelineBufferDescriptor>> VertexPipelineBuffers = RenderPipelineDesc.GetVertexBuffers();
+			FMetalShaderBindings& VertexBindings = DomainShader ? DomainShader->Bindings : VertexShader->Bindings;
+			{
+				uint32 ImmutableBuffers = VertexBindings.ConstantBuffers | VertexBindings.ArgumentBuffers;
+				while(ImmutableBuffers)
+				{
+					uint32 Index = __builtin_ctz(ImmutableBuffers);
+					ImmutableBuffers &= ~(1 << Index);
+					
+					if (Index < ML_MaxBuffers)
+					{
+						ns::AutoReleased<mtlpp::PipelineBufferDescriptor> PipelineBuffer = VertexPipelineBuffers[Index];
+						PipelineBuffer.SetMutability(mtlpp::Mutability::Immutable);
+					}
+				}
+			}
+			
+			if (PixelShader)
+			{
+				ns::AutoReleased<ns::Array<mtlpp::PipelineBufferDescriptor>> FragmentPipelineBuffers = RenderPipelineDesc.GetFragmentBuffers();
+				uint32 ImmutableBuffers = PixelShader->Bindings.ConstantBuffers | PixelShader->Bindings.ArgumentBuffers;
+				while(ImmutableBuffers)
+				{
+					uint32 Index = __builtin_ctz(ImmutableBuffers);
+					ImmutableBuffers &= ~(1 << Index);
+					
+					if (Index < ML_MaxBuffers)
+					{
+						ns::AutoReleased<mtlpp::PipelineBufferDescriptor> PipelineBuffer = FragmentPipelineBuffers[Index];
+						PipelineBuffer.SetMutability(mtlpp::Mutability::Immutable);
+					}
+				}
+			}
+		}
+		
         FMetalBlendState* BlendState = (FMetalBlendState*)Init.BlendState;
 		
 		ns::Array<mtlpp::RenderPipelineColorAttachmentDescriptor> ColorAttachments = RenderPipelineDesc.GetColorAttachments();
@@ -527,6 +564,24 @@ static FMetalShaderPipeline* CreateMTLRenderPipeline(bool const bSync, FMetalGra
 			ComputePipelineDesc = mtlpp::ComputePipelineDescriptor();
             check(ComputePipelineDesc);
 			
+			if (FMetalCommandQueue::SupportsFeature(EMetalFeaturesGPUCommandBufferTimes))
+			{
+				ns::AutoReleased<ns::Array<mtlpp::PipelineBufferDescriptor>> PipelineBuffers = ComputePipelineDesc.GetBuffers();
+				
+				uint32 ImmutableBuffers = VertexShader->Bindings.ConstantBuffers | VertexShader->Bindings.ArgumentBuffers;
+				while(ImmutableBuffers)
+				{
+					uint32 Index = __builtin_ctz(ImmutableBuffers);
+					ImmutableBuffers &= ~(1 << Index);
+					
+					if (Index < ML_MaxBuffers)
+					{
+						ns::AutoReleased<mtlpp::PipelineBufferDescriptor> PipelineBuffer = PipelineBuffers[Index];
+						PipelineBuffer.SetMutability(mtlpp::Mutability::Immutable);
+					}
+				}
+			}
+			
 			mtlpp::VertexDescriptor DomainVertexDesc;
 			mtlpp::StageInputOutputDescriptor ComputeStageInOut;
 			ComputeStageInOut.SetIndexBufferIndex(VertexShader->TessellationControlPointIndexBuffer);
diff --git a/Engine/Source/Runtime/Apple/MetalRHI/Private/MetalShaders.cpp b/Engine/Source/Runtime/Apple/MetalRHI/Private/MetalShaders.cpp
--- a/Engine/Source/Runtime/Apple/MetalRHI/Private/MetalShaders.cpp
+++ b/Engine/Source/Runtime/Apple/MetalRHI/Private/MetalShaders.cpp
@@ -658,6 +658,24 @@ FMetalShaderPipeline* FMetalComputeShader::GetPipeline(EPixelFormat const* const
 			Descriptor.SetMaxTotalThreadsPerThreadgroup(NumThreadsX*NumThreadsY*NumThreadsZ);
 		}
 		
+		if (FMetalCommandQueue::SupportsFeature(EMetalFeaturesGPUCommandBufferTimes))
+		{
+			ns::AutoReleased<ns::Array<mtlpp::PipelineBufferDescriptor>> PipelineBuffers = Descriptor.GetBuffers();
+			
+			uint32 ImmutableBuffers = Bindings.ConstantBuffers | Bindings.ArgumentBuffers;
+			while(ImmutableBuffers)
+			{
+				uint32 Index = __builtin_ctz(ImmutableBuffers);
+				ImmutableBuffers &= ~(1 << Index);
+				
+				if (Index < ML_MaxBuffers)
+				{
+					ns::AutoReleased<mtlpp::PipelineBufferDescriptor> PipelineBuffer = PipelineBuffers[Index];
+					PipelineBuffer.SetMutability(mtlpp::Mutability::Immutable);
+				}
+			}
+		}
+		
 		mtlpp::ComputePipelineState Kernel;
         mtlpp::ComputePipelineReflection Reflection;
 		
