Side-tables are also immutable.

diff --git a/Engine/Source/Runtime/Apple/MetalRHI/Private/MetalPipeline.cpp b/Engine/Source/Runtime/Apple/MetalRHI/Private/MetalPipeline.cpp
--- a/Engine/Source/Runtime/Apple/MetalRHI/Private/MetalPipeline.cpp
+++ b/Engine/Source/Runtime/Apple/MetalRHI/Private/MetalPipeline.cpp
@@ -392,6 +392,7 @@ static FMetalShaderPipeline* CreateMTLRenderPipeline(bool const bSync, FMetalGra
 		{
 			ns::AutoReleased<ns::Array<mtlpp::PipelineBufferDescriptor>> VertexPipelineBuffers = RenderPipelineDesc.GetVertexBuffers();
 			FMetalShaderBindings& VertexBindings = DomainShader ? DomainShader->Bindings : VertexShader->Bindings;
+			int8 VertexSideTable = DomainShader ? DomainShader->SideTableBinding : VertexShader->SideTableBinding;
 			{
 				uint32 ImmutableBuffers = VertexBindings.ConstantBuffers | VertexBindings.ArgumentBuffers;
 				while(ImmutableBuffers)
@@ -405,6 +406,11 @@ static FMetalShaderPipeline* CreateMTLRenderPipeline(bool const bSync, FMetalGra
 						PipelineBuffer.SetMutability(mtlpp::Mutability::Immutable);
 					}
 				}
+				if (VertexSideTable > 0)
+				{
+					ns::AutoReleased<mtlpp::PipelineBufferDescriptor> PipelineBuffer = VertexPipelineBuffers[VertexSideTable];
+					PipelineBuffer.SetMutability(mtlpp::Mutability::Immutable);
+				}
 			}
 			
 			if (PixelShader)
@@ -422,6 +428,11 @@ static FMetalShaderPipeline* CreateMTLRenderPipeline(bool const bSync, FMetalGra
 						PipelineBuffer.SetMutability(mtlpp::Mutability::Immutable);
 					}
 				}
+				if (PixelShader->SideTableBinding > 0)
+				{
+					ns::AutoReleased<mtlpp::PipelineBufferDescriptor> PipelineBuffer = FragmentPipelineBuffers[PixelShader->SideTableBinding];
+					PipelineBuffer.SetMutability(mtlpp::Mutability::Immutable);
+				}
 			}
 		}
 		
@@ -580,6 +591,11 @@ static FMetalShaderPipeline* CreateMTLRenderPipeline(bool const bSync, FMetalGra
 						PipelineBuffer.SetMutability(mtlpp::Mutability::Immutable);
 					}
 				}
+				if (VertexShader->SideTableBinding > 0)
+				{
+					ns::AutoReleased<mtlpp::PipelineBufferDescriptor> PipelineBuffer = PipelineBuffers[VertexShader->SideTableBinding];
+					PipelineBuffer.SetMutability(mtlpp::Mutability::Immutable);
+				}
 			}
 			
 			mtlpp::VertexDescriptor DomainVertexDesc;
diff --git a/Engine/Source/Runtime/Apple/MetalRHI/Private/MetalShaders.cpp b/Engine/Source/Runtime/Apple/MetalRHI/Private/MetalShaders.cpp
--- a/Engine/Source/Runtime/Apple/MetalRHI/Private/MetalShaders.cpp
+++ b/Engine/Source/Runtime/Apple/MetalRHI/Private/MetalShaders.cpp
@@ -674,6 +674,11 @@ FMetalShaderPipeline* FMetalComputeShader::GetPipeline(EPixelFormat const* const
 					PipelineBuffer.SetMutability(mtlpp::Mutability::Immutable);
 				}
 			}
+			if (SideTableBinding > 0)
+			{
+				ns::AutoReleased<mtlpp::PipelineBufferDescriptor> PipelineBuffer = PipelineBuffers[SideTableBinding];
+				PipelineBuffer.SetMutability(mtlpp::Mutability::Immutable);
+			}
 		}
 		
 		mtlpp::ComputePipelineState Kernel;
