Only call the pipeline buffer descriptor functions on iOS/tvOS 11+ and macOS 10.13+ using an independent feature enum.

diff --git a/Engine/Source/Runtime/Apple/MetalRHI/Private/MetalCommandQueue.cpp b/Engine/Source/Runtime/Apple/MetalRHI/Private/MetalCommandQueue.cpp
--- a/Engine/Source/Runtime/Apple/MetalRHI/Private/MetalCommandQueue.cpp
+++ b/Engine/Source/Runtime/Apple/MetalRHI/Private/MetalCommandQueue.cpp
@@ -79,7 +79,7 @@ FMetalCommandQueue::FMetalCommandQueue(mtlpp::Device InDevice, uint32 const MaxN
 			
 			if(Vers.majorVersion >= 11)
 			{
-				Features |= EMetalFeaturesGPUCaptureManager | EMetalFeaturesBufferSubAllocation | EMetalFeaturesParallelRenderEncoders;
+				Features |= EMetalFeaturesGPUCaptureManager | EMetalFeaturesBufferSubAllocation | EMetalFeaturesParallelRenderEncoders | EMetalFeaturesPipelineBufferMutability;
 				
 				if (MaxShaderVersion >= 3)
 				{
@@ -164,7 +164,7 @@ FMetalCommandQueue::FMetalCommandQueue(mtlpp::Device InDevice, uint32 const MaxN
 		
 		if(Vers.majorVersion >= 11)
 		{
-			Features |= EMetalFeaturesPresentMinDuration | EMetalFeaturesGPUCaptureManager | EMetalFeaturesBufferSubAllocation | EMetalFeaturesParallelRenderEncoders;
+			Features |= EMetalFeaturesPresentMinDuration | EMetalFeaturesGPUCaptureManager | EMetalFeaturesBufferSubAllocation | EMetalFeaturesParallelRenderEncoders | EMetalFeaturesPipelineBufferMutability;
         }
 #endif
 	}
@@ -252,7 +252,7 @@ FMetalCommandQueue::FMetalCommandQueue(mtlpp::Device InDevice, uint32 const MaxN
     
     if(Device.SupportsFeatureSet(mtlpp::FeatureSet::macOS_GPUFamily1_v3) && FPlatformMisc::MacOSXVersionCompare(10,13,0) >= 0)
     {
-        Features |= EMetalFeaturesMultipleViewports | EMetalFeaturesGPUCommandBufferTimes | EMetalFeaturesGPUCaptureManager | EMetalFeaturesAbsoluteTimeQueries | EMetalFeaturesSupportsVSyncToggle;
+        Features |= EMetalFeaturesMultipleViewports | EMetalFeaturesGPUCommandBufferTimes | EMetalFeaturesPipelineBufferMutability | EMetalFeaturesGPUCaptureManager | EMetalFeaturesAbsoluteTimeQueries | EMetalFeaturesSupportsVSyncToggle;
 		
 		if (FParse::Param(FCommandLine::Get(),TEXT("metalfence")))
 		{
diff --git a/Engine/Source/Runtime/Apple/MetalRHI/Private/MetalCommandQueue.h b/Engine/Source/Runtime/Apple/MetalRHI/Private/MetalCommandQueue.h
--- a/Engine/Source/Runtime/Apple/MetalRHI/Private/MetalCommandQueue.h
+++ b/Engine/Source/Runtime/Apple/MetalRHI/Private/MetalCommandQueue.h
@@ -101,6 +101,8 @@ typedef NS_OPTIONS(uint64, EMetalFeatures)
 	EMetalFeaturesParallelRenderEncoders = 1llu << 40llu,
 	/** Supports indirect argument buffers */
 	EMetalFeaturesIABs = 1llu << 41llu,
+	/** Supports specifying the mutability of buffers bound to PSOs */
+	EMetalFeaturesPipelineBufferMutability = 1llu << 42llu,
 };
 
 /**
diff --git a/Engine/Source/Runtime/Apple/MetalRHI/Private/MetalPipeline.cpp b/Engine/Source/Runtime/Apple/MetalRHI/Private/MetalPipeline.cpp
--- a/Engine/Source/Runtime/Apple/MetalRHI/Private/MetalPipeline.cpp
+++ b/Engine/Source/Runtime/Apple/MetalRHI/Private/MetalPipeline.cpp
@@ -388,7 +388,7 @@ static FMetalShaderPipeline* CreateMTLRenderPipeline(bool const bSync, FMetalGra
 		mtlpp::RenderPipelineDescriptor RenderPipelineDesc;
 		mtlpp::ComputePipelineDescriptor ComputePipelineDesc(nil);
 		
-		if (FMetalCommandQueue::SupportsFeature(EMetalFeaturesGPUCommandBufferTimes))
+		if (FMetalCommandQueue::SupportsFeature(EMetalFeaturesPipelineBufferMutability))
 		{
 			ns::AutoReleased<ns::Array<mtlpp::PipelineBufferDescriptor>> VertexPipelineBuffers = RenderPipelineDesc.GetVertexBuffers();
 			FMetalShaderBindings& VertexBindings = DomainShader ? DomainShader->Bindings : VertexShader->Bindings;
@@ -575,7 +575,7 @@ static FMetalShaderPipeline* CreateMTLRenderPipeline(bool const bSync, FMetalGra
 			ComputePipelineDesc = mtlpp::ComputePipelineDescriptor();
             check(ComputePipelineDesc);
 			
-			if (FMetalCommandQueue::SupportsFeature(EMetalFeaturesGPUCommandBufferTimes))
+			if (FMetalCommandQueue::SupportsFeature(EMetalFeaturesPipelineBufferMutability))
 			{
 				ns::AutoReleased<ns::Array<mtlpp::PipelineBufferDescriptor>> PipelineBuffers = ComputePipelineDesc.GetBuffers();
 				
diff --git a/Engine/Source/Runtime/Apple/MetalRHI/Private/MetalShaders.cpp b/Engine/Source/Runtime/Apple/MetalRHI/Private/MetalShaders.cpp
--- a/Engine/Source/Runtime/Apple/MetalRHI/Private/MetalShaders.cpp
+++ b/Engine/Source/Runtime/Apple/MetalRHI/Private/MetalShaders.cpp
@@ -658,7 +658,7 @@ FMetalShaderPipeline* FMetalComputeShader::GetPipeline(EPixelFormat const* const
 			Descriptor.SetMaxTotalThreadsPerThreadgroup(NumThreadsX*NumThreadsY*NumThreadsZ);
 		}
 		
-		if (FMetalCommandQueue::SupportsFeature(EMetalFeaturesGPUCommandBufferTimes))
+		if (FMetalCommandQueue::SupportsFeature(EMetalFeaturesPipelineBufferMutability))
 		{
 			ns::AutoReleased<ns::Array<mtlpp::PipelineBufferDescriptor>> PipelineBuffers = Descriptor.GetBuffers();
 			
