Don't initialise IABs for each referenced resource, they only need to be initialised once!

diff --git a/Engine/Source/Runtime/Apple/MetalRHI/Private/MetalUniformBuffer.cpp b/Engine/Source/Runtime/Apple/MetalRHI/Private/MetalUniformBuffer.cpp
--- a/Engine/Source/Runtime/Apple/MetalRHI/Private/MetalUniformBuffer.cpp
+++ b/Engine/Source/Runtime/Apple/MetalRHI/Private/MetalUniformBuffer.cpp
@@ -276,14 +276,17 @@ FMetalUniformBuffer::FMetalUniformBuffer(const void* Contents, const FRHIUniform
 						break;
 					}
 				}
-				
-				GetMetalDeviceContext().RegisterUB(this);
-				
-				FRHICommandListImmediate& RHICmdList = FRHICommandListExecutor::GetImmediateCommandList();
-				if (!(UniformUsage & UniformBuffer_SingleDraw) && IsRunningRHIInSeparateThread() && !RHICmdList.Bypass() && IsInRenderingThread())
-				{
-					new (RHICmdList.AllocCommand<FMetalRHICommandInitialiseUniformBufferIAB>()) FMetalRHICommandInitialiseUniformBufferIAB(this);
-				}
+			}
+		}
+		
+		if (FMetalCommandQueue::SupportsFeature(EMetalFeaturesIABs))
+		{
+			GetMetalDeviceContext().RegisterUB(this);
+			
+			FRHICommandListImmediate& RHICmdList = FRHICommandListExecutor::GetImmediateCommandList();
+			if (!(UniformUsage & UniformBuffer_SingleDraw) && IsRunningRHIInSeparateThread() && !RHICmdList.Bypass() && IsInRenderingThread())
+			{
+				new (RHICmdList.AllocCommand<FMetalRHICommandInitialiseUniformBufferIAB>()) FMetalRHICommandInitialiseUniformBufferIAB(this);
 			}
 		}
 	}
