The cost of PixelFormatView on macOS is exorbitant, we need to reallocate on demand to avoid it .

diff --git a/Engine/Source/Runtime/Apple/MetalRHI/Private/MetalTexture.cpp b/Engine/Source/Runtime/Apple/MetalRHI/Private/MetalTexture.cpp
--- a/Engine/Source/Runtime/Apple/MetalRHI/Private/MetalTexture.cpp
+++ b/Engine/Source/Runtime/Apple/MetalRHI/Private/MetalTexture.cpp
@@ -105,7 +105,9 @@ static mtlpp::TextureUsage ConvertFlagsToUsage(uint32 Flags)
 	{
 		Usage |= mtlpp::TextureUsage::RenderTarget;
 		Usage |= mtlpp::TextureUsage::ShaderRead;
+#if !PLATFORM_MAC // The cost of PixelFormatView on macOS is exorbitant, we need to reallocate on demand to avoid it
 		if (!(Flags & (TexCreate_ShaderResource)))
+#endif
 		{
 			Usage &= ~(mtlpp::TextureUsage::PixelFormatView);
 		}
