Disable Metal resource aliasing while it is debugged.

diff --git a/Engine/Source/Runtime/Apple/MetalRHI/Private/MetalRHI.cpp b/Engine/Source/Runtime/Apple/MetalRHI/Private/MetalRHI.cpp
--- a/Engine/Source/Runtime/Apple/MetalRHI/Private/MetalRHI.cpp
+++ b/Engine/Source/Runtime/Apple/MetalRHI/Private/MetalRHI.cpp
@@ -486,7 +486,7 @@ FMetalDynamicRHI::FMetalDynamicRHI(ERHIFeatureLevel::Type RequestedFeatureLevel)
 	GSupportsRenderTargetFormat_PF_G8 = false;
 	GRHISupportsTextureStreaming = true;
 	GSupportsWideMRT = bCanUseWideMRTs;
-	GSupportsTransientResourceAliasing = FMetalCommandQueue::SupportsFeature(EMetalFeaturesHeaps) && FMetalCommandQueue::SupportsFeature(EMetalFeaturesFences);
+	// GSupportsTransientResourceAliasing = FMetalCommandQueue::SupportsFeature(EMetalFeaturesHeaps) && FMetalCommandQueue::SupportsFeature(EMetalFeaturesFences);
 	GSupportsSeparateRenderTargetBlendState = (GMaxRHIFeatureLevel >= ERHIFeatureLevel::SM4);
 
 #if PLATFORM_MAC
