Missing RHI thread fences for Metal resource aliasing and IAB creation.

diff --git a/Engine/Source/Runtime/Apple/MetalRHI/Private/MetalTexture.cpp b/Engine/Source/Runtime/Apple/MetalRHI/Private/MetalTexture.cpp
--- a/Engine/Source/Runtime/Apple/MetalRHI/Private/MetalTexture.cpp
+++ b/Engine/Source/Runtime/Apple/MetalRHI/Private/MetalTexture.cpp
@@ -2565,6 +2565,7 @@ void FMetalDynamicRHI::RHIAcquireTransientResource_RenderThread(FTextureRHIParam
 	else
 	{
 		new (RHICmdList.AllocCommand<FMetalRHICommandUnaliasTextures>()) FMetalRHICommandUnaliasTextures(&Texture, 1);
+		RHICmdList.RHIThreadFence(true);
 	}
 	}
 }
@@ -2603,6 +2604,7 @@ void FMetalDynamicRHI::RHIDiscardTransientResource_RenderThread(FTextureRHIParam
 	else
 	{
 		new (RHICmdList.AllocCommand<FMetalRHICommandAliasTextures>()) FMetalRHICommandAliasTextures(&Texture, 1);
+		RHICmdList.RHIThreadFence(true);
 	}
 	}
 }
@@ -2637,6 +2639,7 @@ void FMetalDynamicRHI::RHIAcquireTransientResource_RenderThread(FVertexBufferRHI
 	else
 	{
 		new (RHICmdList.AllocCommand<FMetalRHICommandAliasBuffer>()) FMetalRHICommandAliasBuffer(MetalBuffer);
+		RHICmdList.RHIThreadFence(true);
 	}
 	}
 }
@@ -2659,6 +2662,7 @@ void FMetalDynamicRHI::RHIAcquireTransientResource_RenderThread(FStructuredBuffe
 	else
 	{
 		new (RHICmdList.AllocCommand<FMetalRHICommandAliasBuffer>()) FMetalRHICommandAliasBuffer(MetalBuffer);
+		RHICmdList.RHIThreadFence(true);
 	}
 	}
 }
diff --git a/Engine/Source/Runtime/Apple/MetalRHI/Private/MetalUniformBuffer.cpp b/Engine/Source/Runtime/Apple/MetalRHI/Private/MetalUniformBuffer.cpp
--- a/Engine/Source/Runtime/Apple/MetalRHI/Private/MetalUniformBuffer.cpp
+++ b/Engine/Source/Runtime/Apple/MetalRHI/Private/MetalUniformBuffer.cpp
@@ -287,6 +287,7 @@ FMetalUniformBuffer::FMetalUniformBuffer(const void* Contents, const FRHIUniform
 			if (!(UniformUsage & UniformBuffer_SingleDraw) && IsRunningRHIInSeparateThread() && !RHICmdList.Bypass() && IsInRenderingThread())
 			{
 				new (RHICmdList.AllocCommand<FMetalRHICommandInitialiseUniformBufferIAB>()) FMetalRHICommandInitialiseUniformBufferIAB(this);
+				RHICmdList.RHIThreadFence(true);
 			}
 		}
 	}
