Track MTLFence count as they are a finite resource and disable MTLFence and MTLHeap inside the Editor (not -game) because all the viewports and previews can exhaust the number of fences. We need to do a better job of detecting when we must flush the GPU and make the fences reusable to avoid this but for the integration, this will do.

diff --git a/Engine/Source/Runtime/Apple/MetalRHI/Private/MetalCommandQueue.cpp b/Engine/Source/Runtime/Apple/MetalRHI/Private/MetalCommandQueue.cpp
--- a/Engine/Source/Runtime/Apple/MetalRHI/Private/MetalCommandQueue.cpp
+++ b/Engine/Source/Runtime/Apple/MetalRHI/Private/MetalCommandQueue.cpp
@@ -233,14 +233,21 @@ FMetalCommandQueue::FMetalCommandQueue(mtlpp::Device InDevice, uint32 const MaxN
                 Features |= EMetalFeaturesIABs;
             }
 			
-			if (!FParse::Param(FCommandLine::Get(),TEXT("nometalfence")))
-			{
-				Features |= EMetalFeaturesFences;
-			}
-			
-			if (!FParse::Param(FCommandLine::Get(),TEXT("nometalheap")) && [Device.GetName().GetPtr() rangeOfString:@"Intel" options:NSCaseInsensitiveSearch].location == NSNotFound)
+            // The editor spawns so many viewports and preview icons that we can run out of hardware fences!
+			// Need to figure out a way to safely flush the rendering and reuse the fences when that happens.
+#if WITH_EDITORONLY_DATA
+			if (!GIsEditor)
+#endif
 			{
-				Features |= EMetalFeaturesHeaps;
+				if (!FParse::Param(FCommandLine::Get(),TEXT("nometalfence")))
+				{
+					Features |= EMetalFeaturesFences;
+				}
+				
+				if (!FParse::Param(FCommandLine::Get(),TEXT("nometalheap")) && [Device.GetName().GetPtr() rangeOfString:@"Intel" options:NSCaseInsensitiveSearch].location == NSNotFound)
+				{
+					Features |= EMetalFeaturesHeaps;
+				}
 			}
 		}
     }
diff --git a/Engine/Source/Runtime/Apple/MetalRHI/Private/MetalFence.cpp b/Engine/Source/Runtime/Apple/MetalRHI/Private/MetalFence.cpp
--- a/Engine/Source/Runtime/Apple/MetalRHI/Private/MetalFence.cpp
+++ b/Engine/Source/Runtime/Apple/MetalRHI/Private/MetalFence.cpp
@@ -6,6 +6,7 @@
 #include "MetalCommandBuffer.h"
 #include "MetalCommandQueue.h"
 #include "MetalContext.h"
+#include "MetalProfiler.h"
 
 @implementation FMetalDebugFence
 @synthesize Inner;
@@ -148,6 +149,9 @@ void FMetalFencePool::Initialise(mtlpp::Device const& InDevice)
 		}
 	}
 	Count = FMetalFencePool::NumFences;
+#if METAL_DEBUG_OPTIONS
+    Allocated = 0;
+#endif
 }
 
 FMetalFence* FMetalFencePool::AllocateFence()
@@ -155,8 +159,10 @@ FMetalFence* FMetalFencePool::AllocateFence()
 	FMetalFence* Fence = Lifo.Pop();
 	if (Fence)
 	{
-		FPlatformAtomics::InterlockedDecrement(&Count);
+        INC_DWORD_STAT(STAT_MetalFenceCount);
+        FPlatformAtomics::InterlockedDecrement(&Count);
 #if METAL_DEBUG_OPTIONS
+        FPlatformAtomics::InterlockedIncrement(&Allocated);
 		if (GMetalRuntimeDebugLevel >= EMetalDebugLevelValidation)
 		{
 			FScopeLock Lock(&Mutex);
@@ -189,7 +195,9 @@ void FMetalFencePool::ReleaseFence(FMetalFence* const InFence)
 {
 	if (InFence)
 	{
+        DEC_DWORD_STAT(STAT_MetalFenceCount);
 #if METAL_DEBUG_OPTIONS
+        FPlatformAtomics::InterlockedDecrement(&Allocated);
 		if (GMetalRuntimeDebugLevel >= EMetalDebugLevelValidation)
 		{
 			FScopeLock Lock(&Mutex);
diff --git a/Engine/Source/Runtime/Apple/MetalRHI/Private/MetalFence.h b/Engine/Source/Runtime/Apple/MetalRHI/Private/MetalFence.h
--- a/Engine/Source/Runtime/Apple/MetalRHI/Private/MetalFence.h
+++ b/Engine/Source/Runtime/Apple/MetalRHI/Private/MetalFence.h
@@ -184,6 +184,7 @@ private:
 #if METAL_DEBUG_OPTIONS
 	TSet<FMetalFence*> Fences;
 	FCriticalSection Mutex;
+    int32 Allocated;
 #endif
 	TLockFreePointerListLIFO<FMetalFence> Lifo;
 	static FMetalFencePool sSelf;
diff --git a/Engine/Source/Runtime/Apple/MetalRHI/Private/MetalProfiler.cpp b/Engine/Source/Runtime/Apple/MetalRHI/Private/MetalProfiler.cpp
--- a/Engine/Source/Runtime/Apple/MetalRHI/Private/MetalProfiler.cpp
+++ b/Engine/Source/Runtime/Apple/MetalRHI/Private/MetalProfiler.cpp
@@ -53,6 +53,7 @@ DEFINE_STAT(STAT_MetalTextureUnusedMemory);
 DEFINE_STAT(STAT_MetalBufferCount);
 DEFINE_STAT(STAT_MetalTextureCount);
 DEFINE_STAT(STAT_MetalHeapCount);
+DEFINE_STAT(STAT_MetalFenceCount);
 
 int64 volatile GMetalTexturePageOnTime = 0;
 int64 volatile GMetalGPUWorkTime = 0;
diff --git a/Engine/Source/Runtime/Apple/MetalRHI/Private/MetalProfiler.h b/Engine/Source/Runtime/Apple/MetalRHI/Private/MetalProfiler.h
--- a/Engine/Source/Runtime/Apple/MetalRHI/Private/MetalProfiler.h
+++ b/Engine/Source/Runtime/Apple/MetalRHI/Private/MetalProfiler.h
@@ -45,6 +45,7 @@ DECLARE_MEMORY_STAT_EXTERN(TEXT("Unused Texture Memory"), STAT_MetalTextureUnuse
 DECLARE_DWORD_ACCUMULATOR_STAT_EXTERN(TEXT("Buffer Count"), STAT_MetalBufferCount, STATGROUP_MetalRHI, );
 DECLARE_DWORD_ACCUMULATOR_STAT_EXTERN(TEXT("Texture Count"), STAT_MetalTextureCount, STATGROUP_MetalRHI, );
 DECLARE_DWORD_ACCUMULATOR_STAT_EXTERN(TEXT("Heap Count"), STAT_MetalHeapCount, STATGROUP_MetalRHI, );
+DECLARE_DWORD_ACCUMULATOR_STAT_EXTERN(TEXT("Fence Count"), STAT_MetalFenceCount, STATGROUP_MetalRHI, );
 
 DECLARE_CYCLE_STAT_EXTERN(TEXT("Texture Page-On time"), STAT_MetalTexturePageOnTime, STATGROUP_MetalRHI, );
 DECLARE_CYCLE_STAT_EXTERN(TEXT("GPU Work time"), STAT_MetalGPUWorkTime, STATGROUP_MetalRHI, );
