Add -forcemetalheap to allow testing of MTLHeap on Intel where it seems to cause issues.

diff --git a/Engine/Source/Runtime/Apple/MetalRHI/Private/MetalCommandQueue.cpp b/Engine/Source/Runtime/Apple/MetalRHI/Private/MetalCommandQueue.cpp
--- a/Engine/Source/Runtime/Apple/MetalRHI/Private/MetalCommandQueue.cpp
+++ b/Engine/Source/Runtime/Apple/MetalRHI/Private/MetalCommandQueue.cpp
@@ -244,7 +244,7 @@ FMetalCommandQueue::FMetalCommandQueue(mtlpp::Device InDevice, uint32 const MaxN
 					Features |= EMetalFeaturesFences;
 				}
 				
-				if (!FParse::Param(FCommandLine::Get(),TEXT("nometalheap")) && [Device.GetName().GetPtr() rangeOfString:@"Intel" options:NSCaseInsensitiveSearch].location == NSNotFound)
+				if (!FParse::Param(FCommandLine::Get(),TEXT("nometalheap")) && ([Device.GetName().GetPtr() rangeOfString:@"Intel" options:NSCaseInsensitiveSearch].location == NSNotFound || FParse::Param(FCommandLine::Get(),TEXT("forcemetalheap"))))
 				{
 					Features |= EMetalFeaturesHeaps;
 				}
