Fix some mac compile errors

diff --git a/Engine/Source/Runtime/Apple/MetalRHI/Private/MetalUniformBuffer.cpp b/Engine/Source/Runtime/Apple/MetalRHI/Private/MetalUniformBuffer.cpp
--- a/Engine/Source/Runtime/Apple/MetalRHI/Private/MetalUniformBuffer.cpp
+++ b/Engine/Source/Runtime/Apple/MetalRHI/Private/MetalUniformBuffer.cpp
@@ -245,7 +245,7 @@ void FMetalUniformBuffer::InitIAB()
 		{
 			FRHIResource* Resource = ResourceTable[i].GetReference();
 			
-			switch(GetLayout().Resources[i])
+			switch(GetLayout().Resources[i].MemberType)
 			{
 				case UBMT_GRAPH_TRACKED_SRV:
 				case UBMT_GRAPH_TRACKED_BUFFER_SRV:
@@ -520,10 +520,10 @@ void FMetalUniformBuffer::Update(const void* Contents, EUniformBufferValidation
         int32 NumResources = Layout.Resources.Num();
         for (int32 i = 0; i < NumResources; ++i)
         {
-            FRHIResource* Resource = *(FRHIResource**)((uint8*)Contents + Layout.ResourceOffsets[i]);
+            FRHIResource* Resource = *(FRHIResource**)((uint8*)Contents + Layout.Resources[i].MemberOffset);
             
             // Allow null SRV's in uniform buffers for feature levels that don't support SRV's in shaders
-            if (Validation == EUniformBufferValidation::ValidateResources && !(GMaxRHIFeatureLevel <= ERHIFeatureLevel::ES3_1 && Layout.Resources[i] == UBMT_SRV))
+            if (Validation == EUniformBufferValidation::ValidateResources && !(GMaxRHIFeatureLevel <= ERHIFeatureLevel::ES3_1 && Layout.Resources[i].MemberType == UBMT_SRV))
             {
                 check(Resource);
             }
@@ -532,7 +532,7 @@ void FMetalUniformBuffer::Update(const void* Contents, EUniformBufferValidation
             
             if (FMetalCommandQueue::SupportsFeature(EMetalFeaturesIABs) && Resource)
             {
-                switch(Layout.Resources[i])
+                switch(Layout.Resources[i].MemberType)
                 {
                     case UBMT_GRAPH_TRACKED_SRV:
                     case UBMT_GRAPH_TRACKED_BUFFER_SRV:
@@ -607,7 +607,7 @@ struct FMetalRHICommandUpateUniformBuffer : public FRHICommand<FMetalRHICommandU
 		uint32 MaxLayoutSize = Layout.ConstantBufferSize;
 		for (int32 i = 0; i < Layout.Resources.Num(); ++i)
 		{
-			MaxLayoutSize = FMath::Max(MaxLayoutSize, (uint32)(Layout.ResourceOffsets[i] + sizeof(FRHIResource*)));
+			MaxLayoutSize = FMath::Max(MaxLayoutSize, (uint32)(Layout.Resources[i].MemberOffset + sizeof(FRHIResource*)));
 		}
 		Contents = new char[MaxLayoutSize];
 		FMemory::Memcpy(Contents, Data, MaxLayoutSize);
