Metal query emulation can only be configured on the RHI thread, so you could not previously issue, submit and then get the query results within the same function unless you looped/waited between submission and accessing the waits. This can be resolved by flushing the RHI thread when we notice that the query has no internal state.

diff --git a/Engine/Source/Runtime/Apple/MetalRHI/Private/MetalQuery.cpp b/Engine/Source/Runtime/Apple/MetalRHI/Private/MetalQuery.cpp
--- a/Engine/Source/Runtime/Apple/MetalRHI/Private/MetalQuery.cpp
+++ b/Engine/Source/Runtime/Apple/MetalRHI/Private/MetalQuery.cpp
@@ -255,8 +255,6 @@ void FMetalRenderQuery::End(FMetalContext* Context)
 				check(StatSample);
 				[StatSample retain];
 				
-				AddRef();
-				
 				// Insert the fence to wait on the current command buffer
 				Context->InsertCommandBufferFence(*(Buffer.CommandBufferFence), [this, StatSample](mtlpp::CommandBuffer const&)
 				{
@@ -271,8 +269,6 @@ void FMetalRenderQuery::End(FMetalContext* Context)
 			else
 #endif
 			{
-				AddRef();
-				
 				// Insert the fence to wait on the current command buffer
 				Context->InsertCommandBufferFence(*(Buffer.CommandBufferFence), [this](mtlpp::CommandBuffer const&)
 				{
@@ -334,6 +330,14 @@ bool FMetalDynamicRHI::RHIGetRenderQueryResult(FRenderQueryRHIParamRef QueryRHI,
 		uint64 WaitMS = (Query->Type == RQT_AbsoluteTime) ? 2000 : 500;
 		if (bWait)
 		{
+			// RHI thread *must* be flushed at this point if the internal handles we rely upon are not yet valid.
+			// We *CANNOT* have one event per query as it consumes too many pthread objects.
+			if (!FRHICommandListExecutor::GetImmediateCommandList().Bypass() && IsRunningRHIInSeparateThread() && !Query->Buffer.CommandBufferFence.IsValid())
+			{
+				FRHICommandListExecutor::GetImmediateCommandList().RHIThreadFence(true);
+				FRHICommandListExecutor::GetImmediateCommandList().ImmediateFlush(EImmediateFlushType::FlushRHIThread);
+			}
+			
 			uint32 IdleStart = FPlatformTime::Cycles();
 		
 			bOK = Query->Buffer.Wait(WaitMS);
