Allow the initial dispatch in a pass to be committed to the prologue encoder, this ensures ordering of tessellation dispatches versus the GPU scene update dispatch, but can and should be generalised to punt any dispatch without prior dependencies.

Partially addresses the following:

However, there is still a bug in the way MTLFences are inserted that leaves one or more resources imporoperly hazard tracked.

diff --git a/Engine/Source/Runtime/Apple/MetalRHI/Private/MetalCommandEncoder.h b/Engine/Source/Runtime/Apple/MetalRHI/Private/MetalCommandEncoder.h
--- a/Engine/Source/Runtime/Apple/MetalRHI/Private/MetalCommandEncoder.h
+++ b/Engine/Source/Runtime/Apple/MetalRHI/Private/MetalCommandEncoder.h
@@ -133,6 +133,9 @@ public:
 	/** @returns The MTLFence for the current encoder or nil if there isn't one. */
 	FMetalFence* GetEncoderFence(void) const;
 	
+	/** @returns The number of encoded passes in the command buffer. */
+	uint32 NumEncodedPasses(void) const { return EncoderNum; }
+	
 #pragma mark - Public Command Encoder Mutators -
 
 	/**
diff --git a/Engine/Source/Runtime/Apple/MetalRHI/Private/MetalRenderPass.cpp b/Engine/Source/Runtime/Apple/MetalRHI/Private/MetalRenderPass.cpp
--- a/Engine/Source/Runtime/Apple/MetalRHI/Private/MetalRenderPass.cpp
+++ b/Engine/Source/Runtime/Apple/MetalRHI/Private/MetalRenderPass.cpp
@@ -721,47 +721,94 @@ void FMetalRenderPass::DrawPatches(uint32 PrimitiveType,FMetalBuffer const& Inde
 
 void FMetalRenderPass::Dispatch(uint32 ThreadGroupCountX, uint32 ThreadGroupCountY, uint32 ThreadGroupCountZ)
 {
-	ConditionalSwitchToCompute();
-	check(CurrentEncoder.GetCommandBuffer());
-	check(CurrentEncoder.IsComputeCommandEncoderActive());
+    if (CurrentEncoder.IsParallel() || CurrentEncoder.NumEncodedPasses() == 0)
+	{
+		ConditionalSwitchToAsyncCompute();
+		check(PrologueEncoder.GetCommandBuffer());
+		check(PrologueEncoder.IsComputeCommandEncoderActive());
+		
+		PrepareToAsyncDispatch();
+		
+		TRefCountPtr<FMetalComputeShader> ComputeShader = State.GetComputeShader();
+		check(ComputeShader);
+		
+		METAL_GPUPROFILE(FMetalProfiler::GetProfiler()->EncodeDispatch(PrologueEncoder.GetCommandBufferStats(), __FUNCTION__));
+		
+		mtlpp::Size ThreadgroupCounts = mtlpp::Size(ComputeShader->NumThreadsX, ComputeShader->NumThreadsY, ComputeShader->NumThreadsZ);
+		check(ComputeShader->NumThreadsX > 0 && ComputeShader->NumThreadsY > 0 && ComputeShader->NumThreadsZ > 0);
+		mtlpp::Size Threadgroups = mtlpp::Size(ThreadGroupCountX, ThreadGroupCountY, ThreadGroupCountZ);
+		PrologueEncoder.GetComputeCommandEncoder().DispatchThreadgroups(Threadgroups, ThreadgroupCounts);
+		METAL_DEBUG_LAYER(EMetalDebugLevelFastValidation, PrologueEncoder.GetComputeCommandEncoderDebugging().DispatchThreadgroups(Threadgroups, ThreadgroupCounts));
+		
+		ConditionalSubmit();
+	}
+	else
+	{
+		ConditionalSwitchToCompute();
+		check(CurrentEncoder.GetCommandBuffer());
+		check(CurrentEncoder.IsComputeCommandEncoderActive());
 
-	PrepareToDispatch();
-	
-	TRefCountPtr<FMetalComputeShader> ComputeShader = State.GetComputeShader();
-	check(ComputeShader);
-	
-	METAL_GPUPROFILE(FMetalProfiler::GetProfiler()->EncodeDispatch(CurrentEncoder.GetCommandBufferStats(), __FUNCTION__));
-	
-	mtlpp::Size ThreadgroupCounts = mtlpp::Size(ComputeShader->NumThreadsX, ComputeShader->NumThreadsY, ComputeShader->NumThreadsZ);
-	check(ComputeShader->NumThreadsX > 0 && ComputeShader->NumThreadsY > 0 && ComputeShader->NumThreadsZ > 0);
-	mtlpp::Size Threadgroups = mtlpp::Size(ThreadGroupCountX, ThreadGroupCountY, ThreadGroupCountZ);
-	CurrentEncoder.GetComputeCommandEncoder().DispatchThreadgroups(Threadgroups, ThreadgroupCounts);
-	METAL_DEBUG_LAYER(EMetalDebugLevelFastValidation, CurrentEncoder.GetComputeCommandEncoderDebugging().DispatchThreadgroups(Threadgroups, ThreadgroupCounts));
-	
-	ConditionalSubmit();
+		PrepareToDispatch();
+		
+		TRefCountPtr<FMetalComputeShader> ComputeShader = State.GetComputeShader();
+		check(ComputeShader);
+		
+		METAL_GPUPROFILE(FMetalProfiler::GetProfiler()->EncodeDispatch(CurrentEncoder.GetCommandBufferStats(), __FUNCTION__));
+		
+		mtlpp::Size ThreadgroupCounts = mtlpp::Size(ComputeShader->NumThreadsX, ComputeShader->NumThreadsY, ComputeShader->NumThreadsZ);
+		check(ComputeShader->NumThreadsX > 0 && ComputeShader->NumThreadsY > 0 && ComputeShader->NumThreadsZ > 0);
+		mtlpp::Size Threadgroups = mtlpp::Size(ThreadGroupCountX, ThreadGroupCountY, ThreadGroupCountZ);
+		CurrentEncoder.GetComputeCommandEncoder().DispatchThreadgroups(Threadgroups, ThreadgroupCounts);
+		METAL_DEBUG_LAYER(EMetalDebugLevelFastValidation, CurrentEncoder.GetComputeCommandEncoderDebugging().DispatchThreadgroups(Threadgroups, ThreadgroupCounts));
+		
+		ConditionalSubmit();
+	}
 }
 
 void FMetalRenderPass::DispatchIndirect(FMetalVertexBuffer* ArgumentBuffer, uint32 ArgumentOffset)
 {
 	check(ArgumentBuffer);
 	
-	ConditionalSwitchToCompute();
-	check(CurrentEncoder.GetCommandBuffer());
-	check(CurrentEncoder.IsComputeCommandEncoderActive());
-	
-	PrepareToDispatch();
-	
-	TRefCountPtr<FMetalComputeShader> ComputeShader = State.GetComputeShader();
-	check(ComputeShader);
-	
-	METAL_GPUPROFILE(FMetalProfiler::GetProfiler()->EncodeDispatch(CurrentEncoder.GetCommandBufferStats(), __FUNCTION__));
-	mtlpp::Size ThreadgroupCounts = mtlpp::Size(ComputeShader->NumThreadsX, ComputeShader->NumThreadsY, ComputeShader->NumThreadsZ);
-	check(ComputeShader->NumThreadsX > 0 && ComputeShader->NumThreadsY > 0 && ComputeShader->NumThreadsZ > 0);
-	
-	CurrentEncoder.GetComputeCommandEncoder().DispatchThreadgroupsWithIndirectBuffer(ArgumentBuffer->Buffer, ArgumentOffset, ThreadgroupCounts);
-	METAL_DEBUG_LAYER(EMetalDebugLevelFastValidation, CurrentEncoder.GetComputeCommandEncoderDebugging().DispatchThreadgroupsWithIndirectBuffer(ArgumentBuffer->Buffer, ArgumentOffset, ThreadgroupCounts));
+    if (CurrentEncoder.IsParallel() || CurrentEncoder.NumEncodedPasses() == 0)
+	{
+		ConditionalSwitchToAsyncCompute();
+		check(PrologueEncoder.GetCommandBuffer());
+		check(PrologueEncoder.IsComputeCommandEncoderActive());
+		
+		PrepareToAsyncDispatch();
+		
+		TRefCountPtr<FMetalComputeShader> ComputeShader = State.GetComputeShader();
+		check(ComputeShader);
+		
+		METAL_GPUPROFILE(FMetalProfiler::GetProfiler()->EncodeDispatch(PrologueEncoder.GetCommandBufferStats(), __FUNCTION__));
+		mtlpp::Size ThreadgroupCounts = mtlpp::Size(ComputeShader->NumThreadsX, ComputeShader->NumThreadsY, ComputeShader->NumThreadsZ);
+		check(ComputeShader->NumThreadsX > 0 && ComputeShader->NumThreadsY > 0 && ComputeShader->NumThreadsZ > 0);
+		
+		PrologueEncoder.GetComputeCommandEncoder().DispatchThreadgroupsWithIndirectBuffer(ArgumentBuffer->Buffer, ArgumentOffset, ThreadgroupCounts);
+		METAL_DEBUG_LAYER(EMetalDebugLevelFastValidation, PrologueEncoder.GetComputeCommandEncoderDebugging().DispatchThreadgroupsWithIndirectBuffer(ArgumentBuffer->Buffer, ArgumentOffset, ThreadgroupCounts));
+		
+		ConditionalSubmit();
+	}
+	else
+	{
+		ConditionalSwitchToCompute();
+		check(CurrentEncoder.GetCommandBuffer());
+		check(CurrentEncoder.IsComputeCommandEncoderActive());
+		
+		PrepareToDispatch();
+		
+		TRefCountPtr<FMetalComputeShader> ComputeShader = State.GetComputeShader();
+		check(ComputeShader);
+		
+		METAL_GPUPROFILE(FMetalProfiler::GetProfiler()->EncodeDispatch(CurrentEncoder.GetCommandBufferStats(), __FUNCTION__));
+		mtlpp::Size ThreadgroupCounts = mtlpp::Size(ComputeShader->NumThreadsX, ComputeShader->NumThreadsY, ComputeShader->NumThreadsZ);
+		check(ComputeShader->NumThreadsX > 0 && ComputeShader->NumThreadsY > 0 && ComputeShader->NumThreadsZ > 0);
+		
+		CurrentEncoder.GetComputeCommandEncoder().DispatchThreadgroupsWithIndirectBuffer(ArgumentBuffer->Buffer, ArgumentOffset, ThreadgroupCounts);
+		METAL_DEBUG_LAYER(EMetalDebugLevelFastValidation, CurrentEncoder.GetComputeCommandEncoderDebugging().DispatchThreadgroupsWithIndirectBuffer(ArgumentBuffer->Buffer, ArgumentOffset, ThreadgroupCounts));
 
-	ConditionalSubmit();
+		ConditionalSubmit();
+	}
 }
 
 FMetalFence* FMetalRenderPass::EndRenderPass(void)
@@ -1431,6 +1478,78 @@ void FMetalRenderPass::ConditionalSwitchToAsyncBlit(void)
 	check(PrologueEncoder.IsBlitCommandEncoderActive());
 }
 
+void FMetalRenderPass::ConditionalSwitchToAsyncCompute(void)
+{
+	SCOPE_CYCLE_COUNTER(STAT_MetalSwitchToComputeTime);
+	
+	if (PrologueEncoder.IsBlitCommandEncoderActive())
+	{
+		PrologueEncoderFence = PrologueEncoder.EndEncoding();
+		LastPrologueEncoderFence = PrologueEncoderFence;
+	}
+	
+	if (!PrologueEncoder.IsComputeCommandEncoderActive())
+	{
+		if (!PrologueEncoder.GetCommandBuffer())
+		{
+			PrologueEncoder.StartCommandBuffer();
+		}
+		State.SetStateDirty();
+		PrologueEncoder.BeginComputeCommandEncoding(ComputeDispatchType);
+		
+		if (PrologueEncoderFence)
+		{
+			if (PrologueEncoderFence->NeedsWait(mtlpp::RenderStages::Fragment))
+			{
+				LastPrologueEncoderFence = nullptr;
+			}
+			PrologueEncoder.WaitForFence(PrologueEncoderFence);
+			PrologueEncoderFence = nullptr;
+		}
+		if (LastPrologueEncoderFence)
+		{
+			PrologueEncoder.WaitAndUpdateFence(LastPrologueEncoderFence);
+			LastPrologueEncoderFence = nullptr;
+		}
+		if (PassStartFence)
+		{
+			PrologueEncoder.WaitForFence(PassStartFence);
+			PassStartFence = nullptr;
+		}
+		if (ParallelPassEndFence)
+		{
+			PrologueEncoder.WaitForFence(ParallelPassEndFence);
+			ParallelPassEndFence = nullptr;
+		}
+		PrologueEncoderFence = PrologueEncoder.GetEncoderFence();
+		LastPrologueEncoderFence = PrologueEncoderFence;
+#if METAL_DEBUG_OPTIONS
+		if (GetEmitDrawEvents() && PrologueEncoderFence)
+		{
+			for (uint32 i = mtlpp::RenderStages::Vertex; i <= mtlpp::RenderStages::Fragment; i++)
+			{
+				if (CmdList.GetCommandQueue().GetRuntimeDebuggingLevel() >= EMetalDebugLevelValidation)
+				{
+					PrologueEncoderFence->Get((mtlpp::RenderStages)i).GetPtr().label = [NSString stringWithFormat:@"Prologue %@", PrologueEncoderFence->Get((mtlpp::RenderStages)i).GetLabel().GetPtr()];
+				}
+				else
+				{
+					PrologueEncoderFence->Get((mtlpp::RenderStages)i).SetLabel([NSString stringWithFormat:@"Prologue %@", PrologueEncoderFence->Get((mtlpp::RenderStages)i).GetLabel().GetPtr()]);
+				}
+			}
+		}
+#endif
+		
+		if (CurrentEncoder.IsRenderCommandEncoderActive() || CurrentEncoder.IsComputeCommandEncoderActive() || CurrentEncoder.IsBlitCommandEncoderActive())
+		{
+			CurrentEncoder.WaitForFence(PrologueEncoderFence);
+			PrologueEncoderFence = nullptr;
+		}
+	}
+	
+	check(PrologueEncoder.IsComputeCommandEncoderActive());
+}
+
 void FMetalRenderPass::CommitRenderResourceTables(void)
 {
 	SCOPE_CYCLE_COUNTER(STAT_MetalCommitRenderResourceTablesTime);
@@ -1489,6 +1608,20 @@ void FMetalRenderPass::CommitDispatchResourceTables(void)
 	}
 }
 
+void FMetalRenderPass::CommitAsyncDispatchResourceTables(void)
+{
+	State.CommitComputeResources(&PrologueEncoder);
+	
+	State.CommitResourceTable(SF_Compute, mtlpp::FunctionType::Kernel, PrologueEncoder);
+	
+	FMetalComputeShader const* ComputeShader = State.GetComputeShader();
+	if (ComputeShader->SideTableBinding >= 0)
+	{
+		PrologueEncoder.SetShaderSideTable(mtlpp::FunctionType::Kernel, ComputeShader->SideTableBinding);
+		State.SetShaderBuffer(SF_Compute, nil, nil, 0, 0, ComputeShader->SideTableBinding, mtlpp::ResourceUsage(0));
+	}
+}
+
 void FMetalRenderPass::PrepareToRender(uint32 PrimitiveType)
 {
 	SCOPE_CYCLE_COUNTER(STAT_MetalPrepareToRenderTime);
@@ -1536,6 +1669,19 @@ void FMetalRenderPass::PrepareToDispatch(void)
     State.SetComputePipelineState(CurrentEncoder);
 }
 
+void FMetalRenderPass::PrepareToAsyncDispatch(void)
+{
+	SCOPE_CYCLE_COUNTER(STAT_MetalPrepareToDispatchTime);
+	
+	check(PrologueEncoder.GetCommandBuffer());
+	check(PrologueEncoder.IsComputeCommandEncoderActive());
+	
+	// Bind shader resources
+	CommitAsyncDispatchResourceTables();
+	
+	State.SetComputePipelineState(PrologueEncoder);
+}
+
 void FMetalRenderPass::ConditionalSubmit()
 {
 	NumOutstandingOps++;
diff --git a/Engine/Source/Runtime/Apple/MetalRHI/Private/MetalRenderPass.h b/Engine/Source/Runtime/Apple/MetalRHI/Private/MetalRenderPass.h
--- a/Engine/Source/Runtime/Apple/MetalRHI/Private/MetalRenderPass.h
+++ b/Engine/Source/Runtime/Apple/MetalRHI/Private/MetalRenderPass.h
@@ -148,14 +148,17 @@ private:
     void ConditionalSwitchToCompute(void);
 	void ConditionalSwitchToBlit(void);
 	void ConditionalSwitchToAsyncBlit(void);
+	void ConditionalSwitchToAsyncCompute(void);
 	
     void PrepareToRender(uint32 PrimType);
     void PrepareToTessellate(uint32 PrimType);
     void PrepareToDispatch(void);
+	void PrepareToAsyncDispatch(void);
 
     void CommitRenderResourceTables(void);
     void CommitTessellationResourceTables(void);
     void CommitDispatchResourceTables(void);
+	void CommitAsyncDispatchResourceTables(void);
     
     void ConditionalSubmit();
 private:
