Remove all support for Metal prior to Metal 1.2 on iOS/tvOS and Metal 2.0 on macOS. This permits some fairly radical simplifications and removal of a number of workarounds.

diff --git a/Engine/Source/Runtime/Apple/MetalRHI/Public/MetalResources.h b/Engine/Source/Runtime/Apple/MetalRHI/Public/MetalResources.h
--- a/Engine/Source/Runtime/Apple/MetalRHI/Public/MetalResources.h
+++ b/Engine/Source/Runtime/Apple/MetalRHI/Public/MetalResources.h
@@ -77,13 +77,6 @@ enum EMetalIndexType
 	EMetalIndexType_Num	   = 3
 };
 
-enum EMetalBufferType
-{
-	EMetalBufferType_Dynamic = 0,
-	EMetalBufferType_Static = 1,
-	EMetalBufferType_Num = 2
-};
-
 /** This represents a vertex shader that hasn't been combined with a specific declaration to create a bound shader. */
 template<typename BaseResourceType, int32 ShaderType>
 class TMetalBaseShader : public BaseResourceType, public IRefCountedObject
@@ -96,12 +89,10 @@ public:
 	: SideTableBinding(-1)
 	, SourceLen(0)
 	, SourceCRC(0)
-    , BufferTypeHash(0)
 	, GlslCodeNSString(nil)
 	, CodeSize(0)
 	{
-		Function[EMetalIndexType_None][EMetalBufferType_Dynamic] = Function[EMetalIndexType_UInt16][EMetalBufferType_Dynamic] = Function[EMetalIndexType_UInt32][EMetalBufferType_Dynamic] = nil;
-		Function[EMetalIndexType_None][EMetalBufferType_Static] = Function[EMetalIndexType_UInt16][EMetalBufferType_Static] = Function[EMetalIndexType_UInt32][EMetalBufferType_Static] = nil;
+		Function[EMetalIndexType_None] = Function[EMetalIndexType_UInt16] = Function[EMetalIndexType_UInt32] = nil;
 	}
 	
 	void Init(const TArray<uint8>& InCode, FMetalCodeHeader& Header, mtlpp::Library InLibrary = nil);
@@ -152,16 +143,12 @@ public:
 	
 	/** Hash for the shader/material permutation constants */
 	uint32 ConstantValueHash;
-
-	/** Hash of the typed_buffer format types */
-    uint32 BufferTypeHash;
     
 protected:
-	mtlpp::Function GetCompiledFunction(EMetalIndexType IndexType, EPixelFormat const* const BufferTypes, uint32 BufferTypeHash, bool const bAsync = false);
-	uint32 GetBufferBindingHash(EPixelFormat const* const BufferTypes) const;
+	mtlpp::Function GetCompiledFunction(EMetalIndexType IndexType, bool const bAsync = false);
 
     // this is the compiler shader
-	mtlpp::Function Function[EMetalIndexType_Num][EMetalBufferType_Num];
+	mtlpp::Function Function[EMetalIndexType_Num];
     
 private:
 	// This is the MTLLibrary for the shader so we can dynamically refine the MTLFunction
@@ -188,8 +175,7 @@ public:
 	FMetalVertexShader(const TArray<uint8>& InCode);
 	FMetalVertexShader(const TArray<uint8>& InCode, mtlpp::Library InLibrary);
 	
-	uint32 GetBindingHash(EPixelFormat const* const BufferTypes) const;
-	mtlpp::Function GetFunction(EMetalIndexType IndexType, EPixelFormat const* const BufferTypes, uint32 BufferTypeHash);
+	mtlpp::Function GetFunction(EMetalIndexType IndexType);
 	
 	// for VSHS
 	FMetalTessellationOutputs TessellationOutputAttribs;
@@ -212,8 +198,7 @@ public:
 	FMetalPixelShader(const TArray<uint8>& InCode);
 	FMetalPixelShader(const TArray<uint8>& InCode, mtlpp::Library InLibrary);
 	
-	uint32 GetBindingHash(EPixelFormat const* const BufferTypes) const;
-	mtlpp::Function GetFunction(EMetalIndexType IndexType, EPixelFormat const* const BufferTypes, uint32 BufferTypeHash);
+	mtlpp::Function GetFunction(EMetalIndexType IndexType);
 };
 
 class FMetalHullShader : public TMetalBaseShader<FRHIHullShader, SF_Hull>
@@ -222,8 +207,7 @@ public:
 	FMetalHullShader(const TArray<uint8>& InCode);
 	FMetalHullShader(const TArray<uint8>& InCode, mtlpp::Library InLibrary);
 	
-	uint32 GetBindingHash(EPixelFormat const* const BufferTypes) const;
-	mtlpp::Function GetFunction(EMetalIndexType IndexType, EPixelFormat const* const BufferTypes, uint32 BufferTypeHash);
+	mtlpp::Function GetFunction(EMetalIndexType IndexType);
 };
 
 class FMetalDomainShader : public TMetalBaseShader<FRHIDomainShader, SF_Domain>
@@ -232,8 +216,7 @@ public:
 	FMetalDomainShader(const TArray<uint8>& InCode);
 	FMetalDomainShader(const TArray<uint8>& InCode, mtlpp::Library InLibrary);
 	
-	uint32 GetBindingHash(EPixelFormat const* const BufferTypes) const;
-	mtlpp::Function GetFunction(EMetalIndexType IndexType, EPixelFormat const* const BufferTypes, uint32 BufferTypeHash);
+	mtlpp::Function GetFunction(EMetalIndexType IndexType);
 	
 	mtlpp::Winding TessellationOutputWinding;
 	mtlpp::TessellationPartitionMode TessellationPartitioning;
@@ -249,8 +232,7 @@ public:
 	FMetalComputeShader(const TArray<uint8>& InCode, mtlpp::Library InLibrary = nil);
 	virtual ~FMetalComputeShader();
 	
-	uint32 GetBindingHash(EPixelFormat const* const BufferTypes) const;
-	FMetalShaderPipeline* GetPipeline(EPixelFormat const* const BufferTypes, uint32 BufferTypeHash);
+	FMetalShaderPipeline* GetPipeline();
 	
 	// thread group counts
 	int32 NumThreadsX;
@@ -259,7 +241,7 @@ public:
     
 private:
     // the state object for a compute shader
-    FMetalShaderPipeline* Pipeline[EMetalBufferType_Num];
+    FMetalShaderPipeline* Pipeline;
 };
 
 struct FMetalRenderPipelineHash
@@ -287,7 +269,7 @@ class FMetalGraphicsPipelineState : public FRHIGraphicsPipelineState
 public:
 	virtual ~FMetalGraphicsPipelineState();
 
-	FMetalShaderPipeline* GetPipeline(EMetalIndexType IndexType, uint32 VertexBufferHash, uint32 PixelBufferHash, uint32 DomainBufferHash, EPixelFormat const* const VertexBufferTypes = nullptr, EPixelFormat const* const PixelBufferTypes = nullptr, EPixelFormat const* const DomainBufferTypes = nullptr);
+	FMetalShaderPipeline* GetPipeline(EMetalIndexType IndexType);
 	
 	/** Cached vertex structure */
 	TRefCountPtr<FMetalVertexDeclaration> VertexDeclaration;
@@ -315,7 +297,7 @@ private:
 	// Needed to runtime refine shaders currently.
 	FGraphicsPipelineStateInitializer Initializer;
 	// Tessellation pipelines have three different variations for the indexing-style.
-	FMetalShaderPipeline* PipelineStates[EMetalIndexType_Num][EMetalBufferType_Num][EMetalBufferType_Num][EMetalBufferType_Num];
+	FMetalShaderPipeline* PipelineStates[EMetalIndexType_Num];
 };
 
 class FMetalComputePipelineState : public FRHIComputePipelineState
diff --git a/Engine/Source/Runtime/Apple/MetalRHI/Public/MetalShaderResources.h b/Engine/Source/Runtime/Apple/MetalRHI/Public/MetalShaderResources.h
--- a/Engine/Source/Runtime/Apple/MetalRHI/Public/MetalShaderResources.h
+++ b/Engine/Source/Runtime/Apple/MetalRHI/Public/MetalShaderResources.h
@@ -132,12 +132,10 @@ struct FMetalShaderBindings
 	TArray<TArray<CrossCompiler::FPackedArrayInfo>>	PackedUniformBuffers;
 	TArray<CrossCompiler::FPackedArrayInfo>			PackedGlobalArrays;
 	FMetalShaderResourceTable				ShaderResourceTable;
-	TArray<uint8> 							TypedBufferFormats;
 	TMap<uint8, TArray<uint8>>				ArgumentBufferMasks;
 
     uint32  LinearBuffer;
 	uint32	TypedBuffers;
-	uint32 	InvariantBuffers;
 	uint32 	ConstantBuffers;
 	uint32  ArgumentBuffers;
 	uint16	InOutMask;
@@ -149,7 +147,6 @@ struct FMetalShaderBindings
 	FMetalShaderBindings() :
 		LinearBuffer(0),
         TypedBuffers(0),
-		InvariantBuffers(0),
 		ConstantBuffers(0),
 		ArgumentBuffers(0),
 		InOutMask(0),
@@ -166,11 +163,9 @@ inline FArchive& operator<<(FArchive& Ar, FMetalShaderBindings& Bindings)
 	Ar << Bindings.PackedUniformBuffers;
 	Ar << Bindings.PackedGlobalArrays;
 	Ar << Bindings.ShaderResourceTable;
-	Ar << Bindings.TypedBufferFormats;
 	Ar << Bindings.ArgumentBufferMasks;
     Ar << Bindings.LinearBuffer;
     Ar << Bindings.TypedBuffers;
-	Ar << Bindings.InvariantBuffers;
 	Ar << Bindings.ConstantBuffers;
 	Ar << Bindings.ArgumentBuffers;
 	Ar << Bindings.InOutMask;
