Mask based on the correct indices and make the error reporting a little bit more meaningful when a vertex stream doesn't have enough data to work.

diff --git a/Engine/Source/Runtime/Apple/MetalRHI/Private/MetalRenderPass.cpp b/Engine/Source/Runtime/Apple/MetalRHI/Private/MetalRenderPass.cpp
--- a/Engine/Source/Runtime/Apple/MetalRHI/Private/MetalRenderPass.cpp
+++ b/Engine/Source/Runtime/Apple/MetalRHI/Private/MetalRenderPass.cpp
@@ -458,7 +458,7 @@ void FMetalRenderPass::DrawIndexedPrimitive(FMetalBuffer const& IndexBuffer, uin
 		for(int VertexElemIdx = 0;VertexElemIdx < VertexDecl->Elements.Num();++VertexElemIdx)
 		{
 			FVertexElement const & VertexElem = VertexDecl->Elements[VertexElemIdx];
-			if(VertexElem.Stride > 0 && VertexElem.bUseInstanceIndex && ((InOutMask & (1 << VertexElemIdx))))
+			if(VertexElem.Stride > 0 && VertexElem.bUseInstanceIndex && ((InOutMask & (1 << VertexElem.AttributeIndex))))
 			{
 				uint32 AvailElementCount = 0;
 				
@@ -471,18 +471,18 @@ void FMetalRenderPass::DrawIndexedPrimitive(FMetalBuffer const& IndexBuffer, uin
 				}
 				
 				ClampedNumInstances = FMath::Clamp<uint32>(ClampedNumInstances, 0, AvailElementCount);
-			}
-		}
-		
-		if(ClampedNumInstances < NumInstances)
-		{
-			FString ShaderName = TEXT("Unknown");
+				
+				if(ClampedNumInstances < NumInstances)
+				{
+					FString ShaderName = TEXT("Unknown");
 #if !(UE_BUILD_SHIPPING || UE_BUILD_TEST)
-			ShaderName = PipelineState->PixelShader->ShaderName;
+					ShaderName = PipelineState->PixelShader->ShaderName;
 #endif
-			// Setting NumInstances to ClampedNumInstances would fix any visual rendering bugs resulting from this bad call but these draw calls are wrong - don't hide the issue
-			UE_LOG(LogMetal, Error, TEXT("Metal DrawIndexedPrimitive requested to draw %d Instances but vertex stream only has %d instance data available. ShaderName: %s"), NumInstances, ClampedNumInstances,
-				*ShaderName);
+					// Setting NumInstances to ClampedNumInstances would fix any visual rendering bugs resulting from this bad call but these draw calls are wrong - don't hide the issue
+					UE_LOG(LogMetal, Error, TEXT("Metal DrawIndexedPrimitive requested to draw %d Instances but vertex stream only has %d instance data available. ShaderName: %s, Deficient Attribute Index: %u"), NumInstances, ClampedNumInstances,
+						   *ShaderName, VertexElem.AttributeIndex);
+				}
+			}
 		}
 	}
 #endif
