Hack to make Metal tessellation sort of work again - before the refactor all the global uniform bindings were replicated between vertex & hull shaders, but that is no longer true. Thus Apple's plan of merging the vertex & hull stages doesn't quite work as the high-level sets all hull shader parameters to 0, which clobbers the correct values provided to the vertex shader. For now we simply ignore the hull shader parameters and hope it works, but to really fix this will mean changing the way tessellation is supported to compile and execute vertex & hull shaders separately as I had originally envisioned, as this will naturally match D3D.

diff --git a/Engine/Source/Runtime/Apple/MetalRHI/Private/MetalStateCache.cpp b/Engine/Source/Runtime/Apple/MetalRHI/Private/MetalStateCache.cpp
--- a/Engine/Source/Runtime/Apple/MetalRHI/Private/MetalStateCache.cpp
+++ b/Engine/Source/Runtime/Apple/MetalRHI/Private/MetalStateCache.cpp
@@ -1837,7 +1837,12 @@ void FMetalStateCache::CommitTessellationResources(FMetalCommandEncoder* Raster,
     }
     
     SetResourcesFromTables(GraphicsPSO->HullShader, CrossCompiler::SHADER_STAGE_HULL);
-    GetShaderParameters(CrossCompiler::SHADER_STAGE_HULL).CommitPackedGlobals(this, Compute, SF_Hull, GraphicsPSO->HullShader->Bindings);
+	
+	// @todo Can't put global uniforms in Hull shaders - the high-level no longer
+	// replicates the bindings to both vertex & hull shaders, so we have conflicts and the
+	// hull shader versions are almost always 0-filled which causes chaos.
+	// We need to split the shaders out into two distinct stages to match D3D
+    // GetShaderParameters(CrossCompiler::SHADER_STAGE_HULL).CommitPackedGlobals(this, Compute, SF_Hull, GraphicsPSO->HullShader->Bindings);
 	
 	SetResourcesFromTables(GraphicsPSO->DomainShader, CrossCompiler::SHADER_STAGE_DOMAIN);
    	GetShaderParameters(CrossCompiler::SHADER_STAGE_DOMAIN).CommitPackedGlobals(this, Raster, SF_Domain, GraphicsPSO->DomainShader->Bindings);
