Remove tessellation function constant as we can use linear-textures to achieve the same result which means we no longer need separate variants of the shader. DrawIndexedPrimitiveUP can no longer draw a tessellated mesh - but that was never advisable anyway.

diff --git a/Engine/Source/Developer/Apple/MetalShaderFormat/Private/MetalBackend.cpp b/Engine/Source/Developer/Apple/MetalShaderFormat/Private/MetalBackend.cpp
--- a/Engine/Source/Developer/Apple/MetalShaderFormat/Private/MetalBackend.cpp
+++ b/Engine/Source/Developer/Apple/MetalShaderFormat/Private/MetalBackend.cpp
@@ -1474,9 +1474,9 @@ protected:
             int32 patchIndex = Buffers.GetIndex(patchCount);
             check(patchIndex >= 0 && patchIndex < 30);
             
-            ir_variable* indexBuffer = new(ParseState)ir_variable(glsl_type::void_type, "indexBuffer", ir_var_in);
+            ir_variable* indexBuffer = new(ParseState)ir_variable(glsl_type::uint_type, "indexBuffer", ir_var_in);
             indexBuffer->semantic = "";
-            Buffers.Buffers.Add(indexBuffer);
+            Buffers.Textures.Add(indexBuffer);
             
             int32 IndexBufferIndex = Buffers.GetIndex(indexBuffer);
             check(IndexBufferIndex >= 0 && IndexBufferIndex < 30);
@@ -1487,7 +1487,8 @@ protected:
 				"ushort2 thread_position_in_threadgroup [[thread_position_in_threadgroup]],\n"
 				"uint2 threadgroup_position_in_grid [[threadgroup_position_in_grid]],\n"
 				"constant uint *patchCount [[ buffer(%d) ]],\n"
-				"const device void *indexBuffer [[ buffer(%d) ]]", // indexBufferType == 0 means not indexed (and will strip) 1 means uint16_t* and 2 means uint32_t*
+				"#define METAL_INDEX_BUFFER_ID %d\n"
+				"typedBuffer1_read(uint, indexBuffer, METAL_INDEX_BUFFER_ID)",
                 patchIndex, IndexBufferIndex
 			);
 			bPrintComma = true;
@@ -1624,19 +1625,12 @@ protected:
 						ralloc_asprintf_append(buffer, "#define GET_INTERNAL_PATCH_ID() (GET_INSTANCE_ID() * GET_PATCH_COUNT() + GET_PATCH_ID())\n");
 						ralloc_asprintf_append(buffer, "#define GET_PATCH_ID_IN_THREADGROUP() (GET_PATCH_ID() %% TessellationPatchesPerThreadGroup)\n");
 						ralloc_asprintf_append(buffer, "#define GET_INPUT_CP_ID() (thread_position_in_grid.x %% TessellationInputControlPoints)\n");
-						/* NOTE: relies upon
-						enum EMetalIndexType
-						{
-							EMetalIndexType_None   = 0,
-							EMetalIndexType_UInt16 = 1,
-							EMetalIndexType_UInt32 = 2
-						};*/
-						ralloc_asprintf_append(buffer, "constant uint indexBufferType [[ function_constant(32) ]];\n");
+
+						ir_variable* indexBuffer = ParseState->symbols->get_variable("indexBuffer");
+						int32 IndexBufferIndex = Buffers.GetIndex(indexBuffer);
 						ralloc_asprintf_append(buffer, "#define GET_VERTEX_ID() \\\n");
-						ralloc_asprintf_append(buffer, "	(indexBufferType == 0) ? thread_position_in_grid.x : \\\n");
-						ralloc_asprintf_append(buffer, "	(indexBufferType == 1) ? ((const device ushort *)indexBuffer)[thread_position_in_grid.x] : \\\n");
-						ralloc_asprintf_append(buffer, "	(indexBufferType == 2) ? ((const device uint   *)indexBuffer)[thread_position_in_grid.x] : \\\n");
-						ralloc_asprintf_append(buffer, "	thread_position_in_grid.x\n");
+						ralloc_asprintf_append(buffer, "	(is_null_texture(indexBuffer)) ? thread_position_in_grid.x : \\\n");
+						ralloc_asprintf_append(buffer, "	buffer::load<uint, METAL_INDEX_BUFFER_ID>(indexBuffer, thread_position_in_grid.x + patchCount[1])\n");
 						ralloc_asprintf_append(buffer, "/* optionally vertex_id = GET_VERTEX_ID() + grid_origin.x */\n");
 					}
 
diff --git a/Engine/Source/Developer/Apple/MetalShaderFormat/Private/MetalShaderCompiler.cpp b/Engine/Source/Developer/Apple/MetalShaderFormat/Private/MetalShaderCompiler.cpp
--- a/Engine/Source/Developer/Apple/MetalShaderFormat/Private/MetalShaderCompiler.cpp
+++ b/Engine/Source/Developer/Apple/MetalShaderFormat/Private/MetalShaderCompiler.cpp
@@ -1319,7 +1319,6 @@ void BuildMetalShaderOutput(
 	Header.TessellationControlPointOutBuffer    = CCHeader.TessellationControlPointOutBuffer;
 	Header.TessellationControlPointIndexBuffer  = CCHeader.TessellationControlPointIndexBuffer;
 	Header.TessellationOutputAttribs            = TessOutputAttribs;
-	Header.bTessFunctionConstants				= (FCStringAnsi::Strstr(USFSource, "indexBufferType [[ function_constant(32) ]]") != nullptr);
 	Header.bDeviceFunctionConstants				= (FCStringAnsi::Strstr(USFSource, "#define __METAL_DEVICE_CONSTANT_INDEX__ 1") != nullptr);
 	Header.SideTable 							= CCHeader.SideTable;
 	Header.Bindings.ArgumentBufferMasks			= CCHeader.ArgumentBuffers;
diff --git a/Engine/Source/Developer/Apple/MetalShaderFormat/Private/MetalShaderFormat.cpp b/Engine/Source/Developer/Apple/MetalShaderFormat/Private/MetalShaderFormat.cpp
--- a/Engine/Source/Developer/Apple/MetalShaderFormat/Private/MetalShaderFormat.cpp
+++ b/Engine/Source/Developer/Apple/MetalShaderFormat/Private/MetalShaderFormat.cpp
@@ -239,7 +239,7 @@ class FMetalShaderFormat : public IShaderFormat
 public:
 	enum
 	{
-		HEADER_VERSION = 62,
+		HEADER_VERSION = 63,
 	};
 	
 	struct FVersion
diff --git a/Engine/Source/Runtime/Apple/MetalRHI/Private/MetalCommands.cpp b/Engine/Source/Runtime/Apple/MetalRHI/Private/MetalCommands.cpp
--- a/Engine/Source/Runtime/Apple/MetalRHI/Private/MetalCommands.cpp
+++ b/Engine/Source/Runtime/Apple/MetalRHI/Private/MetalCommands.cpp
@@ -737,7 +737,8 @@ void FMetalRHICommandContext::RHIDrawIndexedPrimitive(FIndexBufferRHIParamRef In
 
 	FMetalIndexBuffer* IndexBuffer = ResourceCast(IndexBufferRHI);
 	
-	Context->DrawIndexedPrimitive(IndexBuffer->Buffer, IndexBuffer->GetStride(), IndexBuffer->IndexType, PrimitiveType, BaseVertexIndex, FirstInstance, NumVertices, StartIndex, NumPrimitives, NumInstances);
+	EPixelFormat Format = IndexBuffer->IndexType == mtlpp::IndexType::UInt16 ? PF_R16_UINT : PF_R32_UINT;
+	Context->DrawIndexedPrimitive(IndexBuffer->Buffer, IndexBuffer->GetLinearTexture(Format), IndexBuffer->GetStride(), IndexBuffer->IndexType, PrimitiveType, BaseVertexIndex, FirstInstance, NumVertices, StartIndex, NumPrimitives, NumInstances);
 	}
 }
 
@@ -879,7 +880,7 @@ void FMetalRHICommandContext::RHIEndDrawIndexedPrimitiveUP()
 	// how many to draw
 	uint32 NumIndices = GetVertexCountForPrimitiveCount(PendingNumPrimitives, PendingPrimitiveType);
 	
-	Context->DrawIndexedPrimitive(PendingIndexBuffer, PendingIndexDataStride, (PendingIndexDataStride == 2) ? mtlpp::IndexType::UInt16 : mtlpp::IndexType::UInt32, PendingPrimitiveType, 0, 0, NumIndices, 0, PendingNumPrimitives, 1);
+	Context->DrawIndexedPrimitive(PendingIndexBuffer, FMetalTexture(), PendingIndexDataStride, (PendingIndexDataStride == 2) ? mtlpp::IndexType::UInt16 : mtlpp::IndexType::UInt32, PendingPrimitiveType, 0, 0, NumIndices, 0, PendingNumPrimitives, 1);
 	
 	// mark temp memory as usable
 	PendingVertexBuffer = nil;
diff --git a/Engine/Source/Runtime/Apple/MetalRHI/Private/MetalContext.cpp b/Engine/Source/Runtime/Apple/MetalRHI/Private/MetalContext.cpp
--- a/Engine/Source/Runtime/Apple/MetalRHI/Private/MetalContext.cpp
+++ b/Engine/Source/Runtime/Apple/MetalRHI/Private/MetalContext.cpp
@@ -1482,7 +1482,7 @@ void FMetalContext::DrawPrimitiveIndirect(uint32 PrimitiveType, FMetalVertexBuff
 	RenderPass.DrawPrimitiveIndirect(PrimitiveType, VertexBuffer, ArgumentOffset);
 }
 
-void FMetalContext::DrawIndexedPrimitive(FMetalBuffer const& IndexBuffer, uint32 IndexStride, mtlpp::IndexType IndexType, uint32 PrimitiveType, int32 BaseVertexIndex, uint32 FirstInstance, uint32 NumVertices, uint32 StartIndex, uint32 NumPrimitives, uint32 NumInstances)
+void FMetalContext::DrawIndexedPrimitive(FMetalBuffer const& IndexBuffer, FMetalTexture const& IndexTexture, uint32 IndexStride, mtlpp::IndexType IndexType, uint32 PrimitiveType, int32 BaseVertexIndex, uint32 FirstInstance, uint32 NumVertices, uint32 StartIndex, uint32 NumPrimitives, uint32 NumInstances)
 {
 	// finalize any pending state
 	if(!PrepareToDraw(PrimitiveType, GetRHIMetalIndexType(IndexType)))
@@ -1490,7 +1490,7 @@ void FMetalContext::DrawIndexedPrimitive(FMetalBuffer const& IndexBuffer, uint32
 		return;
 	}
 	
-	RenderPass.DrawIndexedPrimitive(IndexBuffer, IndexStride, PrimitiveType, BaseVertexIndex, FirstInstance, NumVertices, StartIndex, NumPrimitives, NumInstances);
+	RenderPass.DrawIndexedPrimitive(IndexBuffer, IndexTexture, IndexStride, PrimitiveType, BaseVertexIndex, FirstInstance, NumVertices, StartIndex, NumPrimitives, NumInstances);
 }
 
 void FMetalContext::DrawIndexedIndirect(FMetalIndexBuffer* IndexBuffer, uint32 PrimitiveType, FMetalStructuredBuffer* VertexBuffer, int32 DrawArgumentsIndex, uint32 NumInstances)
diff --git a/Engine/Source/Runtime/Apple/MetalRHI/Private/MetalContext.h b/Engine/Source/Runtime/Apple/MetalRHI/Private/MetalContext.h
--- a/Engine/Source/Runtime/Apple/MetalRHI/Private/MetalContext.h
+++ b/Engine/Source/Runtime/Apple/MetalRHI/Private/MetalContext.h
@@ -71,7 +71,7 @@ public:
 	
 	void DrawPrimitiveIndirect(uint32 PrimitiveType, FMetalVertexBuffer* VertexBuffer, uint32 ArgumentOffset);
 	
-	void DrawIndexedPrimitive(FMetalBuffer const& IndexBuffer, uint32 IndexStride, mtlpp::IndexType IndexType, uint32 PrimitiveType, int32 BaseVertexIndex, uint32 FirstInstance,
+	void DrawIndexedPrimitive(FMetalBuffer const& IndexBuffer, FMetalTexture const& IndexTexture, uint32 IndexStride, mtlpp::IndexType IndexType, uint32 PrimitiveType, int32 BaseVertexIndex, uint32 FirstInstance,
 							  uint32 NumVertices, uint32 StartIndex, uint32 NumPrimitives, uint32 NumInstances);
 	
 	void DrawIndexedIndirect(FMetalIndexBuffer* IndexBufferRHI, uint32 PrimitiveType, FMetalStructuredBuffer* VertexBufferRHI, int32 DrawArgumentsIndex, uint32 NumInstances);
diff --git a/Engine/Source/Runtime/Apple/MetalRHI/Private/MetalIndexBuffer.cpp b/Engine/Source/Runtime/Apple/MetalRHI/Private/MetalIndexBuffer.cpp
--- a/Engine/Source/Runtime/Apple/MetalRHI/Private/MetalIndexBuffer.cpp
+++ b/Engine/Source/Runtime/Apple/MetalRHI/Private/MetalIndexBuffer.cpp
@@ -12,12 +12,28 @@
 #include "RenderUtils.h"
 #include "HAL/LowLevelMemTracker.h"
 
+static uint32 MetalIndexBufferUsage(uint32 InUsage)
+{
+	uint32 Usage = InUsage;
+	if (RHISupportsTessellation(GMaxRHIShaderPlatform))
+	{
+		Usage |= BUF_ShaderResource;
+	}
+	Usage |= EMetalBufferUsage_GPUOnly|EMetalBufferUsage_LinearTex;
+	return Usage;
+}
+
 /** Constructor */
 FMetalIndexBuffer::FMetalIndexBuffer(uint32 InStride, uint32 InSize, uint32 InUsage)
 	: FRHIIndexBuffer(InStride, InSize, InUsage)
-	, FMetalRHIBuffer(InSize, InUsage|EMetalBufferUsage_GPUOnly|EMetalBufferUsage_LinearTex, RRT_IndexBuffer)
+	, FMetalRHIBuffer(InSize, MetalIndexBufferUsage(InUsage), RRT_IndexBuffer)
 	, IndexType((InStride == 2) ? mtlpp::IndexType::UInt16 : mtlpp::IndexType::UInt32)
 {
+	if (RHISupportsTessellation(GMaxRHIShaderPlatform))
+	{
+		EPixelFormat Format = IndexType == mtlpp::IndexType::UInt16 ? PF_R16_UINT : PF_R32_UINT;
+		CreateLinearTexture(Format, this);
+	}
 }
 
 FMetalIndexBuffer::~FMetalIndexBuffer()
diff --git a/Engine/Source/Runtime/Apple/MetalRHI/Private/MetalPipeline.cpp b/Engine/Source/Runtime/Apple/MetalRHI/Private/MetalPipeline.cpp
--- a/Engine/Source/Runtime/Apple/MetalRHI/Private/MetalPipeline.cpp
+++ b/Engine/Source/Runtime/Apple/MetalRHI/Private/MetalPipeline.cpp
@@ -342,9 +342,9 @@ static FMetalShaderPipeline* CreateMTLRenderPipeline(bool const bSync, FMetalGra
     FMetalDomainShader* DomainShader = (FMetalDomainShader*)Init.BoundShaderState.DomainShaderRHI;
     FMetalPixelShader* PixelShader = (FMetalPixelShader*)Init.BoundShaderState.PixelShaderRHI;
     
-    mtlpp::Function vertexFunction = VertexShader->GetFunction(IndexType);
-    mtlpp::Function fragmentFunction = PixelShader ? PixelShader->GetFunction(EMetalIndexType_None) : nil;
-    mtlpp::Function domainFunction = DomainShader ? DomainShader->GetFunction(EMetalIndexType_None) : nil;
+    mtlpp::Function vertexFunction = VertexShader->GetFunction();
+    mtlpp::Function fragmentFunction = PixelShader ? PixelShader->GetFunction() : nil;
+    mtlpp::Function domainFunction = DomainShader ? DomainShader->GetFunction() : nil;
     
     FMetalShaderPipeline* Pipeline = nil;
     if (vertexFunction && ((PixelShader != nullptr) == (fragmentFunction != nil)) && ((DomainShader != nullptr) == (domainFunction != nil)))
diff --git a/Engine/Source/Runtime/Apple/MetalRHI/Private/MetalRenderPass.cpp b/Engine/Source/Runtime/Apple/MetalRHI/Private/MetalRenderPass.cpp
--- a/Engine/Source/Runtime/Apple/MetalRHI/Private/MetalRenderPass.cpp
+++ b/Engine/Source/Runtime/Apple/MetalRHI/Private/MetalRenderPass.cpp
@@ -409,7 +409,7 @@ void FMetalRenderPass::DrawPrimitive(uint32 PrimitiveType, uint32 BaseVertexInde
 	}
 	else
 	{
-		DrawPatches(PrimitiveType, nullptr, 0, BaseVertexIndex, 0, 0, NumPrimitives, NumInstances);
+		DrawPatches(PrimitiveType, nullptr, nullptr, 0, BaseVertexIndex, 0, 0, NumPrimitives, NumInstances);
 	}
 	
 	ConditionalSubmit();	
@@ -437,7 +437,7 @@ void FMetalRenderPass::DrawPrimitiveIndirect(uint32 PrimitiveType, FMetalVertexB
 	}
 }
 
-void FMetalRenderPass::DrawIndexedPrimitive(FMetalBuffer const& IndexBuffer, uint32 IndexStride, uint32 PrimitiveType, int32 BaseVertexIndex, uint32 FirstInstance,
+void FMetalRenderPass::DrawIndexedPrimitive(FMetalBuffer const& IndexBuffer, FMetalTexture const& IndexTexture, uint32 IndexStride, uint32 PrimitiveType, int32 BaseVertexIndex, uint32 FirstInstance,
 											 uint32 NumVertices, uint32 StartIndex, uint32 NumPrimitives, uint32 NumInstances)
 {
 	// We need at least one to cover all use cases
@@ -511,7 +511,7 @@ void FMetalRenderPass::DrawIndexedPrimitive(FMetalBuffer const& IndexBuffer, uin
 	}
 	else
 	{
-		DrawPatches(PrimitiveType, IndexBuffer, IndexStride, BaseVertexIndex, FirstInstance, StartIndex, NumPrimitives, NumInstances);
+		DrawPatches(PrimitiveType, IndexBuffer, IndexTexture, IndexStride, BaseVertexIndex, FirstInstance, StartIndex, NumPrimitives, NumInstances);
 	}
 	
 	ConditionalSubmit();
@@ -564,7 +564,7 @@ void FMetalRenderPass::DrawIndexedPrimitiveIndirect(uint32 PrimitiveType,FMetalI
 	}
 }
 
-void FMetalRenderPass::DrawPatches(uint32 PrimitiveType,FMetalBuffer const& IndexBuffer, uint32 IndexBufferStride, int32 BaseVertexIndex, uint32 FirstInstance, uint32 StartIndex,
+void FMetalRenderPass::DrawPatches(uint32 PrimitiveType,FMetalBuffer const& IndexBuffer, FMetalTexture const& IndexTexture, uint32 IndexBufferStride, int32 BaseVertexIndex, uint32 FirstInstance, uint32 StartIndex,
 									uint32 NumPrimitives, uint32 NumInstances)
 {
 	if (GetMetalDeviceContext().SupportsFeature(EMetalFeaturesTessellation))
@@ -617,9 +617,14 @@ void FMetalRenderPass::DrawPatches(uint32 PrimitiveType,FMetalBuffer const& Inde
 		if(IndexBuffer && Pipeline->TessellationPipelineDesc.TessellationControlPointIndexBufferIndex != UINT_MAX)
 		{
 			PrologueEncoder.SetShaderBuffer(mtlpp::FunctionType::Kernel, IndexBuffer, StartIndex * IndexBufferStride, IndexBuffer.GetLength() - (StartIndex * IndexBufferStride), Pipeline->TessellationPipelineDesc.TessellationControlPointIndexBufferIndex, mtlpp::ResourceUsage::Read);
-			PrologueEncoder.SetShaderBuffer(mtlpp::FunctionType::Kernel, IndexBuffer, StartIndex * IndexBufferStride, IndexBuffer.GetLength() - (StartIndex * IndexBufferStride), Pipeline->TessellationPipelineDesc.TessellationIndexBufferIndex, mtlpp::ResourceUsage::Read);
+			
 			State.SetShaderBuffer(SF_Vertex, nil, nil, 0, 0, Pipeline->TessellationPipelineDesc.TessellationControlPointIndexBufferIndex, mtlpp::ResourceUsage(0));
-			State.SetShaderBuffer(SF_Vertex, nil, nil, 0, 0, Pipeline->TessellationPipelineDesc.TessellationIndexBufferIndex, mtlpp::ResourceUsage(0));
+		}
+		
+		if (Pipeline->TessellationPipelineDesc.TessellationIndexBufferIndex != UINT_MAX)
+		{
+			PrologueEncoder.SetShaderTexture(mtlpp::FunctionType::Kernel, IndexTexture, Pipeline->TessellationPipelineDesc.TessellationIndexBufferIndex, mtlpp::ResourceUsage::Read);
+			State.SetShaderTexture(SF_Vertex, nil, Pipeline->TessellationPipelineDesc.TessellationIndexBufferIndex, mtlpp::ResourceUsage(0));
 		}
 		
 		if(Pipeline->TessellationPipelineDesc.TessellationOutputControlPointBufferIndex != UINT_MAX) //TessellationOutputControlPointBufferIndex -> hullShaderOutputBuffer
@@ -652,7 +657,8 @@ void FMetalRenderPass::DrawPatches(uint32 PrimitiveType,FMetalBuffer const& Inde
 		}
 		
 		// set the patchCount
-		PrologueEncoder.SetShaderBytes(mtlpp::FunctionType::Kernel, (const uint8*)&NumPrimitives, sizeof(NumPrimitives), Pipeline->TessellationPipelineDesc.TessellationPatchCountBufferIndex);
+		uint32 patchCountData[] = { NumPrimitives, StartIndex };
+		PrologueEncoder.SetShaderBytes(mtlpp::FunctionType::Kernel, (const uint8*)&patchCountData[0], sizeof(patchCountData), Pipeline->TessellationPipelineDesc.TessellationPatchCountBufferIndex);
 		State.SetShaderBuffer(SF_Vertex, nil, nil, 0, 0, Pipeline->TessellationPipelineDesc.TessellationPatchCountBufferIndex, mtlpp::ResourceUsage(0));
 		
 		if (boundShaderState->VertexShader->SideTableBinding >= 0)
diff --git a/Engine/Source/Runtime/Apple/MetalRHI/Private/MetalRenderPass.h b/Engine/Source/Runtime/Apple/MetalRHI/Private/MetalRenderPass.h
--- a/Engine/Source/Runtime/Apple/MetalRHI/Private/MetalRenderPass.h
+++ b/Engine/Source/Runtime/Apple/MetalRHI/Private/MetalRenderPass.h
@@ -40,14 +40,14 @@ public:
     
     void DrawPrimitiveIndirect(uint32 PrimitiveType, FMetalVertexBuffer* VertexBuffer, uint32 ArgumentOffset);
     
-    void DrawIndexedPrimitive(FMetalBuffer const& IndexBuffer, uint32 IndexStride, uint32 PrimitiveType, int32 BaseVertexIndex, uint32 FirstInstance,
+    void DrawIndexedPrimitive(FMetalBuffer const& IndexBuffer, FMetalTexture const& IndexTexture, uint32 IndexStride, uint32 PrimitiveType, int32 BaseVertexIndex, uint32 FirstInstance,
                          uint32 NumVertices, uint32 StartIndex, uint32 NumPrimitives, uint32 NumInstances);
     
     void DrawIndexedIndirect(FMetalIndexBuffer* IndexBufferRHI, uint32 PrimitiveType, FMetalStructuredBuffer* VertexBufferRHI, int32 DrawArgumentsIndex, uint32 NumInstances);
     
     void DrawIndexedPrimitiveIndirect(uint32 PrimitiveType,FMetalIndexBuffer* IndexBufferRHI,FMetalVertexBuffer* VertexBufferRHI,uint32 ArgumentOffset);
     
-    void DrawPatches(uint32 PrimitiveType, FMetalBuffer const& IndexBuffer, uint32 IndexBufferStride, int32 BaseVertexIndex, uint32 FirstInstance, uint32 StartIndex,
+    void DrawPatches(uint32 PrimitiveType, FMetalBuffer const& IndexBuffer, FMetalTexture const& IndexTexture, uint32 IndexBufferStride, int32 BaseVertexIndex, uint32 FirstInstance, uint32 StartIndex,
                      uint32 NumPrimitives, uint32 NumInstances);
     
     void Dispatch(uint32 ThreadGroupCountX, uint32 ThreadGroupCountY, uint32 ThreadGroupCountZ);
diff --git a/Engine/Source/Runtime/Apple/MetalRHI/Private/MetalShaders.cpp b/Engine/Source/Runtime/Apple/MetalRHI/Private/MetalShaders.cpp
--- a/Engine/Source/Runtime/Apple/MetalRHI/Private/MetalShaders.cpp
+++ b/Engine/Source/Runtime/Apple/MetalRHI/Private/MetalShaders.cpp
@@ -221,7 +221,6 @@ void TMetalBaseShader<BaseResourceType, ShaderType>::Init(const TArray<uint8>& I
 	// If this triggers than a level above us has failed to provide valid shader data and the cook is probably bogus
 	UE_CLOG(Header.SourceLen == 0 || Header.SourceCRC == 0, LogMetal, Fatal, TEXT("Invalid Shader Bytecode provided."));
 	
-    bTessFunctionConstants = Header.bTessFunctionConstants;
     bDeviceFunctionConstants = Header.bDeviceFunctionConstants;
 
 	// remember where the header ended and code (precompiled or source) begins
@@ -267,29 +266,26 @@ void TMetalBaseShader<BaseResourceType, ShaderType>::Init(const TArray<uint8>& I
 		[GlslCodeNSString retain];
 	}
 	
-	bHasFunctionConstants = (Header.bTessFunctionConstants || Header.bDeviceFunctionConstants || Header.Bindings.TypedBuffers);
+	bHasFunctionConstants = (Header.bDeviceFunctionConstants);
 
 	ConstantValueHash = 0;
 	
 	Library = InLibrary;
 	
 	bool bNeedsCompiling = false;
-	uint32 const Count = Header.bTessFunctionConstants ? EMetalIndexType_Num : 1;
-	for (uint32 i = 0; i < Count; i++)
+
+	// Find the existing compiled shader in the cache.
+	uint32 FunctionConstantHash = ConstantValueHash;
+	FMetalCompiledShaderKey Key(Header.SourceLen, Header.SourceCRC, FunctionConstantHash);
+	
+	Function = GetMetalCompiledShaderCache().FindRef(Key);
+	if (!Library && Function)
 	{
-		// Find the existing compiled shader in the cache.
-		uint32 FunctionConstantHash = i ^ ConstantValueHash;
-		FMetalCompiledShaderKey Key(Header.SourceLen, Header.SourceCRC, FunctionConstantHash);
-		
-		Function[i] = GetMetalCompiledShaderCache().FindRef(Key);
-		if (!Library && Function[i])
-		{
-			Library = GetMetalCompiledShaderCache().FindLibrary(Function[i]);
-		}
-		else
-		{
-			bNeedsCompiling = true;
-		}
+		Library = GetMetalCompiledShaderCache().FindLibrary(Function);
+	}
+	else
+	{
+		bNeedsCompiling = true;
 	}
 	
     Bindings = Header.Bindings;
@@ -453,11 +449,8 @@ void TMetalBaseShader<BaseResourceType, ShaderType>::Init(const TArray<uint8>& I
 		UE_CLOG(bHasFunctionConstants && !FMetalCommandQueue::SupportsFeature(EMetalFeaturesFunctionConstants), LogMetal, Error, TEXT("Metal shader has function constants but current OS/device does not support them."));
 		UE_CLOG(bHasFunctionConstants && !FMetalCommandQueue::SupportsFeature(EMetalFeaturesFunctionConstants), LogMetal, Fatal, TEXT("%s"), *FString(GetSourceCode()));
 		
-        for (uint32 i = 0; i < Count; i++)
-        {
-            GetCompiledFunction((EMetalIndexType)i, true);
-        }
-	}
+        GetCompiledFunction(true);
+    }
 	UniformBuffersCopyInfo = Header.UniformBuffersCopyInfo;
 	SideTableBinding = Header.SideTable;
 }
@@ -470,16 +463,16 @@ TMetalBaseShader<BaseResourceType, ShaderType>::~TMetalBaseShader()
 }
 
 template<typename BaseResourceType, int32 ShaderType>
-mtlpp::Function TMetalBaseShader<BaseResourceType, ShaderType>::GetCompiledFunction(EMetalIndexType IndexType, bool const bAsync)
+mtlpp::Function TMetalBaseShader<BaseResourceType, ShaderType>::GetCompiledFunction( bool const bAsync)
 {
-	mtlpp::Function Func = Function[IndexType];
+	mtlpp::Function Func = Function;
     
 	if (!Func)
 	{
 		// Find the existing compiled shader in the cache.
-		uint32 FunctionConstantHash = IndexType ^ ConstantValueHash;
+		uint32 FunctionConstantHash = ConstantValueHash;
 		FMetalCompiledShaderKey Key(SourceLen, SourceCRC, FunctionConstantHash);
-		Func = Function[IndexType] = GetMetalCompiledShaderCache().FindRef(Key);
+		Func = Function = GetMetalCompiledShaderCache().FindRef(Key);
 		
 		if (!Func)
 		{
@@ -491,11 +484,6 @@ mtlpp::Function TMetalBaseShader<BaseResourceType, ShaderType>::GetCompiledFunct
             {
                 ConstantValues = mtlpp::FunctionConstantValues();
 				
-                if(bTessFunctionConstants)
-                {
-                    // Index 32 is the tessellation index-buffer presence constant
-					ConstantValues.SetConstantValues(&IndexType, mtlpp::DataType::UInt, ns::Range(32, 1));
-                }
                 if (bDeviceFunctionConstants)
                 {
                     // Index 33 is the device vendor id constant
@@ -508,21 +496,21 @@ mtlpp::Function TMetalBaseShader<BaseResourceType, ShaderType>::GetCompiledFunct
 				METAL_GPUPROFILE(FScopedMetalCPUStats CPUStat(FString::Printf(TEXT("NewFunction: %s"), *FString(Name))));
                 if (!bHasFunctionConstants)
                 {
-                    Function[IndexType] = Library.NewFunction(Name);
+                    Function = Library.NewFunction(Name);
                 }
                 else
                 {
 					ns::AutoReleasedError AError;
-					Function[IndexType] = Library.NewFunction(Name, ConstantValues, &AError);
+					Function = Library.NewFunction(Name, ConstantValues, &AError);
 					ns::Error Error = AError;
-					UE_CLOG(Function[IndexType] == nil, LogMetal, Error, TEXT("Failed to create function: %s"), *FString(Error.GetPtr().description));
-                    UE_CLOG(Function[IndexType] == nil, LogMetal, Fatal, TEXT("*********** Error\n%s"), *FString(GetSourceCode()));
+					UE_CLOG(Function == nil, LogMetal, Error, TEXT("Failed to create function: %s"), *FString(Error.GetPtr().description));
+                    UE_CLOG(Function == nil, LogMetal, Fatal, TEXT("*********** Error\n%s"), *FString(GetSourceCode()));
                 }
                 
-                check(Function[IndexType]);
-                GetMetalCompiledShaderCache().Add(Key, Library, Function[IndexType]);
+                check(Function);
+                GetMetalCompiledShaderCache().Add(Key, Library, Function);
                 
-                Func = Function[IndexType];
+                Func = Function;
             }
             else
             {
@@ -580,7 +568,7 @@ FMetalShaderPipeline* FMetalComputeShader::GetPipeline()
 {
 	if (!Pipeline)
 	{
-		mtlpp::Function Func = GetCompiledFunction(EMetalIndexType_None);
+		mtlpp::Function Func = GetCompiledFunction();
 		check(Func);
         
 		ns::Error Error;
@@ -699,9 +687,9 @@ FMetalVertexShader::FMetalVertexShader(const TArray<uint8>& InCode, mtlpp::Libra
 	TessellationPatchesPerThreadGroup = Header.TessellationPatchesPerThreadGroup;
 }
 
-mtlpp::Function FMetalVertexShader::GetFunction(EMetalIndexType IndexType)
+mtlpp::Function FMetalVertexShader::GetFunction()
 {
-	return GetCompiledFunction(IndexType);
+	return GetCompiledFunction();
 }
 
 FMetalPixelShader::FMetalPixelShader(const TArray<uint8>& InCode)
@@ -716,9 +704,9 @@ FMetalPixelShader::FMetalPixelShader(const TArray<uint8>& InCode, mtlpp::Library
 	Init(InCode, Header, InLibrary);
 }
 
-mtlpp::Function FMetalPixelShader::GetFunction(EMetalIndexType IndexType)
+mtlpp::Function FMetalPixelShader::GetFunction()
 {
-	return GetCompiledFunction(IndexType);
+	return GetCompiledFunction();
 }
 
 FMetalHullShader::FMetalHullShader(const TArray<uint8>& InCode)
@@ -733,9 +721,9 @@ FMetalHullShader::FMetalHullShader(const TArray<uint8>& InCode, mtlpp::Library I
 	Init(InCode, Header, InLibrary);
 }
 
-mtlpp::Function FMetalHullShader::GetFunction(EMetalIndexType IndexType)
+mtlpp::Function FMetalHullShader::GetFunction()
 {
-	return GetCompiledFunction(IndexType);
+	return GetCompiledFunction();
 }
 
 FMetalDomainShader::FMetalDomainShader(const TArray<uint8>& InCode)
@@ -792,9 +780,9 @@ FMetalDomainShader::FMetalDomainShader(const TArray<uint8>& InCode, mtlpp::Libra
 	}
 }
 
-mtlpp::Function FMetalDomainShader::GetFunction(EMetalIndexType IndexType)
+mtlpp::Function FMetalDomainShader::GetFunction()
 {
-	return GetCompiledFunction(IndexType);
+	return GetCompiledFunction();
 }
 
 FVertexShaderRHIRef FMetalDynamicRHI::RHICreateVertexShader(const TArray<uint8>& Code)
@@ -1026,7 +1014,7 @@ FPixelShaderRHIRef FMetalShaderLibrary::CreatePixelShader(const FSHAHash& Hash)
 	if (Code)
 	{
 		FMetalPixelShader* Shader = new FMetalPixelShader(Code->Data, Library[Code->Index]);
-		if (Shader->GetFunction(EMetalIndexType_None))
+		if (Shader->GetFunction())
 		{
 			return Shader;
 		}
@@ -1046,7 +1034,7 @@ FVertexShaderRHIRef FMetalShaderLibrary::CreateVertexShader(const FSHAHash& Hash
 	if (Code)
 	{
 		FMetalVertexShader* Shader = new FMetalVertexShader(Code->Data, Library[Code->Index]);
-		if (Shader->GetFunction(EMetalIndexType_None))
+		if (Shader->GetFunction())
 		{
 			return Shader;
 		}
@@ -1065,7 +1053,7 @@ FHullShaderRHIRef FMetalShaderLibrary::CreateHullShader(const FSHAHash& Hash)
 	if (Code)
 	{
 		FMetalHullShader* Shader = new FMetalHullShader(Code->Data, Library[Code->Index]);
-		if(Shader->GetFunction(EMetalIndexType_None))
+		if(Shader->GetFunction())
 		{
 			return Shader;
 		}
@@ -1084,7 +1072,7 @@ FDomainShaderRHIRef FMetalShaderLibrary::CreateDomainShader(const FSHAHash& Hash
 	if (Code)
 	{
 		FMetalDomainShader* Shader = new FMetalDomainShader(Code->Data, Library[Code->Index]);
-		if (Shader->GetFunction(EMetalIndexType_None))
+		if (Shader->GetFunction())
 		{
 			return Shader;
 		}
diff --git a/Engine/Source/Runtime/Apple/MetalRHI/Public/MetalResources.h b/Engine/Source/Runtime/Apple/MetalRHI/Public/MetalResources.h
--- a/Engine/Source/Runtime/Apple/MetalRHI/Public/MetalResources.h
+++ b/Engine/Source/Runtime/Apple/MetalRHI/Public/MetalResources.h
@@ -92,7 +92,7 @@ public:
 	, GlslCodeNSString(nil)
 	, CodeSize(0)
 	{
-		Function[EMetalIndexType_None] = Function[EMetalIndexType_UInt16] = Function[EMetalIndexType_UInt32] = nil;
+		Function = nil;
 	}
 	
 	void Init(const TArray<uint8>& InCode, FMetalCodeHeader& Header, mtlpp::Library InLibrary = nil);
@@ -145,10 +145,10 @@ public:
 	uint32 ConstantValueHash;
     
 protected:
-	mtlpp::Function GetCompiledFunction(EMetalIndexType IndexType, bool const bAsync = false);
+	mtlpp::Function GetCompiledFunction(bool const bAsync = false);
 
     // this is the compiler shader
-	mtlpp::Function Function[EMetalIndexType_Num];
+	mtlpp::Function Function;
     
 private:
 	// This is the MTLLibrary for the shader so we can dynamically refine the MTLFunction
@@ -165,7 +165,6 @@ private:
     
     // Function constant states
     bool bHasFunctionConstants;
-    bool bTessFunctionConstants;
     bool bDeviceFunctionConstants;
 };
 
@@ -175,7 +174,7 @@ public:
 	FMetalVertexShader(const TArray<uint8>& InCode);
 	FMetalVertexShader(const TArray<uint8>& InCode, mtlpp::Library InLibrary);
 	
-	mtlpp::Function GetFunction(EMetalIndexType IndexType);
+	mtlpp::Function GetFunction();
 	
 	// for VSHS
 	FMetalTessellationOutputs TessellationOutputAttribs;
@@ -198,7 +197,7 @@ public:
 	FMetalPixelShader(const TArray<uint8>& InCode);
 	FMetalPixelShader(const TArray<uint8>& InCode, mtlpp::Library InLibrary);
 	
-	mtlpp::Function GetFunction(EMetalIndexType IndexType);
+	mtlpp::Function GetFunction();
 };
 
 class FMetalHullShader : public TMetalBaseShader<FRHIHullShader, SF_Hull>
@@ -207,7 +206,7 @@ public:
 	FMetalHullShader(const TArray<uint8>& InCode);
 	FMetalHullShader(const TArray<uint8>& InCode, mtlpp::Library InLibrary);
 	
-	mtlpp::Function GetFunction(EMetalIndexType IndexType);
+	mtlpp::Function GetFunction();
 };
 
 class FMetalDomainShader : public TMetalBaseShader<FRHIDomainShader, SF_Domain>
@@ -216,7 +215,7 @@ public:
 	FMetalDomainShader(const TArray<uint8>& InCode);
 	FMetalDomainShader(const TArray<uint8>& InCode, mtlpp::Library InLibrary);
 	
-	mtlpp::Function GetFunction(EMetalIndexType IndexType);
+	mtlpp::Function GetFunction();
 	
 	mtlpp::Winding TessellationOutputWinding;
 	mtlpp::TessellationPartitionMode TessellationPartitioning;
diff --git a/Engine/Source/Runtime/Apple/MetalRHI/Public/MetalShaderResources.h b/Engine/Source/Runtime/Apple/MetalRHI/Public/MetalShaderResources.h
--- a/Engine/Source/Runtime/Apple/MetalRHI/Public/MetalShaderResources.h
+++ b/Engine/Source/Runtime/Apple/MetalRHI/Public/MetalShaderResources.h
@@ -292,7 +292,6 @@ struct FMetalCodeHeader
 	EMetalOutputWindingMode TessellationOutputWinding;
 	EMetalPartitionMode TessellationPartitioning;
 	
-	bool bTessFunctionConstants;
 	bool bDeviceFunctionConstants;
 	
 	FMetalCodeHeader()
@@ -320,7 +319,6 @@ struct FMetalCodeHeader
 	, SideTable(-1)
 	, TessellationOutputWinding(EMetalOutputWindingMode::Clockwise)
 	, TessellationPartitioning(EMetalPartitionMode::Pow2)
-	, bTessFunctionConstants(false)
 	, bDeviceFunctionConstants(false)
 	{
 		
@@ -386,7 +384,6 @@ inline FArchive& operator<<(FArchive& Ar, FMetalCodeHeader& Header)
 	
 	Ar << Header.TessellationOutputWinding;
 	Ar << Header.TessellationPartitioning;
-	Ar << Header.bTessFunctionConstants;
 	Ar << Header.bDeviceFunctionConstants;
 
     return Ar;
