Hull shaders don't actually set any parameters so it doesn't matter that we never bind them - instead just ignore them entirely and note that should we ever have legitimate Hull-shader only global parameters we'll need to think again.

diff --git a/Engine/Source/Runtime/Apple/MetalRHI/Private/MetalCommands.cpp b/Engine/Source/Runtime/Apple/MetalRHI/Private/MetalCommands.cpp
--- a/Engine/Source/Runtime/Apple/MetalRHI/Private/MetalCommands.cpp
+++ b/Engine/Source/Runtime/Apple/MetalRHI/Private/MetalCommands.cpp
@@ -484,9 +484,8 @@ void FMetalRHICommandContext::RHISetShaderParameter(FVertexShaderRHIParamRef Ver
 
 void FMetalRHICommandContext::RHISetShaderParameter(FHullShaderRHIParamRef HullShaderRHI, uint32 BufferIndex, uint32 BaseIndex, uint32 NumBytes, const void* NewValue)
 {
-	@autoreleasepool {
-	Context->GetCurrentState().GetShaderParameters(CrossCompiler::SHADER_STAGE_HULL).Set(BufferIndex, BaseIndex, NumBytes, NewValue);
-	}
+	// Just ignore Hull shader parameter sets - none of our Hull shaders have any loose parameters to bind.
+	// @todo Whenever we do put a shader parameter into a hull shader we'll need to map it into the vertex-shader parameter buffer so that it can be set on the device.
 }
 
 void FMetalRHICommandContext::RHISetShaderParameter(FPixelShaderRHIParamRef PixelShaderRHI, uint32 BufferIndex, uint32 BaseIndex, uint32 NumBytes, const void* NewValue)
diff --git a/Engine/Source/Runtime/Apple/MetalRHI/Private/MetalStateCache.cpp b/Engine/Source/Runtime/Apple/MetalRHI/Private/MetalStateCache.cpp
--- a/Engine/Source/Runtime/Apple/MetalRHI/Private/MetalStateCache.cpp
+++ b/Engine/Source/Runtime/Apple/MetalRHI/Private/MetalStateCache.cpp
@@ -1838,12 +1838,6 @@ void FMetalStateCache::CommitTessellationResources(FMetalCommandEncoder* Raster,
     
     SetResourcesFromTables(GraphicsPSO->HullShader, CrossCompiler::SHADER_STAGE_HULL);
 	
-	// @todo Can't put global uniforms in Hull shaders - the high-level no longer
-	// replicates the bindings to both vertex & hull shaders, so we have conflicts and the
-	// hull shader versions are almost always 0-filled which causes chaos.
-	// We need to split the shaders out into two distinct stages to match D3D
-    // GetShaderParameters(CrossCompiler::SHADER_STAGE_HULL).CommitPackedGlobals(this, Compute, SF_Hull, GraphicsPSO->HullShader->Bindings);
-	
 	SetResourcesFromTables(GraphicsPSO->DomainShader, CrossCompiler::SHADER_STAGE_DOMAIN);
    	GetShaderParameters(CrossCompiler::SHADER_STAGE_DOMAIN).CommitPackedGlobals(this, Raster, SF_Domain, GraphicsPSO->DomainShader->Bindings);
 }
