Remove some dead [REDACTED_HANDLE]'s.

diff --git a/Engine/Source/Runtime/Apple/MetalRHI/Private/MetalRHI.cpp b/Engine/Source/Runtime/Apple/MetalRHI/Private/MetalRHI.cpp
--- a/Engine/Source/Runtime/Apple/MetalRHI/Private/MetalRHI.cpp
+++ b/Engine/Source/Runtime/Apple/MetalRHI/Private/MetalRHI.cpp
@@ -94,8 +94,7 @@ FMetalDynamicRHI::FMetalDynamicRHI(ERHIFeatureLevel::Type RequestedFeatureLevel)
 	check( IsInGameThread() );
 	check( !GIsThreadedRendering );
 	
-	// @todo Zebra This is now supported on all GPUs in Mac Metal, but not on iOS.
-	// we cannot render to a volume texture without geometry shader support
+	// we cannot render to a volume texture without geometry shader or vertex-shader-layer support, so initialise to false and enable based on platform feature availability
 	GSupportsVolumeTextureRendering = false;
 	
 	// Metal always needs a render target to render with fragment shaders!
@@ -428,7 +427,7 @@ FMetalDynamicRHI::FMetalDynamicRHI(ERHIFeatureLevel::Type RequestedFeatureLevel)
 	if (FPlatformMisc::IsDebuggerPresent() && UE_BUILD_DEBUG)
 	{
 #if PLATFORM_IOS // @todo zebra : needs a RENDER_API or whatever
-		// Enable GL debug markers if we're running in Xcode
+		// Enable debug markers if we're running in Xcode
 		extern int32 GEmitMeshDrawEvent;
 		GEmitMeshDrawEvent = 1;
 #endif
@@ -682,7 +681,7 @@ FMetalDynamicRHI::FMetalDynamicRHI(ERHIFeatureLevel::Type RequestedFeatureLevel)
 		
 	GPixelFormats[PF_BC5				].PlatformFormat	= (uint32)mtlpp::PixelFormat::Invalid;
 	GPixelFormats[PF_R5G6B5_UNORM		].PlatformFormat	= (uint32)mtlpp::PixelFormat::B5G6R5Unorm;
-#else // @todo zebra : srgb?
+#else
     GPixelFormats[PF_DXT1				].PlatformFormat	= (uint32)mtlpp::PixelFormat::BC1_RGBA;
     GPixelFormats[PF_DXT3				].PlatformFormat	= (uint32)mtlpp::PixelFormat::BC2_RGBA;
     GPixelFormats[PF_DXT5				].PlatformFormat	= (uint32)mtlpp::PixelFormat::BC3_RGBA;
@@ -983,7 +982,6 @@ void FMetalRHICommandContext::RHIBeginFrame()
 void FMetalRHIImmediateCommandContext::RHIEndFrame()
 {
 	@autoreleasepool {
-	// @todo zebra: GPUProfilingData.EndFrame();
 #if ENABLE_METAL_GPUPROFILE
 	Profiler->EndFrame();
 #endif
@@ -1023,7 +1021,6 @@ void FMetalRHICommandContext::RHIEndScene()
 void FMetalRHICommandContext::RHIPushEvent(const TCHAR* Name, FColor Color)
 {
 #if ENABLE_METAL_GPUEVENTS
-	// @todo zebra : this was "[NSString stringWithTCHARString:Name]", an extension only on ios for some reason, it should be on Mac also
 	@autoreleasepool
 	{
 		FPlatformMisc::BeginNamedEvent(Color, Name);
diff --git a/Engine/Source/Runtime/Apple/MetalRHI/Private/MetalStateCache.cpp b/Engine/Source/Runtime/Apple/MetalRHI/Private/MetalStateCache.cpp
--- a/Engine/Source/Runtime/Apple/MetalRHI/Private/MetalStateCache.cpp
+++ b/Engine/Source/Runtime/Apple/MetalRHI/Private/MetalStateCache.cpp
@@ -337,7 +337,7 @@ void FMetalStateCache::SetScissorRect(bool const bEnable, mtlpp::ScissorRect con
 
 void FMetalStateCache::SetBlendFactor(FLinearColor const& InBlendFactor)
 {
-	if(BlendFactor != InBlendFactor) // @todo zebra
+	if(BlendFactor != InBlendFactor)
 	{
 		BlendFactor = InBlendFactor;
 		RasterBits |= EMetalRenderFlagBlendColor;
@@ -346,7 +346,7 @@ void FMetalStateCache::SetBlendFactor(FLinearColor const& InBlendFactor)
 
 void FMetalStateCache::SetStencilRef(uint32 const InStencilRef)
 {
-	if(StencilRef != InStencilRef) // @todo zebra
+	if(StencilRef != InStencilRef)
 	{
 		StencilRef = InStencilRef;
 		RasterBits |= EMetalRenderFlagStencilReferenceValue;
@@ -355,7 +355,7 @@ void FMetalStateCache::SetStencilRef(uint32 const InStencilRef)
 
 void FMetalStateCache::SetDepthStencilState(FMetalDepthStencilState* InDepthStencilState)
 {
-	if(DepthStencilState != InDepthStencilState) // @todo zebra
+	if(DepthStencilState != InDepthStencilState)
 	{
 		DepthStencilState = InDepthStencilState;
 		RasterBits |= EMetalRenderFlagDepthStencilState;
@@ -364,7 +364,7 @@ void FMetalStateCache::SetDepthStencilState(FMetalDepthStencilState* InDepthSten
 
 void FMetalStateCache::SetRasterizerState(FMetalRasterizerState* InRasterizerState)
 {
-	if(RasterizerState != InRasterizerState) // @todo zebra
+	if(RasterizerState != InRasterizerState)
 	{
 		RasterizerState = InRasterizerState;
 		RasterBits |= EMetalRenderFlagFrontFacingWinding|EMetalRenderFlagCullMode|EMetalRenderFlagDepthBias|EMetalRenderFlagTriangleFillMode;
@@ -373,7 +373,7 @@ void FMetalStateCache::SetRasterizerState(FMetalRasterizerState* InRasterizerSta
 
 void FMetalStateCache::SetComputeShader(FMetalComputeShader* InComputeShader)
 {
-	if(ComputeShader != InComputeShader) // @todo zebra
+	if(ComputeShader != InComputeShader)
 	{
 		ComputeShader = InComputeShader;
 		
diff --git a/Engine/Source/Runtime/Apple/MetalRHI/Private/MetalTexture.cpp b/Engine/Source/Runtime/Apple/MetalRHI/Private/MetalTexture.cpp
--- a/Engine/Source/Runtime/Apple/MetalRHI/Private/MetalTexture.cpp
+++ b/Engine/Source/Runtime/Apple/MetalRHI/Private/MetalTexture.cpp
@@ -527,8 +527,6 @@ void FMetalSurface::Init(FMetalSurface& Source, NSRange MipRange, EPixelFormat F
 	}
 	
 	ns::Range Slices(0, Source.Texture.GetArrayLength() * (bIsCubemap ? 6 : 1));
-	// @todo Zebra Temporary workaround for absence of X24_G8 or equivalent to GL_STENCIL_INDEX so that the stencil part of a texture may be sampled
-	// For now, if we find ourselves *requiring* this we lazily blit the stencil data out to a separate texture. radr://21813831
 	if(Source.PixelFormat != PF_DepthStencil && !bUseSourceTex)
 	{
 		Texture = Source.Texture.NewTextureView(MetalFormat, Source.Texture.GetTextureType(), ns::Range(MipRange.location, MipRange.length), Slices);
@@ -1398,7 +1396,7 @@ void FMetalSurface::UpdateSurface(FMetalBuffer& Buffer, uint32 MipIndex, uint32
 	else
 	{
 #if !PLATFORM_MAC
-		if (Texture.GetPixelFormat() >= mtlpp::PixelFormat::PVRTC_RGB_2BPP && Texture.GetPixelFormat() <= mtlpp::PixelFormat::ETC2_RGB8A1_sRGB) // @todo zebra
+		if (Texture.GetPixelFormat() >= mtlpp::PixelFormat::PVRTC_RGB_2BPP && Texture.GetPixelFormat() <= mtlpp::PixelFormat::ETC2_RGB8A1_sRGB) // @todo Calculate correct strides and byte-counts
 		{
 			Stride = 0;
 			BytesPerImage = 0;
