Fix incorrect compute state tracking in Metal.

diff --git a/Engine/Source/Runtime/Apple/MetalRHI/Private/MetalStateCache.h b/Engine/Source/Runtime/Apple/MetalRHI/Private/MetalStateCache.h
--- a/Engine/Source/Runtime/Apple/MetalRHI/Private/MetalStateCache.h
+++ b/Engine/Source/Runtime/Apple/MetalRHI/Private/MetalStateCache.h
@@ -9,7 +9,7 @@
 enum EMetalPipelineFlags
 {
 	EMetalPipelineFlagPipelineState = 1 << 0,
-    EMetalPipelineFlagComputeShader = 1 << 1,
+    EMetalPipelineFlagComputeShader = 1 << 5,
     EMetalPipelineFlagRasterMask = 0xF,
     EMetalPipelineFlagComputeMask = 0x30,
     EMetalPipelineFlagMask = 0x3F
