Added a new CVar "rhi.Metal.CacheShaderPipelines"  which controls whether we cache Metal graphics pipeline states for the lifetime of the program. This defaults to 1 (on) which preseves the existing behaviour and improves performance at the expense of memory, but can be set to 0 in the project configuration to release PSO memory back to the OS once all RHI references are released which will save memory but at the expense of performance as PSOs will have to be recreated on-demand. Even when disabled RHI objects will still share references to a single Metal instances of each PSO.

diff --git a/Engine/Source/Runtime/Apple/MetalRHI/Private/MetalPipeline.cpp b/Engine/Source/Runtime/Apple/MetalRHI/Private/MetalPipeline.cpp
--- a/Engine/Source/Runtime/Apple/MetalRHI/Private/MetalPipeline.cpp
+++ b/Engine/Source/Runtime/Apple/MetalRHI/Private/MetalPipeline.cpp
@@ -15,6 +15,13 @@
 #include "Misc/ScopeRWLock.h"
 #include <objc/runtime.h>
 
+static int32 GMetalCacheShaderPipelines = 1;
+static FAutoConsoleVariableRef CVarMetalCacheShaderPipelines(
+	TEXT("rhi.Metal.CacheShaderPipelines"),
+	GMetalCacheShaderPipelines,
+	TEXT("When enabled (1, default) cache all graphics pipeline state objects created in MetalRHI for the life of the program, this trades memory for performance as creating PSOs is expensive in Metal.\n")
+	TEXT("Disable in the project configuration to allow PSOs to be released to save memory at the expense of reduced performance and increased hitching in-game\n. (On by default (1))"), ECVF_ReadOnly);
+
 static int32 GMetalTessellationForcePartitionMode = 0;
 static FAutoConsoleVariableRef CVarMetalTessellationForcePartitionMode(
 	TEXT("rhi.Metal.TessellationForcePartitionMode"),
@@ -49,112 +56,6 @@
 	}
 }
 
-@implementation FMetalShaderPipeline
-
-APPLE_PLATFORM_OBJECT_ALLOC_OVERRIDES(FMetalShaderPipeline)
-
-- (void)dealloc
-{
-	[super dealloc];
-}
-
-#if METAL_DEBUG_OPTIONS
-- (instancetype)init
-{
-	id Self = [super init];
-	if (Self)
-	{
-		RenderPipelineReflection = mtlpp::RenderPipelineReflection(nil);
-		ComputePipelineReflection = mtlpp::ComputePipelineReflection(nil);
-		RenderDesc = mtlpp::RenderPipelineDescriptor(nil);
-		ComputeDesc = mtlpp::ComputePipelineDescriptor(nil);
-	}
-	return Self;
-}
-
-- (void)initResourceMask
-{
-	if (RenderPipelineReflection)
-	{
-		[self initResourceMask:EMetalShaderVertex];
-		[self initResourceMask:EMetalShaderFragment];
-	}
-	if (ComputePipelineReflection)
-	{
-		[self initResourceMask:EMetalShaderCompute];
-	}
-}
-- (void)initResourceMask:(EMetalShaderFrequency)Frequency
-{
-	NSArray<MTLArgument*>* Arguments = nil;
-	switch(Frequency)
-	{
-		case EMetalShaderVertex:
-		{
-			MTLRenderPipelineReflection* Reflection = RenderPipelineReflection;
-			check(Reflection);
-			
-			Arguments = Reflection.vertexArguments;
-			break;
-		}
-		case EMetalShaderFragment:
-		{
-			MTLRenderPipelineReflection* Reflection = RenderPipelineReflection;
-			check(Reflection);
-			
-			Arguments = Reflection.fragmentArguments;
-			break;
-		}
-		case EMetalShaderCompute:
-		{
-			MTLComputePipelineReflection* Reflection = ComputePipelineReflection;
-			check(Reflection);
-			
-			Arguments = Reflection.arguments;
-			break;
-		}
-		default:
-			check(false);
-			break;
-	}
-	
-	for (uint32 i = 0; i < Arguments.count; i++)
-	{
-		MTLArgument* Arg = [Arguments objectAtIndex:i];
-		check(Arg);
-		switch(Arg.type)
-		{
-			case MTLArgumentTypeBuffer:
-			{
-				checkf(Arg.index < ML_MaxBuffers, TEXT("Metal buffer index exceeded!"));
-				ResourceMask[Frequency].BufferMask |= (1 << Arg.index);
-				break;
-			}
-			case MTLArgumentTypeThreadgroupMemory:
-			{
-				break;
-			}
-			case MTLArgumentTypeTexture:
-			{
-				checkf(Arg.index < ML_MaxTextures, TEXT("Metal texture index exceeded!"));
-				ResourceMask[Frequency].TextureMask |= (1 << Arg.index);
-				break;
-			}
-			case MTLArgumentTypeSampler:
-			{
-				checkf(Arg.index < ML_MaxSamplers, TEXT("Metal sampler index exceeded!"));
-				ResourceMask[Frequency].SamplerMask |= (1 << Arg.index);
-				break;
-			}
-			default:
-				check(false);
-				break;
-		}
-	}
-}
-#endif
-@end
-
 struct FMetalGraphicsPipelineKey
 {
 	FMetalRenderPipelineHash RenderPipelineHash;
@@ -270,13 +171,13 @@
 			mtlpp::PixelFormat MetalFormat = (mtlpp::PixelFormat)GPixelFormats[PF_DepthStencil].PlatformFormat;
 			DepthFormatKey = GetMetalPixelFormatKey(MetalFormat);
 		}
-			
+		
 		Key.SetHashValue(Offset_DepthFormat, NumBits_DepthFormat, DepthFormatKey);
 		Key.SetHashValue(Offset_StencilFormat, NumBits_StencilFormat, StencilFormatKey);
 
 		Key.SetHashValue(Offset_SampleCount, NumBits_SampleCount, Init.NumSamples);
 		
-#if PLATFORM_MAC		
+#if PLATFORM_MAC
 		Key.SetHashValue(Offset_PrimitiveTopology, NumBits_PrimitiveTopology, TranslatePrimitiveTopology(Init.PrimitiveType));
 #endif
 
@@ -303,6 +204,205 @@
 	}
 };
 
+static FMetalShaderPipeline* CreateMTLRenderPipeline(bool const bSync, FMetalGraphicsPipelineKey const& Key, const FGraphicsPipelineStateInitializer& Init, EMetalIndexType const IndexType);
+
+class FMetalShaderPipelineCache
+{
+public:
+	static FMetalShaderPipelineCache& Get()
+	{
+		static FMetalShaderPipelineCache sSelf;
+		return sSelf;
+	}
+	
+	FMetalShaderPipeline* GetRenderPipeline(bool const bSync, FMetalGraphicsPipelineState const* State, const FGraphicsPipelineStateInitializer& Init, EMetalIndexType const IndexType)
+	{
+		SCOPE_CYCLE_COUNTER(STAT_MetalPipelineStateTime);
+		
+		FMetalGraphicsPipelineKey Key;
+		InitMetalGraphicsPipelineKey(Key, Init, IndexType);
+		
+		// By default there'll be more threads trying to read this than to write it.
+		FRWScopeLock Lock(PipelineMutex, SLT_ReadOnly);
+		
+		// Try to find the entry in the cache.
+		FMetalShaderPipeline* Desc = Pipelines.FindRef(Key);
+		if (Desc == nil)
+		{
+			Desc = CreateMTLRenderPipeline(bSync, Key, Init, IndexType);
+			
+			// Bail cleanly if compilation fails.
+			if(!Desc)
+			{
+				return nil;
+			}
+			
+			// Now we are a writer as we want to create & add the new pipeline
+			Lock.ReleaseReadOnlyLockAndAcquireWriteLock_USE_WITH_CAUTION();
+			
+			// Retest to ensure no-one beat us here!
+			if (Pipelines.FindRef(Key) == nil)
+			{
+				Pipelines.Add(Key, Desc);
+				ReverseLookup.Add(Desc, Key);
+				
+				if (GMetalCacheShaderPipelines == 0)
+				{
+					// When we aren't caching for program lifetime we autorelease so that the PSO is released to the OS once all RHI references are released.
+					[Desc autorelease];
+				}
+			}
+		}
+		check(Desc);
+		
+		return Desc;
+	}
+	
+	void ReleaseRenderPipeline(FMetalShaderPipeline* Pipeline)
+	{
+		if (GMetalCacheShaderPipelines)
+		{
+			[Pipeline release];
+		}
+		else
+		{
+			// We take a mutex here to prevent anyone from acquiring a reference to the state which might just be about to return memory to the OS.
+			FRWScopeLock Lock(PipelineMutex, SLT_Write);
+			[Pipeline release];
+		}
+	}
+	
+	void RemoveRenderPipeline(FMetalShaderPipeline* Pipeline)
+	{
+		check (GMetalCacheShaderPipelines == 0);
+		{
+			FMetalGraphicsPipelineKey* Desc = ReverseLookup.Find(Pipeline);
+			if (Desc)
+			{
+				Pipelines.Remove(*Desc);
+				ReverseLookup.Remove(Pipeline);
+			}
+		}
+	}
+	
+	
+private:
+	FRWLock PipelineMutex;
+	TMap<FMetalGraphicsPipelineKey, FMetalShaderPipeline*> Pipelines;
+	TMap<FMetalShaderPipeline*, FMetalGraphicsPipelineKey> ReverseLookup;
+};
+
+@implementation FMetalShaderPipeline
+
+APPLE_PLATFORM_OBJECT_ALLOC_OVERRIDES(FMetalShaderPipeline)
+
+- (void)dealloc
+{
+	// For render pipeline states we might need to remove the PSO from the cache when we aren't caching them for program lifetime
+	if (GMetalCacheShaderPipelines == 0 && RenderPipelineState)
+	{
+		FMetalShaderPipelineCache::Get().RemoveRenderPipeline(self);
+	}
+	[super dealloc];
+}
+
+#if METAL_DEBUG_OPTIONS
+- (instancetype)init
+{
+	id Self = [super init];
+	if (Self)
+	{
+		RenderPipelineReflection = mtlpp::RenderPipelineReflection(nil);
+		ComputePipelineReflection = mtlpp::ComputePipelineReflection(nil);
+		RenderDesc = mtlpp::RenderPipelineDescriptor(nil);
+		ComputeDesc = mtlpp::ComputePipelineDescriptor(nil);
+	}
+	return Self;
+}
+
+- (void)initResourceMask
+{
+	if (RenderPipelineReflection)
+	{
+		[self initResourceMask:EMetalShaderVertex];
+		[self initResourceMask:EMetalShaderFragment];
+	}
+	if (ComputePipelineReflection)
+	{
+		[self initResourceMask:EMetalShaderCompute];
+	}
+}
+- (void)initResourceMask:(EMetalShaderFrequency)Frequency
+{
+	NSArray<MTLArgument*>* Arguments = nil;
+	switch(Frequency)
+	{
+		case EMetalShaderVertex:
+		{
+			MTLRenderPipelineReflection* Reflection = RenderPipelineReflection;
+			check(Reflection);
+			
+			Arguments = Reflection.vertexArguments;
+			break;
+		}
+		case EMetalShaderFragment:
+		{
+			MTLRenderPipelineReflection* Reflection = RenderPipelineReflection;
+			check(Reflection);
+			
+			Arguments = Reflection.fragmentArguments;
+			break;
+		}
+		case EMetalShaderCompute:
+		{
+			MTLComputePipelineReflection* Reflection = ComputePipelineReflection;
+			check(Reflection);
+			
+			Arguments = Reflection.arguments;
+			break;
+		}
+		default:
+			check(false);
+			break;
+	}
+	
+	for (uint32 i = 0; i < Arguments.count; i++)
+	{
+		MTLArgument* Arg = [Arguments objectAtIndex:i];
+		check(Arg);
+		switch(Arg.type)
+		{
+			case MTLArgumentTypeBuffer:
+			{
+				checkf(Arg.index < ML_MaxBuffers, TEXT("Metal buffer index exceeded!"));
+				ResourceMask[Frequency].BufferMask |= (1 << Arg.index);
+				break;
+			}
+			case MTLArgumentTypeThreadgroupMemory:
+			{
+				break;
+			}
+			case MTLArgumentTypeTexture:
+			{
+				checkf(Arg.index < ML_MaxTextures, TEXT("Metal texture index exceeded!"));
+				ResourceMask[Frequency].TextureMask |= (1 << Arg.index);
+				break;
+			}
+			case MTLArgumentTypeSampler:
+			{
+				checkf(Arg.index < ML_MaxSamplers, TEXT("Metal sampler index exceeded!"));
+				ResourceMask[Frequency].SamplerMask |= (1 << Arg.index);
+				break;
+			}
+			default:
+				check(false);
+				break;
+		}
+	}
+}
+#endif
+@end
+
 static MTLVertexDescriptor* GetMaskedVertexDescriptor(MTLVertexDescriptor* InputDesc, uint32 InOutMask)
 {
 	for (uint32 Attr = 0; Attr < MaxMetalStreams; Attr++)
@@ -930,41 +1030,12 @@
 
 static FMetalShaderPipeline* GetMTLRenderPipeline(bool const bSync, FMetalGraphicsPipelineState const* State, const FGraphicsPipelineStateInitializer& Init, EMetalIndexType const IndexType)
 {
-	static FRWLock PipelineMutex;
-	static TMap<FMetalGraphicsPipelineKey, FMetalShaderPipeline*> Pipelines;
-	
-	SCOPE_CYCLE_COUNTER(STAT_MetalPipelineStateTime);
-	
-	FMetalGraphicsPipelineKey Key;
-	InitMetalGraphicsPipelineKey(Key, Init, IndexType);
-
-	// By default there'll be more threads trying to read this than to write it.
-	FRWScopeLock Lock(PipelineMutex, SLT_ReadOnly);
-
-	// Try to find the entry in the cache.
-	FMetalShaderPipeline* Desc = Pipelines.FindRef(Key);
-	if (Desc == nil)
-	{
-		Desc = CreateMTLRenderPipeline(bSync, Key, Init, IndexType);
-		
-		// Bail cleanly if compilation fails.
-		if(!Desc)
-		{
-			return nil;
-		}
-
-		// Now we are a writer as we want to create & add the new pipeline
-		Lock.ReleaseReadOnlyLockAndAcquireWriteLock_USE_WITH_CAUTION();
-		
-		// Retest to ensure no-one beat us here!
-		if (Pipelines.FindRef(Key) == nil)
-		{
-			Pipelines.Add(Key, Desc);
-		}
-	}
-	check(Desc);
-	
-	return Desc;
+	return FMetalShaderPipelineCache::Get().GetRenderPipeline(bSync, State, Init, IndexType);
 }
 
+static void ReleaseMTLRenderPipeline(FMetalShaderPipeline* Pipeline)
+{
+	FMetalShaderPipelineCache::Get().ReleaseRenderPipeline(Pipeline);
+}
+
 bool FMetalGraphicsPipelineState::Compile()
@@ -986,7 +1057,7 @@
 {
 	for (uint32 i = 0; i < EMetalIndexType_Num; i++)
 	{
-		[PipelineStates[i] release];
+		ReleaseMTLRenderPipeline(PipelineStates[i]);
 		PipelineStates[i] = nil;
 	}
 }
