EShaderFrequency no longer depends on whether RHI_RAYTRACING is enabled. All frequencies are always defined. Some of the usage of SF_NumFrequencies is replaced with SF_NumStandardFrequencies to avoid memory and run-time overhead.

diff --git a/Engine/Source/Developer/ShaderCompilerCommon/Private/ShaderCompilerCommon.cpp b/Engine/Source/Developer/ShaderCompilerCommon/Private/ShaderCompilerCommon.cpp
--- a/Engine/Source/Developer/ShaderCompilerCommon/Private/ShaderCompilerCommon.cpp
+++ b/Engine/Source/Developer/ShaderCompilerCommon/Private/ShaderCompilerCommon.cpp
@@ -1177,11 +1177,9 @@ namespace CrossCompiler
 		TEXT("Pixel"),
 		TEXT("Geometry"),
 		TEXT("Compute"),
-#if RHI_RAYTRACING
 		TEXT("RayGen"),
 		TEXT("RayMiss"),
 		TEXT("RayHitGroup"),
-#endif
 	};
 
 	/** Compile time check to verify that the GL mapping tables are up-to-date. */
diff --git a/Engine/Source/Developer/ShaderFormatOpenGL/Private/OpenGLShaderCompiler.cpp b/Engine/Source/Developer/ShaderFormatOpenGL/Private/OpenGLShaderCompiler.cpp
--- a/Engine/Source/Developer/ShaderFormatOpenGL/Private/OpenGLShaderCompiler.cpp
+++ b/Engine/Source/Developer/ShaderFormatOpenGL/Private/OpenGLShaderCompiler.cpp
@@ -439,12 +439,10 @@ GLenum GLFrequencyTable[] =
 	GL_FRAGMENT_SHADER, // SF_Pixel
 	GL_GEOMETRY_SHADER,	// SF_Geometry
 	GL_COMPUTE_SHADER,  // SF_Compute
-#if RHI_RAYTRACING
 	// Ray tracing shaders are not supported in OpenGL
 	GLenum(0), // SF_RayGen
 	GLenum(0), // SF_RayMiss
 	GLenum(0), // SF_RayHitGroup (closest hit, any hit, intersection)
-#endif // RHI_RAYTRACING
 };
 
 static_assert(ARRAY_COUNT(GLFrequencyTable) == SF_NumFrequencies, "Frequency table size mismatch.");
diff --git a/Engine/Source/Runtime/Apple/MetalRHI/Private/MetalRenderPass.h b/Engine/Source/Runtime/Apple/MetalRHI/Private/MetalRenderPass.h
--- a/Engine/Source/Runtime/Apple/MetalRHI/Private/MetalRenderPass.h
+++ b/Engine/Source/Runtime/Apple/MetalRHI/Private/MetalRenderPass.h
@@ -168,9 +168,9 @@ private:
     
     // Which of the buffers/textures/sampler slots are bound
     // The state cache is responsible for ensuring we bind the correct 
-    FMetalTextureMask BoundTextures[SF_NumFrequencies];
-    uint32 BoundBuffers[SF_NumFrequencies];
-    uint16 BoundSamplers[SF_NumFrequencies];
+    FMetalTextureMask BoundTextures[SF_NumStandardFrequencies];
+    uint32 BoundBuffers[SF_NumStandardFrequencies];
+    uint16 BoundSamplers[SF_NumStandardFrequencies];
     
     FMetalCommandEncoder CurrentEncoder;
     FMetalCommandEncoder PrologueEncoder;
diff --git a/Engine/Source/Runtime/Apple/MetalRHI/Private/MetalStateCache.cpp b/Engine/Source/Runtime/Apple/MetalRHI/Private/MetalStateCache.cpp
--- a/Engine/Source/Runtime/Apple/MetalRHI/Private/MetalStateCache.cpp
+++ b/Engine/Source/Runtime/Apple/MetalRHI/Private/MetalStateCache.cpp
@@ -190,7 +190,7 @@ FMetalStateCache::~FMetalStateCache()
 		VertexBuffers[i].Length = 0;
 		VertexBuffers[i].Offset = 0;
 	}
-	for (uint32 Frequency = 0; Frequency < SF_NumFrequencies; Frequency++)
+	for (uint32 Frequency = 0; Frequency < SF_NumStandardFrequencies; Frequency++)
 	{
 		ShaderSamplers[Frequency].Bound = 0;
 		for (uint32 i = 0; i < ML_MaxSamplers; i++)
@@ -247,7 +247,7 @@ void FMetalStateCache::Reset(void)
 		VertexBuffers[i].Length = 0;
 		VertexBuffers[i].Offset = 0;
 	}
-	for (uint32 Frequency = 0; Frequency < SF_NumFrequencies; Frequency++)
+	for (uint32 Frequency = 0; Frequency < SF_NumStandardFrequencies; Frequency++)
 	{
 		ShaderSamplers[Frequency].Bound = 0;
 		for (uint32 i = 0; i < ML_MaxSamplers; i++)
@@ -1067,7 +1067,7 @@ void FMetalStateCache::SetGraphicsPipelineState(FMetalGraphicsPipelineState* Sta
 		bool bNewUsingTessellation = (State && State->GetPipeline(IndexType)->TessellationPipelineDesc.DomainVertexDescriptor);
 		if (bNewUsingTessellation != bUsingTessellation)
 		{
-			for (uint32 i = 0; i < SF_NumFrequencies; i++)
+			for (uint32 i = 0; i < SF_NumStandardFrequencies; i++)
 			{
 				ShaderBuffers[i].Bound = UINT32_MAX;
 #if PLATFORM_MAC
@@ -1350,7 +1350,7 @@ bool FMetalStateCache::NeedsToSetRenderTarget(const FRHISetRenderTargetsInfo& In
 
 void FMetalStateCache::SetShaderBuffer(EShaderFrequency const Frequency, FMetalBuffer const& Buffer, FMetalBufferData* const Bytes, NSUInteger const Offset, NSUInteger const Length, NSUInteger const Index, mtlpp::ResourceUsage const Usage, EPixelFormat const Format)
 {
-	check(Frequency < SF_NumFrequencies);
+	check(Frequency < SF_NumStandardFrequencies);
 	check(Index < ML_MaxBuffers);
 	
 	if (ShaderBuffers[Frequency].Buffers[Index].Buffer != Buffer ||
@@ -1381,7 +1381,7 @@ void FMetalStateCache::SetShaderBuffer(EShaderFrequency const Frequency, FMetalB
 
 void FMetalStateCache::SetShaderTexture(EShaderFrequency const Frequency, FMetalTexture const& Texture, NSUInteger const Index, mtlpp::ResourceUsage const Usage)
 {
-	check(Frequency < SF_NumFrequencies);
+	check(Frequency < SF_NumStandardFrequencies);
 	check(Index < ML_MaxTextures);
 	
 	if (ShaderTextures[Frequency].Textures[Index] != Texture
@@ -1408,7 +1408,7 @@ void FMetalStateCache::SetShaderTexture(EShaderFrequency const Frequency, FMetal
 
 void FMetalStateCache::SetShaderSamplerState(EShaderFrequency const Frequency, FMetalSamplerState* const Sampler, NSUInteger const Index)
 {
-	check(Frequency < SF_NumFrequencies);
+	check(Frequency < SF_NumStandardFrequencies);
 	check(Index < ML_MaxSamplers);
 	
 	if (ShaderSamplers[Frequency].Samplers[Index].GetPtr() != (Sampler ? Sampler->State.GetPtr() : nil))
@@ -1928,7 +1928,7 @@ void FMetalStateCache::SetStateDirty(void)
 {	
 	RasterBits = UINT32_MAX;
     PipelineBits = EMetalPipelineFlagMask;
-	for (uint32 i = 0; i < SF_NumFrequencies; i++)
+	for (uint32 i = 0; i < SF_NumStandardFrequencies; i++)
 	{
 		ShaderBuffers[i].Bound = UINT32_MAX;
 #if PLATFORM_MAC
diff --git a/Engine/Source/Runtime/Apple/MetalRHI/Private/MetalStateCache.h b/Engine/Source/Runtime/Apple/MetalRHI/Private/MetalStateCache.h
--- a/Engine/Source/Runtime/Apple/MetalRHI/Private/MetalStateCache.h
+++ b/Engine/Source/Runtime/Apple/MetalRHI/Private/MetalStateCache.h
@@ -231,18 +231,18 @@ private:
 	uint32 SampleCount;
 
 	TSet<TRefCountPtr<FRHIUniformBuffer>> ActiveUniformBuffers;
-	FRHIUniformBuffer* BoundUniformBuffers[SF_NumFrequencies][ML_MaxBuffers];
+	FRHIUniformBuffer* BoundUniformBuffers[SF_NumStandardFrequencies][ML_MaxBuffers];
 	
 	/** Bitfield for which uniform buffers are dirty */
-	uint32 DirtyUniformBuffers[SF_NumFrequencies];
+	uint32 DirtyUniformBuffers[SF_NumStandardFrequencies];
 	
 	/** Vertex attribute buffers */
 	FMetalBufferBinding VertexBuffers[MaxVertexElementCount];
 	
 	/** Bound shader resource tables. */
-	FMetalBufferBindings ShaderBuffers[SF_NumFrequencies];
-	FMetalTextureBindings ShaderTextures[SF_NumFrequencies];
-	FMetalSamplerBindings ShaderSamplers[SF_NumFrequencies];
+	FMetalBufferBindings ShaderBuffers[SF_NumStandardFrequencies];
+	FMetalTextureBindings ShaderTextures[SF_NumStandardFrequencies];
+	FMetalSamplerBindings ShaderSamplers[SF_NumStandardFrequencies];
 	
 	mtlpp::StoreAction ColorStore[MaxSimultaneousRenderTargets];
 	mtlpp::StoreAction DepthStore;
diff --git a/Engine/Source/Runtime/D3D12RHI/Private/D3D12CommandContext.h b/Engine/Source/Runtime/D3D12RHI/Private/D3D12CommandContext.h
--- a/Engine/Source/Runtime/D3D12RHI/Private/D3D12CommandContext.h
+++ b/Engine/Source/Runtime/D3D12RHI/Private/D3D12CommandContext.h
@@ -151,11 +151,11 @@ public:
 
 
 	/** Track the currently bound uniform buffers. */
-	FD3D12UniformBuffer* BoundUniformBuffers[SF_NumFrequencies][MAX_CBS];
-	FUniformBufferRHIRef BoundUniformBufferRefs[SF_NumFrequencies][MAX_CBS];
+	FD3D12UniformBuffer* BoundUniformBuffers[SF_NumStandardFrequencies][MAX_CBS];
+	FUniformBufferRHIRef BoundUniformBufferRefs[SF_NumStandardFrequencies][MAX_CBS];
 
 	/** Bit array to track which uniform buffers have changed since the last draw call. */
-	uint16 DirtyUniformBuffers[SF_NumFrequencies];
+	uint16 DirtyUniformBuffers[SF_NumStandardFrequencies];
 
 	/** Tracks the current depth stencil access type. */
 	FExclusiveDepthStencil CurrentDSVAccessType;
diff --git a/Engine/Source/Runtime/D3D12RHI/Private/D3D12PipelineState.h b/Engine/Source/Runtime/D3D12RHI/Private/D3D12PipelineState.h
--- a/Engine/Source/Runtime/D3D12RHI/Private/D3D12PipelineState.h
+++ b/Engine/Source/Runtime/D3D12RHI/Private/D3D12PipelineState.h
@@ -314,7 +314,7 @@ struct FD3D12GraphicsPipelineState : public FRHIGraphicsPipelineState
 	const FGraphicsPipelineStateInitializer PipelineStateInitializer;
 	const FD3D12RootSignature* RootSignature;
 	uint16 StreamStrides[MaxVertexElementCount];
-	bool bShaderNeedsGlobalConstantBuffer[SF_NumFrequencies];
+	bool bShaderNeedsGlobalConstantBuffer[SF_NumStandardFrequencies];
 
 	FD3D12PipelineState* PipelineState;
 
diff --git a/Engine/Source/Runtime/OpenGLDrv/Private/OpenGLShaders.cpp b/Engine/Source/Runtime/OpenGLDrv/Private/OpenGLShaders.cpp
--- a/Engine/Source/Runtime/OpenGLDrv/Private/OpenGLShaders.cpp
+++ b/Engine/Source/Runtime/OpenGLDrv/Private/OpenGLShaders.cpp
@@ -3874,12 +3874,9 @@ void FOpenGLDynamicRHI::BindPendingShaderState( FOpenGLContextState& ContextStat
 	{
 		int32 NextUniformBufferIndex = OGL_FIRST_UNIFORM_BUFFER;
 
-#if RHI_RAYTRACING
 		static_assert(SF_Compute == 5 && SF_NumFrequencies == 9, "Unexpected SF_ ordering");
 		static_assert(SF_RayGen > SF_Compute, "SF_Compute must be at the end of the list of frequencies supported in OpenGL");
-#else
-		static_assert(SF_Compute == 5 && SF_NumFrequencies == 6, "Unexpected SF_ ordering");
-#endif
+
 		int32 NumUniformBuffers[SF_Compute];
 
 		PendingState.BoundShaderState->GetNumUniformBuffers(NumUniformBuffers);
diff --git a/Engine/Source/Runtime/OpenGLDrv/Private/OpenGLState.cpp b/Engine/Source/Runtime/OpenGLDrv/Private/OpenGLState.cpp
--- a/Engine/Source/Runtime/OpenGLDrv/Private/OpenGLState.cpp
+++ b/Engine/Source/Runtime/OpenGLDrv/Private/OpenGLState.cpp
@@ -518,7 +518,7 @@ void FOpenGLRHIState::InitializeResources(int32 NumCombinedTextures, int32 NumCo
 		ShaderParameters[CrossCompiler::SHADER_STAGE_COMPUTE].InitializeResources(FOpenGL::GetMaxComputeUniformComponents() * 4 * sizeof(float));
 	}
 
-	for (int32 Frequency = 0; Frequency < SF_NumFrequencies; ++Frequency)
+	for (int32 Frequency = 0; Frequency < SF_NumStandardFrequencies; ++Frequency)
 	{
 		DirtyUniformBuffers[Frequency] = MAX_uint16;
 	}
diff --git a/Engine/Source/Runtime/OpenGLDrv/Public/OpenGLState.h b/Engine/Source/Runtime/OpenGLDrv/Public/OpenGLState.h
--- a/Engine/Source/Runtime/OpenGLDrv/Public/OpenGLState.h
+++ b/Engine/Source/Runtime/OpenGLDrv/Public/OpenGLState.h
@@ -467,11 +467,11 @@ struct FOpenGLRHIState : public FOpenGLCommonState
 	enum { MAX_UNIFORM_BUFFERS_PER_SHADER_STAGE = 14 };
 
 	/** Track the currently bound uniform buffers. */
-	FUniformBufferRHIRef BoundUniformBuffers[SF_NumFrequencies][MAX_UNIFORM_BUFFERS_PER_SHADER_STAGE];
+	FUniformBufferRHIRef BoundUniformBuffers[SF_NumStandardFrequencies][MAX_UNIFORM_BUFFERS_PER_SHADER_STAGE];
 
 	/** Bit array to track which uniform buffers have changed since the last draw call. */
 	bool bAnyDirtyGraphicsUniformBuffers;
-	uint16 DirtyUniformBuffers[SF_NumFrequencies];
+	uint16 DirtyUniformBuffers[SF_NumStandardFrequencies];
 
 	// Used for if(!FOpenGL::SupportsFastBufferData())
 	uint32 UpVertexBufferBytes;
@@ -530,7 +530,7 @@ struct FOpenGLRHIState : public FOpenGLCommonState
 		ShaderParameters = NULL;
 
 		// Release references to bound uniform buffers.
-		for (int32 Frequency = 0; Frequency < SF_NumFrequencies; ++Frequency)
+		for (int32 Frequency = 0; Frequency < SF_NumStandardFrequencies; ++Frequency)
 		{
 			for (int32 BindIndex = 0; BindIndex < MAX_UNIFORM_BUFFERS_PER_SHADER_STAGE; ++BindIndex)
 			{
diff --git a/Engine/Source/Runtime/RHI/Public/RHIDefinitions.h b/Engine/Source/Runtime/RHI/Public/RHIDefinitions.h
--- a/Engine/Source/Runtime/RHI/Public/RHIDefinitions.h
+++ b/Engine/Source/Runtime/RHI/Public/RHIDefinitions.h
@@ -28,24 +28,15 @@ enum EShaderFrequency
 	SF_Geometry			= 4,
 	SF_Compute			= 5,
 
+	// Number of standard SM5-style graphics pipeline shader frequencies
 	SF_NumStandardFrequencies = 6,
 
-#if RHI_RAYTRACING
-
 	SF_RayGen			= 6,
 	SF_RayMiss			= 7,
 	SF_RayHitGroup		= 8,
 
 	SF_NumFrequencies	= 9,
 	SF_NumBits			= 4,
-
-#else // RHI_RAYTRACING
-
-	SF_NumFrequencies	= 6,
-	SF_NumBits			= 3,
-
-#endif // RHI_RAYTRACING
-
 };
 static_assert(SF_NumFrequencies <= (1 << SF_NumBits), "SF_NumFrequencies will not fit on SF_NumBits");
 
@@ -1197,15 +1188,12 @@ inline const TCHAR* GetShaderFrequencyString(EShaderFrequency Frequency, bool bI
 	case SF_Domain:			String = TEXT("SF_Domain"); break;
 	case SF_Geometry:		String = TEXT("SF_Geometry"); break;
 	case SF_Pixel:			String = TEXT("SF_Pixel"); break;
-	case SF_Compute:		String = TEXT("SF_Compute"); break; 
-
-#if RHI_RAYTRACING
+	case SF_Compute:		String = TEXT("SF_Compute"); break;
 	case SF_RayGen:			String = TEXT("SF_RayGen"); break;
 	case SF_RayMiss:		String = TEXT("SF_RayMiss"); break;
 	case SF_RayHitGroup:	String = TEXT("SF_RayHitGroup"); break;
-#endif // RHI_RAYTRACING
 
-	default:				
+	default:
 		checkf(0, TEXT("Unknown ShaderFrequency %d"), (int32)Frequency);
 		break;
 	}
diff --git a/Engine/Source/Runtime/RenderCore/Public/ShaderCore.h b/Engine/Source/Runtime/RenderCore/Public/ShaderCore.h
--- a/Engine/Source/Runtime/RenderCore/Public/ShaderCore.h
+++ b/Engine/Source/Runtime/RenderCore/Public/ShaderCore.h
@@ -64,21 +64,15 @@ DECLARE_MEMORY_STAT_EXTERN(TEXT("Shader MapMemory"),STAT_Shaders_ShaderMapMemory
 
 inline TStatId GetMemoryStatType(EShaderFrequency ShaderFrequency)
 {
-#if RHI_RAYTRACING
 	static_assert(9 == SF_NumFrequencies, "EShaderFrequency has a bad size.");
-#else
-	static_assert(6 == SF_NumFrequencies, "EShaderFrequency has a bad size.");
-#endif
 
 	switch(ShaderFrequency)
 	{
 		case SF_Pixel:				return GET_STATID(STAT_PixelShaderMemory);
 		case SF_Compute:			return GET_STATID(STAT_PixelShaderMemory);
-#if RHI_RAYTRACING
 		case SF_RayGen:				return GET_STATID(STAT_PixelShaderMemory);
 		case SF_RayMiss:			return GET_STATID(STAT_PixelShaderMemory);
 		case SF_RayHitGroup:		return GET_STATID(STAT_PixelShaderMemory);
-#endif // RHI_RAYTRACING
 	}
 	return GET_STATID(STAT_VertexShaderMemory);
 }
diff --git a/Engine/Source/Runtime/Renderer/Private/MeshPassProcessor.cpp b/Engine/Source/Runtime/Renderer/Private/MeshPassProcessor.cpp
--- a/Engine/Source/Runtime/Renderer/Private/MeshPassProcessor.cpp
+++ b/Engine/Source/Runtime/Renderer/Private/MeshPassProcessor.cpp
@@ -324,7 +324,7 @@ public:
 	int32 PipelineId;
 	FGraphicsMinimalPipelineStateInitializer PipelineState;
 	uint32 StencilRef;
-	FShaderBindingState ShaderBindings[SF_NumFrequencies];
+	FShaderBindingState ShaderBindings[SF_NumStandardFrequencies];
 	FVertexInputStream VertexStreams[MaxVertexElementCount];
 
 	FMeshDrawCommandStateCache()
diff --git a/Engine/Source/Runtime/Windows/D3D11RHI/Private/D3D11Commands.cpp b/Engine/Source/Runtime/Windows/D3D11RHI/Private/D3D11Commands.cpp
--- a/Engine/Source/Runtime/Windows/D3D11RHI/Private/D3D11Commands.cpp
+++ b/Engine/Source/Runtime/Windows/D3D11RHI/Private/D3D11Commands.cpp
@@ -317,7 +317,7 @@ void FD3D11DynamicRHI::RHISetBoundShaderState( FBoundShaderStateRHIParamRef Boun
 	// Shader changed.  All UB's must be reset by high level code to match other platforms anway.
 	// Clear to catch those bugs, and bugs with stale UB's causing layout mismatches.
 	// Release references to bound uniform buffers.
-	for (int32 Frequency = 0; Frequency < SF_NumFrequencies; ++Frequency)
+	for (int32 Frequency = 0; Frequency < SF_NumStandardFrequencies; ++Frequency)
 	{
 		for (int32 BindIndex = 0; BindIndex < MAX_UNIFORM_BUFFERS_PER_SHADER_STAGE; ++BindIndex)
 		{
diff --git a/Engine/Source/Runtime/Windows/D3D11RHI/Private/D3D11Device.cpp b/Engine/Source/Runtime/Windows/D3D11RHI/Private/D3D11Device.cpp
--- a/Engine/Source/Runtime/Windows/D3D11RHI/Private/D3D11Device.cpp
+++ b/Engine/Source/Runtime/Windows/D3D11RHI/Private/D3D11Device.cpp
@@ -233,7 +233,7 @@ FD3D11DynamicRHI::FD3D11DynamicRHI(IDXGIFactory1* InDXGIFactory1,D3D_FEATURE_LEV
 	uint32 DynamicIBSizes[] = {128,1024,64*1024,1024*1024,0};
 	DynamicIB = new FD3D11DynamicBuffer(this,D3D11_BIND_INDEX_BUFFER,DynamicIBSizes);
 
-	for (int32 Frequency = 0; Frequency < SF_NumFrequencies; ++Frequency)
+	for (int32 Frequency = 0; Frequency < SF_NumStandardFrequencies; ++Frequency)
 	{
 		DirtyUniformBuffers[Frequency] = 0;
 	}
@@ -531,7 +531,7 @@ void FD3D11DynamicRHI::CleanupD3DDevice()
 		DynamicIB = NULL;
 
 		// Release references to bound uniform buffers.
-		for (int32 Frequency = 0; Frequency < SF_NumFrequencies; ++Frequency)
+		for (int32 Frequency = 0; Frequency < SF_NumStandardFrequencies; ++Frequency)
 		{
 			for (int32 BindIndex = 0; BindIndex < MAX_UNIFORM_BUFFERS_PER_SHADER_STAGE; ++BindIndex)
 			{
diff --git a/Engine/Source/Runtime/Windows/D3D11RHI/Private/D3D11RHI.cpp b/Engine/Source/Runtime/Windows/D3D11RHI/Private/D3D11RHI.cpp
--- a/Engine/Source/Runtime/Windows/D3D11RHI/Private/D3D11RHI.cpp
+++ b/Engine/Source/Runtime/Windows/D3D11RHI/Private/D3D11RHI.cpp
@@ -66,7 +66,7 @@ void FD3D11DynamicRHI::ClearState()
 	FMemory::Memzero(CurrentResourcesBoundAsSRVs, sizeof(CurrentResourcesBoundAsSRVs));
 	FMemory::Memzero(CurrentResourcesBoundAsVBs, sizeof(CurrentResourcesBoundAsVBs));
 	CurrentResourceBoundAsIB = nullptr;
-	for (int32 Frequency = 0; Frequency < SF_NumFrequencies; Frequency++)
+	for (int32 Frequency = 0; Frequency < SF_NumStandardFrequencies; Frequency++)
 	{
 		MaxBoundShaderResourcesIndex[Frequency] = INDEX_NONE;
 	}
diff --git a/Engine/Source/Runtime/Windows/D3D11RHI/Private/D3D11RHIPrivate.h b/Engine/Source/Runtime/Windows/D3D11RHI/Private/D3D11RHIPrivate.h
--- a/Engine/Source/Runtime/Windows/D3D11RHI/Private/D3D11RHIPrivate.h
+++ b/Engine/Source/Runtime/Windows/D3D11RHI/Private/D3D11RHIPrivate.h
@@ -767,10 +767,10 @@ protected:
 	TRefCountPtr<ID3D11UnorderedAccessView> CurrentUAVs[D3D11_PS_CS_UAV_REGISTER_COUNT];
 	TRefCountPtr<ID3D11DepthStencilView> CurrentDepthStencilTarget;
 	TRefCountPtr<FD3D11TextureBase> CurrentDepthTexture;
-	FD3D11BaseShaderResource* CurrentResourcesBoundAsSRVs[SF_NumFrequencies][D3D11_COMMONSHADER_INPUT_RESOURCE_SLOT_COUNT];
+	FD3D11BaseShaderResource* CurrentResourcesBoundAsSRVs[SF_NumStandardFrequencies][D3D11_COMMONSHADER_INPUT_RESOURCE_SLOT_COUNT];
 	FD3D11BaseShaderResource* CurrentResourcesBoundAsVBs[D3D11_IA_VERTEX_INPUT_RESOURCE_SLOT_COUNT];
 	FD3D11BaseShaderResource* CurrentResourceBoundAsIB;
-	int32 MaxBoundShaderResourcesIndex[SF_NumFrequencies];
+	int32 MaxBoundShaderResourcesIndex[SF_NumStandardFrequencies];
 	int32 MaxBoundVertexBufferIndex;
 	uint32 NumSimultaneousRenderTargets;
 	uint32 NumUAVs;
@@ -794,10 +794,10 @@ protected:
 	enum { MAX_UNIFORM_BUFFERS_PER_SHADER_STAGE = 14 };
 
 	/** Track the currently bound uniform buffers. */
-	FUniformBufferRHIRef BoundUniformBuffers[SF_NumFrequencies][MAX_UNIFORM_BUFFERS_PER_SHADER_STAGE];
+	FUniformBufferRHIRef BoundUniformBuffers[SF_NumStandardFrequencies][MAX_UNIFORM_BUFFERS_PER_SHADER_STAGE];
 
 	/** Bit array to track which uniform buffers have changed since the last draw call. */
-	uint16 DirtyUniformBuffers[SF_NumFrequencies];
+	uint16 DirtyUniformBuffers[SF_NumStandardFrequencies];
 
 	/** Tracks the current depth stencil access type. */
 	FExclusiveDepthStencil CurrentDSVAccessType;
diff --git a/Engine/Source/Runtime/Windows/D3D11RHI/Private/D3D11Shaders.cpp b/Engine/Source/Runtime/Windows/D3D11RHI/Private/D3D11Shaders.cpp
--- a/Engine/Source/Runtime/Windows/D3D11RHI/Private/D3D11Shaders.cpp
+++ b/Engine/Source/Runtime/Windows/D3D11RHI/Private/D3D11Shaders.cpp
@@ -271,7 +271,7 @@ FD3D11BoundShaderState::FD3D11BoundShaderState(
 	bShaderNeedsGlobalConstantBuffer[SF_Pixel] = InPixelShader ? InPixelShader->bShaderNeedsGlobalConstantBuffer : false;
 	bShaderNeedsGlobalConstantBuffer[SF_Geometry] = InGeometryShader ? InGeometryShader->bShaderNeedsGlobalConstantBuffer : false;
 
-	static_assert(ARRAY_COUNT(bShaderNeedsGlobalConstantBuffer) == SF_NumFrequencies, "EShaderFrequency size should match with array count of bShaderNeedsGlobalConstantBuffer.");
+	static_assert(ARRAY_COUNT(bShaderNeedsGlobalConstantBuffer) == SF_NumStandardFrequencies, "EShaderFrequency size should match with array count of bShaderNeedsGlobalConstantBuffer.");
 }
 
 FD3D11BoundShaderState::~FD3D11BoundShaderState()
diff --git a/Engine/Source/Runtime/Windows/D3D11RHI/Private/D3D11StateCache.cpp b/Engine/Source/Runtime/Windows/D3D11RHI/Private/D3D11StateCache.cpp
--- a/Engine/Source/Runtime/Windows/D3D11RHI/Private/D3D11StateCache.cpp
+++ b/Engine/Source/Runtime/Windows/D3D11RHI/Private/D3D11StateCache.cpp
@@ -286,7 +286,7 @@ void FD3D11StateCacheBase::ClearState()
 
 #if D3D11_ALLOW_STATE_CACHE
 	// Shader Resource View State Cache
-	for (uint32 ShaderFrequency = 0; ShaderFrequency < SF_NumFrequencies; ShaderFrequency++)
+	for (uint32 ShaderFrequency = 0; ShaderFrequency < SF_NumStandardFrequencies; ShaderFrequency++)
 	{
 		for (uint32 Index = 0; Index < D3D11_COMMONSHADER_INPUT_RESOURCE_SLOT_COUNT; Index++)
 		{
@@ -339,7 +339,7 @@ void FD3D11StateCacheBase::ClearState()
 	CurrentIndexOffset = 0;
 	CurrentPrimitiveTopology = D3D11_PRIMITIVE_TOPOLOGY_UNDEFINED;
 
-	for (uint32 Frequency = 0; Frequency < SF_NumFrequencies; Frequency++)
+	for (uint32 Frequency = 0; Frequency < SF_NumStandardFrequencies; Frequency++)
 	{
 		for (uint32 Index = 0; Index < D3D11_COMMONSHADER_CONSTANT_BUFFER_API_SLOT_COUNT; Index++)
 		{
diff --git a/Engine/Source/Runtime/Windows/D3D11RHI/Private/D3D11StateCachePrivate.h b/Engine/Source/Runtime/Windows/D3D11RHI/Private/D3D11StateCachePrivate.h
--- a/Engine/Source/Runtime/Windows/D3D11RHI/Private/D3D11StateCachePrivate.h
+++ b/Engine/Source/Runtime/Windows/D3D11RHI/Private/D3D11StateCachePrivate.h
@@ -71,7 +71,7 @@ protected:
 
 #if D3D11_ALLOW_STATE_CACHE
 	// Shader Resource Views Cache
-	ID3D11ShaderResourceView* CurrentShaderResourceViews[SF_NumFrequencies][D3D11_COMMONSHADER_INPUT_RESOURCE_SLOT_COUNT];
+	ID3D11ShaderResourceView* CurrentShaderResourceViews[SF_NumStandardFrequencies][D3D11_COMMONSHADER_INPUT_RESOURCE_SLOT_COUNT];
 
 	// Rasterizer State Cache
 	ID3D11RasterizerState* CurrentRasterizerState;
@@ -120,7 +120,7 @@ protected:
 	uint16 StreamStrides[MaxVertexElementCount];
 
 	// Sampler State
-	ID3D11SamplerState* CurrentSamplerStates[SF_NumFrequencies][D3D11_COMMONSHADER_SAMPLER_SLOT_COUNT];
+	ID3D11SamplerState* CurrentSamplerStates[SF_NumStandardFrequencies][D3D11_COMMONSHADER_SAMPLER_SLOT_COUNT];
 
 	// Constant Buffer State
 	struct FD3D11ConstantBufferState
@@ -128,7 +128,7 @@ protected:
 		ID3D11Buffer* Buffer;
 		uint32 FirstConstant;
 		uint32 NumConstants;
-	} CurrentConstantBuffers[SF_NumFrequencies][D3D11_COMMONSHADER_CONSTANT_BUFFER_API_SLOT_COUNT];
+	} CurrentConstantBuffers[SF_NumStandardFrequencies][D3D11_COMMONSHADER_CONSTANT_BUFFER_API_SLOT_COUNT];
 
 	bool bAlwaysSetIndexBuffers;
 
diff --git a/Engine/Source/Runtime/Windows/D3D11RHI/Public/D3D11Resources.h b/Engine/Source/Runtime/Windows/D3D11RHI/Public/D3D11Resources.h
--- a/Engine/Source/Runtime/Windows/D3D11RHI/Public/D3D11Resources.h
+++ b/Engine/Source/Runtime/Windows/D3D11RHI/Public/D3D11Resources.h
@@ -123,7 +123,7 @@ public:
 	TRefCountPtr<ID3D11DomainShader> DomainShader;
 	TRefCountPtr<ID3D11GeometryShader> GeometryShader;
 
-	bool bShaderNeedsGlobalConstantBuffer[SF_NumFrequencies];
+	bool bShaderNeedsGlobalConstantBuffer[SF_NumStandardFrequencies];
 
 
 	/** Initialization constructor. */
