Ref-count FMetalFence when not running the internal validation so that they can be reused more aggressively and so avoid running out of fences while loading or streaming large amounts of resources which generates a lot of additional encoders.

diff --git a/Engine/Source/Runtime/Apple/MetalRHI/Private/MetalCommandEncoder.cpp b/Engine/Source/Runtime/Apple/MetalRHI/Private/MetalCommandEncoder.cpp
--- a/Engine/Source/Runtime/Apple/MetalRHI/Private/MetalCommandEncoder.cpp
+++ b/Engine/Source/Runtime/Apple/MetalRHI/Private/MetalCommandEncoder.cpp
@@ -301,7 +301,7 @@ mtlpp::BlitCommandEncoder& FMetalCommandEncoder::GetBlitCommandEncoder(void)
 	return BlitCommandEncoder;
 }
 
-FMetalFence* FMetalCommandEncoder::GetEncoderFence(void) const
+TRefCountPtr<FMetalFence> const& FMetalCommandEncoder::GetEncoderFence(void) const
 {
 	return EncoderFence;
 }
@@ -493,10 +493,10 @@ void FMetalCommandEncoder::BeginBlitCommandEncoding(void)
 	EncoderFence = CommandList.GetCommandQueue().CreateFence(Label);
 }
 
-FMetalFence* FMetalCommandEncoder::EndEncoding(void)
+TRefCountPtr<FMetalFence> FMetalCommandEncoder::EndEncoding(void)
 {
 	static bool bSupportsFences = CommandList.GetCommandQueue().SupportsFeature(EMetalFeaturesFences);
-	FMetalFence* Fence = nullptr;
+	TRefCountPtr<FMetalFence> Fence = nullptr;
 	@autoreleasepool
 	{
 		if(IsRenderCommandEncoderActive())
@@ -573,7 +573,6 @@ FMetalFence* FMetalCommandEncoder::EndEncoding(void)
 				RenderCommandEncoder.EndEncoding();
 				METAL_DEBUG_LAYER(EMetalDebugLevelFastValidation, RenderEncoderDebug.EndEncoder());
 				RenderCommandEncoder = nil;
-				SafeReleaseMetalFence(EncoderFence);
 				EncoderFence = nullptr;
 			}
 			
@@ -694,7 +693,6 @@ FMetalFence* FMetalCommandEncoder::EndEncoding(void)
 			ComputeCommandEncoder.EndEncoding();
 			METAL_DEBUG_LAYER(EMetalDebugLevelFastValidation, ComputeEncoderDebug.EndEncoder());
 			ComputeCommandEncoder = nil;
-			SafeReleaseMetalFence(EncoderFence);
 			EncoderFence = nullptr;
 		}
 		else if(IsBlitCommandEncoderActive())
@@ -737,7 +735,6 @@ FMetalFence* FMetalCommandEncoder::EndEncoding(void)
 			BlitCommandEncoder.EndEncoding();
 			METAL_DEBUG_LAYER(EMetalDebugLevelFastValidation, BlitEncoderDebug.EndEncoder());
 			BlitCommandEncoder = nil;
-			SafeReleaseMetalFence(EncoderFence);
 			EncoderFence = nullptr;
 		}
 	}
diff --git a/Engine/Source/Runtime/Apple/MetalRHI/Private/MetalCommandEncoder.h b/Engine/Source/Runtime/Apple/MetalRHI/Private/MetalCommandEncoder.h
--- a/Engine/Source/Runtime/Apple/MetalRHI/Private/MetalCommandEncoder.h
+++ b/Engine/Source/Runtime/Apple/MetalRHI/Private/MetalCommandEncoder.h
@@ -131,7 +131,7 @@ public:
 	mtlpp::BlitCommandEncoder& GetBlitCommandEncoder(void);
 	
 	/** @returns The MTLFence for the current encoder or nil if there isn't one. */
-	FMetalFence* GetEncoderFence(void) const;
+	TRefCountPtr<FMetalFence> const& GetEncoderFence(void) const;
 	
 	/** @returns The number of encoded passes in the command buffer. */
 	uint32 NumEncodedPasses(void) const { return EncoderNum; }
@@ -156,7 +156,7 @@ public:
 	void BeginBlitCommandEncoding(void);
 	
 	/** Declare that all command generation from this encoder is complete, and detach from the MTLCommandBuffer if there is an encoder active or does nothing if there isn't. */
-	FMetalFence* EndEncoding(void);
+	TRefCountPtr<FMetalFence> EndEncoding(void);
 	
 	/** Initialises a fence for the current command-buffer optionally adding a command-buffer completion handler to the command-buffer */
 	void InsertCommandBufferFence(FMetalCommandBufferFence& Fence, mtlpp::CommandBufferHandler Handler);
@@ -475,7 +475,7 @@ private:
 	METAL_DEBUG_ONLY(FMetalBlitCommandEncoderDebugging BlitEncoderDebug);
 	METAL_DEBUG_ONLY(FMetalParallelRenderCommandEncoderDebugging ParallelEncoderDebug);
 	
-	FMetalFence* EncoderFence;
+	TRefCountPtr<FMetalFence> EncoderFence;
 #if ENABLE_METAL_GPUPROFILE
 	FMetalCommandBufferStats* CommandBufferStats;
 #endif
@@ -497,7 +497,7 @@ private:
 	TSet<mtlpp::Resource::Type> TransitionedResources;
 	TSet<mtlpp::Resource::Type> FenceResources;
 	
-	TSet<FMetalFence*> FragmentFences;
+	TSet<TRefCountPtr<FMetalFence>> FragmentFences;
 	
 	mtlpp::RenderStages FenceStage;
 	uint32 EncoderNum;
diff --git a/Engine/Source/Runtime/Apple/MetalRHI/Private/MetalContext.cpp b/Engine/Source/Runtime/Apple/MetalRHI/Private/MetalContext.cpp
--- a/Engine/Source/Runtime/Apple/MetalRHI/Private/MetalContext.cpp
+++ b/Engine/Source/Runtime/Apple/MetalRHI/Private/MetalContext.cpp
@@ -795,10 +795,10 @@ void FMetalDeviceContext::ReleaseBuffer(FMetalBuffer& Buffer)
 
 struct FMetalRHICommandUpdateFence final : public FRHICommand<FMetalRHICommandUpdateFence>
 {
-	FMetalFence* Fence;
+	TRefCountPtr<FMetalFence> Fence;
 	uint32 Num;
 	
-	FORCEINLINE_DEBUGGABLE FMetalRHICommandUpdateFence(FMetalFence* InFence, uint32 InNum)
+	FORCEINLINE_DEBUGGABLE FMetalRHICommandUpdateFence(TRefCountPtr<FMetalFence>& InFence, uint32 InNum)
 	: Fence(InFence)
 	, Num(InNum)
 	{
@@ -839,8 +839,8 @@ FMetalRHICommandContext* FMetalDeviceContext::AcquireContext(int32 NewIndex, int
 	EndLabel = [NSString stringWithFormat:@"End Parallel Context Index %d Num %d", NewIndex, NewNum];
 #endif
 	
-	FMetalFence* StartFence(NewIndex == 0 ? CommandList.GetCommandQueue().CreateFence(StartLabel) : ParallelFences[NewIndex - 1]);
-	FMetalFence* EndFence(CommandList.GetCommandQueue().CreateFence(EndLabel));
+	TRefCountPtr<FMetalFence> StartFence = (NewIndex == 0 ? CommandList.GetCommandQueue().CreateFence(StartLabel) : ParallelFences[NewIndex - 1].GetReference());
+	TRefCountPtr<FMetalFence> EndFence = CommandList.GetCommandQueue().CreateFence(EndLabel);
 	ParallelFences[NewIndex] = EndFence;
 	
 	// Give the context the fences so that we can properly order the parallel contexts.
@@ -1031,12 +1031,12 @@ void FMetalContext::SetParallelPassFences(FMetalFence* Start, FMetalFence* End)
 	EndFence = End;
 }
 
-FMetalFence* FMetalContext::GetParallelPassStartFence(void) const
+TRefCountPtr<FMetalFence> const& FMetalContext::GetParallelPassStartFence(void) const
 {
 	return StartFence;
 }
 
-FMetalFence* FMetalContext::GetParallelPassEndFence(void) const
+TRefCountPtr<FMetalFence> const& FMetalContext::GetParallelPassEndFence(void) const
 {
 	return EndFence;
 }
@@ -1067,7 +1067,6 @@ void FMetalContext::InitFrame(bool const bImmediateContext, uint32 Index, uint32
 	RenderPass.Begin(StartFence);
 	
 	// Unset the start fence, the render-pass owns it and we can consider it encoded now!
-	SafeReleaseMetalFence(StartFence);
 	StartFence = nullptr;
 	
 	// make sure first SetRenderTarget goes through
@@ -1080,7 +1079,6 @@ void FMetalContext::FinishFrame()
 	RenderPass.Update(EndFence);
 	
 	// Unset the end fence, the render-pass owns it and we can consider it encoded now!
-	SafeReleaseMetalFence(EndFence);
 	EndFence = nullptr;
 	
 	// End the render pass
@@ -1646,7 +1644,6 @@ void FMetalDeviceContext::SetParallelRenderPassDescriptor(FRHISetRenderTargetsIn
 	if (!RenderPass.IsWithinParallelPass())
 	{
 		RenderPass.Begin(EndFence);
-		SafeReleaseMetalFence(EndFence);
 		EndFence = nullptr;
 		StateCache.InvalidateRenderTargets();
 		SetRenderTargetsInfo(TargetInfo, false);
@@ -1670,7 +1667,6 @@ void FMetalDeviceContext::EndParallelRenderCommandEncoding(void)
 	{
 		RenderPass.EndRenderPass();
 		RenderPass.Begin(StartFence, true);
-		SafeReleaseMetalFence(StartFence);
 		StartFence = nullptr;
 		FPlatformAtomics::AtomicStore(&NumParallelContextsInPass, 0);
 	}
@@ -1723,7 +1719,7 @@ public:
 			
 			if (Index == (Num - 1))
 			{
-				FMetalFence* Fence = CmdContext->GetInternalContext().GetParallelPassEndFence();
+				TRefCountPtr<FMetalFence> Fence = CmdContext->GetInternalContext().GetParallelPassEndFence();
 				GetMetalDeviceContext().SetParallelPassFences(Fence, nil);
 			}
 
diff --git a/Engine/Source/Runtime/Apple/MetalRHI/Private/MetalContext.h b/Engine/Source/Runtime/Apple/MetalRHI/Private/MetalContext.h
--- a/Engine/Source/Runtime/Apple/MetalRHI/Private/MetalContext.h
+++ b/Engine/Source/Runtime/Apple/MetalRHI/Private/MetalContext.h
@@ -117,8 +117,8 @@ public:
 #endif
 	
 	void SetParallelPassFences(FMetalFence* Start, FMetalFence* End);
-	FMetalFence* GetParallelPassStartFence(void) const;
-	FMetalFence* GetParallelPassEndFence(void) const;
+	TRefCountPtr<FMetalFence> const& GetParallelPassStartFence(void) const;
+	TRefCountPtr<FMetalFence> const& GetParallelPassEndFence(void) const;
 	
 	void InitFrame(bool const bImmediateContext, uint32 Index, uint32 Num);
 	void FinishFrame();
@@ -150,10 +150,10 @@ protected:
 	TSharedPtr<FMetalQueryBufferPool, ESPMode::ThreadSafe> QueryBuffer;
 	
 	/** Initial fence to wait on for parallel contexts */
-	FMetalFence* StartFence;
+	TRefCountPtr<FMetalFence> StartFence;
 	
 	/** Fence to update at the end for parallel contexts */
-	FMetalFence* EndFence;
+	TRefCountPtr<FMetalFence> EndFence;
 	
 #if ENABLE_METAL_GPUPROFILE
 	/** the slot to store a per-thread context ref */
@@ -278,7 +278,7 @@ private:
 	TLockFreePointerListLIFO<FMetalRHICommandContext> ParallelContexts;
 	
 	/** Fences for parallel execution */
-	TArray<FMetalFence*> ParallelFences;
+	TArray<TRefCountPtr<FMetalFence>> ParallelFences;
 	
 	/** Critical section for FreeList */
 	FCriticalSection FreeListMutex;
diff --git a/Engine/Source/Runtime/Apple/MetalRHI/Private/MetalFence.cpp b/Engine/Source/Runtime/Apple/MetalRHI/Private/MetalFence.cpp
--- a/Engine/Source/Runtime/Apple/MetalRHI/Private/MetalFence.cpp
+++ b/Engine/Source/Runtime/Apple/MetalRHI/Private/MetalFence.cpp
@@ -96,6 +96,29 @@ APPLE_PLATFORM_OBJECT_ALLOC_OVERRIDES(FMetalDebugFence)
 }
 @end
 
+#if METAL_DEBUG_OPTIONS
+extern int32 GMetalRuntimeDebugLevel;
+#endif
+
+uint32 FMetalFence::Release() const
+{
+	uint32 Refs = uint32(--NumRefs);
+	if(Refs == 0)
+	{
+#if METAL_DEBUG_OPTIONS // When using validation we need to use fences only once per-frame in order to make it tractable
+		if (GMetalRuntimeDebugLevel >= EMetalDebugLevelValidation)
+		{
+			SafeReleaseMetalFence(const_cast<FMetalFence*>(this));
+		}
+		else // However in a final game, we need to reuse fences aggressively so that we don't run out when loading into projects
+#endif
+		{
+			FMetalFencePool::Get().ReleaseFence(const_cast<FMetalFence*>(this));
+		}
+	}
+	return Refs;
+}
+
 #if METAL_DEBUG_OPTIONS
 void FMetalFence::Validate(void) const
 {
@@ -110,11 +133,6 @@ void FMetalFence::Validate(void) const
 }
 #endif
 
-FMetalFencePool FMetalFencePool::sSelf;
-#if METAL_DEBUG_OPTIONS
-extern int32 GMetalRuntimeDebugLevel;
-#endif
-
 void FMetalFencePool::Initialise(mtlpp::Device const& InDevice)
 {
 	Device = InDevice;
@@ -149,9 +167,7 @@ void FMetalFencePool::Initialise(mtlpp::Device const& InDevice)
 		}
 	}
 	Count = FMetalFencePool::NumFences;
-#if METAL_DEBUG_OPTIONS
     Allocated = 0;
-#endif
 }
 
 FMetalFence* FMetalFencePool::AllocateFence()
@@ -161,8 +177,8 @@ FMetalFence* FMetalFencePool::AllocateFence()
 	{
         INC_DWORD_STAT(STAT_MetalFenceCount);
         FPlatformAtomics::InterlockedDecrement(&Count);
-#if METAL_DEBUG_OPTIONS
         FPlatformAtomics::InterlockedIncrement(&Allocated);
+#if METAL_DEBUG_OPTIONS
 		if (GMetalRuntimeDebugLevel >= EMetalDebugLevelValidation)
 		{
 			FScopeLock Lock(&Mutex);
@@ -196,8 +212,8 @@ void FMetalFencePool::ReleaseFence(FMetalFence* const InFence)
 	if (InFence)
 	{
         DEC_DWORD_STAT(STAT_MetalFenceCount);
-#if METAL_DEBUG_OPTIONS
         FPlatformAtomics::InterlockedDecrement(&Allocated);
+#if METAL_DEBUG_OPTIONS
 		if (GMetalRuntimeDebugLevel >= EMetalDebugLevelValidation)
 		{
 			FScopeLock Lock(&Mutex);
diff --git a/Engine/Source/Runtime/Apple/MetalRHI/Private/MetalFence.h b/Engine/Source/Runtime/Apple/MetalRHI/Private/MetalFence.h
--- a/Engine/Source/Runtime/Apple/MetalRHI/Private/MetalFence.h
+++ b/Engine/Source/Runtime/Apple/MetalRHI/Private/MetalFence.h
@@ -69,6 +69,7 @@ class FMetalFence
 	};
 public:
 	FMetalFence()
+	: NumRefs(0)
 	{
 		for (uint32 i = 0; i < NumFenceStages; i++)
 		{
@@ -78,12 +79,24 @@ public:
 	}
 	
 	explicit FMetalFence(FMetalFence const& Other)
+	: NumRefs(0)
 	{
 		operator=(Other);
 	}
 	
 	~FMetalFence()
 	{
+		check(!NumRefs);
+	}
+	
+	uint32 AddRef() const
+	{
+		return uint32(++NumRefs);
+	}
+	uint32 Release() const;
+	uint32 GetRefCount() const
+	{
+		return uint32(NumRefs);
 	}
 	
 	FMetalFence& operator=(FMetalFence const& Other)
@@ -157,6 +170,7 @@ private:
 	mtlpp::Fence Fences[NumFenceStages];
 	int8 Writes[NumFenceStages];
 	int8 Waits[NumFenceStages];
+	mutable int32 NumRefs;	
 };
 
 class FMetalFencePool
@@ -170,6 +184,7 @@ public:
 	
 	static FMetalFencePool& Get()
 	{
+		static FMetalFencePool sSelf;
 		return sSelf;
 	}
 	
@@ -178,14 +193,16 @@ public:
 	FMetalFence* AllocateFence();
 	void ReleaseFence(FMetalFence* const InFence);
 	
+	int32 Max() const { return Count; }
+	int32 Num() const { return Allocated; }
+	
 private:
 	int32 Count;
+	int32 Allocated;
 	mtlpp::Device Device;
 #if METAL_DEBUG_OPTIONS
 	TSet<FMetalFence*> Fences;
 	FCriticalSection Mutex;
-    int32 Allocated;
 #endif
-	TLockFreePointerListLIFO<FMetalFence> Lifo;
-	static FMetalFencePool sSelf;
+	TLockFreePointerListFIFO<FMetalFence, PLATFORM_CACHE_LINE_SIZE> Lifo;
 };
diff --git a/Engine/Source/Runtime/Apple/MetalRHI/Private/MetalRenderPass.cpp b/Engine/Source/Runtime/Apple/MetalRHI/Private/MetalRenderPass.cpp
--- a/Engine/Source/Runtime/Apple/MetalRHI/Private/MetalRenderPass.cpp
+++ b/Engine/Source/Runtime/Apple/MetalRHI/Private/MetalRenderPass.cpp
@@ -128,7 +128,7 @@ void FMetalRenderPass::Update(FMetalFence* Fence)
 		}
 		CurrentEncoder.UpdateFence(Fence);
 		State.FlushVisibilityResults(CurrentEncoder);
-		FMetalFence* NewFence = CurrentEncoder.EndEncoding();
+		TRefCountPtr<FMetalFence> NewFence = CurrentEncoder.EndEncoding();
 		check(!CurrentEncoderFence || !NewFence);
 		if (NewFence)
 		{
@@ -137,7 +137,7 @@ void FMetalRenderPass::Update(FMetalFence* Fence)
 	}
 }
 
-FMetalFence* FMetalRenderPass::Submit(EMetalSubmitFlags Flags)
+TRefCountPtr<FMetalFence> const& FMetalRenderPass::Submit(EMetalSubmitFlags Flags)
 {
 	if (CurrentEncoder.GetCommandBuffer() || (Flags & EMetalSubmitFlagsAsyncCommandBuffer))
 	{
@@ -817,7 +817,7 @@ void FMetalRenderPass::DispatchIndirect(FMetalVertexBuffer* ArgumentBuffer, uint
 	}
 }
 
-FMetalFence* FMetalRenderPass::EndRenderPass(void)
+TRefCountPtr<FMetalFence> const& FMetalRenderPass::EndRenderPass(void)
 {
 	if (bWithinRenderPass)
 	{
@@ -1072,7 +1072,7 @@ void FMetalRenderPass::AsyncGenerateMipmapsForTexture(FMetalTexture const& Textu
 	METAL_DEBUG_LAYER(EMetalDebugLevelFastValidation, PrologueEncoder.GetBlitCommandEncoderDebugging().GenerateMipmaps(Texture));
 }
 
-FMetalFence* FMetalRenderPass::End(void)
+TRefCountPtr<FMetalFence> const& FMetalRenderPass::End(void)
 {
 	// EndEncoding should provide the encoder fence...
 	if (PrologueEncoder.IsBlitCommandEncoderActive() || PrologueEncoder.IsComputeCommandEncoderActive())
diff --git a/Engine/Source/Runtime/Apple/MetalRHI/Private/MetalRenderPass.h b/Engine/Source/Runtime/Apple/MetalRHI/Private/MetalRenderPass.h
--- a/Engine/Source/Runtime/Apple/MetalRHI/Private/MetalRenderPass.h
+++ b/Engine/Source/Runtime/Apple/MetalRHI/Private/MetalRenderPass.h
@@ -54,7 +54,7 @@ public:
     
     void DispatchIndirect(FMetalVertexBuffer* ArgumentBufferRHI, uint32 ArgumentOffset);
     
-    FMetalFence* EndRenderPass(void);
+    TRefCountPtr<FMetalFence> const& EndRenderPass(void);
     
     void CopyFromTextureToBuffer(FMetalTexture const& Texture, uint32 sourceSlice, uint32 sourceLevel, mtlpp::Origin sourceOrigin, mtlpp::Size sourceSize, FMetalBuffer const& toBuffer, uint32 destinationOffset, uint32 destinationBytesPerRow, uint32 destinationBytesPerImage, mtlpp::BlitOption options);
     
@@ -80,9 +80,9 @@ public:
 	
 	void AsyncGenerateMipmapsForTexture(FMetalTexture const& Texture);
 	
-    FMetalFence* Submit(EMetalSubmitFlags SubmissionFlags);
+    TRefCountPtr<FMetalFence> const& Submit(EMetalSubmitFlags SubmissionFlags);
     
-    FMetalFence* End(void);
+    TRefCountPtr<FMetalFence> const& End(void);
 	
 	void InsertCommandBufferFence(FMetalCommandBufferFence& Fence, mtlpp::CommandBufferHandler Handler);
 	
@@ -179,11 +179,11 @@ private:
 	// Disjoint ranges *are* permitted!
 	TMap<id<MTLBuffer>, TArray<NSRange>> OutstandingBufferUploads;
     
-    FMetalFence* PassStartFence;
-    FMetalFence* ParallelPassEndFence;
-	FMetalFence* CurrentEncoderFence;
-    FMetalFence* PrologueEncoderFence;
-	FMetalFence* LastPrologueEncoderFence;
+    TRefCountPtr<FMetalFence> PassStartFence;
+    TRefCountPtr<FMetalFence> ParallelPassEndFence;
+	TRefCountPtr<FMetalFence> CurrentEncoderFence;
+    TRefCountPtr<FMetalFence> PrologueEncoderFence;
+	TRefCountPtr<FMetalFence> LastPrologueEncoderFence;
     
     mtlpp::RenderPassDescriptor RenderPassDesc;
     
diff --git a/Engine/Source/Runtime/Apple/MetalRHI/Private/MetalStateCache.cpp b/Engine/Source/Runtime/Apple/MetalRHI/Private/MetalStateCache.cpp
--- a/Engine/Source/Runtime/Apple/MetalRHI/Private/MetalStateCache.cpp
+++ b/Engine/Source/Runtime/Apple/MetalRHI/Private/MetalStateCache.cpp
@@ -1978,7 +1978,7 @@ void FMetalStateCache::FlushVisibilityResults(FMetalCommandEncoder& CommandEncod
 #if PLATFORM_MAC
 	if(VisibilityResults && VisibilityResults->Buffer && VisibilityResults->Buffer.GetStorageMode() == mtlpp::StorageMode::Managed && VisibilityWritten && CommandEncoder.IsRenderCommandEncoderActive())
 	{
-		FMetalFence* Fence = CommandEncoder.EndEncoding();
+		TRefCountPtr<FMetalFence> Fence = CommandEncoder.EndEncoding();
 		
         CommandEncoder.BeginBlitCommandEncoding();
 		CommandEncoder.WaitForFence(Fence);
diff --git a/Engine/Source/Runtime/Apple/MetalRHI/Private/MetalUAV.cpp b/Engine/Source/Runtime/Apple/MetalRHI/Private/MetalUAV.cpp
--- a/Engine/Source/Runtime/Apple/MetalRHI/Private/MetalUAV.cpp
+++ b/Engine/Source/Runtime/Apple/MetalRHI/Private/MetalUAV.cpp
@@ -773,10 +773,24 @@ FComputeFenceRHIRef FMetalDynamicRHI::RHICreateComputeFence(const FName& Name)
 	}
 }
 
+FMetalComputeFence::FMetalComputeFence(FName InName)
+: FRHIComputeFence(InName)
+, Fence(nullptr)
+{}
+
+FMetalComputeFence::~FMetalComputeFence()
+{
+	if (Fence)
+		Fence->Release();
+}
+
 void FMetalComputeFence::Write(FMetalFence* InFence)
 {
 	check(!Fence);
 	Fence = InFence;
+	if (Fence)
+		Fence->AddRef();
+	
 	FRHIComputeFence::WriteFence();
 }
 
@@ -787,12 +801,19 @@ void FMetalComputeFence::Wait(FMetalContext& Context)
 		Context.SubmitCommandsHint(EMetalSubmitFlagsNone);
 	}
 	Context.GetCurrentRenderPass().Begin(Fence);
+	
+	if (Fence)
+		Fence->Release();
+	
 	Fence = nullptr;
 }
 
 void FMetalComputeFence::Reset()
 {
 	FRHIComputeFence::Reset();
+	if (Fence)
+		Fence->Release();
+
 	Fence = nullptr;
 }
 
diff --git a/Engine/Source/Runtime/Apple/MetalRHI/Public/MetalResources.h b/Engine/Source/Runtime/Apple/MetalRHI/Public/MetalResources.h
--- a/Engine/Source/Runtime/Apple/MetalRHI/Public/MetalResources.h
+++ b/Engine/Source/Runtime/Apple/MetalRHI/Public/MetalResources.h
@@ -985,14 +985,9 @@ class FMetalComputeFence : public FRHIComputeFence
 {
 public:
 	
-	FMetalComputeFence(FName InName)
-	: FRHIComputeFence(InName)
-	, Fence(nullptr)
-	{}
+	FMetalComputeFence(FName InName);
 	
-	virtual ~FMetalComputeFence()
-	{
-	}
+	virtual ~FMetalComputeFence();
 	
 	virtual void Reset() final override;
 	
