FMetalFence ref-counting needs to be geniuinely atomic.

diff --git a/Engine/Source/Runtime/Apple/MetalRHI/Private/MetalFence.cpp b/Engine/Source/Runtime/Apple/MetalRHI/Private/MetalFence.cpp
--- a/Engine/Source/Runtime/Apple/MetalRHI/Private/MetalFence.cpp
+++ b/Engine/Source/Runtime/Apple/MetalRHI/Private/MetalFence.cpp
@@ -102,7 +102,7 @@ extern int32 GMetalRuntimeDebugLevel;
 
 uint32 FMetalFence::Release() const
 {
-	uint32 Refs = uint32(--NumRefs);
+	uint32 Refs = uint32(FPlatformAtomics::InterlockedDecrement(&NumRefs));
 	if(Refs == 0)
 	{
 #if METAL_DEBUG_OPTIONS // When using validation we need to use fences only once per-frame in order to make it tractable
diff --git a/Engine/Source/Runtime/Apple/MetalRHI/Private/MetalFence.h b/Engine/Source/Runtime/Apple/MetalRHI/Private/MetalFence.h
--- a/Engine/Source/Runtime/Apple/MetalRHI/Private/MetalFence.h
+++ b/Engine/Source/Runtime/Apple/MetalRHI/Private/MetalFence.h
@@ -91,12 +91,12 @@ public:
 	
 	uint32 AddRef() const
 	{
-		return uint32(++NumRefs);
+		return uint32(FPlatformAtomics::InterlockedIncrement(&NumRefs));
 	}
 	uint32 Release() const;
 	uint32 GetRefCount() const
 	{
-		return uint32(NumRefs);
+		return uint32(FPlatformAtomics::AtomicRead(&NumRefs));
 	}
 	
 	FMetalFence& operator=(FMetalFence const& Other)
