Make error reporting for Metal shader pipeline construction non-fatal for the tessellation shader too - this is *far* from ideal but so long as there are invalid entries in the shader pipeline caches it is unavoidable.

diff --git a/Engine/Source/Runtime/Apple/MetalRHI/Private/MetalPipeline.cpp b/Engine/Source/Runtime/Apple/MetalRHI/Private/MetalPipeline.cpp
--- a/Engine/Source/Runtime/Apple/MetalRHI/Private/MetalPipeline.cpp
+++ b/Engine/Source/Runtime/Apple/MetalRHI/Private/MetalPipeline.cpp
@@ -847,7 +847,7 @@ static FMetalShaderPipeline* CreateMTLRenderPipeline(bool const bSync, FMetalGra
 				UE_CLOG((Pipeline->ComputePipelineState == nil), LogMetal, Error, TEXT("Hull shader: %s"), *FString(HullShader->GetSourceCode()));
 				UE_CLOG((Pipeline->ComputePipelineState == nil), LogMetal, Error, TEXT("Domain shader: %s"), *FString(DomainShader->GetSourceCode()));
 				UE_CLOG((Pipeline->ComputePipelineState == nil), LogMetal, Error, TEXT("Descriptor: %s"), *FString(ComputePipelineDesc.GetPtr().description));
-				UE_CLOG((Pipeline->ComputePipelineState == nil), LogMetal, Fatal, TEXT("Failed to generate a hull pipeline state object:\n\n %s\n\n"), *FString(Error.GetLocalizedDescription()));
+				UE_CLOG((Pipeline->ComputePipelineState == nil), LogMetal, Error, TEXT("Failed to generate a hull pipeline state object:\n\n %s\n\n"), *FString(Error.GetLocalizedDescription()));
 				
 #if METAL_DEBUG_OPTIONS
 				if (Pipeline->ComputePipelineReflection)
