Stop warning about text shaders on iOS - because XGE shader compilation is on by default, even when XGE is not installed, and disables remote shader compilation all Windows based builds of iOS projects will always use text shaders.

diff --git a/Engine/Source/Runtime/Apple/MetalRHI/Private/MetalShaders.cpp b/Engine/Source/Runtime/Apple/MetalRHI/Private/MetalShaders.cpp
--- a/Engine/Source/Runtime/Apple/MetalRHI/Private/MetalShaders.cpp
+++ b/Engine/Source/Runtime/Apple/MetalRHI/Private/MetalShaders.cpp
@@ -212,10 +212,7 @@ void TMetalBaseShader<BaseResourceType, ShaderType>::Init(const TArray<uint8>& I
 	
 	ValidateVersion(Header.Version);
 	
-	// Validate that the compiler flags match the offline compiled flag - somehow they sometimes don't..
-	UE_CLOG((Header.CompileFlags & (1 << CFLAG_Debug)) != 0 && (OfflineCompiledFlag == 0), LogMetal, Warning, TEXT("Metal shader was meant to be compiled as bytecode but stored as text: Header: 0x%x, Offline: 0x%x"), Header.CompileFlags, OfflineCompiledFlag);
-	
-    SourceLen = Header.SourceLen;
+	SourceLen = Header.SourceLen;
     SourceCRC = Header.SourceCRC;
 	
 	// If this triggers than a level above us has failed to provide valid shader data and the cook is probably bogus
