Dupe 4869398 & 4870132:
Fix the ValidateVersion function so that it actually does what it should do even with newer shader platforms.

Fix handling of shader versions for text shaders up to Metal 2.1 - need a better way of dealing with this that adapts more easily to new versions.

diff --git a/Engine/Source/Runtime/Apple/MetalRHI/Private/MetalShaders.cpp b/Engine/Source/Runtime/Apple/MetalRHI/Private/MetalShaders.cpp
--- a/Engine/Source/Runtime/Apple/MetalRHI/Private/MetalShaders.cpp
+++ b/Engine/Source/Runtime/Apple/MetalRHI/Private/MetalShaders.cpp
@@ -145,8 +145,6 @@ mtlpp::LanguageVersion ValidateVersion(uint8 Version)
 		TEXT("Metal 2.1"),
 	};
 	
-	Version = FMath::Min(Version, (uint8)3);
-	
 	mtlpp::LanguageVersion Result = mtlpp::LanguageVersion::Version1_1;
 	switch(Version)
 	{
@@ -179,9 +177,9 @@ mtlpp::LanguageVersion ValidateVersion(uint8 Version)
 #if PLATFORM_MAC
 		Args.Add(TEXT("RequiredOS"), FText::FromString(FString::Printf(TEXT("macOS %d.%d.%d"), MetalMacOSVersions[Version][0], MetalMacOSVersions[Version][1], MetalMacOSVersions[Version][2])));
 #else
-		Args.Add(TEXT("RequiredOS"), FText::FromString(FString::Printf(TEXT("macOS %d.%d.%d"), MetaliOSVersions[Version][0], MetaliOSVersions[Version][1], MetaliOSVersions[Version][2])));
+		Args.Add(TEXT("RequiredOS"), FText::FromString(FString::Printf(TEXT("%d.%d.%d"), MetaliOSVersions[Version][0], MetaliOSVersions[Version][1], MetaliOSVersions[Version][2])));
 #endif
-		FText LocalizedMsg = FText::Format(NSLOCTEXT("MetalRHI", "ShaderVersionUnsupported", "The current OS version does not support {Version} required by the project. You must upgrade to {RequiredOS} to run this project."),Args);
+		FText LocalizedMsg = FText::Format(NSLOCTEXT("MetalRHI", "ShaderVersionUnsupported", "The current OS version does not support {ShaderVersion} required by the project. You must upgrade to {RequiredOS} to run this project."),Args);
 		
 		FText Title = NSLOCTEXT("MetalRHI", "ShaderVersionUnsupportedTitle", "Shader Version Unsupported");
 		FMessageDialog::Open(EAppMsgType::Ok, LocalizedMsg, &Title);
@@ -408,21 +406,34 @@ void TMetalBaseShader<BaseResourceType, ShaderType>::Init(const TArray<uint8>& I
 #endif
 			CompileOptions.SetPreprocessorMacros(PreprocessorMacros);
 #endif
-			if (GetMetalDeviceContext().SupportsFeature(EMetalFeaturesShaderVersions))
+			
+			mtlpp::LanguageVersion MetalVersion;
+			switch(Header.Version)
 			{
-				if (Header.Version < 3)
-				{
+				case 4:
+					MetalVersion = mtlpp::LanguageVersion::Version2_1;
+					break;
+				case 3:
+					MetalVersion = mtlpp::LanguageVersion::Version2_0;
+					break;
+				case 2:
+					MetalVersion = mtlpp::LanguageVersion::Version1_2;
+					break;
+				case 1:
+					MetalVersion = mtlpp::LanguageVersion::Version1_1;
+					break;
+				case 0:
 	#if PLATFORM_MAC
-					CompileOptions.SetLanguageVersion(Header.Version == 0 ? mtlpp::LanguageVersion::Version1_1 : (mtlpp::LanguageVersion)((1 << 16) + FMath::Min(Header.Version, (uint8)2u)));
+					MetalVersion = mtlpp::LanguageVersion::Version1_1;
 	#else
-					CompileOptions.SetLanguageVersion((mtlpp::LanguageVersion)((1 << 16) + FMath::Min(Header.Version, (uint8)2u)));
+					MetalVersion = mtlpp::LanguageVersion::Version1_0;
 	#endif
-				}
-				else if (Header.Version == 3)
-				{
-					CompileOptions.SetLanguageVersion((mtlpp::LanguageVersion)(2 << 16));
-				}
+					break;
+				default:
+					UE_LOG(LogRHI, Fatal, TEXT("Failed to create shader with unknown version %d: %s"), Header.Version, *FString(NewShaderString));
+					break;
 			}
+			CompileOptions.SetLanguageVersion(MetalVersion);
 			
 			ns::AutoReleasedError Error;
 			Library = GetMetalDeviceContext().GetDevice().NewLibrary(NewShaderString, CompileOptions, &Error);
