Dupe 4861012, 4861066 & 4861293:

Very primitive Metal implementation of RHICopyTexture - can only cope with 2D & cube textures at the moment.

Looks like RHICopyTexture should work for 2DArray and probably CubeArray, but 3D textures might not - until we actually try it on one I can't be sure. Can probably use texture-views to avoid the buffer path now that we don't have to support older iOS/macOS versions that lack that feature.

Use texture views in RHICopySubTextureRegion and RHICopyTexture to avoid a slower path that copies via a temporary buffer.

diff --git a/Engine/Source/Runtime/Apple/MetalRHI/Private/MetalTexture.cpp b/Engine/Source/Runtime/Apple/MetalRHI/Private/MetalTexture.cpp
--- a/Engine/Source/Runtime/Apple/MetalRHI/Private/MetalTexture.cpp
+++ b/Engine/Source/Runtime/Apple/MetalRHI/Private/MetalTexture.cpp
@@ -2785,10 +2785,25 @@ void FMetalDynamicRHI::RHICopySubTextureRegion(FTexture2DRHIParamRef SourceTextu
 			check(DestinationSize.width == SourceSize.width);
 			check(DestinationSize.height == SourceSize.height);
 			
+			FMetalTexture SrcTexture;
+			mtlpp::TextureUsage Usage = MetalSrcTexture->Surface.Texture.GetUsage();
+			if(Usage & mtlpp::TextureUsage::PixelFormatView)
+			{
+				ns::Range Slices(0, MetalSrcTexture->Surface.Texture.GetArrayLength() * (MetalSrcTexture->Surface.bIsCubemap ? 6 : 1));
+				if(MetalSrcTexture->Surface.Texture.GetPixelFormat() != MetalDestTexture->Surface.Texture.GetPixelFormat())
+				{
+					SrcTexture = MetalSrcTexture->Surface.Texture.NewTextureView(MetalDestTexture->Surface.Texture.GetPixelFormat(), MetalSrcTexture->Surface.Texture.GetTextureType(), ns::Range(0, MetalSrcTexture->Surface.Texture.GetMipmapLevelCount()), Slices);
+				}
+			}
+			if (!SrcTexture)
+			{
+				SrcTexture = MetalSrcTexture->Surface.Texture;
+			}
+			
 			// Account for create with TexCreate_SRGB flag which could make these different
-			if(MetalSrcTexture->Surface.Texture.GetPixelFormat() == MetalDestTexture->Surface.Texture.GetPixelFormat())
+			if(SrcTexture.GetPixelFormat() == MetalDestTexture->Surface.Texture.GetPixelFormat())
 			{
-				ImmediateContext.GetInternalContext().CopyFromTextureToTexture(MetalSrcTexture->Surface.Texture, 0, 0, SourceOrigin,SourceSize,MetalDestTexture->Surface.Texture, 0, 0, DestinationOrigin);
+				ImmediateContext.GetInternalContext().CopyFromTextureToTexture(SrcTexture, 0, 0, SourceOrigin,SourceSize,MetalDestTexture->Surface.Texture, 0, 0, DestinationOrigin);
 			}
 			else
 			{
@@ -2816,6 +2831,11 @@ void FMetalDynamicRHI::RHICopySubTextureRegion(FTexture2DRHIParamRef SourceTextu
 				
 				GetMetalDeviceContext().ReleaseBuffer(Buffer);
 			}
+			
+			if (SrcTexture != MetalSrcTexture->Surface.Texture)
+			{
+				SafeReleaseMetalTexture(SrcTexture);
+			}
 		}
 		else
 		{
@@ -2824,3 +2844,102 @@ void FMetalDynamicRHI::RHICopySubTextureRegion(FTexture2DRHIParamRef SourceTextu
 	}
 }
 
+
+
+void FMetalRHICommandContext::RHICopyTexture(FTextureRHIParamRef SourceTextureRHI, FTextureRHIParamRef DestTextureRHI, const FRHICopyTextureInfo& CopyInfo)
+{
+	if (!SourceTextureRHI || !DestTextureRHI || SourceTextureRHI == DestTextureRHI)
+	{
+		// no need to do anything (silently ignored)
+		return;
+	}
+	
+	RHITransitionResources(EResourceTransitionAccess::EReadable, &SourceTextureRHI, 1);
+	
+	@autoreleasepool {
+		check(SourceTextureRHI);
+		check(DestTextureRHI);
+		
+		if(SourceTextureRHI->GetFormat() == DestTextureRHI->GetFormat())
+		{
+			FMetalSurface* MetalSrcTexture = GetMetalSurfaceFromRHITexture(SourceTextureRHI);
+			FMetalSurface* MetalDestTexture = GetMetalSurfaceFromRHITexture(DestTextureRHI);
+			
+			FIntVector Size = (CopyInfo.Size != FIntVector::ZeroValue) ? CopyInfo.Size : FIntVector(MetalSrcTexture->Texture.GetWidth(), MetalSrcTexture->Texture.GetHeight(), MetalSrcTexture->Texture.GetDepth());
+			
+			mtlpp::Origin SourceOrigin(CopyInfo.SourcePosition.X, CopyInfo.SourcePosition.Y, CopyInfo.SourcePosition.Z);
+			mtlpp::Origin DestinationOrigin(CopyInfo.DestPosition.X, CopyInfo.DestPosition.Y, CopyInfo.DestPosition.Z);
+
+			FMetalTexture SrcTexture;
+			mtlpp::TextureUsage Usage = MetalSrcTexture->Texture.GetUsage();
+			if(Usage & mtlpp::TextureUsage::PixelFormatView)
+			{
+				ns::Range Slices(0, MetalSrcTexture->Texture.GetArrayLength() * (MetalSrcTexture->bIsCubemap ? 6 : 1));
+				if(MetalSrcTexture->Texture.GetPixelFormat() != MetalDestTexture->Texture.GetPixelFormat())
+				{
+					SrcTexture = MetalSrcTexture->Texture.NewTextureView(MetalDestTexture->Texture.GetPixelFormat(), MetalSrcTexture->Texture.GetTextureType(), ns::Range(0, MetalSrcTexture->Texture.GetMipmapLevelCount()), Slices);
+				}
+			}
+			if (!SrcTexture)
+			{
+				SrcTexture = MetalSrcTexture->Texture;
+			}
+			
+			for (uint32 SliceIndex = 0; SliceIndex < CopyInfo.NumSlices; ++SliceIndex)
+			{
+				uint32 SourceSliceIndex = CopyInfo.SourceSliceIndex + SliceIndex;
+				uint32 DestSliceIndex = CopyInfo.DestSliceIndex + SliceIndex;
+
+				for (uint32 MipIndex = 0; MipIndex < CopyInfo.NumMips; ++MipIndex)
+				{
+					uint32 SourceMipIndex = CopyInfo.SourceMipIndex + MipIndex;
+					uint32 DestMipIndex = CopyInfo.DestMipIndex + MipIndex;
+					mtlpp::Size SourceSize(FMath::Max(CopyInfo.Size.X >> MipIndex, 1), FMath::Max(CopyInfo.Size.Y >> MipIndex, 1), FMath::Max(CopyInfo.Size.Z >> MipIndex, 1));
+					
+					// Account for create with TexCreate_SRGB flag which could make these different
+					if(SrcTexture.GetPixelFormat() == MetalDestTexture->Texture.GetPixelFormat())
+					{
+						GetInternalContext().CopyFromTextureToTexture(SrcTexture, SourceSliceIndex, SourceMipIndex, SourceOrigin,SourceSize,MetalDestTexture->Texture, DestSliceIndex, DestMipIndex, DestinationOrigin);
+					}
+					else
+					{
+						// Linear and sRGB mismatch then try to go via metal buffer
+						// Modified clone of logic from MetalRenderTarget.cpp
+						uint32 BytesPerPixel = (MetalSrcTexture->PixelFormat != PF_DepthStencil) ? GPixelFormats[MetalSrcTexture->PixelFormat].BlockBytes : 1;
+						const uint32 Stride = BytesPerPixel * SourceSize.width;
+						const uint32 Alignment = PLATFORM_MAC ? 1u : 64u;
+						const uint32 AlignedStride = ((Stride - 1) & ~(Alignment - 1)) + Alignment;
+						const uint32 BytesPerImage = AlignedStride *  SourceSize.height;
+						const uint32 DataSize = BytesPerImage * SourceSize.depth;
+						
+						FMetalBuffer Buffer = GetMetalDeviceContext().CreatePooledBuffer(FMetalPooledBufferArgs(GetInternalContext().GetDevice(), DataSize, mtlpp::StorageMode::Shared));
+						
+						check(Buffer);
+						
+						mtlpp::BlitOption Options = mtlpp::BlitOption::None;
+#if !PLATFORM_MAC
+						if (MetalSrcTexture->Surface.Texture.GetPixelFormat() >= mtlpp::PixelFormat::PVRTC_RGB_2BPP && MetalSrcTexture->Surface.Texture.GetPixelFormat() <= mtlpp::PixelFormat::PVRTC_RGBA_4BPP_sRGB)
+						{
+							Options = mtlpp::BlitOption::RowLinearPVRTC;
+						}
+#endif
+						GetInternalContext().CopyFromTextureToBuffer(MetalSrcTexture->Texture, SourceSliceIndex, SourceMipIndex, SourceOrigin, SourceSize, Buffer, 0, AlignedStride, BytesPerImage, Options);
+						GetInternalContext().CopyFromBufferToTexture(Buffer, 0, Stride, BytesPerImage, SourceSize, MetalDestTexture->Texture, DestSliceIndex, DestMipIndex, DestinationOrigin, Options);
+						
+						GetMetalDeviceContext().ReleaseBuffer(Buffer);
+					}
+				}
+			}
+			
+			if (SrcTexture != MetalSrcTexture->Texture)
+			{
+				SafeReleaseMetalTexture(SrcTexture);
+			}
+		}
+		else
+		{
+			UE_LOG(LogMetal, Warning, TEXT("RHICopyTexture Source <-> Destination texture format mismatch"));
+		}
+	}
+}
+
diff --git a/Engine/Source/Runtime/Apple/MetalRHI/Public/MetalRHIContext.h b/Engine/Source/Runtime/Apple/MetalRHI/Public/MetalRHIContext.h
--- a/Engine/Source/Runtime/Apple/MetalRHI/Public/MetalRHIContext.h
+++ b/Engine/Source/Runtime/Apple/MetalRHI/Public/MetalRHIContext.h
@@ -47,6 +47,8 @@ public:
 	/** Clears a UAV to the multi-component value provided. */
 	virtual void RHIClearTinyUAV(FUnorderedAccessViewRHIParamRef UnorderedAccessViewRHI, const uint32* Values) final override;
 	
+	virtual void RHICopyTexture(FTextureRHIParamRef SourceTextureRHI, FTextureRHIParamRef DestTextureRHI, const FRHICopyTextureInfo& CopyInfo) final override;
+	
 	/**
 	 * Resolves from one texture to another.
 	 * @param SourceTexture - texture to resolve from, 0 is silenty ignored
