Expose support for CopyTexture in MetalRHI.

diff --git a/Engine/Source/Runtime/Apple/MetalRHI/Private/MetalRHI.cpp b/Engine/Source/Runtime/Apple/MetalRHI/Private/MetalRHI.cpp
--- a/Engine/Source/Runtime/Apple/MetalRHI/Private/MetalRHI.cpp
+++ b/Engine/Source/Runtime/Apple/MetalRHI/Private/MetalRHI.cpp
@@ -366,6 +366,8 @@ FMetalDynamicRHI::FMetalDynamicRHI(ERHIFeatureLevel::Type RequestedFeatureLevel)
 	{
 		GMetalCommandBufferHasStartEndTimeAPI = true;
 	}
+	
+	GRHISupportsCopyToTextureMultipleMips = true;
 		
 	if(
 	   #if PLATFORM_MAC
