Missed duplicating CL #4872592 which keeps clang happy.

diff --git a/Engine/Source/Runtime/Apple/MetalRHI/Private/MetalShaders.cpp b/Engine/Source/Runtime/Apple/MetalRHI/Private/MetalShaders.cpp
--- a/Engine/Source/Runtime/Apple/MetalRHI/Private/MetalShaders.cpp
+++ b/Engine/Source/Runtime/Apple/MetalRHI/Private/MetalShaders.cpp
@@ -431,6 +431,11 @@ void TMetalBaseShader<BaseResourceType, ShaderType>::Init(const TArray<uint8>& I
 					break;
 				default:
 					UE_LOG(LogRHI, Fatal, TEXT("Failed to create shader with unknown version %d: %s"), Header.Version, *FString(NewShaderString));
+	#if PLATFORM_MAC
+					MetalVersion = mtlpp::LanguageVersion::Version1_1;
+	#else
+					MetalVersion = mtlpp::LanguageVersion::Version1_0;
+	#endif
 					break;
 			}
 			CompileOptions.SetLanguageVersion(MetalVersion);
