Partially duplicate 4861533:
iOS compile fixes.

diff --git a/Engine/Source/Runtime/Apple/MetalRHI/Private/MetalTexture.cpp b/Engine/Source/Runtime/Apple/MetalRHI/Private/MetalTexture.cpp
--- a/Engine/Source/Runtime/Apple/MetalRHI/Private/MetalTexture.cpp
+++ b/Engine/Source/Runtime/Apple/MetalRHI/Private/MetalTexture.cpp
@@ -2918,7 +2918,7 @@ void FMetalRHICommandContext::RHICopyTexture(FTextureRHIParamRef SourceTextureRH
 						
 						mtlpp::BlitOption Options = mtlpp::BlitOption::None;
 #if !PLATFORM_MAC
-						if (MetalSrcTexture->Surface.Texture.GetPixelFormat() >= mtlpp::PixelFormat::PVRTC_RGB_2BPP && MetalSrcTexture->Surface.Texture.GetPixelFormat() <= mtlpp::PixelFormat::PVRTC_RGBA_4BPP_sRGB)
+						if (MetalSrcTexture->Texture.GetPixelFormat() >= mtlpp::PixelFormat::PVRTC_RGB_2BPP && MetalSrcTexture->Texture.GetPixelFormat() <= mtlpp::PixelFormat::PVRTC_RGBA_4BPP_sRGB)
 						{
 							Options = mtlpp::BlitOption::RowLinearPVRTC;
 						}
