Requests from Nvidia and Intel to aid Metal performance:
- Swap back to Managed uniform buffers to improve performance on discrete GPUs.
- Disable the parallel command-buffer encoding on Intel & Nvidia as the small command-buffers hurt more than they help - can be forced on with -metalparallel command-line option.
- Allow forcing on async. compute with -metalasynccompute so NV can A/B test whether it helps.

diff --git a/Engine/Source/Runtime/Apple/MetalRHI/Private/MetalCommandQueue.cpp b/Engine/Source/Runtime/Apple/MetalRHI/Private/MetalCommandQueue.cpp
--- a/Engine/Source/Runtime/Apple/MetalRHI/Private/MetalCommandQueue.cpp
+++ b/Engine/Source/Runtime/Apple/MetalRHI/Private/MetalCommandQueue.cpp
@@ -19,7 +19,6 @@
 #pragma mark - Private C++ Statics -
 uint64 FMetalCommandQueue::Features = 0;
 extern mtlpp::VertexFormat GMetalFColorVertexFormat;
-extern bool GMetalManagedUniformBuffers;
 
 #pragma mark - Public C++ Boilerplate -
 
@@ -294,8 +293,6 @@ FMetalCommandQueue::FMetalCommandQueue(mtlpp::Device InDevice, uint32 const MaxN
 	{
 		Features |= EMetalFeaturesAbsoluteTimeQueries;
 	}
-	
-	GMetalManagedUniformBuffers = FParse::Param(FCommandLine::Get(),TEXT("metalmanagedubs"));
 #endif
 	
 #if !UE_BUILD_SHIPPING
diff --git a/Engine/Source/Runtime/Apple/MetalRHI/Private/MetalRHI.cpp b/Engine/Source/Runtime/Apple/MetalRHI/Private/MetalRHI.cpp
--- a/Engine/Source/Runtime/Apple/MetalRHI/Private/MetalRHI.cpp
+++ b/Engine/Source/Runtime/Apple/MetalRHI/Private/MetalRHI.cpp
@@ -415,9 +415,9 @@ FMetalDynamicRHI::FMetalDynamicRHI(ERHIFeatureLevel::Type RequestedFeatureLevel)
 #else
 		GRHISupportsRHIThread = bSupportsRHIThread;
 #endif
-		GRHISupportsParallelRHIExecute = GRHISupportsRHIThread;
+		GRHISupportsParallelRHIExecute = GRHISupportsRHIThread && ((!IsRHIDeviceIntel() && !IsRHIDeviceNVIDIA()) || FParse::Param(FCommandLine::Get(),TEXT("metalparallel")));
 #endif
-		GSupportsEfficientAsyncCompute = GRHISupportsParallelRHIExecute && (IsRHIDeviceAMD() || PLATFORM_IOS); // Only AMD currently support async. compute and it requires parallel execution to be useful.
+		GSupportsEfficientAsyncCompute = GRHISupportsParallelRHIExecute && (IsRHIDeviceAMD() || PLATFORM_IOS || FParse::Param(FCommandLine::Get(),TEXT("metalasynccompute"))); // Only AMD currently support async. compute and it requires parallel execution to be useful.
 		GSupportsParallelOcclusionQueries = GRHISupportsRHIThread;
 	}
 	else
diff --git a/Engine/Source/Runtime/Apple/MetalRHI/Private/MetalUniformBuffer.cpp b/Engine/Source/Runtime/Apple/MetalRHI/Private/MetalUniformBuffer.cpp
--- a/Engine/Source/Runtime/Apple/MetalRHI/Private/MetalUniformBuffer.cpp
+++ b/Engine/Source/Runtime/Apple/MetalRHI/Private/MetalUniformBuffer.cpp
@@ -29,7 +29,6 @@ struct FMetalRHICommandInitialiseUniformBufferIAB : public FRHICommand<FMetalRHI
 		Buffer->InitIAB();
 	}
 };
-bool GMetalManagedUniformBuffers = false;
 
 struct FMetalArgumentDesc
 {
