Fix silly mistake that led to color channel inversion on iOS 11/12 when using Metal 1.2.

diff --git a/Engine/Source/Runtime/Apple/MetalRHI/Private/MetalCommandQueue.cpp b/Engine/Source/Runtime/Apple/MetalRHI/Private/MetalCommandQueue.cpp
--- a/Engine/Source/Runtime/Apple/MetalRHI/Private/MetalCommandQueue.cpp
+++ b/Engine/Source/Runtime/Apple/MetalRHI/Private/MetalCommandQueue.cpp
@@ -85,7 +85,10 @@ FMetalCommandQueue::FMetalCommandQueue(mtlpp::Device InDevice, uint32 const MaxN
 			{
 				Features |= EMetalFeaturesGPUCaptureManager | EMetalFeaturesBufferSubAllocation | EMetalFeaturesParallelRenderEncoders | EMetalFeaturesPipelineBufferMutability;
 				
-				GMetalFColorVertexFormat = mtlpp::VertexFormat::UChar4Normalized_BGRA;
+				if (MaxShaderVersion >= 3)
+				{
+					GMetalFColorVertexFormat = mtlpp::VertexFormat::UChar4Normalized_BGRA;
+				}
 				
 				if (MaxShaderVersion >= 3)
 				{
@@ -147,7 +150,10 @@ FMetalCommandQueue::FMetalCommandQueue(mtlpp::Device InDevice, uint32 const MaxN
 			
 			if(Vers.majorVersion >= 11)
 			{
-				GMetalFColorVertexFormat = mtlpp::VertexFormat::UChar4Normalized_BGRA;
+				if (MaxShaderVersion >= 3)
+				{
+					GMetalFColorVertexFormat = mtlpp::VertexFormat::UChar4Normalized_BGRA;
+				}
 				
 				Features |= EMetalFeaturesPresentMinDuration | EMetalFeaturesGPUCaptureManager | EMetalFeaturesBufferSubAllocation | EMetalFeaturesParallelRenderEncoders | EMetalFeaturesPipelineBufferMutability;
 				
