Use dynamic casting to access Metal tessellation shader index buffers - means restoring the code path for R16Uint and R32Uint but that should be fine. This avoids the AMD driver bug that causes incorrect reads at the extreme beginning/end of a texture_buffer or texture2d created from an index buffer. Need to get to the bottom of this with the driver team.

diff --git a/Engine/Shaders/Public/Platform/Metal/MetalCommon.ush b/Engine/Shaders/Public/Platform/Metal/MetalCommon.ush
--- a/Engine/Shaders/Public/Platform/Metal/MetalCommon.ush
+++ b/Engine/Shaders/Public/Platform/Metal/MetalCommon.ush
@@ -5,7 +5,7 @@
 =============================================================================*/
 
 // Update this GUID to improve shader recompilation for Metal only shaders
-// GUID = 0FDEF9A8-0978-44ED-B2ED-046924F18E73
+// GUID = 4E19BD71-3C18-4138-B085-451EAC550287
   
 #pragma once
 
diff --git a/Engine/Source/Developer/Apple/MetalShaderFormat/Private/MetalBackend.cpp b/Engine/Source/Developer/Apple/MetalShaderFormat/Private/MetalBackend.cpp
--- a/Engine/Source/Developer/Apple/MetalShaderFormat/Private/MetalBackend.cpp
+++ b/Engine/Source/Developer/Apple/MetalShaderFormat/Private/MetalBackend.cpp
@@ -1454,7 +1454,7 @@ protected:
 		print_type_full(sig->return_type);
 		ralloc_asprintf_append(buffer, " %s(", sig->function_name());
 		
-        bInsertSideTable = false;
+        bInsertSideTable = Backend.bIsTessellationVSHS;
         if (sig->is_main && Backend.bBoundsChecks)
 		{
             foreach_iter(exec_list_iterator, iter, sig->parameters)
@@ -1492,7 +1492,7 @@ protected:
             
             ir_variable* indexBuffer = new(ParseState)ir_variable(glsl_type::uint_type, "indexBuffer", ir_var_in);
             indexBuffer->semantic = "";
-            Buffers.Textures.Add(indexBuffer);
+            Buffers.Buffers.Add(indexBuffer);
             
             int32 IndexBufferIndex = Buffers.GetIndex(indexBuffer);
             check(IndexBufferIndex >= 0 && IndexBufferIndex < 30);
@@ -1504,7 +1504,7 @@ protected:
 				"uint2 threadgroup_position_in_grid [[threadgroup_position_in_grid]],\n"
 				"constant uint *patchCount [[ buffer(%d) ]],\n"
 				"#define METAL_INDEX_BUFFER_ID %d\n"
-				"typedBuffer1_read(uint, indexBuffer, METAL_INDEX_BUFFER_ID)",
+				"const device typed_buffer<uint>* indexBuffer [[ buffer(METAL_INDEX_BUFFER_ID) ]]",
                 patchIndex, IndexBufferIndex
 			);
 			bPrintComma = true;
@@ -1645,8 +1645,8 @@ protected:
 						ir_variable* indexBuffer = ParseState->symbols->get_variable("indexBuffer");
 						int32 IndexBufferIndex = Buffers.GetIndex(indexBuffer);
 						ralloc_asprintf_append(buffer, "#define GET_VERTEX_ID() \\\n");
-						ralloc_asprintf_append(buffer, "	(is_null_texture(indexBuffer)) ? thread_position_in_grid.x : \\\n");
-						ralloc_asprintf_append(buffer, "	buffer::load<uint, METAL_INDEX_BUFFER_ID>(indexBuffer, thread_position_in_grid.x + patchCount[1])\n");
+						ralloc_asprintf_append(buffer, "	(BufferSizes[(METAL_INDEX_BUFFER_ID*2)+1] == 0) ? thread_position_in_grid.x : \\\n");
+						ralloc_asprintf_append(buffer, "	buffer::load<uint, METAL_INDEX_BUFFER_ID>(indexBuffer, thread_position_in_grid.x, BufferSizes)\n");
 						ralloc_asprintf_append(buffer, "/* optionally vertex_id = GET_VERTEX_ID() + grid_origin.x */\n");
 					}
 
diff --git a/Engine/Source/Runtime/Apple/MetalRHI/Private/MetalCommands.cpp b/Engine/Source/Runtime/Apple/MetalRHI/Private/MetalCommands.cpp
--- a/Engine/Source/Runtime/Apple/MetalRHI/Private/MetalCommands.cpp
+++ b/Engine/Source/Runtime/Apple/MetalRHI/Private/MetalCommands.cpp
@@ -735,9 +735,7 @@ void FMetalRHICommandContext::RHIDrawIndexedPrimitive(FIndexBufferRHIParamRef In
 	RHI_DRAW_CALL_STATS(PrimitiveType,FMath::Max(NumInstances,1u)*NumPrimitives);
 
 	FMetalIndexBuffer* IndexBuffer = ResourceCast(IndexBufferRHI);
-	
-	EPixelFormat Format = IndexBuffer->IndexType == mtlpp::IndexType::UInt16 ? PF_R16_UINT : PF_R32_UINT;
-	Context->DrawIndexedPrimitive(IndexBuffer->Buffer, IndexBuffer->GetLinearTexture(Format), IndexBuffer->GetStride(), IndexBuffer->IndexType, PrimitiveType, BaseVertexIndex, FirstInstance, NumVertices, StartIndex, NumPrimitives, NumInstances);
+	Context->DrawIndexedPrimitive(IndexBuffer->Buffer, IndexBuffer->GetStride(), IndexBuffer->IndexType, PrimitiveType, BaseVertexIndex, FirstInstance, NumVertices, StartIndex, NumPrimitives, NumInstances);
 	}
 }
 
@@ -879,7 +877,7 @@ void FMetalRHICommandContext::RHIEndDrawIndexedPrimitiveUP()
 	// how many to draw
 	uint32 NumIndices = GetVertexCountForPrimitiveCount(PendingNumPrimitives, PendingPrimitiveType);
 	
-	Context->DrawIndexedPrimitive(PendingIndexBuffer, FMetalTexture(), PendingIndexDataStride, (PendingIndexDataStride == 2) ? mtlpp::IndexType::UInt16 : mtlpp::IndexType::UInt32, PendingPrimitiveType, 0, 0, NumIndices, 0, PendingNumPrimitives, 1);
+	Context->DrawIndexedPrimitive(PendingIndexBuffer, PendingIndexDataStride, (PendingIndexDataStride == 2) ? mtlpp::IndexType::UInt16 : mtlpp::IndexType::UInt32, PendingPrimitiveType, 0, 0, NumIndices, 0, PendingNumPrimitives, 1);
 	
 	// mark temp memory as usable
 	PendingVertexBuffer = nil;
diff --git a/Engine/Source/Runtime/Apple/MetalRHI/Private/MetalContext.cpp b/Engine/Source/Runtime/Apple/MetalRHI/Private/MetalContext.cpp
--- a/Engine/Source/Runtime/Apple/MetalRHI/Private/MetalContext.cpp
+++ b/Engine/Source/Runtime/Apple/MetalRHI/Private/MetalContext.cpp
@@ -1480,7 +1480,7 @@ void FMetalContext::DrawPrimitiveIndirect(uint32 PrimitiveType, FMetalVertexBuff
 	RenderPass.DrawPrimitiveIndirect(PrimitiveType, VertexBuffer, ArgumentOffset);
 }
 
-void FMetalContext::DrawIndexedPrimitive(FMetalBuffer const& IndexBuffer, FMetalTexture const& IndexTexture, uint32 IndexStride, mtlpp::IndexType IndexType, uint32 PrimitiveType, int32 BaseVertexIndex, uint32 FirstInstance, uint32 NumVertices, uint32 StartIndex, uint32 NumPrimitives, uint32 NumInstances)
+void FMetalContext::DrawIndexedPrimitive(FMetalBuffer const& IndexBuffer, uint32 IndexStride, mtlpp::IndexType IndexType, uint32 PrimitiveType, int32 BaseVertexIndex, uint32 FirstInstance, uint32 NumVertices, uint32 StartIndex, uint32 NumPrimitives, uint32 NumInstances)
 {
 	// finalize any pending state
 	if(!PrepareToDraw(PrimitiveType, GetRHIMetalIndexType(IndexType)))
@@ -1488,7 +1488,7 @@ void FMetalContext::DrawIndexedPrimitive(FMetalBuffer const& IndexBuffer, FMetal
 		return;
 	}
 	
-	RenderPass.DrawIndexedPrimitive(IndexBuffer, IndexTexture, IndexStride, PrimitiveType, BaseVertexIndex, FirstInstance, NumVertices, StartIndex, NumPrimitives, NumInstances);
+	RenderPass.DrawIndexedPrimitive(IndexBuffer, IndexStride, PrimitiveType, BaseVertexIndex, FirstInstance, NumVertices, StartIndex, NumPrimitives, NumInstances);
 }
 
 void FMetalContext::DrawIndexedIndirect(FMetalIndexBuffer* IndexBuffer, uint32 PrimitiveType, FMetalStructuredBuffer* VertexBuffer, int32 DrawArgumentsIndex, uint32 NumInstances)
diff --git a/Engine/Source/Runtime/Apple/MetalRHI/Private/MetalContext.h b/Engine/Source/Runtime/Apple/MetalRHI/Private/MetalContext.h
--- a/Engine/Source/Runtime/Apple/MetalRHI/Private/MetalContext.h
+++ b/Engine/Source/Runtime/Apple/MetalRHI/Private/MetalContext.h
@@ -71,7 +71,7 @@ public:
 	
 	void DrawPrimitiveIndirect(uint32 PrimitiveType, FMetalVertexBuffer* VertexBuffer, uint32 ArgumentOffset);
 	
-	void DrawIndexedPrimitive(FMetalBuffer const& IndexBuffer, FMetalTexture const& IndexTexture, uint32 IndexStride, mtlpp::IndexType IndexType, uint32 PrimitiveType, int32 BaseVertexIndex, uint32 FirstInstance,
+	void DrawIndexedPrimitive(FMetalBuffer const& IndexBuffer, uint32 IndexStride, mtlpp::IndexType IndexType, uint32 PrimitiveType, int32 BaseVertexIndex, uint32 FirstInstance,
 							  uint32 NumVertices, uint32 StartIndex, uint32 NumPrimitives, uint32 NumInstances);
 	
 	void DrawIndexedIndirect(FMetalIndexBuffer* IndexBufferRHI, uint32 PrimitiveType, FMetalStructuredBuffer* VertexBufferRHI, int32 DrawArgumentsIndex, uint32 NumInstances);
diff --git a/Engine/Source/Runtime/Apple/MetalRHI/Private/MetalRenderPass.cpp b/Engine/Source/Runtime/Apple/MetalRHI/Private/MetalRenderPass.cpp
--- a/Engine/Source/Runtime/Apple/MetalRHI/Private/MetalRenderPass.cpp
+++ b/Engine/Source/Runtime/Apple/MetalRHI/Private/MetalRenderPass.cpp
@@ -407,7 +407,7 @@ void FMetalRenderPass::DrawPrimitive(uint32 PrimitiveType, uint32 BaseVertexInde
 	}
 	else
 	{
-		DrawPatches(PrimitiveType, nullptr, nullptr, 0, BaseVertexIndex, 0, 0, NumPrimitives, NumInstances);
+		DrawPatches(PrimitiveType, nullptr, 0, BaseVertexIndex, 0, 0, NumPrimitives, NumInstances);
 	}
 	
 	ConditionalSubmit();	
@@ -435,7 +435,7 @@ void FMetalRenderPass::DrawPrimitiveIndirect(uint32 PrimitiveType, FMetalVertexB
 	}
 }
 
-void FMetalRenderPass::DrawIndexedPrimitive(FMetalBuffer const& IndexBuffer, FMetalTexture const& IndexTexture, uint32 IndexStride, uint32 PrimitiveType, int32 BaseVertexIndex, uint32 FirstInstance,
+void FMetalRenderPass::DrawIndexedPrimitive(FMetalBuffer const& IndexBuffer, uint32 IndexStride, uint32 PrimitiveType, int32 BaseVertexIndex, uint32 FirstInstance,
 											 uint32 NumVertices, uint32 StartIndex, uint32 NumPrimitives, uint32 NumInstances)
 {
 	// We need at least one to cover all use cases
@@ -509,7 +509,7 @@ void FMetalRenderPass::DrawIndexedPrimitive(FMetalBuffer const& IndexBuffer, FMe
 	}
 	else
 	{
-		DrawPatches(PrimitiveType, IndexBuffer, IndexTexture, IndexStride, BaseVertexIndex, FirstInstance, StartIndex, NumPrimitives, NumInstances);
+		DrawPatches(PrimitiveType, IndexBuffer, IndexStride, BaseVertexIndex, FirstInstance, StartIndex, NumPrimitives, NumInstances);
 	}
 	
 	ConditionalSubmit();
@@ -562,7 +562,7 @@ void FMetalRenderPass::DrawIndexedPrimitiveIndirect(uint32 PrimitiveType,FMetalI
 	}
 }
 
-void FMetalRenderPass::DrawPatches(uint32 PrimitiveType,FMetalBuffer const& IndexBuffer, FMetalTexture const& IndexTexture, uint32 IndexBufferStride, int32 BaseVertexIndex, uint32 FirstInstance, uint32 StartIndex,
+void FMetalRenderPass::DrawPatches(uint32 PrimitiveType,FMetalBuffer const& IndexBuffer, uint32 IndexBufferStride, int32 BaseVertexIndex, uint32 FirstInstance, uint32 StartIndex,
 									uint32 NumPrimitives, uint32 NumInstances)
 {
 	if (GetMetalDeviceContext().SupportsFeature(EMetalFeaturesTessellation))
@@ -621,8 +621,15 @@ void FMetalRenderPass::DrawPatches(uint32 PrimitiveType,FMetalBuffer const& Inde
 		
 		if (Pipeline->TessellationPipelineDesc.TessellationIndexBufferIndex != UINT_MAX)
 		{
-			PrologueEncoder.SetShaderTexture(mtlpp::FunctionType::Kernel, IndexTexture, Pipeline->TessellationPipelineDesc.TessellationIndexBufferIndex, mtlpp::ResourceUsage::Read);
-			State.SetShaderTexture(SF_Vertex, nil, Pipeline->TessellationPipelineDesc.TessellationIndexBufferIndex, mtlpp::ResourceUsage(0));
+			if (IndexBuffer)
+			{
+				PrologueEncoder.SetShaderBuffer(mtlpp::FunctionType::Kernel, IndexBuffer, StartIndex * IndexBufferStride, IndexBuffer.GetLength() - (StartIndex * IndexBufferStride), Pipeline->TessellationPipelineDesc.TessellationIndexBufferIndex, mtlpp::ResourceUsage::Read, IndexBufferStride == 2 ? PF_R16_UINT : PF_R32_UINT);
+			}
+			else
+			{
+				PrologueEncoder.SetShaderBuffer(mtlpp::FunctionType::Kernel, nil, 0, 0, Pipeline->TessellationPipelineDesc.TessellationIndexBufferIndex, mtlpp::ResourceUsage::Read, (EPixelFormat)0);
+			}
+			State.SetShaderBuffer(SF_Vertex, nil, nil, 0, 0, Pipeline->TessellationPipelineDesc.TessellationControlPointIndexBufferIndex, mtlpp::ResourceUsage(0));
 		}
 		
 		if(Pipeline->TessellationPipelineDesc.TessellationOutputControlPointBufferIndex != UINT_MAX) //TessellationOutputControlPointBufferIndex -> hullShaderOutputBuffer
diff --git a/Engine/Source/Runtime/Apple/MetalRHI/Private/MetalRenderPass.h b/Engine/Source/Runtime/Apple/MetalRHI/Private/MetalRenderPass.h
--- a/Engine/Source/Runtime/Apple/MetalRHI/Private/MetalRenderPass.h
+++ b/Engine/Source/Runtime/Apple/MetalRHI/Private/MetalRenderPass.h
@@ -40,14 +40,14 @@ public:
     
     void DrawPrimitiveIndirect(uint32 PrimitiveType, FMetalVertexBuffer* VertexBuffer, uint32 ArgumentOffset);
     
-    void DrawIndexedPrimitive(FMetalBuffer const& IndexBuffer, FMetalTexture const& IndexTexture, uint32 IndexStride, uint32 PrimitiveType, int32 BaseVertexIndex, uint32 FirstInstance,
+    void DrawIndexedPrimitive(FMetalBuffer const& IndexBuffer, uint32 IndexStride, uint32 PrimitiveType, int32 BaseVertexIndex, uint32 FirstInstance,
                          uint32 NumVertices, uint32 StartIndex, uint32 NumPrimitives, uint32 NumInstances);
     
     void DrawIndexedIndirect(FMetalIndexBuffer* IndexBufferRHI, uint32 PrimitiveType, FMetalStructuredBuffer* VertexBufferRHI, int32 DrawArgumentsIndex, uint32 NumInstances);
     
     void DrawIndexedPrimitiveIndirect(uint32 PrimitiveType,FMetalIndexBuffer* IndexBufferRHI,FMetalVertexBuffer* VertexBufferRHI,uint32 ArgumentOffset);
     
-    void DrawPatches(uint32 PrimitiveType, FMetalBuffer const& IndexBuffer, FMetalTexture const& IndexTexture, uint32 IndexBufferStride, int32 BaseVertexIndex, uint32 FirstInstance, uint32 StartIndex,
+    void DrawPatches(uint32 PrimitiveType, FMetalBuffer const& IndexBuffer, uint32 IndexBufferStride, int32 BaseVertexIndex, uint32 FirstInstance, uint32 StartIndex,
                      uint32 NumPrimitives, uint32 NumInstances);
     
     void Dispatch(uint32 ThreadGroupCountX, uint32 ThreadGroupCountY, uint32 ThreadGroupCountZ);
