UE4 - Copy from 4981899
- Save a few bytes on the PSO initializer
#rnx

diff --git a/Engine/Source/Runtime/Apple/MetalRHI/Private/MetalPipeline.cpp b/Engine/Source/Runtime/Apple/MetalRHI/Private/MetalPipeline.cpp
--- a/Engine/Source/Runtime/Apple/MetalRHI/Private/MetalPipeline.cpp
+++ b/Engine/Source/Runtime/Apple/MetalRHI/Private/MetalPipeline.cpp
@@ -110,7 +110,7 @@ struct FMetalGraphicsPipelineKey
 		bool bHasActiveTargets = false;
 		for (uint32 i = 0; i < NumActiveTargets; i++)
 		{
-			EPixelFormat TargetFormat = Init.RenderTargetFormats[i];
+			EPixelFormat TargetFormat = (EPixelFormat)Init.RenderTargetFormats[i];
 			if (TargetFormat == PF_Unknown) { continue; }
 
 			mtlpp::PixelFormat MetalFormat = (mtlpp::PixelFormat)GPixelFormats[TargetFormat].PlatformFormat;
@@ -525,7 +525,7 @@ static FMetalShaderPipeline* CreateMTLRenderPipeline(bool const bSync, FMetalGra
 		
         for (uint32 i = 0; i < NumActiveTargets; i++)
         {
-            EPixelFormat TargetFormat = Init.RenderTargetFormats[i];
+            EPixelFormat TargetFormat = (EPixelFormat)Init.RenderTargetFormats[i];
             if (TargetFormat == PF_Unknown && PixelShader && (((PixelShader->Bindings.InOutMask & 0x7fff) & (1 << i))))
             {
 				UE_LOG(LogMetal, Fatal, TEXT("Pipeline pixel shader expects target %u to be bound but it isn't: %s."), i, *FString(PixelShader->GetSourceCode()));
diff --git a/Engine/Source/Runtime/RHI/Private/PipelineFileCache.cpp b/Engine/Source/Runtime/RHI/Private/PipelineFileCache.cpp
--- a/Engine/Source/Runtime/RHI/Private/PipelineFileCache.cpp
+++ b/Engine/Source/Runtime/RHI/Private/PipelineFileCache.cpp
@@ -894,7 +894,7 @@ FPipelineCacheFileFormatPSO::FPipelineCacheFileFormatPSO()
 	
 	for (uint32 i = 0; i < MaxSimultaneousRenderTargets; i++)
 	{
-		PSO.GraphicsDesc.RenderTargetFormats[i] = Init.RenderTargetFormats[i];
+		PSO.GraphicsDesc.RenderTargetFormats[i] = (EPixelFormat)Init.RenderTargetFormats[i];
 		PSO.GraphicsDesc.RenderTargetFlags[i] = Init.RenderTargetFlags[i];
 	}
 	
diff --git a/Engine/Source/Runtime/RHI/Public/RHIResources.h b/Engine/Source/Runtime/RHI/Public/RHIResources.h
--- a/Engine/Source/Runtime/RHI/Public/RHIResources.h
+++ b/Engine/Source/Runtime/RHI/Public/RHIResources.h
@@ -1871,7 +1871,8 @@ public:
 class FGraphicsPipelineStateInitializer final : public FGraphicsMinimalPipelineStateInitializer
 {
 public:
-	using TRenderTargetFormats		= TStaticArray<TEnumAsByte<EPixelFormat>, MaxSimultaneousRenderTargets>;
+	// Can't use TEnumByte<EPixelFormat> as it changes the struct to be non trivially constructible, breaking memset
+	using TRenderTargetFormats		= TStaticArray<uint8/*EPixelFormat*/, MaxSimultaneousRenderTargets>;
 	using TRenderTargetFlags		= TStaticArray<uint32, MaxSimultaneousRenderTargets>;
 
 	FGraphicsPipelineStateInitializer()
@@ -1887,7 +1888,8 @@ public:
 		, NumSamples(0)
 		, Flags(0)
 	{
-		static_assert(PF_MAX <= 255, "EPixelFormat is assumed to be a byte; check usage of TEnumAsByte<EPixelFormat>");
+		static_assert(sizeof(EPixelFormat) != sizeof(uint8), "Change TRenderTargetFormats's uint8 to EPixelFormat");
+		static_assert(PF_MAX < MAX_uint8, "TRenderTargetFormats assumes EPixelFormat can fit in a uint8!");
 	}
 
 	FGraphicsPipelineStateInitializer(const FGraphicsMinimalPipelineStateInitializer& InMinimalState)
diff --git a/Engine/Source/Runtime/VulkanRHI/Private/VulkanRenderTarget.cpp b/Engine/Source/Runtime/VulkanRHI/Private/VulkanRenderTarget.cpp
--- a/Engine/Source/Runtime/VulkanRHI/Private/VulkanRenderTarget.cpp
+++ b/Engine/Source/Runtime/VulkanRHI/Private/VulkanRenderTarget.cpp
@@ -2141,7 +2141,7 @@ FVulkanRenderTargetLayout::FVulkanRenderTargetLayout(const FGraphicsPipelineStat
 	NumSamples = Initializer.NumSamples;
 	for (uint32 Index = 0; Index < Initializer.RenderTargetsEnabled; ++Index)
 	{
-		EPixelFormat UEFormat = Initializer.RenderTargetFormats[Index];
+		EPixelFormat UEFormat = (EPixelFormat)Initializer.RenderTargetFormats[Index];
 		if (UEFormat != PF_Unknown)
 		{
 			VkAttachmentDescription& CurrDesc = Desc[NumAttachmentDescriptions];
