changes from Dev-Rendering which provide better metal usage for buffers, allow defragging the metal heaps

diff --git a/Engine/Source/Runtime/Apple/MetalRHI/Private/MetalBuffer.cpp b/Engine/Source/Runtime/Apple/MetalRHI/Private/MetalBuffer.cpp
--- a/Engine/Source/Runtime/Apple/MetalRHI/Private/MetalBuffer.cpp
+++ b/Engine/Source/Runtime/Apple/MetalRHI/Private/MetalBuffer.cpp
@@ -8,9 +8,23 @@
 #include "MetalCommandBuffer.h"
 #include "MetalProfiler.h"
 
+DECLARE_MEMORY_STAT(TEXT("Used Device Buffer Memory"), STAT_MetalDeviceBufferMemory, STATGROUP_MetalRHI);
+DECLARE_MEMORY_STAT(TEXT("Used Pooled Buffer Memory"), STAT_MetalPooledBufferMemory, STATGROUP_MetalRHI);
+DECLARE_MEMORY_STAT(TEXT("Used Magazine Buffer Memory"), STAT_MetalMagazineBufferMemory, STATGROUP_MetalRHI);
+DECLARE_MEMORY_STAT(TEXT("Used Heap Buffer Memory"), STAT_MetalHeapBufferMemory, STATGROUP_MetalRHI);
+DECLARE_MEMORY_STAT(TEXT("Used Linear Buffer Memory"), STAT_MetalLinearBufferMemory, STATGROUP_MetalRHI);
+
 DECLARE_MEMORY_STAT(TEXT("Unused Pooled Buffer Memory"), STAT_MetalPooledBufferUnusedMemory, STATGROUP_MetalRHI);
 DECLARE_MEMORY_STAT(TEXT("Unused Magazine Buffer Memory"), STAT_MetalMagazineBufferUnusedMemory, STATGROUP_MetalRHI);
 DECLARE_MEMORY_STAT(TEXT("Unused Heap Buffer Memory"), STAT_MetalHeapBufferUnusedMemory, STATGROUP_MetalRHI);
+DECLARE_MEMORY_STAT(TEXT("Unused Linear Buffer Memory"), STAT_MetalLinearBufferUnusedMemory, STATGROUP_MetalRHI);
+
+static int32 GMetalHeapBufferBytesToCompact = 0;
+static FAutoConsoleVariableRef CVarMetalHeapBufferBytesToCompact(
+    TEXT("rhi.Metal.HeapBufferBytesToCompact"),
+    GMetalHeapBufferBytesToCompact,
+    TEXT("When enabled (> 0) this will force MetalRHI to compact the given number of bytes each frame into older buffer heaps from newer ones in order to defragment memory and reduce wastage.\n")
+    TEXT("(Off by default (0))"));
 
 #if METAL_DEBUG_OPTIONS
 extern int32 GMetalBufferScribble;
@@ -135,6 +149,14 @@ void FMetalBuffer::Release()
 	}
 }
 
+void FMetalBuffer::SetOwner(class FMetalRHIBuffer* Owner)
+{
+    if (Heap)
+    {
+        Heap->SetOwner(ns::Range(GetOffset(), GetLength()), Owner);
+    }
+}
+
 FMetalSubBufferHeap::FMetalSubBufferHeap(NSUInteger Size, NSUInteger Alignment, mtlpp::ResourceOptions Options, FCriticalSection& InPoolMutex)
 : PoolMutex(InPoolMutex)
 , OutstandingAllocs(0)
@@ -189,13 +211,41 @@ FMetalSubBufferHeap::~FMetalSubBufferHeap()
 	}
 }
 
+void FMetalSubBufferHeap::SetOwner(ns::Range const& Range, FMetalRHIBuffer* Owner)
+{
+    FScopeLock Lock(&PoolMutex);
+    for (uint32 i = 0; i < AllocRanges.Num(); i++)
+    {
+        if (AllocRanges[i].Range.Location == Range.Location)
+        {
+            check(AllocRanges[i].Range.Length == Range.Length);
+            check(AllocRanges[i].Owner == nullptr || Owner == nullptr);
+            AllocRanges[i].Owner = Owner;
+            break;
+        }
+    }
+}
+
 void FMetalSubBufferHeap::FreeRange(ns::Range const& Range)
 {
-	FPlatformAtomics::InterlockedDecrement(&OutstandingAllocs);
+    FPlatformAtomics::InterlockedDecrement(&OutstandingAllocs);
+    {
+        FScopeLock Lock(&PoolMutex);
+        for (uint32 i = 0; i < AllocRanges.Num(); i++)
+        {
+            if (AllocRanges[i].Range.Location == Range.Location)
+            {
+                check(AllocRanges[i].Range.Length == Range.Length);
+                AllocRanges.RemoveAt(i);
+                break;
+            }
+        }
+    }
 	if (ParentHeap)
 	{
-		SET_MEMORY_STAT(STAT_MetalBufferUnusedMemory, ParentHeap.GetSize() - ParentHeap.GetUsedSize());
-		SET_MEMORY_STAT(STAT_MetalHeapBufferUnusedMemory, ParentHeap.GetSize() - ParentHeap.GetUsedSize());
+        INC_MEMORY_STAT_BY(STAT_MetalBufferUnusedMemory, Range.Length);
+        INC_MEMORY_STAT_BY(STAT_MetalHeapBufferUnusedMemory, Range.Length);
+        DEC_MEMORY_STAT_BY(STAT_MetalHeapBufferMemory, Range.Length);
 	}
 	else
 	{
@@ -248,7 +298,8 @@ void FMetalSubBufferHeap::FreeRange(ns::Range const& Range)
 			
 			INC_MEMORY_STAT_BY(STAT_MetalBufferUnusedMemory, Range.Length);
 			INC_MEMORY_STAT_BY(STAT_MetalHeapBufferUnusedMemory, Range.Length);
-		
+            DEC_MEMORY_STAT_BY(STAT_MetalHeapBufferMemory, Range.Length);
+
 #if METAL_DEBUG_OPTIONS
 			uint64 LostSize = GetSize() - UsedSize;
 			for (ns::Range const& FreeRange : FreeRanges)
@@ -335,7 +386,7 @@ NSUInteger FMetalSubBufferHeap::GetUsedSize() const
 
 int64 FMetalSubBufferHeap::NumCurrentAllocations() const
 {
-	return OutstandingAllocs;
+    return OutstandingAllocs;
 }
 
 void FMetalSubBufferHeap::SetLabel(const ns::String& label)
@@ -369,6 +420,23 @@ NSUInteger FMetalSubBufferHeap::MaxAvailableSize() const
 	}
 }
 
+bool FMetalSubBufferHeap::CanAllocateSize(NSUInteger Size) const
+{
+    if (ParentHeap)
+    {
+        NSUInteger Storage = (NSUInteger(GetStorageMode()) << mtlpp::ResourceStorageModeShift);
+        NSUInteger Cache = (NSUInteger(GetCpuCacheMode()) << mtlpp::ResourceCpuCacheModeShift);
+        mtlpp::ResourceOptions Opt = mtlpp::ResourceOptions(Storage | Cache);
+        
+        NSUInteger Align = ParentHeap.GetDevice().HeapBufferSizeAndAlign(Size, Opt).Align;
+        return Size <= ParentHeap.MaxAvailableSizeWithAlignment(Align);
+    }
+    else
+    {
+        return Size <= MaxAvailableSize();
+    }
+}
+
 FMetalBuffer FMetalSubBufferHeap::NewBuffer(NSUInteger length)
 {
 	NSUInteger Size = Align(length, MinAlign);
@@ -384,8 +452,9 @@ FMetalBuffer FMetalSubBufferHeap::NewBuffer(NSUInteger length)
 #if STATS || ENABLE_LOW_LEVEL_MEM_TRACKER
 		MetalLLM::LogAllocBuffer(GetMetalDeviceContext().GetDevice(), Result);
 #endif
-		SET_MEMORY_STAT(STAT_MetalBufferUnusedMemory, ParentHeap.GetSize() - ParentHeap.GetUsedSize());
-		SET_MEMORY_STAT(STAT_MetalHeapBufferUnusedMemory, ParentHeap.GetSize() - ParentHeap.GetUsedSize());
+        DEC_MEMORY_STAT_BY(STAT_MetalBufferUnusedMemory, Result.GetLength());
+        DEC_MEMORY_STAT_BY(STAT_MetalHeapBufferUnusedMemory, Result.GetLength());
+        INC_MEMORY_STAT_BY(STAT_MetalHeapBufferMemory, Result.GetLength());
 	}
 	else
 	{
@@ -405,11 +474,12 @@ FMetalBuffer FMetalSubBufferHeap::NewBuffer(NSUInteger length)
 					
 					DEC_MEMORY_STAT_BY(STAT_MetalBufferUnusedMemory, Range.Length);
 					DEC_MEMORY_STAT_BY(STAT_MetalHeapBufferUnusedMemory, Range.Length);
-				
+                    INC_MEMORY_STAT_BY(STAT_MetalHeapBufferMemory, Range.Length);
+
 					if (Range.Length > Size)
 					{
 						ns::Range Split = ns::Range(Range.Location + Size, Range.Length - Size);
-						FPlatformAtomics::InterlockedIncrement(&OutstandingAllocs);
+                        FPlatformAtomics::InterlockedIncrement(&OutstandingAllocs);
 						FreeRange(Split);
 						
 						Range.Length = Size;
@@ -425,13 +495,19 @@ FMetalBuffer FMetalSubBufferHeap::NewBuffer(NSUInteger length)
 #endif
 					
 					Result = FMetalBuffer(MTLPP_VALIDATE(mtlpp::Buffer, ParentBuffer, SafeGetRuntimeDebuggingLevel() >= EMetalDebugLevelValidation, NewBuffer(Range)), this);
-				
+                    
+                    Allocation Alloc;
+                    Alloc.Range = Range;
+                    Alloc.Resource = Result.GetPtr();
+                    Alloc.Owner = nullptr;
+                    AllocRanges.Add(Alloc);
+
 					break;
 				}
 			}
 		}
 	}
-	FPlatformAtomics::InterlockedIncrement(&OutstandingAllocs);
+    FPlatformAtomics::InterlockedIncrement(&OutstandingAllocs);
 	check(Result && Result.GetPtr());
 	return Result;
 }
@@ -468,13 +544,13 @@ FMetalSubBufferLinear::FMetalSubBufferLinear(NSUInteger Size, NSUInteger Alignme
 	MetalLLM::LogAllocBuffer(GetMetalDeviceContext().GetDevice(), ParentBuffer);
 #endif
 	INC_MEMORY_STAT_BY(STAT_MetalBufferUnusedMemory, FullSize);
-	INC_MEMORY_STAT_BY(STAT_MetalHeapBufferUnusedMemory, FullSize);
+    INC_MEMORY_STAT_BY(STAT_MetalLinearBufferUnusedMemory, FullSize);
 }
 
 FMetalSubBufferLinear::~FMetalSubBufferLinear()
 {
 	DEC_MEMORY_STAT_BY(STAT_MetalBufferUnusedMemory, ParentBuffer.GetLength());
-	DEC_MEMORY_STAT_BY(STAT_MetalHeapBufferUnusedMemory, ParentBuffer.GetLength());
+    DEC_MEMORY_STAT_BY(STAT_MetalLinearBufferUnusedMemory, ParentBuffer.GetLength());
 }
 
 void FMetalSubBufferLinear::FreeRange(ns::Range const& Range)
@@ -492,7 +568,8 @@ void FMetalSubBufferLinear::FreeRange(ns::Range const& Range)
 	{
 		FreedSize += Range.Length;
 		INC_MEMORY_STAT_BY(STAT_MetalBufferUnusedMemory, Range.Length);
-		INC_MEMORY_STAT_BY(STAT_MetalHeapBufferUnusedMemory, Range.Length);
+        INC_MEMORY_STAT_BY(STAT_MetalLinearBufferUnusedMemory, Range.Length);
+        DEC_MEMORY_STAT_BY(STAT_MetalLinearBufferMemory, Range.Length);
 		if (FreedSize == UsedSize)
 		{
 			UsedSize = 0;
@@ -563,7 +640,8 @@ FMetalBuffer FMetalSubBufferLinear::NewBuffer(NSUInteger length)
 	{
 		ns::Range Range(NewWriteHead, Size);
 		DEC_MEMORY_STAT_BY(STAT_MetalBufferUnusedMemory, Range.Length);
-		DEC_MEMORY_STAT_BY(STAT_MetalHeapBufferUnusedMemory, Range.Length);
+        DEC_MEMORY_STAT_BY(STAT_MetalLinearBufferUnusedMemory, Range.Length);
+        INC_MEMORY_STAT_BY(STAT_MetalLinearBufferMemory, Range.Length);
 		Result = FMetalBuffer(MTLPP_VALIDATE(mtlpp::Buffer, ParentBuffer, SafeGetRuntimeDebuggingLevel() >= EMetalDebugLevelValidation, NewBuffer(Range)), this);
 		UsedSize += Size;
 		WriteHead = NewWriteHead + Size;
@@ -581,18 +659,24 @@ mtlpp::PurgeableState FMetalSubBufferLinear::SetPurgeableState(mtlpp::PurgeableS
 
 FMetalSubBufferMagazine::FMetalSubBufferMagazine(NSUInteger Size, NSUInteger ChunkSize, mtlpp::ResourceOptions Options)
 : MinAlign(ChunkSize)
+, BlockSize(ChunkSize)
 , OutstandingAllocs(0)
 , UsedSize(0)
 {
-	NSUInteger FullSize = Align(Size, ChunkSize);
+    static bool bSupportsHeaps = GetMetalDeviceContext().SupportsFeature(EMetalFeaturesHeaps);
+    mtlpp::StorageMode Storage = (mtlpp::StorageMode)((Options & mtlpp::ResourceStorageModeMask) >> mtlpp::ResourceStorageModeShift);
+    if (bSupportsHeaps && Storage == mtlpp::StorageMode::Private)
+    {
+        MinAlign = GetMetalDeviceContext().GetDevice().HeapBufferSizeAndAlign(BlockSize, Options).Align;
+    }
+    
+    NSUInteger FullSize = Align(Size, MinAlign);
 	METAL_GPUPROFILE(FScopedMetalCPUStats CPUStat(FString::Printf(TEXT("AllocBuffer: %llu, %llu"), FullSize, Options)));
 	
-	mtlpp::StorageMode Storage = (mtlpp::StorageMode)((Options & mtlpp::ResourceStorageModeMask) >> mtlpp::ResourceStorageModeShift);
 #if PLATFORM_MAC
 	check(Storage != mtlpp::StorageMode::Managed /* Managed memory cannot be safely suballocated! When you overwrite existing data the GPU buffer is immediately disposed of! */);
 #endif
 
-	static bool bSupportsHeaps = GetMetalDeviceContext().SupportsFeature(EMetalFeaturesHeaps);
 	if (bSupportsHeaps && Storage == mtlpp::StorageMode::Private)
 	{
 		mtlpp::HeapDescriptor Desc;
@@ -628,7 +712,7 @@ FMetalSubBufferMagazine::~FMetalSubBufferMagazine()
 	if (ParentHeap)
 	{
 		DEC_MEMORY_STAT_BY(STAT_MetalBufferUnusedMemory, ParentHeap.GetSize());
-		DEC_MEMORY_STAT_BY(STAT_MetalHeapBufferUnusedMemory, ParentHeap.GetSize());
+        DEC_MEMORY_STAT_BY(STAT_MetalMagazineBufferUnusedMemory, ParentHeap.GetSize());
 	}
 	else
 	{
@@ -639,11 +723,12 @@ FMetalSubBufferMagazine::~FMetalSubBufferMagazine()
 
 void FMetalSubBufferMagazine::FreeRange(ns::Range const& Range)
 {
-	FPlatformAtomics::InterlockedDecrement(&OutstandingAllocs);
+    FPlatformAtomics::InterlockedDecrement(&OutstandingAllocs);
 	if (ParentHeap)
 	{
-		SET_MEMORY_STAT(STAT_MetalBufferUnusedMemory, ParentHeap.GetSize() - ParentHeap.GetUsedSize());
-		SET_MEMORY_STAT(STAT_MetalHeapBufferUnusedMemory, ParentHeap.GetSize() - ParentHeap.GetUsedSize());
+        INC_MEMORY_STAT_BY(STAT_MetalBufferUnusedMemory, Range.Length);
+        INC_MEMORY_STAT_BY(STAT_MetalMagazineBufferUnusedMemory, Range.Length);
+        DEC_MEMORY_STAT_BY(STAT_MetalMagazineBufferMemory, Range.Length);
 	}
 	else
 	{
@@ -661,6 +746,7 @@ void FMetalSubBufferMagazine::FreeRange(ns::Range const& Range)
 		
 		INC_MEMORY_STAT_BY(STAT_MetalBufferUnusedMemory, Range.Length);
 		INC_MEMORY_STAT_BY(STAT_MetalMagazineBufferUnusedMemory, Range.Length);
+        DEC_MEMORY_STAT_BY(STAT_MetalMagazineBufferMemory, Range.Length);
 	}
 }
 
@@ -743,7 +829,12 @@ NSUInteger FMetalSubBufferMagazine::GetFreeSize() const
 
 int64 FMetalSubBufferMagazine::NumCurrentAllocations() const
 {
-	return OutstandingAllocs;
+    return OutstandingAllocs;
+}
+
+bool FMetalSubBufferMagazine::CanAllocateSize(NSUInteger Size) const
+{
+    return GetFreeSize() >= Size;
 }
 
 void FMetalSubBufferMagazine::SetLabel(const ns::String& label)
@@ -760,7 +851,7 @@ void FMetalSubBufferMagazine::SetLabel(const ns::String& label)
 
 FMetalBuffer FMetalSubBufferMagazine::NewBuffer()
 {
-	NSUInteger Size = MinAlign;
+    NSUInteger Size = BlockSize;
 	FMetalBuffer Result;
 
 	if (ParentHeap)
@@ -773,8 +864,9 @@ FMetalBuffer FMetalSubBufferMagazine::NewBuffer()
 #if STATS || ENABLE_LOW_LEVEL_MEM_TRACKER
 		MetalLLM::LogAllocBuffer(GetMetalDeviceContext().GetDevice(), Result);
 #endif
-		SET_MEMORY_STAT(STAT_MetalBufferUnusedMemory, ParentHeap.GetSize() - ParentHeap.GetUsedSize());
-		SET_MEMORY_STAT(STAT_MetalHeapBufferUnusedMemory, ParentHeap.GetSize() - ParentHeap.GetUsedSize());
+        DEC_MEMORY_STAT_BY(STAT_MetalBufferUnusedMemory, Result.GetLength());
+        DEC_MEMORY_STAT_BY(STAT_MetalMagazineBufferUnusedMemory, Result.GetLength());
+        INC_MEMORY_STAT_BY(STAT_MetalMagazineBufferMemory, Result.GetLength());
 	}
 	else
 	{
@@ -786,12 +878,13 @@ FMetalBuffer FMetalSubBufferMagazine::NewBuffer()
 			FPlatformAtomics::InterlockedAdd(&UsedSize, ((int64)Range->Length));
 			DEC_MEMORY_STAT_BY(STAT_MetalBufferUnusedMemory, Range->Length);
 			DEC_MEMORY_STAT_BY(STAT_MetalMagazineBufferUnusedMemory, Range->Length);
+            INC_MEMORY_STAT_BY(STAT_MetalMagazineBufferMemory, Range->Length);
 			Result = FMetalBuffer(MTLPP_VALIDATE(mtlpp::Buffer, ParentBuffer, SafeGetRuntimeDebuggingLevel() >= EMetalDebugLevelValidation, NewBuffer(*Range)), this);
 			delete Range;
 		}
 	}
 
-	FPlatformAtomics::InterlockedIncrement(&OutstandingAllocs);
+    FPlatformAtomics::InterlockedIncrement(&OutstandingAllocs);
 	check(Result && Result.GetPtr());
 	return Result;
 }
@@ -872,7 +965,7 @@ FMetalBuffer FMetalSubBufferRing::NewBuffer(NSUInteger Size, uint32 Alignment)
 	// Allocate on first use
 	if(!Buffer.IsValid())
 	{
-		Buffer = MakeShared<FMetalRingBufferRef, ESPMode::ThreadSafe>(GetMetalDeviceContext().GetResourceHeap().CreateBuffer(InitialSize, MinAlign, Options, true));
+        Buffer = MakeShared<FMetalRingBufferRef, ESPMode::ThreadSafe>(GetMetalDeviceContext().GetResourceHeap().CreateBuffer(InitialSize, MinAlign, BUF_Dynamic, Options, true));
 	}
 	
 	if(Buffer->LastRead <= WriteHead)
@@ -891,7 +984,7 @@ FMetalBuffer FMetalSubBufferRing::NewBuffer(NSUInteger Size, uint32 Alignment)
 		else if (Storage == mtlpp::StorageMode::Managed)
 		{
 			Submit();
-			Buffer = MakeShared<FMetalRingBufferRef, ESPMode::ThreadSafe>(GetMetalDeviceContext().GetResourceHeap().CreateBuffer(Buffer->Buffer.GetLength(), MinAlign, Options, true));
+            Buffer = MakeShared<FMetalRingBufferRef, ESPMode::ThreadSafe>(GetMetalDeviceContext().GetResourceHeap().CreateBuffer(Buffer->Buffer.GetLength(), MinAlign, BUF_Dynamic, Options, true));
 			WriteHead = 0;
 			CommitHead = 0;
 			SubmitHead = 0;
@@ -912,7 +1005,7 @@ FMetalBuffer FMetalSubBufferRing::NewBuffer(NSUInteger Size, uint32 Alignment)
 		
 		Submit();
 		
-		Buffer = MakeShared<FMetalRingBufferRef, ESPMode::ThreadSafe>(GetMetalDeviceContext().GetResourceHeap().CreateBuffer(NewBufferSize, MinAlign, Options, true));
+        Buffer = MakeShared<FMetalRingBufferRef, ESPMode::ThreadSafe>(GetMetalDeviceContext().GetResourceHeap().CreateBuffer(NewBufferSize, MinAlign, BUF_Dynamic, Options, true));
 		WriteHead = 0;
 		CommitHead = 0;
 		SubmitHead = 0;
@@ -949,8 +1042,8 @@ void FMetalSubBufferRing::Shrink()
 			
 			UE_LOG(LogMetal, Verbose, TEXT("Shrinking RingBuffer from %u to %u as max. usage is %u at frame %lld]"), (uint32)Buffer->Buffer.GetLength(), (uint32)ThreeQuarterSize, (uint32)FrameMax, GFrameNumberRenderThread);
 			
-			Buffer = MakeShared<FMetalRingBufferRef, ESPMode::ThreadSafe>(GetMetalDeviceContext().GetResourceHeap().CreateBuffer(ThreeQuarterSize, MinAlign, Options, true));
-			
+            Buffer = MakeShared<FMetalRingBufferRef, ESPMode::ThreadSafe>(GetMetalDeviceContext().GetResourceHeap().CreateBuffer(ThreeQuarterSize, MinAlign, BUF_Dynamic, Options, true));
+
 			WriteHead = 0;
 			CommitHead = 0;
 			SubmitHead = 0;
@@ -1092,7 +1185,7 @@ FMetalBuffer FMetalBufferPoolPolicyData::CreateResource(CreationArguments Args)
 
 FMetalBufferPoolPolicyData::CreationArguments FMetalBufferPoolPolicyData::GetCreationArguments(FMetalBuffer const& Resource)
 {
-	return FMetalBufferPoolPolicyData::CreationArguments(Resource.GetDevice(), Resource.GetLength(), Resource.GetStorageMode());
+    return FMetalBufferPoolPolicyData::CreationArguments(Resource.GetDevice(), Resource.GetLength(), 0, Resource.GetStorageMode());
 }
 
 void FMetalBufferPoolPolicyData::FreeResource(FMetalBuffer& Resource)
@@ -1149,6 +1242,7 @@ FMetalTexture FMetalTexturePool::CreateTexture(mtlpp::Device Device, mtlpp::Text
 	{
 		Texture = *Tex;
 		Pool.Remove(Descriptor);
+        Texture.SetPurgeableState(mtlpp::PurgeableState::NonVolatile);
 	}
 	else
 	{
@@ -1176,6 +1270,8 @@ void FMetalTexturePool::ReleaseTexture(FMetalTexture& Texture)
 	Descriptor.usage = Texture.GetUsage();
 	Descriptor.freedFrame = GFrameNumberRenderThread;
 	
+    Texture.SetPurgeableState(mtlpp::PurgeableState::Volatile);
+    
 	FScopeLock Lock(&PoolMutex);
 	Pool.Add(Descriptor, Texture);
 }
@@ -1195,6 +1291,10 @@ void FMetalTexturePool::Drain(bool const bForce)
 			{
 				It.RemoveCurrent();
 			}
+            else if ((GFrameNumberRenderThread - It->Key.freedFrame) >= PurgeAfterNumFrames)
+            {
+                It->Value.SetPurgeableState(mtlpp::PurgeableState::Empty);
+            }
 		}
 	}
 }
@@ -1342,7 +1442,7 @@ mtlpp::Heap FMetalResourceHeap::GetTextureHeap(mtlpp::TextureDescriptor Desc, mt
 	return Result;
 }
 
-FMetalBuffer FMetalResourceHeap::CreateBuffer(uint32 Size, uint32 Alignment, mtlpp::ResourceOptions Options, bool bForceUnique)
+FMetalBuffer FMetalResourceHeap::CreateBuffer(uint32 Size, uint32 Alignment, uint32 Flags, mtlpp::ResourceOptions Options, bool bForceUnique)
 {
 	LLM_SCOPE_METAL(ELLMTagMetal::Buffers);
 	LLM_PLATFORM_SCOPE_METAL(ELLMTagMetal::Buffers);
@@ -1351,6 +1451,8 @@ FMetalBuffer FMetalResourceHeap::CreateBuffer(uint32 Size, uint32 Alignment, mtl
 	static bool bSupportsBufferSubAllocation = FMetalCommandQueue::SupportsFeature(EMetalFeaturesBufferSubAllocation);
 	bForceUnique |= (!bSupportsBufferSubAllocation && !bSupportsHeaps);
 	
+    uint32 Usage = (Flags & BUF_Static) ? UsageStatic : UsageDynamic;
+    
 	FMetalBuffer Buffer;
 	uint32 BlockSize = Align(Size, Alignment);
 	mtlpp::StorageMode StorageMode = (mtlpp::StorageMode)(((NSUInteger)Options & mtlpp::ResourceStorageModeMask) >> mtlpp::ResourceStorageModeShift);
@@ -1386,9 +1488,11 @@ FMetalBuffer FMetalResourceHeap::CreateBuffer(uint32 Size, uint32 Alignment, mtl
 				 }
 				 else
 				{
-					Buffer = ManagedBuffers.CreatePooledResource(FMetalPooledBufferArgs(Queue->GetDevice(), BlockSize, StorageMode));
+                    Buffer = ManagedBuffers.CreatePooledResource(FMetalPooledBufferArgs(Queue->GetDevice(), BlockSize, Flags, StorageMode));
+                    Buffer.SetPurgeableState(mtlpp::PurgeableState::NonVolatile);
 					DEC_MEMORY_STAT_BY(STAT_MetalBufferUnusedMemory, Buffer.GetLength());
 					DEC_MEMORY_STAT_BY(STAT_MetalPooledBufferUnusedMemory, Buffer.GetLength());
+                    INC_MEMORY_STAT_BY(STAT_MetalPooledBufferMemory, Buffer.GetLength());
 				}
 				break;
 			}
@@ -1405,12 +1509,12 @@ FMetalBuffer FMetalResourceHeap::CreateBuffer(uint32 Size, uint32 Alignment, mtl
 					FScopeLock Lock(&Mutex);
 					
 					uint32 i = GetMagazineIndex(BlockSize);
-					TArray<FMetalSubBufferMagazine*>& Heaps = SmallBuffers[Storage][i];
-					
+                    TArray<FMetalSubBufferMagazine*>& Heaps = SmallBuffers[Usage][Storage][i];
+
 					FMetalSubBufferMagazine* Found = nullptr;
 					for (FMetalSubBufferMagazine* Heap : Heaps)
 					{
-						if (Heap->GetFreeSize() >= BlockSize)
+                        if (Heap->CanAllocateSize(BlockSize))
 						{
 							Found = Heap;
 							break;
@@ -1420,7 +1524,7 @@ FMetalBuffer FMetalResourceHeap::CreateBuffer(uint32 Size, uint32 Alignment, mtl
 					if (!Found)
 					{
 						Found = new FMetalSubBufferMagazine(MagazineAllocSizes[i], MagazineSizes[i], mtlpp::ResourceOptions((NSUInteger)Options & (mtlpp::ResourceStorageModeMask|mtlpp::ResourceHazardTrackingModeMask)));
-						SmallBuffers[Storage][i].Add(Found);
+                        SmallBuffers[Usage][Storage][i].Add(Found);
 					}
 					check(Found);
 					
@@ -1432,12 +1536,12 @@ FMetalBuffer FMetalResourceHeap::CreateBuffer(uint32 Size, uint32 Alignment, mtl
 					FScopeLock Lock(&Mutex);
 					
 					uint32 i = GetHeapIndex(BlockSize);
-					TArray<FMetalSubBufferHeap*>& Heaps = BufferHeaps[Storage][i];
-					
+                    TArray<FMetalSubBufferHeap*>& Heaps = BufferHeaps[Usage][Storage][i];
+
 					FMetalSubBufferHeap* Found = nullptr;
 					for (FMetalSubBufferHeap* Heap : Heaps)
 					{
-						if (Heap->MaxAvailableSize() >= BlockSize)
+                        if (Heap->CanAllocateSize(BlockSize))
 						{
 							Found = Heap;
 							break;
@@ -1446,8 +1550,9 @@ FMetalBuffer FMetalResourceHeap::CreateBuffer(uint32 Size, uint32 Alignment, mtl
 					
 					if (!Found)
 					{
-						Found = new FMetalSubBufferHeap(HeapAllocSizes[i], HeapSizes[i], mtlpp::ResourceOptions((NSUInteger)Options & (mtlpp::ResourceStorageModeMask|mtlpp::ResourceHazardTrackingModeMask)), Mutex);
-						BufferHeaps[Storage][i].Add(Found);
+                        uint32 MinAlign = PLATFORM_MAC ? 256 : 64;
+                        Found = new FMetalSubBufferHeap(HeapAllocSizes[i], MinAlign, mtlpp::ResourceOptions((NSUInteger)Options & (mtlpp::ResourceStorageModeMask|mtlpp::ResourceHazardTrackingModeMask)), Mutex);
+                        BufferHeaps[Usage][Storage][i].Add(Found);
 					}
 					check(Found);
 					
@@ -1457,9 +1562,11 @@ FMetalBuffer FMetalResourceHeap::CreateBuffer(uint32 Size, uint32 Alignment, mtl
 				else
 				{
 					FScopeLock Lock(&Mutex);
-					Buffer = Buffers[Storage].CreatePooledResource(FMetalPooledBufferArgs(Queue->GetDevice(), BlockSize, StorageMode));
+                    Buffer = Buffers[Storage].CreatePooledResource(FMetalPooledBufferArgs(Queue->GetDevice(), BlockSize, Flags, StorageMode));
+                    Buffer.SetPurgeableState(mtlpp::PurgeableState::NonVolatile);
 					DEC_MEMORY_STAT_BY(STAT_MetalBufferUnusedMemory, Buffer.GetLength());
 					DEC_MEMORY_STAT_BY(STAT_MetalPooledBufferUnusedMemory, Buffer.GetLength());
+                    INC_MEMORY_STAT_BY(STAT_MetalPooledBufferMemory, Buffer.GetLength());
 				}
 				break;
 			}
@@ -1477,6 +1584,7 @@ FMetalBuffer FMetalResourceHeap::CreateBuffer(uint32 Size, uint32 Alignment, mtl
 #if STATS || ENABLE_LOW_LEVEL_MEM_TRACKER
 		MetalLLM::LogAllocBuffer(Queue->GetDevice(), Buffer);
 #endif
+        INC_MEMORY_STAT_BY(STAT_MetalDeviceBufferMemory, Buffer.GetLength());
 	}
 	
 	if (GMetalBufferZeroFill && Buffer.GetStorageMode() != mtlpp::StorageMode::Private)
@@ -1498,7 +1606,10 @@ void FMetalResourceHeap::ReleaseBuffer(FMetalBuffer& Buffer)
 		
 		INC_MEMORY_STAT_BY(STAT_MetalBufferUnusedMemory, Buffer.GetLength());
 		INC_MEMORY_STAT_BY(STAT_MetalPooledBufferUnusedMemory, Buffer.GetLength());
-		
+        DEC_MEMORY_STAT_BY(STAT_MetalPooledBufferMemory, Buffer.GetLength());
+
+        Buffer.SetPurgeableState(mtlpp::PurgeableState::Volatile);
+        
 		switch (StorageMode)
 		{
 	#if PLATFORM_MAC
@@ -1527,7 +1638,8 @@ void FMetalResourceHeap::ReleaseBuffer(FMetalBuffer& Buffer)
 	}
 	else
 	{
-		Buffer.Release();
+        DEC_MEMORY_STAT_BY(STAT_MetalDeviceBufferMemory, Buffer.GetLength());
+        Buffer.Release();
 	}
 }
 
@@ -1575,48 +1687,101 @@ void FMetalResourceHeap::ReleaseTexture(FMetalSurface* Surface, FMetalTexture& T
 void FMetalResourceHeap::Compact(bool const bForce)
 {
 	FScopeLock Lock(&Mutex);
-	for (uint32 t = 0; t < NumAllocTypes; t++)
-	{
-		for (uint32 i = 0; i < NumMagazineSizes; i++)
-		{
-			for (auto It = SmallBuffers[t][i].CreateIterator(); It; ++It)
-			{
-				FMetalSubBufferMagazine* Data = *It;
-				if (Data->NumCurrentAllocations() == 0 || bForce)
-				{
-					It.RemoveCurrent();
-					delete Data;
-				}
-			}
-		}
-		
-		for (uint32 i = 0; i < NumHeapSizes; i++)
-		{
-			for (auto It = BufferHeaps[t][i].CreateIterator(); It; ++It)
-			{
-				FMetalSubBufferHeap* Data = *It;
-				if (Data->NumCurrentAllocations() == 0 || bForce)
-				{
-					It.RemoveCurrent();
-					delete Data;
-				}
-			}
-		}
-	}
-	
-	for (uint32 i = 0; i < EMetalHeapTextureUsageNum; i++)
-	{
-		for (uint32 j = 0; j < NumTextureHeapSizes; j++)
-		{
-			for (auto It = TextureHeaps[i][j].CreateIterator(); It; ++It)
-			{
-				if (It->GetUsedSize() == 0 || bForce)
-				{
-					It.RemoveCurrent();
-				}
-			}
-		}
-	}
+    for (uint32 u = 0; u < NumUsageTypes; u++)
+    {
+        for (uint32 t = 0; t < NumAllocTypes; t++)
+        {
+            for (uint32 i = 0; i < NumMagazineSizes; i++)
+            {
+                for (auto It = SmallBuffers[u][t][i].CreateIterator(); It; ++It)
+                {
+                    FMetalSubBufferMagazine* Data = *It;
+                    if (Data->NumCurrentAllocations() == 0 || bForce)
+                    {
+                        It.RemoveCurrent();
+                        delete Data;
+                    }
+                }
+            }
+            
+            mtlpp::CommandBuffer Buffer;
+            mtlpp::BlitCommandEncoder Encoder;
+            
+            uint32 BytesCompacted = 0;
+            uint32 const BytesToCompact = GMetalHeapBufferBytesToCompact;
+            
+            for (uint32 i = 0; i < NumHeapSizes; i++)
+            {
+                for (uint32 j = 1; j < BufferHeaps[u][t][i].Num(); j++)
+                {
+                    FMetalSubBufferHeap* Data = BufferHeaps[u][t][i][j];
+                    if (Data->AllocRanges.Num() > 0 && BytesCompacted < BytesToCompact)
+                    {
+                        for (FMetalSubBufferHeap::Allocation const& Alloc : Data->AllocRanges)
+                        {
+                            if (Alloc.Owner)
+                            {
+                                for (uint32 AllocIndex = 0; AllocIndex < j && BytesCompacted < BytesToCompact; AllocIndex++)
+                                {
+                                    FMetalSubBufferHeap* Prev = BufferHeaps[u][t][i][AllocIndex];
+                                    if (Prev->MaxAvailableSize() >= Alloc.Range.Length)
+                                    {
+                                        if (!Encoder)
+                                        {
+                                            Buffer = Queue->CreateCommandBuffer();
+                                            Encoder = Buffer.BlitCommandEncoder();
+                                        }
+                                        
+                                        FMetalBuffer NewBuffer = Prev->NewBuffer(Alloc.Range.Length);
+                                        check(NewBuffer && NewBuffer.GetPtr());
+                                        
+                                        FMetalBuffer SourceResource = FMetalBuffer((id<MTLBuffer>)Alloc.Resource);
+                                        Encoder.Copy(SourceResource, Alloc.Range.Location, NewBuffer, 0, Alloc.Range.Length);
+                                        
+                                        FMetalBuffer PrevBuffer;
+                                        FMetalBuffer PrevCPUBuffer;
+                                        if (Alloc.Owner->Buffer.GetPtr() == Alloc.Resource && Alloc.Owner->Buffer.GetOffset() == Alloc.Range.Location)
+                                        {
+                                            PrevBuffer = Alloc.Owner->Buffer;
+                                            Alloc.Owner->Buffer = NewBuffer;
+                                        }
+                                        if (Alloc.Owner->CPUBuffer.GetPtr() == Alloc.Resource && Alloc.Owner->CPUBuffer.GetOffset() == Alloc.Range.Location)
+                                        {
+                                            PrevCPUBuffer = Alloc.Owner->CPUBuffer;
+                                            Alloc.Owner->CPUBuffer = NewBuffer;
+                                        }
+                                        if (PrevBuffer)
+                                            SafeReleaseMetalBuffer(PrevBuffer);
+                                        if (PrevCPUBuffer)
+                                            SafeReleaseMetalBuffer(PrevCPUBuffer);
+                                        
+                                        BytesCompacted += Alloc.Range.Length;
+                                        break;
+                                    }
+                                }
+                            }
+                        }
+                    }
+                }
+                
+                for (auto It = BufferHeaps[u][t][i].CreateIterator(); It; ++It)
+                {
+                    FMetalSubBufferHeap* Data = *It;
+                    if (Data->NumCurrentAllocations() == 0 || bForce)
+                    {
+                        It.RemoveCurrent();
+                        delete Data;
+                    }
+                }
+            }
+            
+            if (Encoder)
+            {
+                Encoder.EndEncoding();
+                Queue->CommitCommandBuffer(Buffer);
+            }
+        }
+    }
 
 	Buffers[AllocShared].DrainPool(bForce);
 	Buffers[AllocPrivate].DrainPool(bForce);
@@ -1673,15 +1838,10 @@ uint32 FMetalResourceHeap::MagazineSizes[FMetalResourceHeap::NumMagazineSizes] =
 	1024,
 	2048,
 	4096,
+    8192,
 };
 
 uint32 FMetalResourceHeap::HeapSizes[FMetalResourceHeap::NumHeapSizes] = {
-	16384,
-	32768,
-	65536,
-	131072,
-	262144,
-	524288,
 	1048576,
 	2097152,
 };
@@ -1696,17 +1856,12 @@ uint32 FMetalResourceHeap::MagazineAllocSizes[FMetalResourceHeap::NumMagazineSiz
 	16384,
 	16384,
 	16384,
+    32768,
 };
 
 uint32 FMetalResourceHeap::HeapAllocSizes[FMetalResourceHeap::NumHeapSizes] = {
-	131072,
-	131072,
-	524288,
-	524288,
-	1048576,
 	2097152,
 	4194304,
-	4194304,
 };
 
 uint32 FMetalResourceHeap::HeapTextureHeapSizes[FMetalResourceHeap::NumTextureHeapSizes] = {
diff --git a/Engine/Source/Runtime/Apple/MetalRHI/Private/MetalBuffer.h b/Engine/Source/Runtime/Apple/MetalRHI/Private/MetalBuffer.h
--- a/Engine/Source/Runtime/Apple/MetalRHI/Private/MetalBuffer.h
+++ b/Engine/Source/Runtime/Apple/MetalRHI/Private/MetalBuffer.h
@@ -11,22 +11,26 @@
 
 struct FMetalPooledBufferArgs
 {
-	FMetalPooledBufferArgs() : Device(nil), Size(0), Storage(mtlpp::StorageMode::Shared) {}
-	
-	FMetalPooledBufferArgs(mtlpp::Device InDevice, uint32 InSize, mtlpp::StorageMode InStorage)
-	: Device(InDevice)
-	, Size(InSize)
-	, Storage(InStorage)
-	{
-	}
-	
-	mtlpp::Device Device;
-	uint32 Size;
-	mtlpp::StorageMode Storage;
+    FMetalPooledBufferArgs() : Device(nil), Size(0), Flags(0), Storage(mtlpp::StorageMode::Shared) {}
+    
+    FMetalPooledBufferArgs(mtlpp::Device InDevice, uint32 InSize, uint32 InFlags, mtlpp::StorageMode InStorage)
+    : Device(InDevice)
+    , Size(InSize)
+    , Flags(InFlags)
+    , Storage(InStorage)
+    {
+    }
+    
+    mtlpp::Device Device;
+    uint32 Size;
+    uint32 Flags;
+    mtlpp::StorageMode Storage;
 };
 
 class FMetalSubBufferHeap
 {
+    friend class FMetalResourceHeap;
+    
 public:
 	FMetalSubBufferHeap(NSUInteger Size, NSUInteger Alignment, mtlpp::ResourceOptions, FCriticalSection& PoolMutex);
 	~FMetalSubBufferHeap();
@@ -38,22 +42,33 @@ public:
     NSUInteger     GetSize() const;
     NSUInteger     GetUsedSize() const;
 	NSUInteger	 MaxAvailableSize() const;
-	int64     NumCurrentAllocations() const;
+    int64     NumCurrentAllocations() const;
+    bool     CanAllocateSize(NSUInteger Size) const;
 
     void SetLabel(const ns::String& label);
 	
     FMetalBuffer NewBuffer(NSUInteger length);
     mtlpp::PurgeableState SetPurgeableState(mtlpp::PurgeableState state);
 	void FreeRange(ns::Range const& Range);
+    
+    void SetOwner(ns::Range const& Range, FMetalRHIBuffer* Owner);
 
 private:
+    struct Allocation
+    {
+        ns::Range Range;
+        mtlpp::Buffer::Type Resource;
+        FMetalRHIBuffer* Owner;
+    };
+    
 	FCriticalSection& PoolMutex;
-	int64 volatile OutstandingAllocs;
+    int64 volatile OutstandingAllocs;
 	NSUInteger MinAlign;
 	NSUInteger UsedSize;
 	mtlpp::Buffer ParentBuffer;
 	mutable mtlpp::Heap ParentHeap;
 	TArray<ns::Range> FreeRanges;
+    TArray<Allocation> AllocRanges;
 };
 
 class FMetalSubBufferLinear
@@ -98,7 +113,8 @@ public:
     NSUInteger     GetSize() const;
     NSUInteger     GetUsedSize() const;
 	NSUInteger	 GetFreeSize() const;
-	int64     NumCurrentAllocations() const;
+    int64     NumCurrentAllocations() const;
+    bool     CanAllocateSize(NSUInteger Size) const;
 
     void SetLabel(const ns::String& label);
 	void FreeRange(ns::Range const& Range);
@@ -108,7 +124,8 @@ public:
 
 private:
 	NSUInteger MinAlign;
-	int64 volatile OutstandingAllocs;
+    NSUInteger BlockSize;
+    int64 volatile OutstandingAllocs;
 	int64 volatile UsedSize;
 	mtlpp::Buffer ParentBuffer;
 	mutable mtlpp::Heap ParentHeap;
@@ -250,6 +267,7 @@ class FMetalTexturePool
 {
 	enum
 	{
+        PurgeAfterNumFrames = 2, /* Textures must be reused fairly rapidly but after this number of frames we reclaim the memory, even though the object persists */
 		CullAfterNumFrames = 3, /* Textures must be reused fairly rapidly or we bin them as they are much larger than buffers */
 	};
 public:
@@ -345,17 +363,12 @@ class FMetalResourceHeap
 		Size1024,
 		Size2048,
 		Size4096,
+        Size8192,
 		NumMagazineSizes
 	};
 	
 	enum HeapSize
 	{
-		Size16k,
-		Size32k,
-		Size64k,
-		Size128k,
-		Size256k,
-		Size512k,
 		Size1Mb,
 		Size2Mb,
 		NumHeapSizes
@@ -392,13 +405,20 @@ class FMetalResourceHeap
 		EMetalHeapTextureUsageNum = 2
 	};
 
+    enum UsageTypes
+    {
+        UsageStatic,
+        UsageDynamic,
+        NumUsageTypes = 2
+    };
+    
 public:
 	FMetalResourceHeap(void);
 	~FMetalResourceHeap();
 	
 	void Init(FMetalCommandQueue& Queue);
 	
-	FMetalBuffer CreateBuffer(uint32 Size, uint32 Alignment, mtlpp::ResourceOptions Options, bool bForceUnique = false);
+    FMetalBuffer CreateBuffer(uint32 Size, uint32 Alignment, uint32 Flags, mtlpp::ResourceOptions Options, bool bForceUnique = false);
 	FMetalTexture CreateTexture(mtlpp::TextureDescriptor Desc, FMetalSurface* Surface);
 	
 	void ReleaseBuffer(FMetalBuffer& Buffer);
@@ -423,13 +443,13 @@ private:
 	FCriticalSection Mutex;
 	FMetalCommandQueue* Queue;
 	
-	/** Small allocations (<= 4KB) are made from magazine allocators that use sub-ranges of a buffer */
-	TArray<FMetalSubBufferMagazine*> SmallBuffers[NumAllocTypes][NumMagazineSizes];
+    /** Small allocations (<= 4KB) are made from magazine allocators that use sub-ranges of a buffer */
+    TArray<FMetalSubBufferMagazine*> SmallBuffers[NumUsageTypes][NumAllocTypes][NumMagazineSizes];
+    
+    /** Typical allocations (4KB - 4MB) are made from heap allocators that use sub-ranges of a buffer */
+    /** There are two alignment categories for heaps - 16b for Vertes/Index data and 256b for constant data (macOS-only) */
+    TArray<FMetalSubBufferHeap*> BufferHeaps[NumUsageTypes][NumAllocTypes][NumHeapSizes];
 
-	/** Typical allocations (4KB - 4MB) are made from heap allocators that use sub-ranges of a buffer */
-	/** There are two alignment categories for heaps - 16b for Vertes/Index data and 256b for constant data (macOS-only) */
-	TArray<FMetalSubBufferHeap*> BufferHeaps[NumAllocTypes][NumHeapSizes];
-	
 	/** Larger buffers (up-to 32MB) that are subject to bucketing & pooling rather than sub-allocation */
 	FMetalBufferPool Buffers[NumAllocTypes];
 #if PLATFORM_MAC // All managed buffers are bucketed & pooled rather than sub-allocated to avoid memory consistency complexities
diff --git a/Engine/Source/Runtime/Apple/MetalRHI/Private/MetalContext.cpp b/Engine/Source/Runtime/Apple/MetalRHI/Private/MetalContext.cpp
--- a/Engine/Source/Runtime/Apple/MetalRHI/Private/MetalContext.cpp
+++ b/Engine/Source/Runtime/Apple/MetalRHI/Private/MetalContext.cpp
@@ -492,12 +492,6 @@ void FMetalDeviceContext::DrainHeap()
 
 void FMetalDeviceContext::EndFrame()
 {
-	FlushFreeList();
-	
-	ClearFreeList();
-	
-    Heap.Compact(false);
-    
 	// A 'frame' in this context is from the beginning of encoding on the CPU
 	// to the end of all rendering operations on the GPU. So the semaphore is
 	// signalled when the last command buffer finishes GPU execution.
@@ -533,6 +527,12 @@ void FMetalDeviceContext::EndFrame()
 #endif
 	SubmitCommandsHint((uint32)SubmitFlags);
 	
+    FlushFreeList();
+    
+    ClearFreeList();
+    
+    Heap.Compact(false);
+    
 	InitFrame(true, 0, 0);
 }
 
@@ -767,7 +767,7 @@ FMetalTexture FMetalDeviceContext::CreateTexture(FMetalSurface* Surface, mtlpp::
 
 FMetalBuffer FMetalDeviceContext::CreatePooledBuffer(FMetalPooledBufferArgs const& Args)
 {
-	FMetalBuffer Buffer = Heap.CreateBuffer(Args.Size, BufferOffsetAlignment, GetCommandQueue().GetCompatibleResourceOptions((mtlpp::ResourceOptions)(BUFFER_CACHE_MODE | mtlpp::ResourceOptions::HazardTrackingModeUntracked | ((NSUInteger)Args.Storage << mtlpp::ResourceStorageModeShift))));
+    FMetalBuffer Buffer = Heap.CreateBuffer(Args.Size, BufferOffsetAlignment, Args.Flags, GetCommandQueue().GetCompatibleResourceOptions((mtlpp::ResourceOptions)(BUFFER_CACHE_MODE | mtlpp::ResourceOptions::HazardTrackingModeUntracked | ((NSUInteger)Args.Storage << mtlpp::ResourceStorageModeShift))));
 	check(Buffer && Buffer.GetPtr());
 #if METAL_DEBUG_OPTIONS
 	if (GMetalResourcePurgeOnDelete && !Buffer.GetHeap())
diff --git a/Engine/Source/Runtime/Apple/MetalRHI/Private/MetalQuery.cpp b/Engine/Source/Runtime/Apple/MetalRHI/Private/MetalQuery.cpp
--- a/Engine/Source/Runtime/Apple/MetalRHI/Private/MetalQuery.cpp
+++ b/Engine/Source/Runtime/Apple/MetalRHI/Private/MetalQuery.cpp
@@ -57,7 +57,7 @@ FMetalQueryBuffer* FMetalQueryBufferPool::GetCurrentQueryBuffer()
 			LLM_PLATFORM_SCOPE_METAL(ELLMTagMetal::Buffers);
 
 			METAL_GPUPROFILE(FScopedMetalCPUStats CPUStat(FString::Printf(TEXT("AllocBuffer: %llu, %llu"), EQueryBufferMaxSize, mtlpp::ResourceOptions::StorageModeShared)));
-			Buffer = FMetalBuffer(MTLPP_VALIDATE(mtlpp::Device, Context->GetDevice(), SafeGetRuntimeDebuggingLevel() >= EMetalDebugLevelValidation, NewBuffer(EQueryBufferMaxSize, GetMetalDeviceContext().GetCommandQueue().GetCompatibleResourceOptions((mtlpp::ResourceOptions)(BUFFER_CACHE_MODE | mtlpp::ResourceOptions::HazardTrackingModeUntracked | mtlpp::ResourceOptions::StorageModeShared)))), false);
+            Buffer = GetMetalDeviceContext().GetResourceHeap().CreateBuffer(EQueryBufferMaxSize, 16, BUF_Dynamic, Context->GetCommandQueue().GetCompatibleResourceOptions((mtlpp::ResourceOptions)(BUFFER_CACHE_MODE | mtlpp::ResourceOptions::HazardTrackingModeUntracked | mtlpp::ResourceOptions::StorageModeShared)), true);
 			FMemory::Memzero((((uint8*)Buffer.GetContents())), EQueryBufferMaxSize);
 
 #if STATS || ENABLE_LOW_LEVEL_MEM_TRACKER
diff --git a/Engine/Source/Runtime/Apple/MetalRHI/Private/MetalRHIContext.cpp b/Engine/Source/Runtime/Apple/MetalRHI/Private/MetalRHIContext.cpp
--- a/Engine/Source/Runtime/Apple/MetalRHI/Private/MetalRHIContext.cpp
+++ b/Engine/Source/Runtime/Apple/MetalRHI/Private/MetalRHIContext.cpp
@@ -42,6 +42,7 @@ void SafeReleaseMetalBuffer(FMetalBuffer& Buffer)
 {
 	if(GIsMetalInitialized && GDynamicRHI && Buffer)
 	{
+        Buffer.SetOwner(nullptr);
 		FMetalRHICommandContext* Context = static_cast<FMetalRHICommandContext*>(RHIGetDefaultContext());
 		if(Context)
 		{
diff --git a/Engine/Source/Runtime/Apple/MetalRHI/Private/MetalRenderPass.cpp b/Engine/Source/Runtime/Apple/MetalRHI/Private/MetalRenderPass.cpp
--- a/Engine/Source/Runtime/Apple/MetalRHI/Private/MetalRenderPass.cpp
+++ b/Engine/Source/Runtime/Apple/MetalRHI/Private/MetalRenderPass.cpp
@@ -591,19 +591,19 @@ void FMetalRenderPass::DrawPatches(uint32 PrimitiveType,FMetalBuffer const& Inde
 		FMetalBuffer hullShaderOutputBuffer = nil;
 		if(hullShaderOutputBufferSize)
 		{
-			hullShaderOutputBuffer = deviceContext.CreatePooledBuffer(FMetalPooledBufferArgs(device, hullShaderOutputBufferSize, mtlpp::StorageMode::Private));
+            hullShaderOutputBuffer = deviceContext.CreatePooledBuffer(FMetalPooledBufferArgs(device, hullShaderOutputBufferSize, BUF_Dynamic, mtlpp::StorageMode::Private));
 		}
 		
 		FMetalBuffer hullConstShaderOutputBuffer = nil;
 		if(hullConstShaderOutputBufferSize)
 		{
-			hullConstShaderOutputBuffer = deviceContext.CreatePooledBuffer(FMetalPooledBufferArgs(device, hullConstShaderOutputBufferSize, mtlpp::StorageMode::Private));
+            hullConstShaderOutputBuffer = deviceContext.CreatePooledBuffer(FMetalPooledBufferArgs(device, hullConstShaderOutputBufferSize, BUF_Dynamic, mtlpp::StorageMode::Private));
 		}
 		
 		FMetalBuffer tessellationFactorBuffer = nil;
 		if(tessellationFactorBufferSize)
 		{
-			tessellationFactorBuffer = deviceContext.CreatePooledBuffer(FMetalPooledBufferArgs(device, tessellationFactorBufferSize, mtlpp::StorageMode::Private));
+            tessellationFactorBuffer = deviceContext.CreatePooledBuffer(FMetalPooledBufferArgs(device, tessellationFactorBufferSize, BUF_Dynamic, mtlpp::StorageMode::Private));
 		}
 	
 		auto& computeEncoder = PrologueEncoder.GetComputeCommandEncoder();
diff --git a/Engine/Source/Runtime/Apple/MetalRHI/Private/MetalRenderTarget.cpp b/Engine/Source/Runtime/Apple/MetalRHI/Private/MetalRenderTarget.cpp
--- a/Engine/Source/Runtime/Apple/MetalRHI/Private/MetalRenderTarget.cpp
+++ b/Engine/Source/Runtime/Apple/MetalRHI/Private/MetalRenderTarget.cpp
@@ -670,7 +670,7 @@ void FMetalDynamicRHI::RHIReadSurfaceData(FTextureRHIParamRef TextureRHI, FIntRe
 		const uint32 Alignment = PLATFORM_MAC ? 1u : 64u; // Mac permits natural row alignment (tightly-packed) but iOS does not.
 		const uint32 AlignedStride = ((Stride - 1) & ~(Alignment - 1)) + Alignment;
 		const uint32 BytesPerImage = AlignedStride * SizeY;
-		FMetalBuffer Buffer = ((FMetalDeviceContext*)ImmediateContext.Context)->CreatePooledBuffer(FMetalPooledBufferArgs(ImmediateContext.Context->GetDevice(), BytesPerImage, mtlpp::StorageMode::Shared));
+        FMetalBuffer Buffer = ((FMetalDeviceContext*)ImmediateContext.Context)->CreatePooledBuffer(FMetalPooledBufferArgs(ImmediateContext.Context->GetDevice(), BytesPerImage, BUF_Dynamic, mtlpp::StorageMode::Shared));
 		{
 			// Synchronise the texture with the CPU
 			SCOPE_CYCLE_COUNTER(STAT_MetalTexturePageOffTime);
@@ -771,7 +771,7 @@ void FMetalDynamicRHI::RHIReadSurfaceFloatData(FTextureRHIParamRef TextureRHI, F
 	const uint32 AlignedStride = ((Stride - 1) & ~(Alignment - 1)) + Alignment;
 	const uint32 BytesPerImage = AlignedStride  * SizeY;
 	int32 FloatBGRADataSize = BytesPerImage;
-	FMetalBuffer Buffer = ((FMetalDeviceContext*)ImmediateContext.Context)->CreatePooledBuffer(FMetalPooledBufferArgs(ImmediateContext.Context->GetDevice(), FloatBGRADataSize, mtlpp::StorageMode::Shared));
+    FMetalBuffer Buffer = ((FMetalDeviceContext*)ImmediateContext.Context)->CreatePooledBuffer(FMetalPooledBufferArgs(ImmediateContext.Context->GetDevice(), FloatBGRADataSize, BUF_Dynamic, mtlpp::StorageMode::Shared));
 	{
 		// Synchronise the texture with the CPU
 		SCOPE_CYCLE_COUNTER(STAT_MetalTexturePageOffTime);
@@ -837,7 +837,7 @@ void FMetalDynamicRHI::RHIRead3DSurfaceFloatData(FTextureRHIParamRef TextureRHI,
 	const uint32 AlignedStride = ((Stride - 1) & ~(Alignment - 1)) + Alignment;
 	const uint32 BytesPerImage = AlignedStride  * SizeY;
 	int32 FloatBGRADataSize = BytesPerImage * SizeZ;
-	FMetalBuffer Buffer = ((FMetalDeviceContext*)ImmediateContext.Context)->CreatePooledBuffer(FMetalPooledBufferArgs(ImmediateContext.Context->GetDevice(), FloatBGRADataSize, mtlpp::StorageMode::Shared));
+    FMetalBuffer Buffer = ((FMetalDeviceContext*)ImmediateContext.Context)->CreatePooledBuffer(FMetalPooledBufferArgs(ImmediateContext.Context->GetDevice(), FloatBGRADataSize, BUF_Dynamic, mtlpp::StorageMode::Shared));
 	{
 		// Synchronise the texture with the CPU
 		SCOPE_CYCLE_COUNTER(STAT_MetalTexturePageOffTime);
diff --git a/Engine/Source/Runtime/Apple/MetalRHI/Private/MetalTexture.cpp b/Engine/Source/Runtime/Apple/MetalRHI/Private/MetalTexture.cpp
--- a/Engine/Source/Runtime/Apple/MetalRHI/Private/MetalTexture.cpp
+++ b/Engine/Source/Runtime/Apple/MetalRHI/Private/MetalTexture.cpp
@@ -1273,7 +1273,7 @@ FMetalBuffer FMetalSurface::AllocSurface(uint32 MipIndex, uint32 ArrayIndex, ERe
 	// allocate some temporary memory
 	mtlpp::ResourceOptions ResMode = GetMetalDeviceContext().GetCommandQueue().GetCompatibleResourceOptions((mtlpp::ResourceOptions)(mtlpp::ResourceOptions::StorageModeShared | (bSupportsResourceOptions && !(PLATFORM_MAC && PixelFormat == PF_G8 && (Flags & TexCreate_SRGB)) ? mtlpp::ResourceOptions::CpuCacheModeWriteCombined : 0)));
 	
-	FMetalBuffer Buffer = GetMetalDeviceContext().GetResourceHeap().CreateBuffer(MipBytes, BufferOffsetAlignment, ResMode);
+    FMetalBuffer Buffer = GetMetalDeviceContext().GetResourceHeap().CreateBuffer(MipBytes, BufferOffsetAlignment, BUF_Dynamic, ResMode);
 	
 #if PLATFORM_MAC
 	// Expand R8_sRGB into RGBA8_sRGB for Mac.
@@ -2047,7 +2047,7 @@ static FMetalBuffer InternalCopyTexture2DUpdateRegion(FTexture2DRHIParamRef Text
 	if(Texture->GetFormat() == PF_G8 && (Texture->GetFlags() & TexCreate_SRGB))
 	{
 		const uint32 BufferSize = UpdateRegion.Height * UpdateRegion.Width * sizeof(uint32);		
-		Buffer = GetMetalDeviceContext().CreatePooledBuffer(FMetalPooledBufferArgs(GetMetalDeviceContext().GetDevice(), BufferSize, mtlpp::StorageMode::Shared));
+        Buffer = GetMetalDeviceContext().CreatePooledBuffer(FMetalPooledBufferArgs(GetMetalDeviceContext().GetDevice(), BufferSize, BUF_Dynamic, mtlpp::StorageMode::Shared));
 		InternalExpandR8ToStandardRGBA((uint32*)Buffer.GetContents(), UpdateRegion, InOutSourcePitch, SourceData);
 	}
 
@@ -2055,7 +2055,7 @@ static FMetalBuffer InternalCopyTexture2DUpdateRegion(FTexture2DRHIParamRef Text
 #endif
 	{
 		const uint32 BufferSize = UpdateRegion.Height * InSourcePitch;
-		Buffer = GetMetalDeviceContext().CreatePooledBuffer(FMetalPooledBufferArgs(GetMetalDeviceContext().GetDevice(), BufferSize, mtlpp::StorageMode::Shared));
+        Buffer = GetMetalDeviceContext().CreatePooledBuffer(FMetalPooledBufferArgs(GetMetalDeviceContext().GetDevice(), BufferSize, BUF_Dynamic, mtlpp::StorageMode::Shared));
 
 		// Limit copy to line by line by update region pitch otherwise we can go off the end of source data on the last row
 		uint8* pDestRow = (uint8*)Buffer.GetContents();
@@ -2259,7 +2259,7 @@ struct FMetalDynamicRHIUpdateTexture3DCommand final : public FRHICommand<FMetalD
 		FMetalTexture Tex = Texture->Surface.Texture;
 		const uint32 BufferSize = UpdateRegion.Height * UpdateRegion.Depth* SourceRowPitch;
 		
-		Buffer = GetMetalDeviceContext().CreatePooledBuffer(FMetalPooledBufferArgs(GetMetalDeviceContext().GetDevice(), BufferSize, mtlpp::StorageMode::Shared));
+        Buffer = GetMetalDeviceContext().CreatePooledBuffer(FMetalPooledBufferArgs(GetMetalDeviceContext().GetDevice(), BufferSize, BUF_Dynamic, mtlpp::StorageMode::Shared));
 		InternalCopyTexture3DUpdateRegionData(Texture, UpdateRegion, SourceRowPitch, SourceDepthPitch, SourceData, (uint8*)Buffer.GetContents());
 	}
 	
@@ -2332,7 +2332,7 @@ void FMetalDynamicRHI::RHIUpdateTexture3D(FTexture3DRHIParamRef TextureRHI,uint3
 			SCOPED_AUTORELEASE_POOL;
 
 			const uint32 BufferSize = UpdateRegion.Height * UpdateRegion.Depth * SourceRowPitch;
-			FMetalBuffer IntermediateBuffer = GetMetalDeviceContext().CreatePooledBuffer(FMetalPooledBufferArgs(GetMetalDeviceContext().GetDevice(), BufferSize, mtlpp::StorageMode::Shared));
+            FMetalBuffer IntermediateBuffer = GetMetalDeviceContext().CreatePooledBuffer(FMetalPooledBufferArgs(GetMetalDeviceContext().GetDevice(), BufferSize, BUF_Dynamic, mtlpp::StorageMode::Shared));
 			InternalCopyTexture3DUpdateRegionData(TextureRHI, UpdateRegion, SourceRowPitch, SourceDepthPitch, SourceData, (uint8*)IntermediateBuffer.GetContents());
 			InternalUpdateTexture3D(ImmediateContext.GetInternalContext(), TextureRHI, MipIndex, UpdateRegion, SourceRowPitch, SourceDepthPitch, IntermediateBuffer);
 			GetMetalDeviceContext().ReleaseBuffer(IntermediateBuffer);
@@ -2815,8 +2815,8 @@ void FMetalDynamicRHI::RHICopySubTextureRegion(FTexture2DRHIParamRef SourceTextu
 				const uint32 AlignedStride = ((Stride - 1) & ~(Alignment - 1)) + Alignment;
 				const uint32 BytesPerImage = AlignedStride *  SourceSize.height;
 				
-				FMetalBuffer Buffer = GetMetalDeviceContext().CreatePooledBuffer(FMetalPooledBufferArgs(ImmediateContext.Context->GetDevice(), BytesPerImage, mtlpp::StorageMode::Shared));
-				
+                FMetalBuffer Buffer = GetMetalDeviceContext().CreatePooledBuffer(FMetalPooledBufferArgs(ImmediateContext.Context->GetDevice(), BytesPerImage, BUF_Dynamic, mtlpp::StorageMode::Shared));
+
 				check(Buffer);
 				
 				mtlpp::BlitOption Options = mtlpp::BlitOption::None;
diff --git a/Engine/Source/Runtime/Apple/MetalRHI/Private/MetalUAV.cpp b/Engine/Source/Runtime/Apple/MetalRHI/Private/MetalUAV.cpp
--- a/Engine/Source/Runtime/Apple/MetalRHI/Private/MetalUAV.cpp
+++ b/Engine/Source/Runtime/Apple/MetalRHI/Private/MetalUAV.cpp
@@ -544,7 +544,7 @@ void FMetalRHICommandContext::RHIClearTinyUAV(FUnorderedAccessViewRHIParamRef Un
 			FMetalBuffer Temp = nil;
 			bool bBufferPooled = false;
 			
-			FMetalPooledBufferArgs Args(GetMetalDeviceContext().GetDevice(), AlignedSize, mtlpp::StorageMode::Shared);
+            FMetalPooledBufferArgs Args(GetMetalDeviceContext().GetDevice(), AlignedSize, BUF_Dynamic, mtlpp::StorageMode::Shared);
 			Temp = GetMetalDeviceContext().CreatePooledBuffer(Args);
 			bBufferPooled = true;
 			
@@ -899,7 +899,7 @@ void FMetalRHICommandContext::RHICopyToStagingBuffer(FVertexBufferRHIParamRef So
 			{
 				SafeReleaseMetalBuffer(ReadbackBuffer);
 			}
-			FMetalPooledBufferArgs ArgsCPU(GetMetalDeviceContext().GetDevice(), NumBytes, mtlpp::StorageMode::Shared);
+            FMetalPooledBufferArgs ArgsCPU(GetMetalDeviceContext().GetDevice(), NumBytes, BUF_Dynamic, mtlpp::StorageMode::Shared);
 			ReadbackBuffer = GetMetalDeviceContext().CreatePooledBuffer(ArgsCPU);
 		}
 
diff --git a/Engine/Source/Runtime/Apple/MetalRHI/Private/MetalVertexBuffer.cpp b/Engine/Source/Runtime/Apple/MetalRHI/Private/MetalVertexBuffer.cpp
--- a/Engine/Source/Runtime/Apple/MetalRHI/Private/MetalVertexBuffer.cpp
+++ b/Engine/Source/Runtime/Apple/MetalRHI/Private/MetalVertexBuffer.cpp
@@ -242,9 +242,11 @@ void FMetalRHIBuffer::Alloc(uint32 InSize, EResourceLockMode LockMode)
 	{
         check(LockMode != RLM_ReadOnly);
 		mtlpp::StorageMode Mode = (bUsePrivateMem ? mtlpp::StorageMode::Private : BUFFER_STORAGE_MODE);
-		FMetalPooledBufferArgs Args(GetMetalDeviceContext().GetDevice(), InSize, Mode);
+        FMetalPooledBufferArgs Args(GetMetalDeviceContext().GetDevice(), InSize, Usage, Mode);
 		Buffer = GetMetalDeviceContext().CreatePooledBuffer(Args);
 		check(Buffer && Buffer.GetPtr());
+        
+        Buffer.SetOwner(this);
 
         METAL_INC_DWORD_STAT_BY(Type, MemAlloc, InSize);
         
@@ -263,9 +265,10 @@ void FMetalRHIBuffer::Alloc(uint32 InSize, EResourceLockMode LockMode)
 	
 	if (bUsePrivateMem && !CPUBuffer)
 	{
-		FMetalPooledBufferArgs ArgsCPU(GetMetalDeviceContext().GetDevice(), InSize, mtlpp::StorageMode::Shared);
+        FMetalPooledBufferArgs ArgsCPU(GetMetalDeviceContext().GetDevice(), InSize, BUF_Dynamic, mtlpp::StorageMode::Shared);
 		CPUBuffer = GetMetalDeviceContext().CreatePooledBuffer(ArgsCPU);
 		check(CPUBuffer && CPUBuffer.GetPtr());
+        CPUBuffer.SetOwner(this);
         METAL_INC_DWORD_STAT_BY(Type, MemAlloc, InSize);
 		check(CPUBuffer.GetLength() >= Buffer.GetLength());
 	}
diff --git a/Engine/Source/Runtime/Apple/MetalRHI/Public/MetalResources.h b/Engine/Source/Runtime/Apple/MetalRHI/Public/MetalResources.h
--- a/Engine/Source/Runtime/Apple/MetalRHI/Public/MetalResources.h
+++ b/Engine/Source/Runtime/Apple/MetalRHI/Public/MetalResources.h
@@ -348,6 +348,7 @@ public:
 	inline bool IsPooled() const { return bPooled; }
 	inline bool IsSingleUse() const { return bSingleUse; }
 	inline void MarkSingleUse() { bSingleUse = true; }
+    void SetOwner(class FMetalRHIBuffer* Owner);
 	void Release();
 	
 	friend uint32 GetTypeHash(FMetalBuffer const& Hash)
