CIS fix

diff --git a/Engine/Source/Runtime/Apple/MetalRHI/Private/MetalTexture.cpp b/Engine/Source/Runtime/Apple/MetalRHI/Private/MetalTexture.cpp
--- a/Engine/Source/Runtime/Apple/MetalRHI/Private/MetalTexture.cpp
+++ b/Engine/Source/Runtime/Apple/MetalRHI/Private/MetalTexture.cpp
@@ -970,7 +970,7 @@ FMetalSurface::FMetalSurface(ERHIResourceType ResourceType, EPixelFormat Format,
 		{
 			mtlpp::Device Device = GetMetalDeviceContext().GetDevice();
 			mtlpp::SizeAndAlign SizeAlign = Device.HeapTextureSizeAndAlign(Desc);
-			FMetalPooledBufferArgs Args(Device, SizeAlign.Size, mtlpp::StorageMode::Private);
+			FMetalPooledBufferArgs Args(Device, SizeAlign.Size, 0, mtlpp::StorageMode::Private);
 			FMetalBuffer Buffer = GetMetalDeviceContext().CreatePooledBuffer(Args);
 			Texture = Buffer.NewTexture(Desc, 0, SizeAlign.Size);
 		}
@@ -2912,7 +2912,7 @@ void FMetalRHICommandContext::RHICopyTexture(FTextureRHIParamRef SourceTextureRH
 						const uint32 BytesPerImage = AlignedStride *  SourceSize.height;
 						const uint32 DataSize = BytesPerImage * SourceSize.depth;
 						
-						FMetalBuffer Buffer = GetMetalDeviceContext().CreatePooledBuffer(FMetalPooledBufferArgs(GetInternalContext().GetDevice(), DataSize, mtlpp::StorageMode::Shared));
+						FMetalBuffer Buffer = GetMetalDeviceContext().CreatePooledBuffer(FMetalPooledBufferArgs(GetInternalContext().GetDevice(), DataSize, 0, mtlpp::StorageMode::Shared));
 						
 						check(Buffer);
 						
diff --git a/Engine/Source/Runtime/Apple/MetalRHI/Private/MetalUniformBuffer.cpp b/Engine/Source/Runtime/Apple/MetalRHI/Private/MetalUniformBuffer.cpp
--- a/Engine/Source/Runtime/Apple/MetalRHI/Private/MetalUniformBuffer.cpp
+++ b/Engine/Source/Runtime/Apple/MetalRHI/Private/MetalUniformBuffer.cpp
@@ -419,7 +419,7 @@ void FMetalUniformBuffer::InitIAB()
 			Desc.SetAccess(mtlpp::ArgumentAccess::ReadOnly);
 			Desc.SetDataType(mtlpp::DataType::Pointer);
 			
-			FMetalPooledBufferArgs Args(GetMetalDeviceContext().GetDevice(), BufferSizes.Num() * sizeof(uint32), BUFFER_STORAGE_MODE);
+			FMetalPooledBufferArgs Args(GetMetalDeviceContext().GetDevice(), BufferSizes.Num() * sizeof(uint32), 0, BUFFER_STORAGE_MODE);
 			NewIAB->IndirectArgumentBufferSideTable = GetMetalDeviceContext().CreatePooledBuffer(Args);
 			
 			FMemory::Memcpy(NewIAB->IndirectArgumentBufferSideTable.GetContents(), BufferSizes.GetData(), BufferSizes.Num() * sizeof(uint32));
@@ -446,7 +446,7 @@ void FMetalUniformBuffer::InitIAB()
 		
 		mtlpp::ArgumentEncoder Encoder = FMetalArgumentEncoderCache::Get().CreateEncoder(Arguments);
 		
-		NewIAB->IndirectArgumentBuffer = GetMetalDeviceContext().GetResourceHeap().CreateBuffer(Encoder.GetEncodedLength(), 16, mtlpp::ResourceOptions(BUFFER_CACHE_MODE | ((NSUInteger)BUFFER_STORAGE_MODE << mtlpp::ResourceStorageModeShift)), true);
+		NewIAB->IndirectArgumentBuffer = GetMetalDeviceContext().GetResourceHeap().CreateBuffer(Encoder.GetEncodedLength(), 16, 0, mtlpp::ResourceOptions(BUFFER_CACHE_MODE | ((NSUInteger)BUFFER_STORAGE_MODE << mtlpp::ResourceStorageModeShift)), true);
 		
 		Encoder.SetArgumentBuffer(NewIAB->IndirectArgumentBuffer, 0);
 		
