Compile out all support for Metal tessellation shaders from the MetalRHI runtime & remove all vestigial references to geometry shaders using the new pre-processor flags. This resizes all the various arrays to eliminate unused entries which were previously SF_NumFrequencies in size.

diff --git a/Engine/Source/Developer/Apple/MetalShaderFormat/Private/MetalShaderCompiler.cpp b/Engine/Source/Developer/Apple/MetalShaderFormat/Private/MetalShaderCompiler.cpp
--- a/Engine/Source/Developer/Apple/MetalShaderFormat/Private/MetalShaderCompiler.cpp
+++ b/Engine/Source/Developer/Apple/MetalShaderFormat/Private/MetalShaderCompiler.cpp
@@ -1305,20 +1305,27 @@ void BuildMetalShaderOutput(
 	Header.NumThreadsY = CCHeader.NumThreads[1];
 	Header.NumThreadsZ = CCHeader.NumThreads[2];
 	
-	Header.TessellationOutputControlPoints 		= CCHeader.TessellationOutputControlPoints;
-	Header.TessellationDomain					= CCHeader.TessellationDomain;
-	Header.TessellationInputControlPoints       = CCHeader.TessellationInputControlPoints;
-	Header.TessellationMaxTessFactor            = CCHeader.TessellationMaxTessFactor;
-	Header.TessellationOutputWinding			= CCHeader.TessellationOutputWinding;
-	Header.TessellationPartitioning				= CCHeader.TessellationPartitioning;
-	Header.TessellationPatchesPerThreadGroup    = CCHeader.TessellationPatchesPerThreadGroup;
-	Header.TessellationPatchCountBuffer         = CCHeader.TessellationPatchCountBuffer;
-	Header.TessellationIndexBuffer              = CCHeader.TessellationIndexBuffer;
-	Header.TessellationHSOutBuffer              = CCHeader.TessellationHSOutBuffer;
-	Header.TessellationHSTFOutBuffer            = CCHeader.TessellationHSTFOutBuffer;
-	Header.TessellationControlPointOutBuffer    = CCHeader.TessellationControlPointOutBuffer;
-	Header.TessellationControlPointIndexBuffer  = CCHeader.TessellationControlPointIndexBuffer;
-	Header.TessellationOutputAttribs            = TessOutputAttribs;
+	if (ShaderInput.Target.Platform == SP_METAL_SM5 && (Frequency == SF_Vertex || Frequency == SF_Hull || Frequency == SF_Domain))
+	{
+		FMetalTessellationHeader TessHeader;
+		TessHeader.TessellationOutputControlPoints 		= CCHeader.TessellationOutputControlPoints;
+		TessHeader.TessellationDomain					= CCHeader.TessellationDomain;
+		TessHeader.TessellationInputControlPoints       = CCHeader.TessellationInputControlPoints;
+		TessHeader.TessellationMaxTessFactor            = CCHeader.TessellationMaxTessFactor;
+		TessHeader.TessellationOutputWinding			= CCHeader.TessellationOutputWinding;
+		TessHeader.TessellationPartitioning				= CCHeader.TessellationPartitioning;
+		TessHeader.TessellationPatchesPerThreadGroup    = CCHeader.TessellationPatchesPerThreadGroup;
+		TessHeader.TessellationPatchCountBuffer         = CCHeader.TessellationPatchCountBuffer;
+		TessHeader.TessellationIndexBuffer              = CCHeader.TessellationIndexBuffer;
+		TessHeader.TessellationHSOutBuffer              = CCHeader.TessellationHSOutBuffer;
+		TessHeader.TessellationHSTFOutBuffer            = CCHeader.TessellationHSTFOutBuffer;
+		TessHeader.TessellationControlPointOutBuffer    = CCHeader.TessellationControlPointOutBuffer;
+		TessHeader.TessellationControlPointIndexBuffer  = CCHeader.TessellationControlPointIndexBuffer;
+		TessHeader.TessellationOutputAttribs            = TessOutputAttribs;
+
+		Header.Tessellation.Add(TessHeader);
+		
+	}
 	Header.bDeviceFunctionConstants				= (FCStringAnsi::Strstr(USFSource, "#define __METAL_DEVICE_CONSTANT_INDEX__ 1") != nullptr);
 	Header.SideTable 							= CCHeader.SideTable;
 	Header.Bindings.ArgumentBufferMasks			= CCHeader.ArgumentBuffers;
diff --git a/Engine/Source/Developer/Apple/MetalShaderFormat/Private/MetalShaderFormat.cpp b/Engine/Source/Developer/Apple/MetalShaderFormat/Private/MetalShaderFormat.cpp
--- a/Engine/Source/Developer/Apple/MetalShaderFormat/Private/MetalShaderFormat.cpp
+++ b/Engine/Source/Developer/Apple/MetalShaderFormat/Private/MetalShaderFormat.cpp
@@ -239,7 +239,7 @@ class FMetalShaderFormat : public IShaderFormat
 public:
 	enum
 	{
-		HEADER_VERSION = 63,
+		HEADER_VERSION = 64,
 	};
 	
 	struct FVersion
diff --git a/Engine/Source/Runtime/Apple/MetalRHI/Private/MetalCommands.cpp b/Engine/Source/Runtime/Apple/MetalRHI/Private/MetalCommands.cpp
--- a/Engine/Source/Runtime/Apple/MetalRHI/Private/MetalCommands.cpp
+++ b/Engine/Source/Runtime/Apple/MetalRHI/Private/MetalCommands.cpp
@@ -244,7 +244,7 @@ void FMetalRHICommandContext::RHISetUAVParameter(FComputeShaderRHIParamRef Compu
 {
 	@autoreleasepool {
 	FMetalUnorderedAccessView* UAV = ResourceCast(UAVRHI);
-	Context->GetCurrentState().SetShaderUnorderedAccessView(SF_Compute, UAVIndex, UAV);
+	Context->GetCurrentState().SetShaderUnorderedAccessView(EMetalShaderStages::Compute, UAVIndex, UAV);
 	}
 }
 
@@ -252,7 +252,7 @@ void FMetalRHICommandContext::RHISetUAVParameter(FComputeShaderRHIParamRef Compu
 {
 	@autoreleasepool {
 		FMetalUnorderedAccessView* UAV = ResourceCast(UAVRHI);
-		Context->GetCurrentState().SetShaderUnorderedAccessView(SF_Compute, UAVIndex, UAV);
+		Context->GetCurrentState().SetShaderUnorderedAccessView(EMetalShaderStages::Compute, UAVIndex, UAV);
 	}
 }
 
@@ -265,65 +265,69 @@ void FMetalRHICommandContext::RHISetShaderTexture(FVertexShaderRHIParamRef Verte
 	{
         if (Surface->Texture || !(Surface->Flags & TexCreate_Presentable))
         {
-            Context->GetCurrentState().SetShaderTexture(SF_Vertex, Surface->Texture, TextureIndex, (mtlpp::ResourceUsage)(mtlpp::ResourceUsage::Read|mtlpp::ResourceUsage::Sample));
+            Context->GetCurrentState().SetShaderTexture(EMetalShaderStages::Vertex, Surface->Texture, TextureIndex, (mtlpp::ResourceUsage)(mtlpp::ResourceUsage::Read|mtlpp::ResourceUsage::Sample));
         }
         else
         {
             FMetalTexture Tex = Surface->GetCurrentTexture();
-            Context->GetCurrentState().SetShaderTexture(SF_Vertex, Tex, TextureIndex, (mtlpp::ResourceUsage)(mtlpp::ResourceUsage::Read|mtlpp::ResourceUsage::Sample));
+            Context->GetCurrentState().SetShaderTexture(EMetalShaderStages::Vertex, Tex, TextureIndex, (mtlpp::ResourceUsage)(mtlpp::ResourceUsage::Read|mtlpp::ResourceUsage::Sample));
         }
 	}
 	else
 	{
-		Context->GetCurrentState().SetShaderTexture(SF_Vertex, nil, TextureIndex, mtlpp::ResourceUsage(0));
+		Context->GetCurrentState().SetShaderTexture(EMetalShaderStages::Vertex, nil, TextureIndex, mtlpp::ResourceUsage(0));
 	}
 	}
 }
 
 void FMetalRHICommandContext::RHISetShaderTexture(FHullShaderRHIParamRef HullShader, uint32 TextureIndex, FTextureRHIParamRef NewTextureRHI)
 {
+#if PLATFORM_SUPPORTS_TESSELLATION_SHADERS
 	@autoreleasepool {
 	FMetalSurface* Surface = GetMetalSurfaceFromRHITexture(NewTextureRHI);
 	if (Surface != nullptr)
     {
         if (Surface->Texture || !(Surface->Flags & TexCreate_Presentable))
         {
-            Context->GetCurrentState().SetShaderTexture(SF_Hull, Surface->Texture, TextureIndex, (mtlpp::ResourceUsage)(mtlpp::ResourceUsage::Read|mtlpp::ResourceUsage::Sample));
+            Context->GetCurrentState().SetShaderTexture(EMetalShaderStages::Hull, Surface->Texture, TextureIndex, (mtlpp::ResourceUsage)(mtlpp::ResourceUsage::Read|mtlpp::ResourceUsage::Sample));
         }
         else
         {
             FMetalTexture Tex = Surface->GetCurrentTexture();
-            Context->GetCurrentState().SetShaderTexture(SF_Hull, Tex, TextureIndex, (mtlpp::ResourceUsage)(mtlpp::ResourceUsage::Read|mtlpp::ResourceUsage::Sample));
+            Context->GetCurrentState().SetShaderTexture(EMetalShaderStages::Hull, Tex, TextureIndex, (mtlpp::ResourceUsage)(mtlpp::ResourceUsage::Read|mtlpp::ResourceUsage::Sample));
         }
 	}
 	else
 	{
-		Context->GetCurrentState().SetShaderTexture(SF_Hull, nil, TextureIndex, mtlpp::ResourceUsage(0));
+		Context->GetCurrentState().SetShaderTexture(EMetalShaderStages::Hull, nil, TextureIndex, mtlpp::ResourceUsage(0));
 	}
 	}
+#endif
 }
 
 void FMetalRHICommandContext::RHISetShaderTexture(FDomainShaderRHIParamRef DomainShader, uint32 TextureIndex, FTextureRHIParamRef NewTextureRHI)
 {
+#if PLATFORM_SUPPORTS_TESSELLATION_SHADERS
 	@autoreleasepool {
 	FMetalSurface* Surface = GetMetalSurfaceFromRHITexture(NewTextureRHI);
 	if (Surface != nullptr)
     {
         if (Surface->Texture || !(Surface->Flags & TexCreate_Presentable))
         {
-            Context->GetCurrentState().SetShaderTexture(SF_Domain, Surface->Texture, TextureIndex, (mtlpp::ResourceUsage)(mtlpp::ResourceUsage::Read|mtlpp::ResourceUsage::Sample));
+            Context->GetCurrentState().SetShaderTexture(EMetalShaderStages::Domain, Surface->Texture, TextureIndex, (mtlpp::ResourceUsage)(mtlpp::ResourceUsage::Read|mtlpp::ResourceUsage::Sample));
         }
         else
         {
             FMetalTexture Tex = Surface->GetCurrentTexture();
-            Context->GetCurrentState().SetShaderTexture(SF_Domain, Tex, TextureIndex, (mtlpp::ResourceUsage)(mtlpp::ResourceUsage::Read|mtlpp::ResourceUsage::Sample));
+            Context->GetCurrentState().SetShaderTexture(EMetalShaderStages::Domain, Tex, TextureIndex, (mtlpp::ResourceUsage)(mtlpp::ResourceUsage::Read|mtlpp::ResourceUsage::Sample));
         }
 	}
 	else
 	{
-		Context->GetCurrentState().SetShaderTexture(SF_Domain, nil, TextureIndex, mtlpp::ResourceUsage(0));
+		Context->GetCurrentState().SetShaderTexture(EMetalShaderStages::Domain, nil, TextureIndex, mtlpp::ResourceUsage(0));
 	}
 	}
+#endif
 }
 
 void FMetalRHICommandContext::RHISetShaderTexture(FGeometryShaderRHIParamRef GeometryShader, uint32 TextureIndex, FTextureRHIParamRef NewTextureRHI)
@@ -340,17 +344,17 @@ void FMetalRHICommandContext::RHISetShaderTexture(FPixelShaderRHIParamRef PixelS
     {
         if (Surface->Texture || !(Surface->Flags & TexCreate_Presentable))
         {
-            Context->GetCurrentState().SetShaderTexture(SF_Pixel, Surface->Texture, TextureIndex, (mtlpp::ResourceUsage)(mtlpp::ResourceUsage::Read|mtlpp::ResourceUsage::Sample));
+            Context->GetCurrentState().SetShaderTexture(EMetalShaderStages::Pixel, Surface->Texture, TextureIndex, (mtlpp::ResourceUsage)(mtlpp::ResourceUsage::Read|mtlpp::ResourceUsage::Sample));
         }
         else
         {
             FMetalTexture Tex = Surface->GetCurrentTexture();
-            Context->GetCurrentState().SetShaderTexture(SF_Pixel, Tex, TextureIndex, (mtlpp::ResourceUsage)(mtlpp::ResourceUsage::Read|mtlpp::ResourceUsage::Sample));
+            Context->GetCurrentState().SetShaderTexture(EMetalShaderStages::Pixel, Tex, TextureIndex, (mtlpp::ResourceUsage)(mtlpp::ResourceUsage::Read|mtlpp::ResourceUsage::Sample));
         }
 	}
 	else
 	{
-		Context->GetCurrentState().SetShaderTexture(SF_Pixel, nil, TextureIndex, mtlpp::ResourceUsage(0));
+		Context->GetCurrentState().SetShaderTexture(EMetalShaderStages::Pixel, nil, TextureIndex, mtlpp::ResourceUsage(0));
 	}
 	}
 }
@@ -363,17 +367,17 @@ void FMetalRHICommandContext::RHISetShaderTexture(FComputeShaderRHIParamRef Comp
     {
         if (Surface->Texture || !(Surface->Flags & TexCreate_Presentable))
         {
-            Context->GetCurrentState().SetShaderTexture(SF_Compute, Surface->Texture, TextureIndex, (mtlpp::ResourceUsage)(mtlpp::ResourceUsage::Read|mtlpp::ResourceUsage::Sample));
+            Context->GetCurrentState().SetShaderTexture(EMetalShaderStages::Compute, Surface->Texture, TextureIndex, (mtlpp::ResourceUsage)(mtlpp::ResourceUsage::Read|mtlpp::ResourceUsage::Sample));
         }
         else
         {
             FMetalTexture Tex = Surface->GetCurrentTexture();
-            Context->GetCurrentState().SetShaderTexture(SF_Compute, Tex, TextureIndex, (mtlpp::ResourceUsage)(mtlpp::ResourceUsage::Read|mtlpp::ResourceUsage::Sample));
+            Context->GetCurrentState().SetShaderTexture(EMetalShaderStages::Compute, Tex, TextureIndex, (mtlpp::ResourceUsage)(mtlpp::ResourceUsage::Read|mtlpp::ResourceUsage::Sample));
         }
 	}
 	else
 	{
-		Context->GetCurrentState().SetShaderTexture(SF_Compute, nil, TextureIndex, mtlpp::ResourceUsage(0));
+		Context->GetCurrentState().SetShaderTexture(EMetalShaderStages::Compute, nil, TextureIndex, mtlpp::ResourceUsage(0));
 	}
 	}
 }
@@ -383,24 +387,28 @@ void FMetalRHICommandContext::RHISetShaderResourceViewParameter(FVertexShaderRHI
 {
 	@autoreleasepool {
 	FMetalShaderResourceView* SRV = ResourceCast(SRVRHI);
-	Context->GetCurrentState().SetShaderResourceView(Context, SF_Vertex, TextureIndex, SRV);
+	Context->GetCurrentState().SetShaderResourceView(Context, EMetalShaderStages::Vertex, TextureIndex, SRV);
 	}
 }
 
 void FMetalRHICommandContext::RHISetShaderResourceViewParameter(FHullShaderRHIParamRef HullShaderRHI,uint32 TextureIndex,FShaderResourceViewRHIParamRef SRVRHI)
 {
+#if PLATFORM_SUPPORTS_TESSELLATION_SHADERS
 	@autoreleasepool {
 	FMetalShaderResourceView* SRV = ResourceCast(SRVRHI);
-	Context->GetCurrentState().SetShaderResourceView(Context, SF_Hull, TextureIndex, SRV);
+	Context->GetCurrentState().SetShaderResourceView(Context, EMetalShaderStages::Hull, TextureIndex, SRV);
 	}
+#endif
 }
 
 void FMetalRHICommandContext::RHISetShaderResourceViewParameter(FDomainShaderRHIParamRef DomainShaderRHI,uint32 TextureIndex,FShaderResourceViewRHIParamRef SRVRHI)
 {
+#if PLATFORM_SUPPORTS_TESSELLATION_SHADERS
 	@autoreleasepool {
 	FMetalShaderResourceView* SRV = ResourceCast(SRVRHI);
-	Context->GetCurrentState().SetShaderResourceView(Context, SF_Domain, TextureIndex, SRV);
+	Context->GetCurrentState().SetShaderResourceView(Context, EMetalShaderStages::Domain, TextureIndex, SRV);
 	}
+#endif
 }
 
 void FMetalRHICommandContext::RHISetShaderResourceViewParameter(FGeometryShaderRHIParamRef GeometryShaderRHI,uint32 TextureIndex,FShaderResourceViewRHIParamRef SRVRHI)
@@ -412,7 +420,7 @@ void FMetalRHICommandContext::RHISetShaderResourceViewParameter(FPixelShaderRHIP
 {
 	@autoreleasepool {
 	FMetalShaderResourceView* SRV = ResourceCast(SRVRHI);
-	Context->GetCurrentState().SetShaderResourceView(Context, SF_Pixel, TextureIndex, SRV);
+	Context->GetCurrentState().SetShaderResourceView(Context, EMetalShaderStages::Pixel, TextureIndex, SRV);
 	}
 }
 
@@ -420,7 +428,7 @@ void FMetalRHICommandContext::RHISetShaderResourceViewParameter(FComputeShaderRH
 {
 	@autoreleasepool {
 	FMetalShaderResourceView* SRV = ResourceCast(SRVRHI);
-	Context->GetCurrentState().SetShaderResourceView(Context, SF_Compute, TextureIndex, SRV);
+	Context->GetCurrentState().SetShaderResourceView(Context, EMetalShaderStages::Compute, TextureIndex, SRV);
 	}
 }
 
@@ -430,26 +438,30 @@ void FMetalRHICommandContext::RHISetShaderSampler(FVertexShaderRHIParamRef Verte
 	@autoreleasepool {
 	FMetalSamplerState* NewState = ResourceCast(NewStateRHI);
 
-	Context->GetCurrentState().SetShaderSamplerState(SF_Vertex, NewState, SamplerIndex);
+	Context->GetCurrentState().SetShaderSamplerState(EMetalShaderStages::Vertex, NewState, SamplerIndex);
 	}
 }
 
 void FMetalRHICommandContext::RHISetShaderSampler(FHullShaderRHIParamRef HullShader, uint32 SamplerIndex, FSamplerStateRHIParamRef NewStateRHI)
 {
+#if PLATFORM_SUPPORTS_TESSELLATION_SHADERS
 	@autoreleasepool {
 	FMetalSamplerState* NewState = ResourceCast(NewStateRHI);
 
-	Context->GetCurrentState().SetShaderSamplerState(SF_Hull, NewState, SamplerIndex);
+	Context->GetCurrentState().SetShaderSamplerState(EMetalShaderStages::Hull, NewState, SamplerIndex);
 	}
+#endif
 }
 
 void FMetalRHICommandContext::RHISetShaderSampler(FDomainShaderRHIParamRef DomainShader, uint32 SamplerIndex, FSamplerStateRHIParamRef NewStateRHI)
 {
+#if PLATFORM_SUPPORTS_TESSELLATION_SHADERS
 	@autoreleasepool {
 	FMetalSamplerState* NewState = ResourceCast(NewStateRHI);
 
-	Context->GetCurrentState().SetShaderSamplerState(SF_Domain, NewState, SamplerIndex);
+	Context->GetCurrentState().SetShaderSamplerState(EMetalShaderStages::Domain, NewState, SamplerIndex);
 	}
+#endif
 }
 
 void FMetalRHICommandContext::RHISetShaderSampler(FGeometryShaderRHIParamRef GeometryShader, uint32 SamplerIndex, FSamplerStateRHIParamRef NewStateRHI)
@@ -462,7 +474,7 @@ void FMetalRHICommandContext::RHISetShaderSampler(FPixelShaderRHIParamRef PixelS
 	@autoreleasepool {
 	FMetalSamplerState* NewState = ResourceCast(NewStateRHI);
 
-	Context->GetCurrentState().SetShaderSamplerState(SF_Pixel, NewState, SamplerIndex);
+	Context->GetCurrentState().SetShaderSamplerState(EMetalShaderStages::Pixel, NewState, SamplerIndex);
 	}
 }
 
@@ -471,14 +483,14 @@ void FMetalRHICommandContext::RHISetShaderSampler(FComputeShaderRHIParamRef Comp
 	@autoreleasepool {
 	FMetalSamplerState* NewState = ResourceCast(NewStateRHI);
 
-	Context->GetCurrentState().SetShaderSamplerState(SF_Compute, NewState, SamplerIndex);
+	Context->GetCurrentState().SetShaderSamplerState(EMetalShaderStages::Compute, NewState, SamplerIndex);
 	}
 }
 
 void FMetalRHICommandContext::RHISetShaderParameter(FVertexShaderRHIParamRef VertexShaderRHI, uint32 BufferIndex, uint32 BaseIndex, uint32 NumBytes, const void* NewValue)
 {
 	@autoreleasepool {
-	Context->GetCurrentState().GetShaderParameters(CrossCompiler::SHADER_STAGE_VERTEX).Set(BufferIndex, BaseIndex, NumBytes, NewValue);
+	Context->GetCurrentState().GetShaderParameters(EMetalShaderStages::Vertex).Set(BufferIndex, BaseIndex, NumBytes, NewValue);
 	}
 }
 
@@ -491,15 +503,17 @@ void FMetalRHICommandContext::RHISetShaderParameter(FHullShaderRHIParamRef HullS
 void FMetalRHICommandContext::RHISetShaderParameter(FPixelShaderRHIParamRef PixelShaderRHI, uint32 BufferIndex, uint32 BaseIndex, uint32 NumBytes, const void* NewValue)
 {
 	@autoreleasepool {
-	Context->GetCurrentState().GetShaderParameters(CrossCompiler::SHADER_STAGE_PIXEL).Set(BufferIndex, BaseIndex, NumBytes, NewValue);
+	Context->GetCurrentState().GetShaderParameters(EMetalShaderStages::Pixel).Set(BufferIndex, BaseIndex, NumBytes, NewValue);
 	}
 }
 
 void FMetalRHICommandContext::RHISetShaderParameter(FDomainShaderRHIParamRef DomainShaderRHI, uint32 BufferIndex, uint32 BaseIndex, uint32 NumBytes, const void* NewValue)
 {
+#if PLATFORM_SUPPORTS_TESSELLATION_SHADERS
 	@autoreleasepool {
-	Context->GetCurrentState().GetShaderParameters(CrossCompiler::SHADER_STAGE_DOMAIN).Set(BufferIndex, BaseIndex, NumBytes, NewValue);
+	Context->GetCurrentState().GetShaderParameters(EMetalShaderStages::Domain).Set(BufferIndex, BaseIndex, NumBytes, NewValue);
 	}
+#endif
 }
 
 void FMetalRHICommandContext::RHISetShaderParameter(FGeometryShaderRHIParamRef GeometryShaderRHI, uint32 BufferIndex, uint32 BaseIndex, uint32 NumBytes, const void* NewValue)
@@ -510,7 +524,7 @@ void FMetalRHICommandContext::RHISetShaderParameter(FGeometryShaderRHIParamRef G
 void FMetalRHICommandContext::RHISetShaderParameter(FComputeShaderRHIParamRef ComputeShaderRHI,uint32 BufferIndex, uint32 BaseIndex, uint32 NumBytes, const void* NewValue)
 {
 	@autoreleasepool {
-	Context->GetCurrentState().GetShaderParameters(CrossCompiler::SHADER_STAGE_COMPUTE).Set(BufferIndex, BaseIndex, NumBytes, NewValue);
+	Context->GetCurrentState().GetShaderParameters(EMetalShaderStages::Compute).Set(BufferIndex, BaseIndex, NumBytes, NewValue);
 	}
 }
 
@@ -518,48 +532,52 @@ void FMetalRHICommandContext::RHISetShaderUniformBuffer(FVertexShaderRHIParamRef
 {
 	@autoreleasepool {
 	FMetalVertexShader* VertexShader = ResourceCast(VertexShaderRHI);
-	Context->GetCurrentState().BindUniformBuffer(SF_Vertex, BufferIndex, BufferRHI);
+	Context->GetCurrentState().BindUniformBuffer(EMetalShaderStages::Vertex, BufferIndex, BufferRHI);
 
 	auto& Bindings = VertexShader->Bindings;
 	check(BufferIndex < Bindings.NumUniformBuffers);
 	if (Bindings.ConstantBuffers & 1 << BufferIndex)
 	{
 		auto* UB = (FMetalUniformBuffer*)BufferRHI;
-		Context->GetCurrentState().SetShaderBuffer(SF_Vertex, (Bindings.ArgumentBuffers & 1 << BufferIndex) ? UB->GetIAB().IndirectArgumentBuffer : UB->Buffer, UB->Data, 0, UB->GetSize(), BufferIndex, mtlpp::ResourceUsage::Read);
+		Context->GetCurrentState().SetShaderBuffer(EMetalShaderStages::Vertex, (Bindings.ArgumentBuffers & 1 << BufferIndex) ? UB->GetIAB().IndirectArgumentBuffer : UB->Buffer, UB->Data, 0, UB->GetSize(), BufferIndex, mtlpp::ResourceUsage::Read);
 	}
 	}
 }
 
 void FMetalRHICommandContext::RHISetShaderUniformBuffer(FHullShaderRHIParamRef HullShaderRHI, uint32 BufferIndex, FUniformBufferRHIParamRef BufferRHI)
 {
+#if PLATFORM_SUPPORTS_TESSELLATION_SHADERS
 	@autoreleasepool {
 	FMetalHullShader* HullShader = ResourceCast(HullShaderRHI);
-	Context->GetCurrentState().BindUniformBuffer(SF_Hull, BufferIndex, BufferRHI);
+	Context->GetCurrentState().BindUniformBuffer(EMetalShaderStages::Hull, BufferIndex, BufferRHI);
 
 	auto& Bindings = HullShader->Bindings;
 	check(BufferIndex < Bindings.NumUniformBuffers);
 	if (Bindings.ConstantBuffers & 1 << BufferIndex)
 	{
 		auto* UB = (FMetalUniformBuffer*)BufferRHI;
-		Context->GetCurrentState().SetShaderBuffer(SF_Hull, (Bindings.ArgumentBuffers & 1 << BufferIndex) ? UB->GetIAB().IndirectArgumentBuffer : UB->Buffer, UB->Data, 0, UB->GetSize(), BufferIndex, mtlpp::ResourceUsage::Read);
+		Context->GetCurrentState().SetShaderBuffer(EMetalShaderStages::Hull, (Bindings.ArgumentBuffers & 1 << BufferIndex) ? UB->GetIAB().IndirectArgumentBuffer : UB->Buffer, UB->Data, 0, UB->GetSize(), BufferIndex, mtlpp::ResourceUsage::Read);
 	}
 	}
+#endif
 }
 
 void FMetalRHICommandContext::RHISetShaderUniformBuffer(FDomainShaderRHIParamRef DomainShaderRHI, uint32 BufferIndex, FUniformBufferRHIParamRef BufferRHI)
 {
+#if PLATFORM_SUPPORTS_TESSELLATION_SHADERS
 	@autoreleasepool {
 	FMetalDomainShader* DomainShader = ResourceCast(DomainShaderRHI);
-	Context->GetCurrentState().BindUniformBuffer(SF_Domain, BufferIndex, BufferRHI);
+	Context->GetCurrentState().BindUniformBuffer(EMetalShaderStages::Domain, BufferIndex, BufferRHI);
 
 	auto& Bindings = DomainShader->Bindings;
 	check(BufferIndex < Bindings.NumUniformBuffers);
 	if (Bindings.ConstantBuffers & 1 << BufferIndex)
 	{
 		auto* UB = (FMetalUniformBuffer*)BufferRHI;
-		Context->GetCurrentState().SetShaderBuffer(SF_Domain, (Bindings.ArgumentBuffers & 1 << BufferIndex) ? UB->GetIAB().IndirectArgumentBuffer : UB->Buffer, UB->Data, 0, UB->GetSize(), BufferIndex, mtlpp::ResourceUsage::Read);
+		Context->GetCurrentState().SetShaderBuffer(EMetalShaderStages::Domain, (Bindings.ArgumentBuffers & 1 << BufferIndex) ? UB->GetIAB().IndirectArgumentBuffer : UB->Buffer, UB->Data, 0, UB->GetSize(), BufferIndex, mtlpp::ResourceUsage::Read);
 	}
 	}
+#endif
 }
 
 void FMetalRHICommandContext::RHISetShaderUniformBuffer(FGeometryShaderRHIParamRef GeometryShader, uint32 BufferIndex, FUniformBufferRHIParamRef BufferRHI)
@@ -571,14 +589,14 @@ void FMetalRHICommandContext::RHISetShaderUniformBuffer(FPixelShaderRHIParamRef
 {
 	@autoreleasepool {
 	FMetalPixelShader* PixelShader = ResourceCast(PixelShaderRHI);
-	Context->GetCurrentState().BindUniformBuffer(SF_Pixel, BufferIndex, BufferRHI);
+	Context->GetCurrentState().BindUniformBuffer(EMetalShaderStages::Pixel, BufferIndex, BufferRHI);
 
 	auto& Bindings = PixelShader->Bindings;
 	check(BufferIndex < Bindings.NumUniformBuffers);
 	if (Bindings.ConstantBuffers & 1 << BufferIndex)
 	{
 		auto* UB = (FMetalUniformBuffer*)BufferRHI;
-		Context->GetCurrentState().SetShaderBuffer(SF_Pixel, (Bindings.ArgumentBuffers & 1 << BufferIndex) ? UB->GetIAB().IndirectArgumentBuffer : UB->Buffer, UB->Data, 0, UB->GetSize(), BufferIndex, mtlpp::ResourceUsage::Read);
+		Context->GetCurrentState().SetShaderBuffer(EMetalShaderStages::Pixel, (Bindings.ArgumentBuffers & 1 << BufferIndex) ? UB->GetIAB().IndirectArgumentBuffer : UB->Buffer, UB->Data, 0, UB->GetSize(), BufferIndex, mtlpp::ResourceUsage::Read);
 	}
 	}
 }
@@ -587,14 +605,14 @@ void FMetalRHICommandContext::RHISetShaderUniformBuffer(FComputeShaderRHIParamRe
 {
 	@autoreleasepool {
 	FMetalComputeShader* ComputeShader = ResourceCast(ComputeShaderRHI);
-	Context->GetCurrentState().BindUniformBuffer(SF_Compute, BufferIndex, BufferRHI);
+	Context->GetCurrentState().BindUniformBuffer(EMetalShaderStages::Compute, BufferIndex, BufferRHI);
 
 	auto& Bindings = ComputeShader->Bindings;
 	check(BufferIndex < Bindings.NumUniformBuffers);
 	if (Bindings.ConstantBuffers & 1 << BufferIndex)
 	{
 		auto* UB = (FMetalUniformBuffer*)BufferRHI;
-		Context->GetCurrentState().SetShaderBuffer(SF_Compute, (Bindings.ArgumentBuffers & 1 << BufferIndex) ? UB->GetIAB().IndirectArgumentBuffer : UB->Buffer, UB->Data, 0, UB->GetSize(), BufferIndex, mtlpp::ResourceUsage::Read);
+		Context->GetCurrentState().SetShaderBuffer(EMetalShaderStages::Compute, (Bindings.ArgumentBuffers & 1 << BufferIndex) ? UB->GetIAB().IndirectArgumentBuffer : UB->Buffer, UB->Data, 0, UB->GetSize(), BufferIndex, mtlpp::ResourceUsage::Read);
 	}
 	}
 }
@@ -821,11 +839,13 @@ void FMetalRHICommandContext::RHIEndDrawPrimitiveUP()
 	// set the vertex buffer
 	uint32 NumVertices = GetVertexCountForPrimitiveCount(PendingNumPrimitives, PendingPrimitiveType);
 	Context->GetCurrentState().SetVertexStream(0, &PendingVertexBuffer, nil, 0, PendingVertexDataStride * NumVertices);
+#if PLATFORM_SUPPORTS_TESSELLATION_SHADERS
 	if(Context->GetCurrentState().GetUsingTessellation())
 	{
-		Context->GetCurrentState().SetShaderBuffer(SF_Hull, PendingVertexBuffer, nil, 0, PendingVertexDataStride * NumVertices, UNREAL_TO_METAL_BUFFER_INDEX(0), mtlpp::ResourceUsage::Read);
-		Context->GetCurrentState().SetShaderBuffer(SF_Domain, PendingVertexBuffer, nil, 0, PendingVertexDataStride * NumVertices, UNREAL_TO_METAL_BUFFER_INDEX(0), mtlpp::ResourceUsage::Read);
+		Context->GetCurrentState().SetShaderBuffer(EMetalShaderStages::Hull, PendingVertexBuffer, nil, 0, PendingVertexDataStride * NumVertices, UNREAL_TO_METAL_BUFFER_INDEX(0), mtlpp::ResourceUsage::Read);
+		Context->GetCurrentState().SetShaderBuffer(EMetalShaderStages::Domain, PendingVertexBuffer, nil, 0, PendingVertexDataStride * NumVertices, UNREAL_TO_METAL_BUFFER_INDEX(0), mtlpp::ResourceUsage::Read);
 	}
+#endif
 	
 	Context->DrawPrimitive(PendingPrimitiveType, 0, PendingNumPrimitives, Context->GetCurrentState().GetRenderTargetArraySize());
 	
diff --git a/Engine/Source/Runtime/Apple/MetalRHI/Private/MetalPipeline.cpp b/Engine/Source/Runtime/Apple/MetalRHI/Private/MetalPipeline.cpp
--- a/Engine/Source/Runtime/Apple/MetalRHI/Private/MetalPipeline.cpp
+++ b/Engine/Source/Runtime/Apple/MetalRHI/Private/MetalPipeline.cpp
@@ -185,15 +185,17 @@ struct FMetalGraphicsPipelineKey
 		Key.VertexDescriptorHash = VertexDecl->Layout;
 		
 		FMetalVertexShader* VertexShader = (FMetalVertexShader*)Init.BoundShaderState.VertexShaderRHI;
+		Key.VertexFunction = VertexShader->GetHash();
+
+#if PLATFORM_SUPPORTS_TESSELLATION_SHADERS
 		FMetalDomainShader* DomainShader = (FMetalDomainShader*)Init.BoundShaderState.DomainShaderRHI;
-		
-        Key.VertexFunction = VertexShader->GetHash();
 		if (DomainShader)
 		{
 			Key.DomainFunction = DomainShader->GetHash();
 			Key.SetHashValue(Offset_IndexType, NumBits_IndexType, IndexType);
 		}
 		else
+#endif
 		{
 			Key.SetHashValue(Offset_IndexType, NumBits_IndexType, EMetalIndexType_None);
 		}
@@ -439,15 +441,22 @@ static MTLVertexDescriptor* GetMaskedVertexDescriptor(MTLVertexDescriptor* Input
 static FMetalShaderPipeline* CreateMTLRenderPipeline(bool const bSync, FMetalGraphicsPipelineKey const& Key, const FGraphicsPipelineStateInitializer& Init, EMetalIndexType const IndexType)
 {
     FMetalVertexShader* VertexShader = (FMetalVertexShader*)Init.BoundShaderState.VertexShaderRHI;
-    FMetalDomainShader* DomainShader = (FMetalDomainShader*)Init.BoundShaderState.DomainShaderRHI;
     FMetalPixelShader* PixelShader = (FMetalPixelShader*)Init.BoundShaderState.PixelShaderRHI;
     
     mtlpp::Function vertexFunction = VertexShader->GetFunction();
     mtlpp::Function fragmentFunction = PixelShader ? PixelShader->GetFunction() : nil;
+
+#if PLATFORM_SUPPORTS_TESSELLATION_SHADERS
+	FMetalDomainShader* DomainShader = (FMetalDomainShader*)Init.BoundShaderState.DomainShaderRHI;
     mtlpp::Function domainFunction = DomainShader ? DomainShader->GetFunction() : nil;
+#endif
     
     FMetalShaderPipeline* Pipeline = nil;
-    if (vertexFunction && ((PixelShader != nullptr) == (fragmentFunction != nil)) && ((DomainShader != nullptr) == (domainFunction != nil)))
+    if (vertexFunction && ((PixelShader != nullptr) == (fragmentFunction != nil))
+#if PLATFORM_SUPPORTS_TESSELLATION_SHADERS
+		&& ((DomainShader != nullptr) == (domainFunction != nil))
+#endif
+	)
     {
 		ns::Error Error;
 		mtlpp::Device Device = GetMetalDeviceContext().GetDevice();
@@ -603,10 +612,10 @@ static FMetalShaderPipeline* CreateMTLRenderPipeline(bool const bSync, FMetalGra
         }
         
         check(Init.BoundShaderState.VertexShaderRHI != nullptr);
+#if PLATFORM_SUPPORTS_GEOMETRY_SHADERS
         check(Init.BoundShaderState.GeometryShaderRHI == nullptr);
+#endif
 
-        FMetalHullShader* HullShader = (FMetalHullShader*)Init.BoundShaderState.HullShaderRHI;
-        
         if(RenderPipelineDesc.GetDepthAttachmentPixelFormat() == mtlpp::PixelFormat::Invalid && PixelShader && ((PixelShader->Bindings.InOutMask & 0x8000) || (NumActiveTargets == 0 && (PixelShader->Bindings.NumUAVs > 0))))
         {
             RenderPipelineDesc.SetDepthAttachmentPixelFormat((mtlpp::PixelFormat)GPixelFormats[PF_DepthStencil].PlatformFormat);
@@ -619,10 +628,15 @@ static FMetalShaderPipeline* CreateMTLRenderPipeline(bool const bSync, FMetalGra
     #endif
         
         FMetalVertexDeclaration* VertexDecl = (FMetalVertexDeclaration*)Init.BoundShaderState.VertexDeclarationRHI;
-        
-        if (Init.BoundShaderState.HullShaderRHI == nullptr)
+		
+#if PLATFORM_SUPPORTS_TESSELLATION_SHADERS
+		FMetalHullShader* HullShader = (FMetalHullShader*)Init.BoundShaderState.HullShaderRHI;
+		if (Init.BoundShaderState.HullShaderRHI == nullptr)
+#endif
         {
+#if PLATFORM_SUPPORTS_TESSELLATION_SHADERS
             check(Init.BoundShaderState.DomainShaderRHI == nullptr);
+#endif
             RenderPipelineDesc.SetVertexDescriptor(GetMaskedVertexDescriptor(VertexDecl->Layout.VertexDesc, VertexShader->Bindings.InOutMask));
             RenderPipelineDesc.SetVertexFunction(vertexFunction);
             RenderPipelineDesc.SetFragmentFunction(fragmentFunction);
@@ -632,6 +646,7 @@ static FMetalShaderPipeline* CreateMTLRenderPipeline(bool const bSync, FMetalGra
 			RenderPipelineDesc.SetLabel([NSString stringWithFormat:@"%@+%@", VertexName.GetPtr(), FragmentName.GetPtr()]);
 #endif
         }
+#if PLATFORM_SUPPORTS_TESSELLATION_SHADERS
         else
         {
             check(Init.BoundShaderState.DomainShaderRHI != nullptr);
@@ -947,6 +962,7 @@ static FMetalShaderPipeline* CreateMTLRenderPipeline(bool const bSync, FMetalGra
 #endif
             }
         }
+#endif
         
         NSUInteger RenderOption = mtlpp::PipelineOption::NoPipelineOption;
 		mtlpp::AutoReleasedRenderPipelineReflection* Reflection = nullptr;
@@ -976,8 +992,10 @@ static FMetalShaderPipeline* CreateMTLRenderPipeline(bool const bSync, FMetalGra
 		UE_CLOG((Pipeline->RenderPipelineState == nil), LogMetal, Error, TEXT("Failed to generate a pipeline state object: %s"), *FString(Error.GetPtr().description));
 		UE_CLOG((Pipeline->RenderPipelineState == nil), LogMetal, Error, TEXT("Vertex shader: %s"), *FString(VertexShader->GetSourceCode()));
 		UE_CLOG((Pipeline->RenderPipelineState == nil), LogMetal, Error, TEXT("Pixel shader: %s"), PixelShader ? *FString(PixelShader->GetSourceCode()) : TEXT("NULL"));
+#if PLATFORM_SUPPORTS_TESSELLATION_SHADERS
 		UE_CLOG((Pipeline->RenderPipelineState == nil), LogMetal, Error, TEXT("Hull shader: %s"), HullShader ? *FString(HullShader->GetSourceCode()) : TEXT("NULL"));
 		UE_CLOG((Pipeline->RenderPipelineState == nil), LogMetal, Error, TEXT("Domain shader: %s"), DomainShader ? *FString(DomainShader->GetSourceCode()) : TEXT("NULL"));
+#endif
 		UE_CLOG((Pipeline->RenderPipelineState == nil), LogMetal, Error, TEXT("Descriptor: %s"), *FString(RenderPipelineDesc.GetPtr().description));
 		UE_CLOG((Pipeline->RenderPipelineState == nil), LogMetal, Error, TEXT("Failed to generate a render pipeline state object:\n\n %s\n\n"), *FString(Error.GetLocalizedDescription()));
 		
@@ -989,34 +1007,15 @@ static FMetalShaderPipeline* CreateMTLRenderPipeline(bool const bSync, FMetalGra
 		}
 		
     #if METAL_DEBUG_OPTIONS
-        Pipeline->ComputeSource = DomainShader ? VertexShader->GetSourceCode() : nil;
+	#if PLATFORM_SUPPORTS_TESSELLATION_SHADERS
+		Pipeline->ComputeSource = DomainShader ? VertexShader->GetSourceCode() : nil;
         Pipeline->VertexSource = DomainShader ? DomainShader->GetSourceCode() : VertexShader->GetSourceCode();
+	#else
+		Pipeline->VertexSource = VertexShader->GetSourceCode();
+	#endif
         Pipeline->FragmentSource = PixelShader ? PixelShader->GetSourceCode() : nil;
     #endif
         
-    #if 0
-        if (GFrameCounter > 3)
-        {
-            NSLog(@"===============================================================");
-            NSLog(@"Creating a BSS at runtime frame %lld... this may hitch! [this = %p]", GFrameCounter, Pipeline);
-            NSLog(@"Vertex declaration:");
-            FVertexDeclarationElementList& Elements = VertexDecl->Elements;
-            for (int32 i = 0; i < Elements.Num(); i++)
-            {
-                FVertexElement& Elem = Elements[i];
-                NSLog(@"   Elem %d: attr: %d, stream: %d, type: %d, stride: %d, offset: %d", i, Elem.AttributeIndex, Elem.StreamIndex, (uint32)Elem.Type, Elem.Stride, Elem.Offset);
-            }
-            
-            NSLog(@"\nVertexShader:");
-            NSLog(@"%@", VertexShader ? VertexShader->GetSourceCode() : @"NONE");
-            NSLog(@"\nPixelShader:");
-            NSLog(@"%@", PixelShader ? PixelShader->GetSourceCode() : @"NONE");
-            NSLog(@"\nDomainShader:");
-            NSLog(@"%@", DomainShader ? DomainShader->GetSourceCode() : @"NONE");
-            NSLog(@"===============================================================");
-        }
-    #endif
-        
 #if METAL_DEBUG_OPTIONS
         if (GFrameCounter > 3)
         {
@@ -1092,9 +1091,13 @@ FGraphicsPipelineStateRHIRef FMetalDynamicRHI::RHICreateGraphicsPipelineState(co
 	State->VertexDeclaration = ResourceCast(Initializer.BoundShaderState.VertexDeclarationRHI);
 	State->VertexShader = ResourceCast(Initializer.BoundShaderState.VertexShaderRHI);
 	State->PixelShader = ResourceCast(Initializer.BoundShaderState.PixelShaderRHI);
+#if PLATFORM_SUPPORTS_TESSELLATION_SHADERS
 	State->HullShader = ResourceCast(Initializer.BoundShaderState.HullShaderRHI);
 	State->DomainShader = ResourceCast(Initializer.BoundShaderState.DomainShaderRHI);
+#endif
+#if PLATFORM_SUPPORTS_GEOMETRY_SHADERS
 	State->GeometryShader = ResourceCast(Initializer.BoundShaderState.GeometryShaderRHI);
+#endif
 	State->DepthStencilState = ResourceCast(Initializer.DepthStencilState);
 	State->RasterizerState = ResourceCast(Initializer.RasterizerState);
 	return State;
diff --git a/Engine/Source/Runtime/Apple/MetalRHI/Private/MetalRHIPrivate.h b/Engine/Source/Runtime/Apple/MetalRHI/Private/MetalRHIPrivate.h
--- a/Engine/Source/Runtime/Apple/MetalRHI/Private/MetalRHIPrivate.h
+++ b/Engine/Source/Runtime/Apple/MetalRHI/Private/MetalRHIPrivate.h
@@ -231,5 +231,72 @@ extern int32 GMetalBufferZeroFill;
 
 mtlpp::LanguageVersion ValidateVersion(uint8 Version);
 
+// Needs to be the same as EShaderFrequency when all stages are supported, but unlike EShaderFrequency you can compile out stages.
+enum EMetalShaderStages
+{
+	Vertex,
+#if PLATFORM_SUPPORTS_TESSELLATION_SHADERS
+	Hull,
+	Domain,
+#endif
+	Pixel,
+#if PLATFORM_SUPPORTS_GEOMETRY_SHADERS
+	Geometry,
+#endif
+	Compute,
+	
+	Num,
+};
+
+FORCEINLINE EShaderFrequency GetRHIShaderFrequency(EMetalShaderStages Stage)
+{
+	switch (Stage)
+	{
+		case EMetalShaderStages::Vertex:
+			return SF_Vertex;
+#if PLATFORM_SUPPORTS_TESSELLATION_SHADERS
+		case EMetalShaderStages::Hull:
+			return SF_Hull;
+		case EMetalShaderStages::Domain:
+			return SF_Domain;
+#endif
+		case EMetalShaderStages::Pixel:
+			return SF_Pixel;
+#if PLATFORM_SUPPORTS_GEOMETRY_SHADERS
+		case EMetalShaderStages::Geometry:
+			return SF_Geometry;
+#endif
+		case EMetalShaderStages::Compute:
+			return SF_Compute;
+		default:
+			return SF_NumFrequencies;
+	}
+}
+
+FORCEINLINE EMetalShaderStages GetMetalShaderFrequency(EShaderFrequency Stage)
+{
+	switch (Stage)
+	{
+		case SF_Vertex:
+			return EMetalShaderStages::Vertex;
+#if PLATFORM_SUPPORTS_TESSELLATION_SHADERS
+		case SF_Hull:
+			return EMetalShaderStages::Hull;
+		case SF_Domain:
+			return EMetalShaderStages::Domain;
+#endif
+		case SF_Pixel:
+			return EMetalShaderStages::Pixel;
+#if PLATFORM_SUPPORTS_GEOMETRY_SHADERS
+		case SF_Geometry:
+			return EMetalShaderStages::Geometry;
+#endif
+		case SF_Compute:
+			return EMetalShaderStages::Compute;
+		default:
+			return EMetalShaderStages::Num;
+	}
+}
+
 #include "MetalStateCache.h"
 #include "MetalContext.h"
diff --git a/Engine/Source/Runtime/Apple/MetalRHI/Private/MetalRenderPass.cpp b/Engine/Source/Runtime/Apple/MetalRHI/Private/MetalRenderPass.cpp
--- a/Engine/Source/Runtime/Apple/MetalRHI/Private/MetalRenderPass.cpp
+++ b/Engine/Source/Runtime/Apple/MetalRHI/Private/MetalRenderPass.cpp
@@ -389,7 +389,9 @@ void FMetalRenderPass::DrawPrimitive(uint32 PrimitiveType, uint32 BaseVertexInde
 {
 	NumInstances = FMath::Max(NumInstances,1u);
 	
+#if PLATFORM_SUPPORTS_TESSELLATION_SHADERS
 	if(!State.GetUsingTessellation())
+#endif
 	{
 		ConditionalSwitchToRender();
 		check(CurrentEncoder.GetCommandBuffer());
@@ -405,10 +407,12 @@ void FMetalRenderPass::DrawPrimitive(uint32 PrimitiveType, uint32 BaseVertexInde
 		CurrentEncoder.GetRenderCommandEncoder().Draw(TranslatePrimitiveType(PrimitiveType), BaseVertexIndex, NumVertices, NumInstances);
 		METAL_DEBUG_LAYER(EMetalDebugLevelFastValidation, CurrentEncoder.GetRenderCommandEncoderDebugging().Draw(TranslatePrimitiveType(PrimitiveType), BaseVertexIndex, NumVertices, NumInstances));
 	}
+#if PLATFORM_SUPPORTS_TESSELLATION_SHADERS
 	else
 	{
 		DrawPatches(PrimitiveType, nullptr, 0, BaseVertexIndex, 0, 0, NumPrimitives, NumInstances);
 	}
+#endif
 	
 	ConditionalSubmit();	
 }
@@ -485,7 +489,9 @@ void FMetalRenderPass::DrawIndexedPrimitive(FMetalBuffer const& IndexBuffer, uin
 	}
 #endif
 	
+#if PLATFORM_SUPPORTS_TESSELLATION_SHADERS
 	if (!State.GetUsingTessellation())
+#endif
 	{
 		ConditionalSwitchToRender();
 		check(CurrentEncoder.GetCommandBuffer());
@@ -507,10 +513,12 @@ void FMetalRenderPass::DrawIndexedPrimitive(FMetalBuffer const& IndexBuffer, uin
 			METAL_DEBUG_LAYER(EMetalDebugLevelFastValidation, CurrentEncoder.GetRenderCommandEncoderDebugging().DrawIndexed(TranslatePrimitiveType(PrimitiveType), NumIndices, ((IndexStride == 2) ? mtlpp::IndexType::UInt16 : mtlpp::IndexType::UInt32), IndexBuffer, StartIndex * IndexStride, NumInstances));
 		}
 	}
+#if PLATFORM_SUPPORTS_TESSELLATION_SHADERS
 	else
 	{
 		DrawPatches(PrimitiveType, IndexBuffer, IndexStride, BaseVertexIndex, FirstInstance, StartIndex, NumPrimitives, NumInstances);
 	}
+#endif
 	
 	ConditionalSubmit();
 }
@@ -562,6 +570,7 @@ void FMetalRenderPass::DrawIndexedPrimitiveIndirect(uint32 PrimitiveType,FMetalI
 	}
 }
 
+#if PLATFORM_SUPPORTS_TESSELLATION_SHADERS
 void FMetalRenderPass::DrawPatches(uint32 PrimitiveType,FMetalBuffer const& IndexBuffer, uint32 IndexBufferStride, int32 BaseVertexIndex, uint32 FirstInstance, uint32 StartIndex,
 									uint32 NumPrimitives, uint32 NumInstances)
 {
@@ -616,7 +625,7 @@ void FMetalRenderPass::DrawPatches(uint32 PrimitiveType,FMetalBuffer const& Inde
 		{
 			PrologueEncoder.SetShaderBuffer(mtlpp::FunctionType::Kernel, IndexBuffer, StartIndex * IndexBufferStride, IndexBuffer.GetLength() - (StartIndex * IndexBufferStride), Pipeline->TessellationPipelineDesc.TessellationControlPointIndexBufferIndex, mtlpp::ResourceUsage::Read);
 			
-			State.SetShaderBuffer(SF_Vertex, nil, nil, 0, 0, Pipeline->TessellationPipelineDesc.TessellationControlPointIndexBufferIndex, mtlpp::ResourceUsage(0));
+			State.SetShaderBuffer(EMetalShaderStages::Vertex, nil, nil, 0, 0, Pipeline->TessellationPipelineDesc.TessellationControlPointIndexBufferIndex, mtlpp::ResourceUsage(0));
 		}
 		
 		if (Pipeline->TessellationPipelineDesc.TessellationIndexBufferIndex != UINT_MAX)
@@ -629,59 +638,59 @@ void FMetalRenderPass::DrawPatches(uint32 PrimitiveType,FMetalBuffer const& Inde
 			{
 				PrologueEncoder.SetShaderBuffer(mtlpp::FunctionType::Kernel, nil, 0, 0, Pipeline->TessellationPipelineDesc.TessellationIndexBufferIndex, mtlpp::ResourceUsage::Read, (EPixelFormat)0);
 			}
-			State.SetShaderBuffer(SF_Vertex, nil, nil, 0, 0, Pipeline->TessellationPipelineDesc.TessellationControlPointIndexBufferIndex, mtlpp::ResourceUsage(0));
+			State.SetShaderBuffer(EMetalShaderStages::Vertex, nil, nil, 0, 0, Pipeline->TessellationPipelineDesc.TessellationControlPointIndexBufferIndex, mtlpp::ResourceUsage(0));
 		}
 		
 		if(Pipeline->TessellationPipelineDesc.TessellationOutputControlPointBufferIndex != UINT_MAX) //TessellationOutputControlPointBufferIndex -> hullShaderOutputBuffer
 		{
 			PrologueEncoder.SetShaderBuffer(mtlpp::FunctionType::Kernel, hullShaderOutputBuffer, hullShaderOutputOffset, hullShaderOutputBuffer.GetLength() - hullShaderOutputOffset, Pipeline->TessellationPipelineDesc.TessellationOutputControlPointBufferIndex, mtlpp::ResourceUsage(0));
-			State.SetShaderBuffer(SF_Vertex, nil, nil, 0, 0, Pipeline->TessellationPipelineDesc.TessellationOutputControlPointBufferIndex, mtlpp::ResourceUsage(0));
+			State.SetShaderBuffer(EMetalShaderStages::Vertex, nil, nil, 0, 0, Pipeline->TessellationPipelineDesc.TessellationOutputControlPointBufferIndex, mtlpp::ResourceUsage(0));
 		}
 		
 		if(Pipeline->TessellationPipelineDesc.TessellationPatchConstBufferIndex != UINT_MAX) //TessellationPatchConstBufferIndex -> hullConstShaderOutputBuffer
 		{
 			PrologueEncoder.SetShaderBuffer(mtlpp::FunctionType::Kernel, hullConstShaderOutputBuffer, hullConstShaderOutputOffset, hullConstShaderOutputBuffer.GetLength() - hullConstShaderOutputOffset, Pipeline->TessellationPipelineDesc.TessellationPatchConstBufferIndex, mtlpp::ResourceUsage(0));
-			State.SetShaderBuffer(SF_Vertex, nil, nil, 0, 0, Pipeline->TessellationPipelineDesc.TessellationPatchConstBufferIndex, mtlpp::ResourceUsage(0));
+			State.SetShaderBuffer(EMetalShaderStages::Vertex, nil, nil, 0, 0, Pipeline->TessellationPipelineDesc.TessellationPatchConstBufferIndex, mtlpp::ResourceUsage(0));
 		}
 		
 		if(Pipeline->TessellationPipelineDesc.TessellationFactorBufferIndex != UINT_MAX) // TessellationFactorBufferIndex->tessellationFactorBuffer
 		{
 			PrologueEncoder.SetShaderBuffer(mtlpp::FunctionType::Kernel, tessellationFactorBuffer, tessellationFactorsOffset, tessellationFactorBuffer.GetLength() - tessellationFactorsOffset, Pipeline->TessellationPipelineDesc.TessellationFactorBufferIndex, mtlpp::ResourceUsage(0));
-			State.SetShaderBuffer(SF_Vertex, nil, nil, 0, 0, Pipeline->TessellationPipelineDesc.TessellationFactorBufferIndex, mtlpp::ResourceUsage(0));
+			State.SetShaderBuffer(EMetalShaderStages::Vertex, nil, nil, 0, 0, Pipeline->TessellationPipelineDesc.TessellationFactorBufferIndex, mtlpp::ResourceUsage(0));
 		}
 		
 		if(Pipeline->TessellationPipelineDesc.TessellationInputControlPointBufferIndex != UINT_MAX) //TessellationInputControlPointBufferIndex->hullShaderOutputBuffer
 		{
 			CurrentEncoder.SetShaderBuffer(mtlpp::FunctionType::Vertex, hullShaderOutputBuffer, hullShaderOutputOffset, hullShaderOutputBuffer.GetLength() - hullShaderOutputOffset, Pipeline->TessellationPipelineDesc.TessellationInputControlPointBufferIndex, mtlpp::ResourceUsage(0));
-			State.SetShaderBuffer(SF_Domain, nil, nil, 0, 0, Pipeline->TessellationPipelineDesc.TessellationInputControlPointBufferIndex, mtlpp::ResourceUsage(0));
+			State.SetShaderBuffer(EMetalShaderStages::Domain, nil, nil, 0, 0, Pipeline->TessellationPipelineDesc.TessellationInputControlPointBufferIndex, mtlpp::ResourceUsage(0));
 		}
 		if(Pipeline->TessellationPipelineDesc.TessellationInputPatchConstBufferIndex != UINT_MAX) //TessellationInputPatchConstBufferIndex->hullConstShaderOutputBuffer
 		{
 			CurrentEncoder.SetShaderBuffer(mtlpp::FunctionType::Vertex, hullConstShaderOutputBuffer, hullConstShaderOutputOffset, hullConstShaderOutputBuffer.GetLength() - hullConstShaderOutputOffset, Pipeline->TessellationPipelineDesc.TessellationInputPatchConstBufferIndex, mtlpp::ResourceUsage(0));
-			State.SetShaderBuffer(SF_Domain, nil, nil, 0, 0, Pipeline->TessellationPipelineDesc.TessellationInputPatchConstBufferIndex, mtlpp::ResourceUsage(0));
+			State.SetShaderBuffer(EMetalShaderStages::Domain, nil, nil, 0, 0, Pipeline->TessellationPipelineDesc.TessellationInputPatchConstBufferIndex, mtlpp::ResourceUsage(0));
 		}
 		
 		// set the patchCount
 		uint32 patchCountData[] = { NumPrimitives, StartIndex };
 		PrologueEncoder.SetShaderBytes(mtlpp::FunctionType::Kernel, (const uint8*)&patchCountData[0], sizeof(patchCountData), Pipeline->TessellationPipelineDesc.TessellationPatchCountBufferIndex);
-		State.SetShaderBuffer(SF_Vertex, nil, nil, 0, 0, Pipeline->TessellationPipelineDesc.TessellationPatchCountBufferIndex, mtlpp::ResourceUsage(0));
+		State.SetShaderBuffer(EMetalShaderStages::Vertex, nil, nil, 0, 0, Pipeline->TessellationPipelineDesc.TessellationPatchCountBufferIndex, mtlpp::ResourceUsage(0));
 		
 		if (boundShaderState->VertexShader->SideTableBinding >= 0)
 		{
 			PrologueEncoder.SetShaderSideTable(mtlpp::FunctionType::Kernel, boundShaderState->VertexShader->SideTableBinding);
-			State.SetShaderBuffer(SF_Vertex, nil, nil, 0, 0, boundShaderState->VertexShader->SideTableBinding, mtlpp::ResourceUsage(0));
+			State.SetShaderBuffer(EMetalShaderStages::Vertex, nil, nil, 0, 0, boundShaderState->VertexShader->SideTableBinding, mtlpp::ResourceUsage(0));
 		}
 		
 		if (boundShaderState->DomainShader->SideTableBinding >= 0)
 		{
 			CurrentEncoder.SetShaderSideTable(mtlpp::FunctionType::Vertex, boundShaderState->DomainShader->SideTableBinding);
-			State.SetShaderBuffer(SF_Domain, nil, nil, 0, 0, boundShaderState->DomainShader->SideTableBinding, mtlpp::ResourceUsage(0));
+			State.SetShaderBuffer(EMetalShaderStages::Domain, nil, nil, 0, 0, boundShaderState->DomainShader->SideTableBinding, mtlpp::ResourceUsage(0));
 		}
 		
 		if (IsValidRef(boundShaderState->PixelShader) && boundShaderState->PixelShader->SideTableBinding >= 0)
 		{
 			CurrentEncoder.SetShaderSideTable(mtlpp::FunctionType::Fragment, boundShaderState->PixelShader->SideTableBinding);
-			State.SetShaderBuffer(SF_Pixel, nil, nil, 0, 0, boundShaderState->PixelShader->SideTableBinding, mtlpp::ResourceUsage(0));
+			State.SetShaderBuffer(EMetalShaderStages::Pixel, nil, nil, 0, 0, boundShaderState->PixelShader->SideTableBinding, mtlpp::ResourceUsage(0));
 		}
 		
 		auto patchesPerThreadGroup = boundShaderState->VertexShader->TessellationPatchesPerThreadGroup;
@@ -729,6 +738,7 @@ void FMetalRenderPass::DrawPatches(uint32 PrimitiveType,FMetalBuffer const& Inde
 		NOT_SUPPORTED("DrawPatches");
 	}
 }
+#endif
 
 void FMetalRenderPass::Dispatch(uint32 ThreadGroupCountX, uint32 ThreadGroupCountY, uint32 ThreadGroupCountZ)
 {
@@ -1573,55 +1583,57 @@ void FMetalRenderPass::CommitRenderResourceTables(void)
 	
 	State.CommitRenderResources(&CurrentEncoder);
 	
-	State.CommitResourceTable(SF_Vertex, mtlpp::FunctionType::Vertex, CurrentEncoder);
+	State.CommitResourceTable(EMetalShaderStages::Vertex, mtlpp::FunctionType::Vertex, CurrentEncoder);
 	
 	FMetalGraphicsPipelineState const* BoundShaderState = State.GetGraphicsPSO();
 	
 	if (BoundShaderState->VertexShader->SideTableBinding >= 0)
 	{
 		CurrentEncoder.SetShaderSideTable(mtlpp::FunctionType::Vertex, BoundShaderState->VertexShader->SideTableBinding);
-		State.SetShaderBuffer(SF_Vertex, nil, nil, 0, 0, BoundShaderState->VertexShader->SideTableBinding, mtlpp::ResourceUsage(0));
+		State.SetShaderBuffer(EMetalShaderStages::Vertex, nil, nil, 0, 0, BoundShaderState->VertexShader->SideTableBinding, mtlpp::ResourceUsage(0));
 	}
 	
 	if (IsValidRef(BoundShaderState->PixelShader))
 	{
-		State.CommitResourceTable(SF_Pixel, mtlpp::FunctionType::Fragment, CurrentEncoder);
+		State.CommitResourceTable(EMetalShaderStages::Pixel, mtlpp::FunctionType::Fragment, CurrentEncoder);
 		if (BoundShaderState->PixelShader->SideTableBinding >= 0)
 		{
 			CurrentEncoder.SetShaderSideTable(mtlpp::FunctionType::Fragment, BoundShaderState->PixelShader->SideTableBinding);
-			State.SetShaderBuffer(SF_Pixel, nil, nil, 0, 0, BoundShaderState->PixelShader->SideTableBinding, mtlpp::ResourceUsage(0));
+			State.SetShaderBuffer(EMetalShaderStages::Pixel, nil, nil, 0, 0, BoundShaderState->PixelShader->SideTableBinding, mtlpp::ResourceUsage(0));
 		}
 	}
 }
 
+#if PLATFORM_SUPPORTS_TESSELLATION_SHADERS
 void FMetalRenderPass::CommitTessellationResourceTables(void)
 {
 	State.CommitTessellationResources(&CurrentEncoder, &PrologueEncoder);
 	
-	State.CommitResourceTable(SF_Vertex, mtlpp::FunctionType::Kernel, PrologueEncoder);
+	State.CommitResourceTable(EMetalShaderStages::Vertex, mtlpp::FunctionType::Kernel, PrologueEncoder);
 	
-	State.CommitResourceTable(SF_Hull, mtlpp::FunctionType::Kernel, PrologueEncoder);
+	State.CommitResourceTable(EMetalShaderStages::Hull, mtlpp::FunctionType::Kernel, PrologueEncoder);
 	
-	State.CommitResourceTable(SF_Domain, mtlpp::FunctionType::Vertex, CurrentEncoder);
+	State.CommitResourceTable(EMetalShaderStages::Domain, mtlpp::FunctionType::Vertex, CurrentEncoder);
 	
 	TRefCountPtr<FMetalGraphicsPipelineState> CurrentBoundShaderState = State.GetGraphicsPSO();
 	if (IsValidRef(CurrentBoundShaderState->PixelShader))
 	{
-		State.CommitResourceTable(SF_Pixel, mtlpp::FunctionType::Fragment, CurrentEncoder);
+		State.CommitResourceTable(EMetalShaderStages::Pixel, mtlpp::FunctionType::Fragment, CurrentEncoder);
 	}
 }
+#endif
 
 void FMetalRenderPass::CommitDispatchResourceTables(void)
 {
 	State.CommitComputeResources(&CurrentEncoder);
 	
-	State.CommitResourceTable(SF_Compute, mtlpp::FunctionType::Kernel, CurrentEncoder);
+	State.CommitResourceTable(EMetalShaderStages::Compute, mtlpp::FunctionType::Kernel, CurrentEncoder);
 	
 	FMetalComputeShader const* ComputeShader = State.GetComputeShader();
 	if (ComputeShader->SideTableBinding >= 0)
 	{
 		CurrentEncoder.SetShaderSideTable(mtlpp::FunctionType::Kernel, ComputeShader->SideTableBinding);
-		State.SetShaderBuffer(SF_Compute, nil, nil, 0, 0, ComputeShader->SideTableBinding, mtlpp::ResourceUsage(0));
+		State.SetShaderBuffer(EMetalShaderStages::Compute, nil, nil, 0, 0, ComputeShader->SideTableBinding, mtlpp::ResourceUsage(0));
 	}
 }
 
@@ -1629,13 +1641,13 @@ void FMetalRenderPass::CommitAsyncDispatchResourceTables(void)
 {
 	State.CommitComputeResources(&PrologueEncoder);
 	
-	State.CommitResourceTable(SF_Compute, mtlpp::FunctionType::Kernel, PrologueEncoder);
+	State.CommitResourceTable(EMetalShaderStages::Compute, mtlpp::FunctionType::Kernel, PrologueEncoder);
 	
 	FMetalComputeShader const* ComputeShader = State.GetComputeShader();
 	if (ComputeShader->SideTableBinding >= 0)
 	{
 		PrologueEncoder.SetShaderSideTable(mtlpp::FunctionType::Kernel, ComputeShader->SideTableBinding);
-		State.SetShaderBuffer(SF_Compute, nil, nil, 0, 0, ComputeShader->SideTableBinding, mtlpp::ResourceUsage(0));
+		State.SetShaderBuffer(EMetalShaderStages::Compute, nil, nil, 0, 0, ComputeShader->SideTableBinding, mtlpp::ResourceUsage(0));
 	}
 }
 
@@ -1655,6 +1667,7 @@ void FMetalRenderPass::PrepareToRender(uint32 PrimitiveType)
     State.SetRenderPipelineState(CurrentEncoder, nullptr);
 }
 
+#if PLATFORM_SUPPORTS_TESSELLATION_SHADERS
 void FMetalRenderPass::PrepareToTessellate(uint32 PrimitiveType)
 {
 	SCOPE_CYCLE_COUNTER(STAT_MetalPrepareToTessellateTime);
@@ -1672,6 +1685,7 @@ void FMetalRenderPass::PrepareToTessellate(uint32 PrimitiveType)
     
     State.SetRenderPipelineState(CurrentEncoder, &PrologueEncoder);
 }
+#endif
 
 void FMetalRenderPass::PrepareToDispatch(void)
 {
diff --git a/Engine/Source/Runtime/Apple/MetalRHI/Private/MetalRenderPass.h b/Engine/Source/Runtime/Apple/MetalRHI/Private/MetalRenderPass.h
--- a/Engine/Source/Runtime/Apple/MetalRHI/Private/MetalRenderPass.h
+++ b/Engine/Source/Runtime/Apple/MetalRHI/Private/MetalRenderPass.h
@@ -46,10 +46,12 @@ public:
     void DrawIndexedIndirect(FMetalIndexBuffer* IndexBufferRHI, uint32 PrimitiveType, FMetalStructuredBuffer* VertexBufferRHI, int32 DrawArgumentsIndex, uint32 NumInstances);
     
     void DrawIndexedPrimitiveIndirect(uint32 PrimitiveType,FMetalIndexBuffer* IndexBufferRHI,FMetalVertexBuffer* VertexBufferRHI,uint32 ArgumentOffset);
-    
+	
+#if PLATFORM_SUPPORTS_TESSELLATION_SHADERS
     void DrawPatches(uint32 PrimitiveType, FMetalBuffer const& IndexBuffer, uint32 IndexBufferStride, int32 BaseVertexIndex, uint32 FirstInstance, uint32 StartIndex,
                      uint32 NumPrimitives, uint32 NumInstances);
-    
+#endif
+	
     void Dispatch(uint32 ThreadGroupCountX, uint32 ThreadGroupCountY, uint32 ThreadGroupCountZ);
     
     void DispatchIndirect(FMetalVertexBuffer* ArgumentBufferRHI, uint32 ArgumentOffset);
@@ -151,12 +153,16 @@ private:
 	void ConditionalSwitchToAsyncCompute(void);
 	
     void PrepareToRender(uint32 PrimType);
+#if PLATFORM_SUPPORTS_TESSELLATION_SHADERS
     void PrepareToTessellate(uint32 PrimType);
+#endif
     void PrepareToDispatch(void);
 	void PrepareToAsyncDispatch(void);
 
     void CommitRenderResourceTables(void);
+#if PLATFORM_SUPPORTS_TESSELLATION_SHADERS
     void CommitTessellationResourceTables(void);
+#endif
     void CommitDispatchResourceTables(void);
 	void CommitAsyncDispatchResourceTables(void);
     
@@ -168,9 +174,9 @@ private:
     
     // Which of the buffers/textures/sampler slots are bound
     // The state cache is responsible for ensuring we bind the correct 
-    FMetalTextureMask BoundTextures[SF_NumStandardFrequencies];
-    uint32 BoundBuffers[SF_NumStandardFrequencies];
-    uint16 BoundSamplers[SF_NumStandardFrequencies];
+	FMetalTextureMask BoundTextures[EMetalShaderStages::Num];
+    uint32 BoundBuffers[EMetalShaderStages::Num];
+    uint16 BoundSamplers[EMetalShaderStages::Num];
     
     FMetalCommandEncoder CurrentEncoder;
     FMetalCommandEncoder PrologueEncoder;
diff --git a/Engine/Source/Runtime/Apple/MetalRHI/Private/MetalShaders.cpp b/Engine/Source/Runtime/Apple/MetalRHI/Private/MetalShaders.cpp
--- a/Engine/Source/Runtime/Apple/MetalRHI/Private/MetalShaders.cpp
+++ b/Engine/Source/Runtime/Apple/MetalRHI/Private/MetalShaders.cpp
@@ -667,18 +667,24 @@ FMetalVertexShader::FMetalVertexShader(const TArray<uint8>& InCode)
 	FMetalCodeHeader Header;
 	Init(InCode, Header);
 	
-	TessellationOutputAttribs = Header.TessellationOutputAttribs;
-	TessellationPatchCountBuffer = Header.TessellationPatchCountBuffer;
-	TessellationIndexBuffer = Header.TessellationIndexBuffer;
-	TessellationHSOutBuffer = Header.TessellationHSOutBuffer;
-	TessellationHSTFOutBuffer = Header.TessellationHSTFOutBuffer;
-	TessellationControlPointOutBuffer = Header.TessellationControlPointOutBuffer;
-	TessellationControlPointIndexBuffer = Header.TessellationControlPointIndexBuffer;
-	TessellationOutputControlPoints = Header.TessellationOutputControlPoints;
-	TessellationDomain = Header.TessellationDomain;
-	TessellationInputControlPoints = Header.TessellationInputControlPoints;
-	TessellationMaxTessFactor = Header.TessellationMaxTessFactor;
-	TessellationPatchesPerThreadGroup = Header.TessellationPatchesPerThreadGroup;
+#if PLATFORM_SUPPORTS_TESSELLATION_SHADERS
+	if (Header.Tessellation.Num())
+	{
+		auto const& Tess = Header.Tessellation[0];
+		TessellationOutputAttribs = Tess.TessellationOutputAttribs;
+		TessellationPatchCountBuffer = Tess.TessellationPatchCountBuffer;
+		TessellationIndexBuffer = Tess.TessellationIndexBuffer;
+		TessellationHSOutBuffer = Tess.TessellationHSOutBuffer;
+		TessellationHSTFOutBuffer = Tess.TessellationHSTFOutBuffer;
+		TessellationControlPointOutBuffer = Tess.TessellationControlPointOutBuffer;
+		TessellationControlPointIndexBuffer = Tess.TessellationControlPointIndexBuffer;
+		TessellationOutputControlPoints = Tess.TessellationOutputControlPoints;
+		TessellationDomain = Tess.TessellationDomain;
+		TessellationInputControlPoints = Tess.TessellationInputControlPoints;
+		TessellationMaxTessFactor = Tess.TessellationMaxTessFactor;
+		TessellationPatchesPerThreadGroup = Tess.TessellationPatchesPerThreadGroup;
+	}
+#endif
 }
 
 FMetalVertexShader::FMetalVertexShader(const TArray<uint8>& InCode, mtlpp::Library InLibrary)
@@ -686,18 +692,24 @@ FMetalVertexShader::FMetalVertexShader(const TArray<uint8>& InCode, mtlpp::Libra
 	FMetalCodeHeader Header;
 	Init(InCode, Header, InLibrary);
 	
-	TessellationOutputAttribs = Header.TessellationOutputAttribs;
-	TessellationPatchCountBuffer = Header.TessellationPatchCountBuffer;
-	TessellationIndexBuffer = Header.TessellationIndexBuffer;
-	TessellationHSOutBuffer = Header.TessellationHSOutBuffer;
-	TessellationHSTFOutBuffer = Header.TessellationHSTFOutBuffer;
-	TessellationControlPointOutBuffer = Header.TessellationControlPointOutBuffer;
-	TessellationControlPointIndexBuffer = Header.TessellationControlPointIndexBuffer;
-	TessellationOutputControlPoints = Header.TessellationOutputControlPoints;
-	TessellationDomain = Header.TessellationDomain;
-	TessellationInputControlPoints = Header.TessellationInputControlPoints;
-	TessellationMaxTessFactor = Header.TessellationMaxTessFactor;
-	TessellationPatchesPerThreadGroup = Header.TessellationPatchesPerThreadGroup;
+#if PLATFORM_SUPPORTS_TESSELLATION_SHADERS
+	if (Header.Tessellation.Num())
+	{
+		auto const& Tess = Header.Tessellation[0];
+		TessellationOutputAttribs = Tess.TessellationOutputAttribs;
+		TessellationPatchCountBuffer = Tess.TessellationPatchCountBuffer;
+		TessellationIndexBuffer = Tess.TessellationIndexBuffer;
+		TessellationHSOutBuffer = Tess.TessellationHSOutBuffer;
+		TessellationHSTFOutBuffer = Tess.TessellationHSTFOutBuffer;
+		TessellationControlPointOutBuffer = Tess.TessellationControlPointOutBuffer;
+		TessellationControlPointIndexBuffer = Tess.TessellationControlPointIndexBuffer;
+		TessellationOutputControlPoints = Tess.TessellationOutputControlPoints;
+		TessellationDomain = Tess.TessellationDomain;
+		TessellationInputControlPoints = Tess.TessellationInputControlPoints;
+		TessellationMaxTessFactor = Tess.TessellationMaxTessFactor;
+		TessellationPatchesPerThreadGroup = Tess.TessellationPatchesPerThreadGroup;
+	}
+#endif
 }
 
 mtlpp::Function FMetalVertexShader::GetFunction()
@@ -722,6 +734,7 @@ mtlpp::Function FMetalPixelShader::GetFunction()
 	return GetCompiledFunction();
 }
 
+#if PLATFORM_SUPPORTS_TESSELLATION_SHADERS
 FMetalHullShader::FMetalHullShader(const TArray<uint8>& InCode)
 {
 	FMetalCodeHeader Header;
@@ -745,10 +758,11 @@ FMetalDomainShader::FMetalDomainShader(const TArray<uint8>& InCode)
 	Init(InCode, Header);
 	
 	// for VSHS
-	TessellationHSOutBuffer = Header.TessellationHSOutBuffer;
-	TessellationControlPointOutBuffer = Header.TessellationControlPointOutBuffer;
+	auto const& Tess = Header.Tessellation[0];
+	TessellationHSOutBuffer = Tess.TessellationHSOutBuffer;
+	TessellationControlPointOutBuffer = Tess.TessellationControlPointOutBuffer;
 	
-	switch (Header.TessellationOutputWinding)
+	switch (Tess.TessellationOutputWinding)
 	{
 		// NOTE: cw and ccw are flipped
 		case EMetalOutputWindingMode::Clockwise:		TessellationOutputWinding = mtlpp::Winding::CounterClockwise; break;
@@ -756,7 +770,7 @@ FMetalDomainShader::FMetalDomainShader(const TArray<uint8>& InCode)
 		default: check(0);
 	}
 	
-	switch (Header.TessellationPartitioning)
+	switch (Tess.TessellationPartitioning)
 	{
 		case EMetalPartitionMode::Pow2:				TessellationPartitioning = mtlpp::TessellationPartitionMode::ModePow2; break;
 		case EMetalPartitionMode::Integer:			TessellationPartitioning = mtlpp::TessellationPartitionMode::ModeInteger; break;
@@ -772,10 +786,11 @@ FMetalDomainShader::FMetalDomainShader(const TArray<uint8>& InCode, mtlpp::Libra
 	Init(InCode, Header, InLibrary);
 	
 	// for VSHS
-	TessellationHSOutBuffer = Header.TessellationHSOutBuffer;
-	TessellationControlPointOutBuffer = Header.TessellationControlPointOutBuffer;
+	auto const& Tess = Header.Tessellation[0];
+	TessellationHSOutBuffer = Tess.TessellationHSOutBuffer;
+	TessellationControlPointOutBuffer = Tess.TessellationControlPointOutBuffer;
 	
-	switch (Header.TessellationOutputWinding)
+	switch (Tess.TessellationOutputWinding)
 	{
 		// NOTE: cw and ccw are flipped
 		case EMetalOutputWindingMode::Clockwise:		TessellationOutputWinding = mtlpp::Winding::CounterClockwise; break;
@@ -783,7 +798,7 @@ FMetalDomainShader::FMetalDomainShader(const TArray<uint8>& InCode, mtlpp::Libra
 		default: check(0);
 	}
 	
-	switch (Header.TessellationPartitioning)
+	switch (Tess.TessellationPartitioning)
 	{
 		case EMetalPartitionMode::Pow2:				TessellationPartitioning = mtlpp::TessellationPartitionMode::ModePow2; break;
 		case EMetalPartitionMode::Integer:			TessellationPartitioning = mtlpp::TessellationPartitionMode::ModeInteger; break;
@@ -797,6 +812,7 @@ mtlpp::Function FMetalDomainShader::GetFunction()
 {
 	return GetCompiledFunction();
 }
+#endif
 
 FVertexShaderRHIRef FMetalDynamicRHI::RHICreateVertexShader(const TArray<uint8>& Code)
 {
@@ -841,7 +857,13 @@ FPixelShaderRHIRef FMetalDynamicRHI::RHICreatePixelShader(FRHIShaderLibraryParam
 FHullShaderRHIRef FMetalDynamicRHI::RHICreateHullShader(const TArray<uint8>& Code) 
 {
 	@autoreleasepool {
+#if PLATFORM_SUPPORTS_TESSELLATION_SHADERS
 	FMetalHullShader* Shader = new FMetalHullShader(Code);
+#else
+	FMetalHullShader* Shader = new FMetalHullShader;
+	FMetalCodeHeader Header;
+	Shader->Init(Code, Header);
+#endif
 	return Shader;
 	}
 }
@@ -861,7 +883,13 @@ FHullShaderRHIRef FMetalDynamicRHI::RHICreateHullShader(FRHIShaderLibraryParamRe
 FDomainShaderRHIRef FMetalDynamicRHI::RHICreateDomainShader(const TArray<uint8>& Code) 
 {
 	@autoreleasepool {
+#if PLATFORM_SUPPORTS_TESSELLATION_SHADERS
 	FMetalDomainShader* Shader = new FMetalDomainShader(Code);
+#else
+	FMetalDomainShader* Shader = new FMetalDomainShader;
+	FMetalCodeHeader Header;
+	Shader->Init(Code, Header);
+#endif
 	return Shader;
 	}
 }
@@ -1062,6 +1090,7 @@ FVertexShaderRHIRef FMetalShaderLibrary::CreateVertexShader(const FSHAHash& Hash
 
 FHullShaderRHIRef FMetalShaderLibrary::CreateHullShader(const FSHAHash& Hash)
 {
+#if PLATFORM_SUPPORTS_TESSELLATION_SHADERS
 	FMetalShadeEntry* Code = Map.HashMap.Find(Hash);
 	if (Code)
 	{
@@ -1077,10 +1106,15 @@ FHullShaderRHIRef FMetalShaderLibrary::CreateHullShader(const FSHAHash& Hash)
 	}
 	UE_LOG(LogMetal, Error, TEXT("Failed to find Hull Shader with SHA: %s"), *Hash.ToString());
 	return FHullShaderRHIRef();
+#else
+	checkf(0, TEXT("Not supported"));
+	return NULL;
+#endif
 }
 
 FDomainShaderRHIRef FMetalShaderLibrary::CreateDomainShader(const FSHAHash& Hash)
 {
+#if PLATFORM_SUPPORTS_TESSELLATION_SHADERS
 	FMetalShadeEntry* Code = Map.HashMap.Find(Hash);
 	if (Code)
 	{
@@ -1096,11 +1130,15 @@ FDomainShaderRHIRef FMetalShaderLibrary::CreateDomainShader(const FSHAHash& Hash
 	}
 	UE_LOG(LogMetal, Error, TEXT("Failed to find Domain Shader with SHA: %s"), *Hash.ToString());
 	return FDomainShaderRHIRef();
+#else
+	checkf(0, TEXT("Not supported"));
+	return NULL;
+#endif
 }
 
 FGeometryShaderRHIRef FMetalShaderLibrary::CreateGeometryShader(const FSHAHash& Hash)
 {
-	checkf(0, TEXT("Not supported yet"));
+	checkf(0, TEXT("Not supported"));
 	return NULL;
 }
 
@@ -1308,7 +1346,7 @@ void FMetalShaderParameterCache::Set(uint32 BufferIndexName, uint32 ByteOffset,
 	FMemory::Memcpy(PackedGlobalUniforms[BufferIndex]->Data + ByteOffset, NewValues, NumBytes);
 }
 
-void FMetalShaderParameterCache::CommitPackedGlobals(FMetalStateCache* Cache, FMetalCommandEncoder* Encoder, EShaderFrequency Frequency, const FMetalShaderBindings& Bindings)
+void FMetalShaderParameterCache::CommitPackedGlobals(FMetalStateCache* Cache, FMetalCommandEncoder* Encoder, uint32 Frequency, const FMetalShaderBindings& Bindings)
 {
 	// copy the current uniform buffer into the ring buffer to submit
 	for (int32 Index = 0; Index < Bindings.PackedGlobalArrays.Num(); ++Index)
@@ -1331,13 +1369,13 @@ void FMetalShaderParameterCache::CommitPackedGlobals(FMetalStateCache* Cache, FM
 				uint8 const* Bytes = PackedGlobalUniforms[Index]->Data;
 				ns::AutoReleased<FMetalBuffer> Buffer(Encoder->GetRingBuffer().NewBuffer(Size, 0));
 				FMemory::Memcpy((uint8*)Buffer.GetContents(), Bytes, Size);
-				Cache->SetShaderBuffer(Frequency, Buffer, nil, 0, Size, UniformBufferIndex, mtlpp::ResourceUsage::Read);
+				Cache->SetShaderBuffer((EMetalShaderStages)Frequency, Buffer, nil, 0, Size, UniformBufferIndex, mtlpp::ResourceUsage::Read);
 			}
 			else
 			{
 				PackedGlobalUniforms[Index]->Len = Size;
-				Cache->SetShaderBuffer(Frequency, nil, nil, 0, 0, UniformBufferIndex, mtlpp::ResourceUsage(0));
-				Cache->SetShaderBuffer(Frequency, nil, PackedGlobalUniforms[Index], 0, Size, UniformBufferIndex, mtlpp::ResourceUsage::Read);
+				Cache->SetShaderBuffer((EMetalShaderStages)Frequency, nil, nil, 0, 0, UniformBufferIndex, mtlpp::ResourceUsage(0));
+				Cache->SetShaderBuffer((EMetalShaderStages)Frequency, nil, PackedGlobalUniforms[Index], 0, Size, UniformBufferIndex, mtlpp::ResourceUsage::Read);
 			}
 
 			// mark as clean
diff --git a/Engine/Source/Runtime/Apple/MetalRHI/Private/MetalStateCache.cpp b/Engine/Source/Runtime/Apple/MetalRHI/Private/MetalStateCache.cpp
--- a/Engine/Source/Runtime/Apple/MetalRHI/Private/MetalStateCache.cpp
+++ b/Engine/Source/Runtime/Apple/MetalRHI/Private/MetalStateCache.cpp
@@ -159,7 +159,9 @@ FMetalStateCache::FMetalStateCache(bool const bInImmediate)
 , bHasValidRenderTarget(false)
 , bHasValidColorTarget(false)
 , bScissorRectEnabled(false)
+#if PLATFORM_SUPPORTS_TESSELLATION_SHADERS
 , bUsingTessellation(false)
+#endif
 , bCanRestartRenderPass(false)
 , bImmediate(bInImmediate)
 , bFallbackDepthStencilBound(false)
@@ -190,7 +192,7 @@ FMetalStateCache::~FMetalStateCache()
 		VertexBuffers[i].Length = 0;
 		VertexBuffers[i].Offset = 0;
 	}
-	for (uint32 Frequency = 0; Frequency < SF_NumStandardFrequencies; Frequency++)
+	for (uint32 Frequency = 0; Frequency < EMetalShaderStages::Num; Frequency++)
 	{
 		ShaderSamplers[Frequency].Bound = 0;
 		for (uint32 i = 0; i < ML_MaxSamplers; i++)
@@ -247,7 +249,7 @@ void FMetalStateCache::Reset(void)
 		VertexBuffers[i].Length = 0;
 		VertexBuffers[i].Offset = 0;
 	}
-	for (uint32 Frequency = 0; Frequency < SF_NumStandardFrequencies; Frequency++)
+	for (uint32 Frequency = 0; Frequency < EMetalShaderStages::Num; Frequency++)
 	{
 		ShaderSamplers[Frequency].Bound = 0;
 		for (uint32 i = 0; i < ML_MaxSamplers; i++)
@@ -294,7 +296,9 @@ void FMetalStateCache::Reset(void)
 	BlendFactor = FLinearColor::Transparent;
 	FrameBufferSize = CGSizeMake(0.0, 0.0);
 	RenderTargetArraySize = 0;
+#if PLATFORM_SUPPORTS_TESSELLATION_SHADERS
     bUsingTessellation = false;
+#endif
     bCanRestartRenderPass = false;
     
     RasterBits = EMetalRenderFlagMask;
@@ -379,13 +383,15 @@ void FMetalStateCache::SetComputeShader(FMetalComputeShader* InComputeShader)
 		
 		PipelineBits |= EMetalPipelineFlagComputeShader;
 		
+#if PLATFORM_SUPPORTS_TESSELLATION_SHADERS
 		bUsingTessellation = false;
+#endif
 		
-		DirtyUniformBuffers[SF_Compute] = 0xffffffff;
+		DirtyUniformBuffers[EMetalShaderStages::Compute] = 0xffffffff;
 
 		for (const auto& PackedGlobalArray : InComputeShader->Bindings.PackedGlobalArrays)
 		{
-			ShaderParameters[CrossCompiler::SHADER_STAGE_COMPUTE].PrepareGlobalUniforms(CrossCompiler::PackedTypeNameToTypeIndex(PackedGlobalArray.TypeName), PackedGlobalArray.Size);
+			ShaderParameters[EMetalShaderStages::Compute].PrepareGlobalUniforms(CrossCompiler::PackedTypeNameToTypeIndex(PackedGlobalArray.TypeName), PackedGlobalArray.Size);
 		}
 	}
 }
@@ -1048,7 +1054,7 @@ void FMetalStateCache::SetVertexStream(uint32 const Index, FMetalBuffer* Buffer,
 	VertexBuffers[Index].Bytes = Bytes;
 	VertexBuffers[Index].Length = Length;
 	
-	SetShaderBuffer(SF_Vertex, VertexBuffers[Index].Buffer, Bytes, Offset, Length, UNREAL_TO_METAL_BUFFER_INDEX(Index), mtlpp::ResourceUsage::Read);
+	SetShaderBuffer(EMetalShaderStages::Vertex, VertexBuffers[Index].Buffer, Bytes, Offset, Length, UNREAL_TO_METAL_BUFFER_INDEX(Index), mtlpp::ResourceUsage::Read);
 }
 
 uint32 FMetalStateCache::GetVertexBufferSize(uint32 const Index)
@@ -1064,10 +1070,11 @@ void FMetalStateCache::SetGraphicsPipelineState(FMetalGraphicsPipelineState* Sta
 	{
 		GraphicsPSO = State;
 		
+#if PLATFORM_SUPPORTS_TESSELLATION_SHADERS
 		bool bNewUsingTessellation = (State && State->GetPipeline(IndexType)->TessellationPipelineDesc.DomainVertexDescriptor);
 		if (bNewUsingTessellation != bUsingTessellation)
 		{
-			for (uint32 i = 0; i < SF_NumStandardFrequencies; i++)
+			for (uint32 i = 0; i < EMetalShaderStages::Num; i++)
 			{
 				ShaderBuffers[i].Bound = UINT32_MAX;
 #if PLATFORM_MAC
@@ -1085,52 +1092,55 @@ void FMetalStateCache::SetGraphicsPipelineState(FMetalGraphicsPipelineState* Sta
 		// previous pipeline with different binding table will overwrite the vertex shader bindings for the current pipeline.
 		if (bNewUsingTessellation)
 		{
-			ShaderBuffers[SF_Hull].Bound = UINT32_MAX;
+			ShaderBuffers[EMetalShaderStages::Hull].Bound = UINT32_MAX;
 #if PLATFORM_MAC
-			ShaderTextures[SF_Hull].Bound = UINT128_MAX;
+			ShaderTextures[EMetalShaderStages::Hull].Bound = UINT128_MAX;
 #else
-			ShaderTextures[SF_Hull].Bound = UINT32_MAX;
+			ShaderTextures[EMetalShaderStages::Hull].Bound = UINT32_MAX;
 #endif
-			ShaderSamplers[SF_Hull].Bound = UINT16_MAX;
+			ShaderSamplers[EMetalShaderStages::Hull].Bound = UINT16_MAX;
 			
 			for (uint32 i = 0; i < ML_MaxBuffers; i++)
 			{
-				BoundUniformBuffers[SF_Hull][i] = nullptr;
-				ShaderBuffers[SF_Hull].Buffers[i].Buffer = nil;
-				ShaderBuffers[SF_Hull].Buffers[i].Bytes = nil;
-				ShaderBuffers[SF_Hull].Buffers[i].Length = 0;
-				ShaderBuffers[SF_Hull].Buffers[i].Offset = 0;
-				ShaderBuffers[SF_Hull].Buffers[i].Usage = mtlpp::ResourceUsage(0);
-				ShaderBuffers[SF_Hull].Formats[i] = PF_Unknown;
+				BoundUniformBuffers[EMetalShaderStages::Hull][i] = nullptr;
+				ShaderBuffers[EMetalShaderStages::Hull].Buffers[i].Buffer = nil;
+				ShaderBuffers[EMetalShaderStages::Hull].Buffers[i].Bytes = nil;
+				ShaderBuffers[EMetalShaderStages::Hull].Buffers[i].Length = 0;
+                ShaderBuffers[EMetalShaderStages::Hull].Buffers[i].Offset = 0;
+                ShaderBuffers[EMetalShaderStages::Hull].Buffers[i].Usage = mtlpp::ResourceUsage(0);
+				ShaderBuffers[EMetalShaderStages::Hull].Formats[i] = PF_Unknown;
 			}
 			for (uint32 i = 0; i < ML_MaxTextures; i++)
 			{
-				ShaderTextures[SF_Hull].Textures[i] = nil;
-				ShaderTextures[SF_Hull].Usage[i] = mtlpp::ResourceUsage(0);
+				ShaderTextures[EMetalShaderStages::Hull].Textures[i] = nil;
+				ShaderTextures[EMetalShaderStages::Hull].Usage[i] = mtlpp::ResourceUsage(0);
 			}
 			
 			for (uint32 i = 0; i < ML_MaxSamplers; i++)
 			{
-				ShaderSamplers[SF_Hull].Samplers[i] = nil;
+				ShaderSamplers[EMetalShaderStages::Hull].Samplers[i] = nil;
 			}
 
 			for (const auto& PackedGlobalArray : State->HullShader->Bindings.PackedGlobalArrays)
 			{
-				ShaderParameters[CrossCompiler::SHADER_STAGE_HULL].PrepareGlobalUniforms(CrossCompiler::PackedTypeNameToTypeIndex(PackedGlobalArray.TypeName), PackedGlobalArray.Size);
+				ShaderParameters[EMetalShaderStages::Hull].PrepareGlobalUniforms(CrossCompiler::PackedTypeNameToTypeIndex(PackedGlobalArray.TypeName), PackedGlobalArray.Size);
 			}
 
 			for (const auto& PackedGlobalArray : State->DomainShader->Bindings.PackedGlobalArrays)
 			{
-				ShaderParameters[CrossCompiler::SHADER_STAGE_DOMAIN].PrepareGlobalUniforms(CrossCompiler::PackedTypeNameToTypeIndex(PackedGlobalArray.TypeName), PackedGlobalArray.Size);
+				ShaderParameters[EMetalShaderStages::Domain].PrepareGlobalUniforms(CrossCompiler::PackedTypeNameToTypeIndex(PackedGlobalArray.TypeName), PackedGlobalArray.Size);
 			}
 		}
 		bUsingTessellation = bNewUsingTessellation;
+		DirtyUniformBuffers[EMetalShaderStages::Hull] = 0xffffffff;
+		DirtyUniformBuffers[EMetalShaderStages::Domain] = 0xffffffff;
+#endif
 		
-		DirtyUniformBuffers[SF_Vertex] = 0xffffffff;
-		DirtyUniformBuffers[SF_Pixel] = 0xffffffff;
-		DirtyUniformBuffers[SF_Hull] = 0xffffffff;
-		DirtyUniformBuffers[SF_Domain] = 0xffffffff;
-		DirtyUniformBuffers[SF_Geometry] = 0xffffffff;
+		DirtyUniformBuffers[EMetalShaderStages::Vertex] = 0xffffffff;
+		DirtyUniformBuffers[EMetalShaderStages::Pixel] = 0xffffffff;
+#if PLATFORM_SUPPORTS_GEOMETRY_SHADERS
+		DirtyUniformBuffers[EMetalShaderStages::Geometry] = 0xffffffff;
+#endif
 		
 		PipelineBits |= EMetalPipelineFlagPipelineState;
 		
@@ -1139,14 +1149,14 @@ void FMetalStateCache::SetGraphicsPipelineState(FMetalGraphicsPipelineState* Sta
 
 		for (const auto& PackedGlobalArray : State->VertexShader->Bindings.PackedGlobalArrays)
 		{
-			ShaderParameters[CrossCompiler::SHADER_STAGE_VERTEX].PrepareGlobalUniforms(CrossCompiler::PackedTypeNameToTypeIndex(PackedGlobalArray.TypeName), PackedGlobalArray.Size);
+			ShaderParameters[EMetalShaderStages::Vertex].PrepareGlobalUniforms(CrossCompiler::PackedTypeNameToTypeIndex(PackedGlobalArray.TypeName), PackedGlobalArray.Size);
 		}
 
 		if (State->PixelShader)
 		{
 			for (const auto& PackedGlobalArray : State->PixelShader->Bindings.PackedGlobalArrays)
 			{
-				ShaderParameters[CrossCompiler::SHADER_STAGE_PIXEL].PrepareGlobalUniforms(CrossCompiler::PackedTypeNameToTypeIndex(PackedGlobalArray.TypeName), PackedGlobalArray.Size);
+				ShaderParameters[EMetalShaderStages::Pixel].PrepareGlobalUniforms(CrossCompiler::PackedTypeNameToTypeIndex(PackedGlobalArray.TypeName), PackedGlobalArray.Size);
 			}
 		}
 		
@@ -1158,7 +1168,7 @@ void FMetalStateCache::SetGraphicsPipelineState(FMetalGraphicsPipelineState* Sta
 				if (IsValidRef(RenderTargetsInfo.UnorderedAccessView[i]))
 				{
 					FMetalUnorderedAccessView* UAV = ResourceCast(RenderTargetsInfo.UnorderedAccessView[i].GetReference());
-					SetShaderUnorderedAccessView(SF_Pixel, i, UAV);
+					SetShaderUnorderedAccessView(EMetalShaderStages::Pixel, i, UAV);
 				}
 			}
 		}
@@ -1175,7 +1185,7 @@ void FMetalStateCache::SetIndexType(EMetalIndexType InIndexType)
 	}
 }
 
-void FMetalStateCache::BindUniformBuffer(EShaderFrequency const Freq, uint32 const BufferIndex, FUniformBufferRHIParamRef BufferRHI)
+void FMetalStateCache::BindUniformBuffer(EMetalShaderStages const Freq, uint32 const BufferIndex, FUniformBufferRHIParamRef BufferRHI)
 {
 	check(BufferIndex < ML_MaxBuffers);
 	if (BoundUniformBuffers[Freq][BufferIndex] != BufferRHI)
@@ -1186,7 +1196,7 @@ void FMetalStateCache::BindUniformBuffer(EShaderFrequency const Freq, uint32 con
 	}
 }
 
-void FMetalStateCache::SetDirtyUniformBuffers(EShaderFrequency const Freq, uint32 const Dirty)
+void FMetalStateCache::SetDirtyUniformBuffers(EMetalShaderStages const Freq, uint32 const Dirty)
 {
 	DirtyUniformBuffers[Freq] = Dirty;
 }
@@ -1348,9 +1358,9 @@ bool FMetalStateCache::NeedsToSetRenderTarget(const FRHISetRenderTargetsInfo& In
 	return bAllChecksPassed == false;
 }
 
-void FMetalStateCache::SetShaderBuffer(EShaderFrequency const Frequency, FMetalBuffer const& Buffer, FMetalBufferData* const Bytes, NSUInteger const Offset, NSUInteger const Length, NSUInteger const Index, mtlpp::ResourceUsage const Usage, EPixelFormat const Format)
+void FMetalStateCache::SetShaderBuffer(EMetalShaderStages const Frequency, FMetalBuffer const& Buffer, FMetalBufferData* const Bytes, NSUInteger const Offset, NSUInteger const Length, NSUInteger const Index, mtlpp::ResourceUsage const Usage, EPixelFormat const Format)
 {
-	check(Frequency < SF_NumStandardFrequencies);
+	check(Frequency < EMetalShaderStages::Num);
 	check(Index < ML_MaxBuffers);
 	
 	if (ShaderBuffers[Frequency].Buffers[Index].Buffer != Buffer ||
@@ -1379,9 +1389,9 @@ void FMetalStateCache::SetShaderBuffer(EShaderFrequency const Frequency, FMetalB
 	}
 }
 
-void FMetalStateCache::SetShaderTexture(EShaderFrequency const Frequency, FMetalTexture const& Texture, NSUInteger const Index, mtlpp::ResourceUsage const Usage)
+void FMetalStateCache::SetShaderTexture(EMetalShaderStages const Frequency, FMetalTexture const& Texture, NSUInteger const Index, mtlpp::ResourceUsage const Usage)
 {
-	check(Frequency < SF_NumStandardFrequencies);
+	check(Frequency < EMetalShaderStages::Num);
 	check(Index < ML_MaxTextures);
 	
 	if (ShaderTextures[Frequency].Textures[Index] != Texture
@@ -1406,9 +1416,9 @@ void FMetalStateCache::SetShaderTexture(EShaderFrequency const Frequency, FMetal
 	}
 }
 
-void FMetalStateCache::SetShaderSamplerState(EShaderFrequency const Frequency, FMetalSamplerState* const Sampler, NSUInteger const Index)
+void FMetalStateCache::SetShaderSamplerState(EMetalShaderStages const Frequency, FMetalSamplerState* const Sampler, NSUInteger const Index)
 {
-	check(Frequency < SF_NumStandardFrequencies);
+	check(Frequency < EMetalShaderStages::Num);
 	check(Index < ML_MaxSamplers);
 	
 	if (ShaderSamplers[Frequency].Samplers[Index].GetPtr() != (Sampler ? Sampler->State.GetPtr() : nil))
@@ -1416,7 +1426,7 @@ void FMetalStateCache::SetShaderSamplerState(EShaderFrequency const Frequency, F
 		if (Sampler)
 		{
 #if !PLATFORM_MAC
-			ShaderSamplers[Frequency].Samplers[Index] = ((Frequency == SF_Vertex || Frequency == SF_Compute) && Sampler->NoAnisoState) ? Sampler->NoAnisoState : Sampler->State;
+			ShaderSamplers[Frequency].Samplers[Index] = ((Frequency == EMetalShaderStages::Vertex || Frequency == EMetalShaderStages::Compute) && Sampler->NoAnisoState) ? Sampler->NoAnisoState : Sampler->State;
 #else
 			ShaderSamplers[Frequency].Samplers[Index] = Sampler->State;
 #endif
@@ -1445,24 +1455,25 @@ void FMetalStateCache::SetResource(uint32 ShaderStage, uint32 BindIndex, FRHITex
 	switch (ShaderStage)
 	{
 		case CrossCompiler::SHADER_STAGE_PIXEL:
-			SetShaderTexture(SF_Pixel, Texture, BindIndex, Usage);
+			SetShaderTexture(EMetalShaderStages::Pixel, Texture, BindIndex, Usage);
 			break;
 			
 		case CrossCompiler::SHADER_STAGE_VERTEX:
-			SetShaderTexture(SF_Vertex, Texture, BindIndex, Usage);
+			SetShaderTexture(EMetalShaderStages::Vertex, Texture, BindIndex, Usage);
 			break;
 			
 		case CrossCompiler::SHADER_STAGE_COMPUTE:
-			SetShaderTexture(SF_Compute, Texture, BindIndex, Usage);
+			SetShaderTexture(EMetalShaderStages::Compute, Texture, BindIndex, Usage);
 			break;
-			
+#if PLATFORM_SUPPORTS_TESSELLATION_SHADERS
 		case CrossCompiler::SHADER_STAGE_HULL:
-			SetShaderTexture(SF_Hull, Texture, BindIndex, Usage);
+			SetShaderTexture(EMetalShaderStages::Hull, Texture, BindIndex, Usage);
 			break;
 			
 		case CrossCompiler::SHADER_STAGE_DOMAIN:
-			SetShaderTexture(SF_Domain, Texture, BindIndex, Usage);
+			SetShaderTexture(EMetalShaderStages::Domain, Texture, BindIndex, Usage);
 			break;
+#endif
 			
 		default:
 			check(0);
@@ -1470,7 +1481,7 @@ void FMetalStateCache::SetResource(uint32 ShaderStage, uint32 BindIndex, FRHITex
 	}
 }
 
-void FMetalStateCache::SetShaderResourceView(FMetalContext* Context, EShaderFrequency ShaderStage, uint32 BindIndex, FMetalShaderResourceView* RESTRICT SRV)
+void FMetalStateCache::SetShaderResourceView(FMetalContext* Context, EMetalShaderStages ShaderStage, uint32 BindIndex, FMetalShaderResourceView* RESTRICT SRV)
 {
 	if (SRV)
 	{
@@ -1521,31 +1532,33 @@ void FMetalStateCache::SetShaderResourceView(FMetalContext* Context, EShaderFreq
 	}
 }
 
-bool FMetalStateCache::IsLinearBuffer(EShaderFrequency ShaderStage, uint32 BindIndex)
+bool FMetalStateCache::IsLinearBuffer(EMetalShaderStages ShaderStage, uint32 BindIndex)
 {
     switch (ShaderStage)
     {
-        case SF_Vertex:
+        case EMetalShaderStages::Vertex:
         {
             return (GraphicsPSO->VertexShader->Bindings.LinearBuffer & (1 << BindIndex)) != 0;
             break;
         }
-        case SF_Pixel:
+        case EMetalShaderStages::Pixel:
         {
             return (GraphicsPSO->PixelShader->Bindings.LinearBuffer & (1 << BindIndex)) != 0;
             break;
         }
-        case SF_Hull:
+#if PLATFORM_SUPPORTS_TESSELLATION_SHADERS
+        case EMetalShaderStages::Hull:
         {
             return (GraphicsPSO->HullShader->Bindings.LinearBuffer & (1 << BindIndex)) != 0;
             break;
         }
-        case SF_Domain:
+        case EMetalShaderStages::Domain:
         {
             return (GraphicsPSO->DomainShader->Bindings.LinearBuffer & (1 << BindIndex)) != 0;
             break;
         }
-        case SF_Compute:
+#endif
+        case EMetalShaderStages::Compute:
         {
             return (ComputeShader->Bindings.LinearBuffer & (1 << BindIndex)) != 0;
         }
@@ -1557,7 +1570,7 @@ bool FMetalStateCache::IsLinearBuffer(EShaderFrequency ShaderStage, uint32 BindI
     }
 }
 
-void FMetalStateCache::SetShaderUnorderedAccessView(EShaderFrequency ShaderStage, uint32 BindIndex, FMetalUnorderedAccessView* RESTRICT UAV)
+void FMetalStateCache::SetShaderUnorderedAccessView(EMetalShaderStages ShaderStage, uint32 BindIndex, FMetalUnorderedAccessView* RESTRICT UAV)
 {
 	if (UAV)
 	{
@@ -1638,24 +1651,26 @@ void FMetalStateCache::SetResource(uint32 ShaderStage, uint32 BindIndex, FMetalS
 	switch (ShaderStage)
 	{
 		case CrossCompiler::SHADER_STAGE_PIXEL:
-			SetShaderResourceView(nullptr, SF_Pixel, BindIndex, SRV);
+			SetShaderResourceView(nullptr, EMetalShaderStages::Pixel, BindIndex, SRV);
 			break;
 			
 		case CrossCompiler::SHADER_STAGE_VERTEX:
-			SetShaderResourceView(nullptr, SF_Vertex, BindIndex, SRV);
+			SetShaderResourceView(nullptr, EMetalShaderStages::Vertex, BindIndex, SRV);
 			break;
 			
 		case CrossCompiler::SHADER_STAGE_COMPUTE:
-			SetShaderResourceView(nullptr, SF_Compute, BindIndex, SRV);
+			SetShaderResourceView(nullptr, EMetalShaderStages::Compute, BindIndex, SRV);
 			break;
 			
+#if PLATFORM_SUPPORTS_TESSELLATION_SHADERS
 		case CrossCompiler::SHADER_STAGE_HULL:
-			SetShaderResourceView(nullptr, SF_Hull, BindIndex, SRV);
+			SetShaderResourceView(nullptr, EMetalShaderStages::Hull, BindIndex, SRV);
 			break;
 			
 		case CrossCompiler::SHADER_STAGE_DOMAIN:
-			SetShaderResourceView(nullptr, SF_Domain, BindIndex, SRV);
+			SetShaderResourceView(nullptr, EMetalShaderStages::Domain, BindIndex, SRV);
 			break;
+#endif
 			
 		default:
 			check(0);
@@ -1669,24 +1684,26 @@ void FMetalStateCache::SetResource(uint32 ShaderStage, uint32 BindIndex, FMetalS
 	switch (ShaderStage)
 	{
 		case CrossCompiler::SHADER_STAGE_PIXEL:
-			SetShaderSamplerState(SF_Pixel, SamplerState, BindIndex);
+			SetShaderSamplerState(EMetalShaderStages::Pixel, SamplerState, BindIndex);
 			break;
 			
 		case CrossCompiler::SHADER_STAGE_VERTEX:
-			SetShaderSamplerState(SF_Vertex, SamplerState, BindIndex);
+			SetShaderSamplerState(EMetalShaderStages::Vertex, SamplerState, BindIndex);
 			break;
 			
 		case CrossCompiler::SHADER_STAGE_COMPUTE:
-			SetShaderSamplerState(SF_Compute, SamplerState, BindIndex);
+			SetShaderSamplerState(EMetalShaderStages::Compute, SamplerState, BindIndex);
 			break;
 			
+#if PLATFORM_SUPPORTS_TESSELLATION_SHADERS
 		case CrossCompiler::SHADER_STAGE_HULL:
-			SetShaderSamplerState(SF_Hull, SamplerState, BindIndex);
+			SetShaderSamplerState(EMetalShaderStages::Hull, SamplerState, BindIndex);
 			break;
 			
 		case CrossCompiler::SHADER_STAGE_DOMAIN:
-			SetShaderSamplerState(SF_Domain, SamplerState, BindIndex);
+			SetShaderSamplerState(EMetalShaderStages::Domain, SamplerState, BindIndex);
 			break;
+#endif
 			
 		default:
 			check(0);
@@ -1699,24 +1716,26 @@ void FMetalStateCache::SetResource(uint32 ShaderStage, uint32 BindIndex, FMetalU
 	switch (ShaderStage)
 	{
 		case CrossCompiler::SHADER_STAGE_PIXEL:
-			SetShaderUnorderedAccessView(SF_Pixel, BindIndex, UAV);
+			SetShaderUnorderedAccessView(EMetalShaderStages::Pixel, BindIndex, UAV);
 			break;
 			
 		case CrossCompiler::SHADER_STAGE_VERTEX:
-			SetShaderUnorderedAccessView(SF_Vertex, BindIndex, UAV);
+			SetShaderUnorderedAccessView(EMetalShaderStages::Vertex, BindIndex, UAV);
 			break;
 			
 		case CrossCompiler::SHADER_STAGE_COMPUTE:
-			SetShaderUnorderedAccessView(SF_Compute, BindIndex, UAV);
+			SetShaderUnorderedAccessView(EMetalShaderStages::Compute, BindIndex, UAV);
 			break;
 			
+#if PLATFORM_SUPPORTS_TESSELLATION_SHADERS
 		case CrossCompiler::SHADER_STAGE_HULL:
-			SetShaderUnorderedAccessView(SF_Hull, BindIndex, UAV);
+			SetShaderUnorderedAccessView(EMetalShaderStages::Hull, BindIndex, UAV);
 			break;
 			
 		case CrossCompiler::SHADER_STAGE_DOMAIN:
-			SetShaderUnorderedAccessView(SF_Domain, BindIndex, UAV);
+			SetShaderUnorderedAccessView(EMetalShaderStages::Domain, BindIndex, UAV);
 			break;
+#endif
 			
 		default:
 			check(0);
@@ -1758,26 +1777,28 @@ void FMetalStateCache::SetResourcesFromTables(ShaderType Shader, uint32 ShaderSt
 {
 	checkSlow(Shader);
 	
-	EShaderFrequency Frequency;
+	EMetalShaderStages Frequency;
 	switch(ShaderStage)
 	{
 		case CrossCompiler::SHADER_STAGE_VERTEX:
-			Frequency = SF_Vertex;
+			Frequency = EMetalShaderStages::Vertex;
 			break;
+#if PLATFORM_SUPPORTS_TESSELLATION_SHADERS
 		case CrossCompiler::SHADER_STAGE_HULL:
-			Frequency = SF_Hull;
+			Frequency = EMetalShaderStages::Hull;
 			break;
 		case CrossCompiler::SHADER_STAGE_DOMAIN:
-			Frequency = SF_Domain;
+			Frequency = EMetalShaderStages::Domain;
 			break;
+#endif
 		case CrossCompiler::SHADER_STAGE_PIXEL:
-			Frequency = SF_Pixel;
+			Frequency = EMetalShaderStages::Pixel;
 			break;
 		case CrossCompiler::SHADER_STAGE_COMPUTE:
-			Frequency = SF_Compute;
+			Frequency = EMetalShaderStages::Compute;
 			break;
 		default:
-			Frequency = SF_NumFrequencies; //Silence a compiler warning/error
+			Frequency = EMetalShaderStages::Num; //Silence a compiler warning/error
 			check(false);
 			break;
 	}
@@ -1813,41 +1834,43 @@ void FMetalStateCache::CommitRenderResources(FMetalCommandEncoder* Raster)
 	check(IsValidRef(GraphicsPSO));
     
     SetResourcesFromTables(GraphicsPSO->VertexShader, CrossCompiler::SHADER_STAGE_VERTEX);
-    GetShaderParameters(CrossCompiler::SHADER_STAGE_VERTEX).CommitPackedGlobals(this, Raster, SF_Vertex, GraphicsPSO->VertexShader->Bindings);
+    GetShaderParameters(EMetalShaderStages::Vertex).CommitPackedGlobals(this, Raster, EMetalShaderStages::Vertex, GraphicsPSO->VertexShader->Bindings);
 	
     if (IsValidRef(GraphicsPSO->PixelShader))
     {
     	SetResourcesFromTables(GraphicsPSO->PixelShader, CrossCompiler::SHADER_STAGE_PIXEL);
-        GetShaderParameters(CrossCompiler::SHADER_STAGE_PIXEL).CommitPackedGlobals(this, Raster, SF_Pixel, GraphicsPSO->PixelShader->Bindings);
+        GetShaderParameters(EMetalShaderStages::Pixel).CommitPackedGlobals(this, Raster, EMetalShaderStages::Pixel, GraphicsPSO->PixelShader->Bindings);
     }
 }
 
+#if PLATFORM_SUPPORTS_TESSELLATION_SHADERS
 void FMetalStateCache::CommitTessellationResources(FMetalCommandEncoder* Raster, FMetalCommandEncoder* Compute)
 {
 	check(IsValidRef(GraphicsPSO));
     check(IsValidRef(GraphicsPSO->HullShader) && IsValidRef(GraphicsPSO->DomainShader));
     
     SetResourcesFromTables(GraphicsPSO->VertexShader, CrossCompiler::SHADER_STAGE_VERTEX);
-    GetShaderParameters(CrossCompiler::SHADER_STAGE_VERTEX).CommitPackedGlobals(this, Compute, SF_Vertex, GraphicsPSO->VertexShader->Bindings);
+    GetShaderParameters(EMetalShaderStages::Vertex).CommitPackedGlobals(this, Compute, EMetalShaderStages::Vertex, GraphicsPSO->VertexShader->Bindings);
 	
     if (IsValidRef(GraphicsPSO->PixelShader))
     {
     	SetResourcesFromTables(GraphicsPSO->PixelShader, CrossCompiler::SHADER_STAGE_PIXEL);
-        GetShaderParameters(CrossCompiler::SHADER_STAGE_PIXEL).CommitPackedGlobals(this, Raster, SF_Pixel, GraphicsPSO->PixelShader->Bindings);
+        GetShaderParameters(EMetalShaderStages::Pixel).CommitPackedGlobals(this, Raster, EMetalShaderStages::Pixel, GraphicsPSO->PixelShader->Bindings);
     }
     
     SetResourcesFromTables(GraphicsPSO->HullShader, CrossCompiler::SHADER_STAGE_HULL);
 	
 	SetResourcesFromTables(GraphicsPSO->DomainShader, CrossCompiler::SHADER_STAGE_DOMAIN);
-   	GetShaderParameters(CrossCompiler::SHADER_STAGE_DOMAIN).CommitPackedGlobals(this, Raster, SF_Domain, GraphicsPSO->DomainShader->Bindings);
+    GetShaderParameters(EMetalShaderStages::Domain).CommitPackedGlobals(this, Raster, EMetalShaderStages::Domain, GraphicsPSO->DomainShader->Bindings);
 }
+#endif
 
 void FMetalStateCache::CommitComputeResources(FMetalCommandEncoder* Compute)
 {
 	check(IsValidRef(ComputeShader));
 	SetResourcesFromTables(ComputeShader, CrossCompiler::SHADER_STAGE_COMPUTE);
 	
-	GetShaderParameters(CrossCompiler::SHADER_STAGE_COMPUTE).CommitPackedGlobals(this, Compute, SF_Compute, ComputeShader->Bindings);
+	GetShaderParameters(EMetalShaderStages::Compute).CommitPackedGlobals(this, Compute, EMetalShaderStages::Compute, ComputeShader->Bindings);
 }
 
 bool FMetalStateCache::PrepareToRestart(bool const bCurrentApplied)
@@ -1928,7 +1951,7 @@ void FMetalStateCache::SetStateDirty(void)
 {	
 	RasterBits = UINT32_MAX;
     PipelineBits = EMetalPipelineFlagMask;
-	for (uint32 i = 0; i < SF_NumStandardFrequencies; i++)
+	for (uint32 i = 0; i < EMetalShaderStages::Num; i++)
 	{
 		ShaderBuffers[i].Bound = UINT32_MAX;
 #if PLATFORM_MAC
@@ -2085,7 +2108,7 @@ void FMetalStateCache::SetComputePipelineState(FMetalCommandEncoder& CommandEnco
     }
 }
 
-void FMetalStateCache::CommitResourceTable(EShaderFrequency const Frequency, mtlpp::FunctionType const Type, FMetalCommandEncoder& CommandEncoder)
+void FMetalStateCache::CommitResourceTable(EMetalShaderStages const Frequency, mtlpp::FunctionType const Type, FMetalCommandEncoder& CommandEncoder)
 {
 	FMetalBufferBindings& BufferBindings = ShaderBuffers[Frequency];
 	while(BufferBindings.Bound)
@@ -2169,23 +2192,25 @@ void FMetalStateCache::CommitResourceTable(EShaderFrequency const Frequency, mtl
 	TMap<uint8, TArray<uint8>> ArgumentBufferMasks;
 	switch(Frequency)
 	{
-		case SF_Vertex:
+		case EMetalShaderStages::Vertex:
 			ArgumentBuffers = GraphicsPSO->VertexShader->Bindings.ArgumentBuffers;
 			ArgumentBufferMasks = GraphicsPSO->VertexShader->Bindings.ArgumentBufferMasks;
-			break;
-		case SF_Hull:
+            break;
+#if PLATFORM_SUPPORTS_TESSELLATION_SHADERS
+		case EMetalShaderStages::Hull:
 			ArgumentBuffers = GraphicsPSO->HullShader->Bindings.ArgumentBuffers;
 			ArgumentBufferMasks = GraphicsPSO->HullShader->Bindings.ArgumentBufferMasks;
 			break;
-		case SF_Domain:
+		case EMetalShaderStages::Domain:
 			ArgumentBuffers = GraphicsPSO->DomainShader->Bindings.ArgumentBuffers;
 			ArgumentBufferMasks = GraphicsPSO->DomainShader->Bindings.ArgumentBufferMasks;
 			break;
-		case SF_Pixel:
+#endif
+		case EMetalShaderStages::Pixel:
 			ArgumentBuffers = GraphicsPSO->PixelShader->Bindings.ArgumentBuffers;
 			ArgumentBufferMasks = GraphicsPSO->PixelShader->Bindings.ArgumentBufferMasks;
 			break;
-		case SF_Compute:
+		case EMetalShaderStages::Compute:
 			ArgumentBuffers = ComputeShader->Bindings.ArgumentBuffers;
 			ArgumentBufferMasks = ComputeShader->Bindings.ArgumentBufferMasks;
 			break;
diff --git a/Engine/Source/Runtime/Apple/MetalRHI/Private/MetalStateCache.h b/Engine/Source/Runtime/Apple/MetalRHI/Private/MetalStateCache.h
--- a/Engine/Source/Runtime/Apple/MetalRHI/Private/MetalStateCache.h
+++ b/Engine/Source/Runtime/Apple/MetalRHI/Private/MetalStateCache.h
@@ -51,8 +51,8 @@ public:
 	void SetVertexStream(uint32 const Index, FMetalBuffer* Buffer, FMetalBufferData* Bytes, uint32 const Offset, uint32 const Length);
 	void SetGraphicsPipelineState(FMetalGraphicsPipelineState* State);
 	void SetIndexType(EMetalIndexType IndexType);
-	void BindUniformBuffer(EShaderFrequency const Freq, uint32 const BufferIndex, FUniformBufferRHIParamRef BufferRHI);
-	void SetDirtyUniformBuffers(EShaderFrequency const Freq, uint32 const Dirty);
+	void BindUniformBuffer(EMetalShaderStages const Freq, uint32 const BufferIndex, FUniformBufferRHIParamRef BufferRHI);
+	void SetDirtyUniformBuffers(EMetalShaderStages const Freq, uint32 const Dirty);
 	
 	/*
 	 * Monitor if samples pass the depth and stencil tests.
@@ -73,7 +73,7 @@ public:
 	 * @param Usage The resource usage flags.
 	 * @param Format The UAV pixel format.
 	 */
-	void SetShaderBuffer(EShaderFrequency const Frequency, FMetalBuffer const& Buffer, FMetalBufferData* const Bytes, NSUInteger const Offset, NSUInteger const Length, NSUInteger const Index, mtlpp::ResourceUsage const Usage, EPixelFormat const Format = PF_Unknown);
+	void SetShaderBuffer(EMetalShaderStages const Frequency, FMetalBuffer const& Buffer, FMetalBufferData* const Bytes, NSUInteger const Offset, NSUInteger const Length, NSUInteger const Index, mtlpp::ResourceUsage const Usage, EPixelFormat const Format = PF_Unknown);
 	
 	/*
 	 * Set a global texture for the specified shader frequency at the given bind point index.
@@ -82,7 +82,7 @@ public:
 	 * @param Index The index to modify.
 	 * @param Usage The resource usage flags.
 	 */
-	void SetShaderTexture(EShaderFrequency const Frequency, FMetalTexture const& Texture, NSUInteger const Index, mtlpp::ResourceUsage const Usage);
+	void SetShaderTexture(EMetalShaderStages const Frequency, FMetalTexture const& Texture, NSUInteger const Index, mtlpp::ResourceUsage const Usage);
 	
 	/*
 	 * Set a global sampler for the specified shader frequency at the given bind point index.
@@ -90,11 +90,11 @@ public:
 	 * @param Sampler The sampler state to bind or nil to clear.
 	 * @param Index The index to modify.
 	 */
-	void SetShaderSamplerState(EShaderFrequency const Frequency, FMetalSamplerState* const Sampler, NSUInteger const Index);
+	void SetShaderSamplerState(EMetalShaderStages const Frequency, FMetalSamplerState* const Sampler, NSUInteger const Index);
 
-	void SetShaderResourceView(FMetalContext* Context, EShaderFrequency ShaderStage, uint32 BindIndex, FMetalShaderResourceView* RESTRICT SRV);
+	void SetShaderResourceView(FMetalContext* Context, EMetalShaderStages ShaderStage, uint32 BindIndex, FMetalShaderResourceView* RESTRICT SRV);
 	
-	void SetShaderUnorderedAccessView(EShaderFrequency ShaderStage, uint32 BindIndex, FMetalUnorderedAccessView* RESTRICT UAV);
+	void SetShaderUnorderedAccessView(EMetalShaderStages ShaderStage, uint32 BindIndex, FMetalUnorderedAccessView* RESTRICT UAV);
 
 	void SetStateDirty(void);
 	
@@ -104,15 +104,17 @@ public:
 
 	void CommitRenderResources(FMetalCommandEncoder* Raster);
 
+#if PLATFORM_SUPPORTS_TESSELLATION_SHADERS
 	void CommitTessellationResources(FMetalCommandEncoder* Raster, FMetalCommandEncoder* Compute);
+#endif
 
 	void CommitComputeResources(FMetalCommandEncoder* Compute);
 	
-	void CommitResourceTable(EShaderFrequency const Frequency, mtlpp::FunctionType const Type, FMetalCommandEncoder& CommandEncoder);
+	void CommitResourceTable(EMetalShaderStages const Frequency, mtlpp::FunctionType const Type, FMetalCommandEncoder& CommandEncoder);
 	
 	bool PrepareToRestart(bool const bCurrentApplied);
 	
-	FMetalShaderParameterCache& GetShaderParameters(uint32 const Stage) { return ShaderParameters[Stage]; }
+	FMetalShaderParameterCache& GetShaderParameters(EMetalShaderStages const Stage) { return ShaderParameters[Stage]; }
 	FLinearColor const& GetBlendFactor() const { return BlendFactor; }
 	uint32 GetStencilRef() const { return StencilRef; }
 	FMetalDepthStencilState* GetDepthStencilState() const { return DepthStencilState; }
@@ -127,18 +129,20 @@ public:
 	const mtlpp::Viewport& GetViewport(uint32 const Index) const { check(Index < ML_MaxViewports); return Viewport[Index]; }
 	uint32 GetVertexBufferSize(uint32 const Index);
 	uint32 GetRenderTargetArraySize() const { return RenderTargetArraySize; }
-	const FRHIUniformBuffer** GetBoundUniformBuffers(EShaderFrequency const Freq) { return (const FRHIUniformBuffer**)&BoundUniformBuffers[Freq][0]; }
-	uint32 GetDirtyUniformBuffers(EShaderFrequency const Freq) const { return DirtyUniformBuffers[Freq]; }
+	const FRHIUniformBuffer** GetBoundUniformBuffers(EMetalShaderStages const Freq) { return (const FRHIUniformBuffer**)&BoundUniformBuffers[Freq][0]; }
+	uint32 GetDirtyUniformBuffers(EMetalShaderStages const Freq) const { return DirtyUniformBuffers[Freq]; }
 	FMetalQueryBuffer* GetVisibilityResultsBuffer() const { return VisibilityResults; }
 	bool GetScissorRectEnabled() const { return bScissorRectEnabled; }
 	bool NeedsToSetRenderTarget(const FRHISetRenderTargetsInfo& RenderTargetsInfo);
 	bool HasValidDepthStencilSurface() const { return IsValidRef(DepthStencilSurface); }
 	EMetalIndexType GetIndexType() const { return IndexType; }
-    bool GetUsingTessellation() const { return bUsingTessellation; }
+#if PLATFORM_SUPPORTS_TESSELLATION_SHADERS
+	bool GetUsingTessellation() const { return bUsingTessellation; }
+#endif
     bool CanRestartRenderPass() const { return bCanRestartRenderPass; }
 	mtlpp::RenderPassDescriptor GetRenderPassDescriptor(void) const { return RenderPassDesc; }
 	uint32 GetSampleCount(void) const { return SampleCount; }
-    bool IsLinearBuffer(EShaderFrequency ShaderStage, uint32 BindIndex);
+    bool IsLinearBuffer(EMetalShaderStages ShaderStage, uint32 BindIndex);
     FMetalShaderPipeline* GetPipelineState(void) const { return GraphicsPSO->GetPipeline(GetIndexType()); }
 	EPrimitiveType GetPrimitiveType() { check(IsValidRef(GraphicsPSO)); return GraphicsPSO->GetPrimitiveType(); }
 	
@@ -225,24 +229,24 @@ private:
 	};
     
 private:
-	FMetalShaderParameterCache ShaderParameters[CrossCompiler::NUM_SHADER_STAGES];
+	FMetalShaderParameterCache ShaderParameters[EMetalShaderStages::Num];
 
 	EMetalIndexType IndexType;
 	uint32 SampleCount;
 
 	TSet<TRefCountPtr<FRHIUniformBuffer>> ActiveUniformBuffers;
-	FRHIUniformBuffer* BoundUniformBuffers[SF_NumStandardFrequencies][ML_MaxBuffers];
+	FRHIUniformBuffer* BoundUniformBuffers[EMetalShaderStages::Num][ML_MaxBuffers];
 	
 	/** Bitfield for which uniform buffers are dirty */
-	uint32 DirtyUniformBuffers[SF_NumStandardFrequencies];
+	uint32 DirtyUniformBuffers[EMetalShaderStages::Num];
 	
 	/** Vertex attribute buffers */
 	FMetalBufferBinding VertexBuffers[MaxVertexElementCount];
 	
 	/** Bound shader resource tables. */
-	FMetalBufferBindings ShaderBuffers[SF_NumStandardFrequencies];
-	FMetalTextureBindings ShaderTextures[SF_NumStandardFrequencies];
-	FMetalSamplerBindings ShaderSamplers[SF_NumStandardFrequencies];
+	FMetalBufferBindings ShaderBuffers[EMetalShaderStages::Num];
+	FMetalTextureBindings ShaderTextures[EMetalShaderStages::Num];
+	FMetalSamplerBindings ShaderSamplers[EMetalShaderStages::Num];
 	
 	mtlpp::StoreAction ColorStore[MaxSimultaneousRenderTargets];
 	mtlpp::StoreAction DepthStore;
@@ -282,7 +286,9 @@ private:
 	bool bHasValidRenderTarget;
 	bool bHasValidColorTarget;
 	bool bScissorRectEnabled;
+#if PLATFORM_SUPPORTS_TESSELLATION_SHADERS
     bool bUsingTessellation;
+#endif
     bool bCanRestartRenderPass;
     bool bImmediate;
 	bool bFallbackDepthStencilBound;
diff --git a/Engine/Source/Runtime/Apple/MetalRHI/Public/MetalResources.h b/Engine/Source/Runtime/Apple/MetalRHI/Public/MetalResources.h
--- a/Engine/Source/Runtime/Apple/MetalRHI/Public/MetalResources.h
+++ b/Engine/Source/Runtime/Apple/MetalRHI/Public/MetalResources.h
@@ -176,6 +176,7 @@ public:
 	
 	mtlpp::Function GetFunction();
 	
+#if PLATFORM_SUPPORTS_TESSELLATION_SHADERS
 	// for VSHS
 	FMetalTessellationOutputs TessellationOutputAttribs;
 	float  TessellationMaxTessFactor;
@@ -189,6 +190,7 @@ public:
 	uint32 TessellationHSTFOutBuffer;
 	uint32 TessellationControlPointOutBuffer;
 	uint32 TessellationControlPointIndexBuffer;
+#endif
 };
 
 class FMetalPixelShader : public TMetalBaseShader<FRHIPixelShader, SF_Pixel>
@@ -200,6 +202,7 @@ public:
 	mtlpp::Function GetFunction();
 };
 
+#if PLATFORM_SUPPORTS_TESSELLATION_SHADERS
 class FMetalHullShader : public TMetalBaseShader<FRHIHullShader, SF_Hull>
 {
 public:
@@ -222,6 +225,10 @@ public:
 	uint32 TessellationHSOutBuffer;
 	uint32 TessellationControlPointOutBuffer;
 };
+#else
+typedef TMetalBaseShader<FRHIHullShader, SF_Hull> FMetalHullShader;
+typedef TMetalBaseShader<FRHIDomainShader, SF_Domain> FMetalDomainShader;
+#endif
 
 typedef TMetalBaseShader<FRHIGeometryShader, SF_Geometry> FMetalGeometryShader;
 
@@ -276,9 +283,13 @@ public:
 	/** Cached shaders */
 	TRefCountPtr<FMetalVertexShader> VertexShader;
 	TRefCountPtr<FMetalPixelShader> PixelShader;
+#if PLATFORM_SUPPORTS_TESSELLATION_SHADERS
 	TRefCountPtr<FMetalHullShader> HullShader;
 	TRefCountPtr<FMetalDomainShader> DomainShader;
+#endif
+#if PLATFORM_SUPPORTS_GEOMETRY_SHADERS
 	TRefCountPtr<FMetalGeometryShader> GeometryShader;
+#endif
 	
 	/** Cached state objects */
 	TRefCountPtr<FMetalDepthStencilState> DepthStencilState;
@@ -961,7 +972,7 @@ public:
 	/**
 	 * Commit shader parameters to the currently bound program.
 	 */
-	void CommitPackedGlobals(class FMetalStateCache* Cache, class FMetalCommandEncoder* Encoder, EShaderFrequency Frequency, const FMetalShaderBindings& Bindings);
+	void CommitPackedGlobals(class FMetalStateCache* Cache, class FMetalCommandEncoder* Encoder, uint32 Frequency, const FMetalShaderBindings& Bindings);
 
 private:
 	/** CPU memory block for storing uniform values. */
diff --git a/Engine/Source/Runtime/Apple/MetalRHI/Public/MetalShaderResources.h b/Engine/Source/Runtime/Apple/MetalRHI/Public/MetalShaderResources.h
--- a/Engine/Source/Runtime/Apple/MetalRHI/Public/MetalShaderResources.h
+++ b/Engine/Source/Runtime/Apple/MetalRHI/Public/MetalShaderResources.h
@@ -228,11 +228,11 @@
 
 struct FMetalTessellationOutputs
 {
+	TArray<FMetalAttribute> HSOut;
+	TArray<FMetalAttribute> PatchControlPointOut;
 	uint32 HSOutSize;
 	uint32 HSTFOutSize;
 	uint32 PatchControlPointOutSize;
-	TArray<FMetalAttribute> HSOut;
-	TArray<FMetalAttribute> PatchControlPointOut;
 	
 	FMetalTessellationOutputs()
 	: HSOutSize(0)
@@ -244,61 +244,36 @@
 	
 	friend FArchive& operator<<(FArchive& Ar, FMetalTessellationOutputs& Attrs)
 	{
+		Ar << Attrs.HSOut;
+		Ar << Attrs.PatchControlPointOut;
 		Ar << Attrs.HSOutSize;
 		Ar << Attrs.HSTFOutSize;
 		Ar << Attrs.PatchControlPointOutSize;
-		Ar << Attrs.HSOut;
-		Ar << Attrs.PatchControlPointOut;
 		return Ar;
 	}
 };
 
-struct FMetalCodeHeader
+struct FMetalTessellationHeader
 {
-	uint32 Frequency;
-	FMetalShaderBindings Bindings;
-	TArray<CrossCompiler::FUniformBufferCopyInfo> UniformBuffersCopyInfo;
-    FString ShaderName;
-    
-    FMetalTessellationOutputs TessellationOutputAttribs;
-
-	uint64 CompilerBuild;
-	uint32 CompilerVersion;
-
+	FMetalTessellationOutputs TessellationOutputAttribs;
+	
 	uint32 TessellationOutputControlPoints;
 	uint32 TessellationDomain; // 3 = tri, 4 = quad // TODO unused
 	uint32 TessellationInputControlPoints; // TODO unused
 	uint32 TessellationPatchesPerThreadGroup;
-    uint32 TessellationPatchCountBuffer;
-    uint32 TessellationIndexBuffer;
-    uint32 TessellationHSOutBuffer;
-    uint32 TessellationHSTFOutBuffer;
-    uint32 TessellationControlPointOutBuffer;
-    uint32 TessellationControlPointIndexBuffer;
+	uint32 TessellationPatchCountBuffer;
+	uint32 TessellationIndexBuffer;
+	uint32 TessellationHSOutBuffer;
+	uint32 TessellationHSTFOutBuffer;
+	uint32 TessellationControlPointOutBuffer;
+	uint32 TessellationControlPointIndexBuffer;
 	float  TessellationMaxTessFactor;
-
-	uint32 SourceLen;
-	uint32 SourceCRC;
 	
-	uint16 CompileFlags;
-	
-	uint32 NumThreadsX;
-	uint32 NumThreadsY;
-	uint32 NumThreadsZ;
-	
-	uint8 Version;
-	int8 SideTable;
-	
 	EMetalOutputWindingMode TessellationOutputWinding;
 	EMetalPartitionMode TessellationPartitioning;
 	
-	bool bDeviceFunctionConstants;
-	
-	FMetalCodeHeader()
-	: Frequency(0)
-	, CompilerBuild(0)
-	, CompilerVersion(0)
-	, TessellationOutputControlPoints(0)
+	FMetalTessellationHeader()
+	: TessellationOutputControlPoints(0)
 	, TessellationDomain(0)
 	, TessellationInputControlPoints(0)
 	, TessellationPatchesPerThreadGroup(0)
@@ -309,16 +284,67 @@
 	, TessellationControlPointOutBuffer(~0u)
 	, TessellationControlPointIndexBuffer(~0u)
 	, TessellationMaxTessFactor(0)
+	, TessellationOutputWinding(EMetalOutputWindingMode::Clockwise)
+	, TessellationPartitioning(EMetalPartitionMode::Pow2)
+	{
+		
+	}
+	
+	friend FArchive& operator<<(FArchive& Ar, FMetalTessellationHeader& Header)
+	{
+		Ar << Header.TessellationOutputAttribs;
+		
+		Ar << Header.TessellationOutputControlPoints;
+		Ar << Header.TessellationDomain;
+		Ar << Header.TessellationInputControlPoints;
+		Ar << Header.TessellationPatchesPerThreadGroup;
+		Ar << Header.TessellationMaxTessFactor;
+		
+		Ar << Header.TessellationPatchCountBuffer;
+		Ar << Header.TessellationIndexBuffer;
+		Ar << Header.TessellationHSOutBuffer;
+		Ar << Header.TessellationHSTFOutBuffer;
+		Ar << Header.TessellationControlPointOutBuffer;
+		Ar << Header.TessellationControlPointIndexBuffer;
+		
+		Ar << Header.TessellationOutputWinding;
+		Ar << Header.TessellationPartitioning;
+		return Ar;
+	}
+};
+
+struct FMetalCodeHeader
+{
+	FMetalShaderBindings Bindings;
+	TArray<CrossCompiler::FUniformBufferCopyInfo> UniformBuffersCopyInfo;
+	TArray<FMetalTessellationHeader> Tessellation;
+	FString ShaderName;
+
+	uint64 CompilerBuild;
+	uint32 CompilerVersion;
+	uint32 Frequency;
+	uint32 SourceLen;
+	uint32 SourceCRC;
+	uint32 NumThreadsX;
+	uint32 NumThreadsY;
+	uint32 NumThreadsZ;
+	uint16 CompileFlags;
+	uint8 Version;
+	int8 SideTable;
+	bool bDeviceFunctionConstants;
+	
+	FMetalCodeHeader()
+	: CompilerBuild(0)
+	, CompilerVersion(0)
+	, Frequency(0)
 	, SourceLen(0)
 	, SourceCRC(0)
-	, CompileFlags(0)
 	, NumThreadsX(0)
 	, NumThreadsY(0)
 	, NumThreadsZ(0)
+	, CompileFlags(0)
 	, Version(0)
 	, SideTable(-1)
-	, TessellationOutputWinding(EMetalOutputWindingMode::Clockwise)
-	, TessellationPartitioning(EMetalPartitionMode::Pow2)
 	, bDeviceFunctionConstants(false)
 	{
 		
@@ -327,7 +353,6 @@
 
 inline FArchive& operator<<(FArchive& Ar, FMetalCodeHeader& Header)
 {
-	Ar << Header.Frequency;
 	Ar << Header.Bindings;
 
 	int32 NumInfos = Header.UniformBuffersCopyInfo.Num();
@@ -350,40 +375,20 @@
 		}
 	}
 	
+	Ar << Header.Tessellation;
 	Ar << Header.ShaderName;
 
-    Ar << Header.TessellationOutputAttribs;
-
-	Ar << Header.CompilerVersion;
 	Ar << Header.CompilerBuild;
-    
-	Ar << Header.TessellationOutputControlPoints;
-	Ar << Header.TessellationDomain;
-	Ar << Header.TessellationInputControlPoints;
-	Ar << Header.TessellationPatchesPerThreadGroup;
-	Ar << Header.TessellationMaxTessFactor;
-    
-    Ar << Header.TessellationPatchCountBuffer;
-    Ar << Header.TessellationIndexBuffer;
-    Ar << Header.TessellationHSOutBuffer;
-    Ar << Header.TessellationHSTFOutBuffer;
-    Ar << Header.TessellationControlPointOutBuffer;
-    Ar << Header.TessellationControlPointIndexBuffer;
-	
+	Ar << Header.CompilerVersion;
+	Ar << Header.Frequency;
 	Ar << Header.SourceLen;
 	Ar << Header.SourceCRC;
-	
-	Ar << Header.CompileFlags;
-	
 	Ar << Header.NumThreadsX;
 	Ar << Header.NumThreadsY;
 	Ar << Header.NumThreadsZ;
-	
+	Ar << Header.CompileFlags;
 	Ar << Header.Version;
 	Ar << Header.SideTable;
-	
-	Ar << Header.TessellationOutputWinding;
-	Ar << Header.TessellationPartitioning;
 	Ar << Header.bDeviceFunctionConstants;
 
     return Ar;
diff --git a/Engine/Source/Runtime/Core/Public/IOS/IOSPlatform.h b/Engine/Source/Runtime/Core/Public/IOS/IOSPlatform.h
--- a/Engine/Source/Runtime/Core/Public/IOS/IOSPlatform.h
+++ b/Engine/Source/Runtime/Core/Public/IOS/IOSPlatform.h
@@ -64,6 +64,8 @@ typedef FIOSPlatformTypes FPlatformTypes;
 #define PLATFORM_UI_NEEDS_FOCUS_OUTLINES				0
 
 #define PLATFORM_NEEDS_RHIRESOURCELIST					0
+#define PLATFORM_SUPPORTS_GEOMETRY_SHADERS				0
+#define PLATFORM_SUPPORTS_TESSELLATION_SHADERS			0
 #if WITH_SIMULATOR
 	#define PLATFORM_BREAK()							__asm__("int $3")
 #elif PLATFORM_64BITS
