Fix newer code that wasn't aware of the new PLATFORM_SUPPORTS_TESSELLATION_SHADERS macros and so broke on iOS.

diff --git a/Engine/Source/Runtime/Apple/MetalRHI/Private/MetalPipeline.cpp b/Engine/Source/Runtime/Apple/MetalRHI/Private/MetalPipeline.cpp
--- a/Engine/Source/Runtime/Apple/MetalRHI/Private/MetalPipeline.cpp
+++ b/Engine/Source/Runtime/Apple/MetalRHI/Private/MetalPipeline.cpp
@@ -483,8 +483,13 @@ static FMetalShaderPipeline* CreateMTLRenderPipeline(bool const bSync, FMetalGra
 		if (FMetalCommandQueue::SupportsFeature(EMetalFeaturesPipelineBufferMutability))
 		{
 			ns::AutoReleased<ns::Array<mtlpp::PipelineBufferDescriptor>> VertexPipelineBuffers = RenderPipelineDesc.GetVertexBuffers();
+#if PLATFORM_SUPPORTS_TESSELLATION_SHADERS
 			FMetalShaderBindings& VertexBindings = DomainShader ? DomainShader->Bindings : VertexShader->Bindings;
 			int8 VertexSideTable = DomainShader ? DomainShader->SideTableBinding : VertexShader->SideTableBinding;
+#else
+			FMetalShaderBindings& VertexBindings = VertexShader->Bindings;
+			int8 VertexSideTable = VertexShader->SideTableBinding;
+#endif
 			{
 				uint32 ImmutableBuffers = VertexBindings.ConstantBuffers | VertexBindings.ArgumentBuffers;
 				while(ImmutableBuffers)
