In MetalRHI tidy up the way we release private textures back in to the pool early so as to avoid placing them back in the pool twice or attempting to put nil into the pool. Also put the use of setPurgeableState behind the CVar "rhi.Metal.ResourcePurgeInPool" which is off by default.

diff --git a/Engine/Source/Runtime/Apple/MetalRHI/Private/MetalBuffer.cpp b/Engine/Source/Runtime/Apple/MetalRHI/Private/MetalBuffer.cpp
--- a/Engine/Source/Runtime/Apple/MetalRHI/Private/MetalBuffer.cpp
+++ b/Engine/Source/Runtime/Apple/MetalRHI/Private/MetalBuffer.cpp
@@ -26,6 +26,12 @@ static FAutoConsoleVariableRef CVarMetalHeapBufferBytesToCompact(
     TEXT("When enabled (> 0) this will force MetalRHI to compact the given number of bytes each frame into older buffer heaps from newer ones in order to defragment memory and reduce wastage.\n")
     TEXT("(Off by default (0))"));
 
+static int32 GMetalResourcePurgeInPool = 0;
+static FAutoConsoleVariableRef CVarMetalResourcePurgeInPool(
+	TEXT("rhi.Metal.ResourcePurgeInPool"),
+	GMetalResourcePurgeInPool,
+	TEXT("Use the SetPurgeableState function to allow the OS to reclaim memory from resources while they are unused in the pools. (Default: 0, Off)"));
+
 #if METAL_DEBUG_OPTIONS
 extern int32 GMetalBufferScribble;
 #endif
@@ -1242,7 +1248,10 @@ FMetalTexture FMetalTexturePool::CreateTexture(mtlpp::Device Device, mtlpp::Text
 	{
 		Texture = *Tex;
 		Pool.Remove(Descriptor);
-        Texture.SetPurgeableState(mtlpp::PurgeableState::NonVolatile);
+		if (GMetalResourcePurgeInPool)
+		{
+        	Texture.SetPurgeableState(mtlpp::PurgeableState::NonVolatile);
+		}
 	}
 	else
 	{
@@ -1270,8 +1279,11 @@ void FMetalTexturePool::ReleaseTexture(FMetalTexture& Texture)
 	Descriptor.usage = Texture.GetUsage();
 	Descriptor.freedFrame = GFrameNumberRenderThread;
 	
-    Texture.SetPurgeableState(mtlpp::PurgeableState::Volatile);
-    
+	if (GMetalResourcePurgeInPool && Texture.SetPurgeableState(mtlpp::PurgeableState::KeepCurrent) == mtlpp::PurgeableState::NonVolatile)
+	{
+		Texture.SetPurgeableState(mtlpp::PurgeableState::Volatile);
+	}
+	
 	FScopeLock Lock(&PoolMutex);
 	Pool.Add(Descriptor, Texture);
 }
@@ -1291,7 +1303,7 @@ void FMetalTexturePool::Drain(bool const bForce)
 			{
 				It.RemoveCurrent();
 			}
-            else if ((GFrameNumberRenderThread - It->Key.freedFrame) >= PurgeAfterNumFrames)
+            else if (GMetalResourcePurgeInPool && (GFrameNumberRenderThread - It->Key.freedFrame) >= PurgeAfterNumFrames)
             {
                 It->Value.SetPurgeableState(mtlpp::PurgeableState::Empty);
             }
@@ -1489,7 +1501,10 @@ FMetalBuffer FMetalResourceHeap::CreateBuffer(uint32 Size, uint32 Alignment, uin
 				 else
 				{
                     Buffer = ManagedBuffers.CreatePooledResource(FMetalPooledBufferArgs(Queue->GetDevice(), BlockSize, Flags, StorageMode));
-                    Buffer.SetPurgeableState(mtlpp::PurgeableState::NonVolatile);
+					if (GMetalResourcePurgeInPool)
+					{
+                    	Buffer.SetPurgeableState(mtlpp::PurgeableState::NonVolatile);
+					}
 					DEC_MEMORY_STAT_BY(STAT_MetalBufferUnusedMemory, Buffer.GetLength());
 					DEC_MEMORY_STAT_BY(STAT_MetalPooledBufferUnusedMemory, Buffer.GetLength());
                     INC_MEMORY_STAT_BY(STAT_MetalPooledBufferMemory, Buffer.GetLength());
@@ -1563,7 +1578,10 @@ FMetalBuffer FMetalResourceHeap::CreateBuffer(uint32 Size, uint32 Alignment, uin
 				{
 					FScopeLock Lock(&Mutex);
                     Buffer = Buffers[Storage].CreatePooledResource(FMetalPooledBufferArgs(Queue->GetDevice(), BlockSize, Flags, StorageMode));
-                    Buffer.SetPurgeableState(mtlpp::PurgeableState::NonVolatile);
+					if (GMetalResourcePurgeInPool)
+					{
+                    	Buffer.SetPurgeableState(mtlpp::PurgeableState::NonVolatile);
+					}
 					DEC_MEMORY_STAT_BY(STAT_MetalBufferUnusedMemory, Buffer.GetLength());
 					DEC_MEMORY_STAT_BY(STAT_MetalPooledBufferUnusedMemory, Buffer.GetLength());
                     INC_MEMORY_STAT_BY(STAT_MetalPooledBufferMemory, Buffer.GetLength());
@@ -1608,7 +1626,10 @@ void FMetalResourceHeap::ReleaseBuffer(FMetalBuffer& Buffer)
 		INC_MEMORY_STAT_BY(STAT_MetalPooledBufferUnusedMemory, Buffer.GetLength());
         DEC_MEMORY_STAT_BY(STAT_MetalPooledBufferMemory, Buffer.GetLength());
 
-        Buffer.SetPurgeableState(mtlpp::PurgeableState::Volatile);
+		if (GMetalResourcePurgeInPool)
+		{
+        	Buffer.SetPurgeableState(mtlpp::PurgeableState::Volatile);
+		}
         
 		switch (StorageMode)
 		{
@@ -1671,7 +1692,7 @@ FMetalTexture FMetalResourceHeap::CreateTexture(mtlpp::TextureDescriptor Desc, F
 
 void FMetalResourceHeap::ReleaseTexture(FMetalSurface* Surface, FMetalTexture& Texture)
 {
-	if (!Texture.GetBuffer() && !Texture.GetParentTexture() && !Texture.GetHeap())
+	if (Texture && !Texture.GetBuffer() && !Texture.GetParentTexture() && !Texture.GetHeap())
 	{
         if (Texture.GetUsage() & mtlpp::TextureUsage::RenderTarget)
         {
diff --git a/Engine/Source/Runtime/Apple/MetalRHI/Private/MetalContext.cpp b/Engine/Source/Runtime/Apple/MetalRHI/Private/MetalContext.cpp
--- a/Engine/Source/Runtime/Apple/MetalRHI/Private/MetalContext.cpp
+++ b/Engine/Source/Runtime/Apple/MetalRHI/Private/MetalContext.cpp
@@ -460,7 +460,7 @@ void FMetalDeviceContext::ClearFreeList()
 			}
 			for ( FMetalTexture& Texture : Pair->UsedTextures )
 			{
-                if (!(Texture.GetBuffer() || Texture.GetParentTexture()) && (Texture.GetStorageMode() != mtlpp::StorageMode::Private))
+                if (!Texture.GetBuffer() && !Texture.GetParentTexture())
 				{
 #if METAL_DEBUG_OPTIONS
 					if (GMetalResourcePurgeOnDelete && !Texture.GetHeap())
@@ -693,8 +693,15 @@ void FMetalDeviceContext::ReleaseTexture(FMetalTexture& Texture)
         if (Texture.GetStorageMode() == mtlpp::StorageMode::Private)
         {
             Heap.ReleaseTexture(nullptr, Texture);
+			
+			// Ensure that the Objective-C handle can't disappear prior to the GPU being done with it without racing with the above
+			if(!ObjectFreeList.Contains(Texture.GetPtr()))
+			{
+				[Texture.GetPtr() retain];
+				ObjectFreeList.Add(Texture.GetPtr());
+			}
         }
-		if(!UsedTextures.Contains(Texture))
+		else if(!UsedTextures.Contains(Texture))
 		{
 			UsedTextures.Add(MoveTemp(Texture));
 		}
