Fix a silly bug that meant MetalRHI's CopyTexture only copied a 1x1 region when the caller expected it to copy the whole surface.

diff --git a/Engine/Source/Runtime/Apple/MetalRHI/Private/MetalTexture.cpp b/Engine/Source/Runtime/Apple/MetalRHI/Private/MetalTexture.cpp
--- a/Engine/Source/Runtime/Apple/MetalRHI/Private/MetalTexture.cpp
+++ b/Engine/Source/Runtime/Apple/MetalRHI/Private/MetalTexture.cpp
@@ -2865,7 +2865,7 @@ void FMetalRHICommandContext::RHICopyTexture(FTextureRHIParamRef SourceTextureRH
 			FMetalSurface* MetalSrcTexture = GetMetalSurfaceFromRHITexture(SourceTextureRHI);
 			FMetalSurface* MetalDestTexture = GetMetalSurfaceFromRHITexture(DestTextureRHI);
 			
-			FIntVector Size = (CopyInfo.Size != FIntVector::ZeroValue) ? CopyInfo.Size : FIntVector(MetalSrcTexture->Texture.GetWidth(), MetalSrcTexture->Texture.GetHeight(), MetalSrcTexture->Texture.GetDepth());
+			FIntVector Size = (CopyInfo.Size != FIntVector::ZeroValue) ? CopyInfo.Size : FIntVector(MetalSrcTexture->SizeX, MetalSrcTexture->SizeY, MetalSrcTexture->SizeZ);
 			
 			mtlpp::Origin SourceOrigin(CopyInfo.SourcePosition.X, CopyInfo.SourcePosition.Y, CopyInfo.SourcePosition.Z);
 			mtlpp::Origin DestinationOrigin(CopyInfo.DestPosition.X, CopyInfo.DestPosition.Y, CopyInfo.DestPosition.Z);
@@ -2894,7 +2894,7 @@ void FMetalRHICommandContext::RHICopyTexture(FTextureRHIParamRef SourceTextureRH
 				{
 					uint32 SourceMipIndex = CopyInfo.SourceMipIndex + MipIndex;
 					uint32 DestMipIndex = CopyInfo.DestMipIndex + MipIndex;
-					mtlpp::Size SourceSize(FMath::Max(CopyInfo.Size.X >> MipIndex, 1), FMath::Max(CopyInfo.Size.Y >> MipIndex, 1), FMath::Max(CopyInfo.Size.Z >> MipIndex, 1));
+					mtlpp::Size SourceSize(FMath::Max(Size.X >> MipIndex, 1), FMath::Max(Size.Y >> MipIndex, 1), FMath::Max(Size.Z >> MipIndex, 1));
 					
 					// Account for create with TexCreate_SRGB flag which could make these different
 					if(SrcTexture.GetPixelFormat() == MetalDestTexture->Texture.GetPixelFormat())
