UE4.22 - Fix deprecation warning
#rnx

diff --git a/Engine/Source/Runtime/Apple/MetalRHI/Private/MetalRenderTarget.cpp b/Engine/Source/Runtime/Apple/MetalRHI/Private/MetalRenderTarget.cpp
--- a/Engine/Source/Runtime/Apple/MetalRHI/Private/MetalRenderTarget.cpp
+++ b/Engine/Source/Runtime/Apple/MetalRHI/Private/MetalRenderTarget.cpp
@@ -114,7 +114,9 @@ void FMetalRHICommandContext::RHICopyToResolveTarget(FTextureRHIParamRef SourceT
 
                 FResolveRect ResolveRect = ResolveParams.Rect;
 
+				PRAGMA_DISABLE_DEPRECATION_WARNINGS
                 SetRenderTargets(RHICmdList, 0, nullptr, DestTextureRHI, ESimpleRenderTargetMode::EClearColorExistingDepth, FExclusiveDepthStencil::DepthWrite_StencilWrite, true);
+				PRAGMA_ENABLE_DEPRECATION_WARNINGS
 
                 FGraphicsPipelineStateInitializer GraphicsPSOInit;
                 RHICmdList.ApplyCachedRenderTargets(GraphicsPSOInit);
