Workaround for some windows having a red edge in the editor on Mac

diff --git a/Engine/Source/Runtime/Apple/MetalRHI/Private/MetalViewport.cpp b/Engine/Source/Runtime/Apple/MetalRHI/Private/MetalViewport.cpp
--- a/Engine/Source/Runtime/Apple/MetalRHI/Private/MetalViewport.cpp
+++ b/Engine/Source/Runtime/Apple/MetalRHI/Private/MetalViewport.cpp
@@ -634,12 +634,12 @@ void FMetalRHIImmediateCommandContext::RHIBeginDrawingViewport(FViewportRHIParam
 	// Set the render target and viewport.
 	if (RenderTargetRHI)
 	{
-		FRHIRenderTargetView RTV(RenderTargetRHI, ERenderTargetLoadAction::ELoad);
+		FRHIRenderTargetView RTV(RenderTargetRHI, GIsEditor ? ERenderTargetLoadAction::EClear : ERenderTargetLoadAction::ELoad);
 		RHISetRenderTargets(1, &RTV, nullptr, 0, NULL);
 	}
 	else
 	{
-		FRHIRenderTargetView RTV(Viewport->GetBackBuffer(EMetalViewportAccessRHI), ERenderTargetLoadAction::ELoad);
+		FRHIRenderTargetView RTV(Viewport->GetBackBuffer(EMetalViewportAccessRHI), GIsEditor ? ERenderTargetLoadAction::EClear : ERenderTargetLoadAction::ELoad);
 		RHISetRenderTargets(1, &RTV, nullptr, 0, NULL);
 	}
 	}
