As a follow-up to 5333983 if you package shaders using native shader libraries and keep the .zip files around (they are stored in the MetaData directory which is retained but not packaged) you can put them in the game's content directory and then when running use the -metalshaderdebug command-line flag to force the game to load shaders from the .zip file as text. This makes it possible to debug shaders in using a build from a build machine rather than having to recook locally which should be a great quality of life improvement for Mac & iOS development. Compiled out in Shipping builds for sanity.

diff --git a/Engine/Source/Runtime/Apple/MetalRHI/Private/MetalShaders.cpp b/Engine/Source/Runtime/Apple/MetalRHI/Private/MetalShaders.cpp
--- a/Engine/Source/Runtime/Apple/MetalRHI/Private/MetalShaders.cpp
+++ b/Engine/Source/Runtime/Apple/MetalRHI/Private/MetalShaders.cpp
@@ -6,6 +6,7 @@
 
 #include "MetalRHIPrivate.h"
 #include "HAL/FileManager.h"
+#include "HAL/PlatformFilemanager.h"
 #include "Misc/Paths.h"
 #include "MetalShaderResources.h"
 #include "MetalResources.h"
@@ -104,6 +105,302 @@ static FMetalCompiledShaderCache& GetMetalCompiledShaderCache()
 	return CompiledShaderCache;
 }
 
+#if !UE_BUILD_SHIPPING
+
+struct FMetalShaderDebugCache
+{
+	static FMetalShaderDebugCache& Get()
+	{
+		static FMetalShaderDebugCache sSelf;
+		return sSelf;
+	}
+	
+	class FMetalShaderDebugZipFile* GetDebugFile(FString Path);
+	ns::String GetShaderCode(uint32 ShaderSrcLen, uint32 ShaderSrcCRC);
+	
+	FCriticalSection Mutex;
+	TMap<FString, class FMetalShaderDebugZipFile*> DebugFiles;
+};
+
+class FMetalShaderDebugZipFile
+{
+	struct FFileEntry
+	{
+		FString Filename;
+		uint32 Crc32;
+		uint64 Length;
+		uint64 Offset;
+		uint32 Time;
+		
+		FFileEntry(const FString& InFilename, uint32 InCrc32, uint64 InLength, uint64 InOffset, uint32 InTime)
+		: Filename(InFilename)
+		, Crc32(InCrc32)
+		, Length(InLength)
+		, Offset(InOffset)
+		, Time(InTime)
+		{}
+	};
+public:
+	FMetalShaderDebugZipFile(FString LibPath)
+	: File(nullptr)
+	{
+		IPlatformFile& PlatformFile = FPlatformFileManager::Get().GetPlatformFile();
+		File = PlatformFile.OpenRead(*LibPath);
+		if(File)
+		{
+			int64 SeekEndOffset = -1;
+			
+			// Write normal end of central directory record
+			const static uint8 EndRecord[] =
+			{
+				0x50, 0x4b, 0x05, 0x06, 0xff, 0xff, 0xff, 0xff,
+				0xff, 0xff, 0xff, 0xff, 0xff, 0xff, 0xff, 0xff,
+				0xff, 0xff, 0xff, 0xff, 0x00, 0x00
+			};
+			
+			SeekEndOffset += sizeof(EndRecord);
+			bool bOK = File->SeekFromEnd(-SeekEndOffset);
+			if (bOK)
+			{
+				TArray<uint8> Data;
+				Data.AddZeroed(sizeof(EndRecord));
+				bOK = File->Read(Data.GetData(), sizeof(EndRecord));
+				if (bOK)
+				{
+					bOK = (FMemory::Memcmp(Data.GetData(), EndRecord, sizeof(EndRecord)) == 0);
+				}
+			}
+			
+			// Write ZIP64 end of central directory locator
+			const static uint8 Locator[] =
+			{
+				0x50, 0x4b, 0x06, 0x07, 0x00, 0x00, 0x00, 0x00,
+			};
+			uint64 DirEndOffset = 0;
+			if (bOK)
+			{
+				SeekEndOffset += sizeof(Locator) + sizeof(uint64) + sizeof(uint32);
+				bOK = File->SeekFromEnd(-SeekEndOffset);
+				if (bOK)
+				{
+					TArray<uint8> Data;
+					Data.AddZeroed(sizeof(Locator));
+					bOK = File->Read(Data.GetData(), sizeof(Locator));
+					if (bOK)
+					{
+						bOK = (FMemory::Memcmp(Data.GetData(), Locator, sizeof(Locator)) == 0);
+					}
+					if (bOK)
+					{
+						bOK = File->Read((uint8*)&DirEndOffset, sizeof(uint64));
+					}
+				}
+			}
+			
+			// Write ZIP64 end of central directory record
+			const static uint8 Record[] =
+			{
+				0x50, 0x4b, 0x06, 0x06, 0x2c, 0x00, 0x00, 0x00,
+				0x00, 0x00, 0x00, 0x00, 0x2d, 0x00, 0x2d, 0x00,
+				0x00, 0x00, 0x00, 0x00, 0x00, 0x00, 0x00, 0x00,
+			};
+			struct FMetalZipRecordData
+			{
+				uint64 FilesNum;
+				uint64 FilesNum2;
+				uint64 DirectorySizeInBytes;
+				uint64 DirStartOffset;
+			} RecordData;
+			if (bOK)
+			{
+				SeekEndOffset += sizeof(Record) + (sizeof(uint64) * 4);
+				bOK = File->SeekFromEnd(-SeekEndOffset);
+				if (bOK)
+				{
+					TArray<uint8> Data;
+					Data.AddZeroed(sizeof(Record));
+					bOK = File->Read(Data.GetData(), sizeof(Record));
+					if (bOK)
+					{
+						bOK = (FMemory::Memcmp(Data.GetData(), Record, sizeof(Record)) == 0);
+					}
+					if (bOK)
+					{
+						bOK = File->Read((uint8*)&RecordData, sizeof(RecordData));
+					}
+				}
+			}
+			
+			if (bOK)
+			{
+				bOK = File->Seek(RecordData.DirStartOffset);
+				if (bOK)
+				{
+					const static uint8 Footer[] =
+					{
+						0x50, 0x4b, 0x01, 0x02, 0x3f, 0x00, 0x2d, 0x00,
+						0x00, 0x00, 0x00, 0x00
+					};
+					const static uint8 Fields[] =
+					{
+						0x20, 0x00, 0x00, 0x00, 0x00, 0x00, 0x00, 0x00,
+						0x20, 0x00, 0x00, 0x00, 0xff, 0xff, 0xff, 0xff
+					};
+					
+					TArray<uint8> FooterData;
+					FooterData.AddZeroed(sizeof(Footer));
+					TArray<uint8> FieldsData;
+					FieldsData.AddZeroed(sizeof(Fields));
+					
+					struct FMetalZipFileHeader {
+						uint32 Time;
+						uint32 CRC;
+						uint64 SizeMarker;
+						uint16 FilenameLen;
+					} __attribute__((packed)) Header;
+					
+					struct FMetalZipFileTrailer {
+						uint16 Flags;
+						uint16 Attribs;
+						uint64 CompressedLen;
+						uint64 UncompressedLen;
+						uint64 Offset;
+						uint32 DiskNum;
+					} __attribute__((packed)) Trailer;
+					
+					static const uint8 FileHeader[] =
+					{
+						0x50, 0x4b, 0x03, 0x04, 0x2d, 0x00, 0x00, 0x00,
+						0x00, 0x00
+					};
+					
+					uint32 FileHeaderFixedSize = sizeof(FileHeader) + sizeof(FMetalZipFileHeader) + sizeof(uint16) + sizeof(FMetalZipFileTrailer);
+					
+					FString Filename;
+					
+					while (bOK && Files.Num() < RecordData.FilesNum && File->Tell() < RecordData.DirStartOffset + RecordData.DirectorySizeInBytes)
+					{
+						bOK = File->Read(FooterData.GetData(), sizeof(Footer));
+						if (bOK)
+						{
+							bOK = (FMemory::Memcmp(FooterData.GetData(), Footer, sizeof(Footer)) == 0);
+						}
+						
+						if (bOK)
+						{
+							bOK = File->Read((uint8*)&Header, sizeof(Header));
+							if (bOK)
+							{
+								bOK = (Header.SizeMarker == (uint64)0xffffffffffffffff);
+							}
+						}
+						
+						if (bOK)
+						{
+							bOK = File->Read(FieldsData.GetData(), sizeof(Fields));
+							if (bOK)
+							{
+								bOK = (FMemory::Memcmp(FieldsData.GetData(), Fields, sizeof(Fields)) == 0);
+							}
+						}
+						
+						if (bOK)
+						{
+							TArray<uint8> FilenameData;
+							FilenameData.AddZeroed(Header.FilenameLen+1);
+							bOK = File->Read(FilenameData.GetData(), Header.FilenameLen);
+							if (bOK)
+							{
+								Filename = UTF8_TO_TCHAR((char const*)FilenameData.GetData());
+							}
+						}
+						
+						if (bOK)
+						{
+							bOK = File->Read((uint8*)&Trailer, sizeof(Trailer));
+							if (bOK)
+							{
+								bOK = (Trailer.Flags == (uint16)0x01 && Trailer.Attribs == (uint16)0x1c && Trailer.DiskNum == 0);
+							}
+						}
+						
+						if (bOK)
+						{
+							FFileEntry NewEntry(Filename, Header.CRC, Trailer.UncompressedLen, Trailer.Offset + FileHeaderFixedSize + Header.FilenameLen, Header.Time);
+							Files.Add(NewEntry);
+						}
+					}
+				}
+			}
+		}
+	}
+
+	~FMetalShaderDebugZipFile()
+	{
+		delete File;
+	}
+	
+	ns::String GetShaderCode(uint32 ShaderSrcLen, uint32 ShaderSrcCRC)
+	{
+		ns::String Source;
+		FScopeLock Lock(&Mutex);
+		FString Name = FString::Printf(TEXT("%u_%u.metal"), ShaderSrcLen, ShaderSrcCRC);
+		for (auto const& Entry : Files)
+		{
+			if (FPaths::GetCleanFilename(Entry.Filename) == Name)
+			{
+				if (File->Seek(Entry.Offset))
+				{
+					TArray<uint8> Data;
+					Data.AddZeroed(Entry.Length+1);
+					if (File->Read(Data.GetData(), Entry.Length))
+					{
+						Source = [NSString stringWithUTF8String:(char const*)Data.GetData()];
+					}
+				}
+				
+				break;
+			}
+		}
+		return Source;
+	}
+	
+private:
+	FCriticalSection Mutex;
+	IFileHandle* File;
+	TArray<FFileEntry> Files;
+};
+
+
+FMetalShaderDebugZipFile* FMetalShaderDebugCache::GetDebugFile(FString Path)
+{
+	FScopeLock Lock(&Mutex);
+	FMetalShaderDebugZipFile* Ref = DebugFiles.FindRef(Path);
+	if (!Ref)
+	{
+		Ref = new FMetalShaderDebugZipFile(Path);
+		DebugFiles.Add(Path, Ref);
+	}
+	return Ref;
+}
+
+ns::String FMetalShaderDebugCache::GetShaderCode(uint32 ShaderSrcLen, uint32 ShaderSrcCRC)
+{
+	ns::String Code;
+	FScopeLock Lock(&Mutex);
+	for (auto const& Ref : DebugFiles)
+	{
+		Code = Ref.Value->GetShaderCode(ShaderSrcLen, ShaderSrcCRC);
+		if (Code)
+		{
+			break;
+		}
+	}
+	return Code;
+}
+
+#endif
+
 NSString* DecodeMetalSourceCode(uint32 CodeSize, TArray<uint8> const& CompressedSource)
 {
 	NSString* GlslCodeNSString = nil;
@@ -249,6 +546,10 @@ void TMetalBaseShader<BaseResourceType, ShaderType>::Init(const TArray<uint8>& I
             memcpy(&CodeSize, UnSourceLen, sizeof(uint32));
 			bHasShaderSource = false;
         }
+		else if(bForceTextShaders)
+		{
+			GlslCodeNSString = [FMetalShaderDebugCache::Get().GetShaderCode(SourceLen, SourceCRC).GetPtr() retain];
+		}
         if (bForceTextShaders)
         {
             bHasShaderSource = (GetSourceCode() != nil);
@@ -1033,7 +1334,19 @@ FMetalShaderLibrary::FMetalShaderLibrary(EShaderPlatform InPlatform, FString con
 , Map(InMap)
 , ShaderLibraryFilename(InShaderLibraryFilename)
 {
+#if !UE_BUILD_SHIPPING
+	DebugFile = nullptr;
+
+	FName PlatformName = LegacyShaderPlatformToShaderFormat(InPlatform);
+	FString LibName = FString::Printf(TEXT("%s_%s"), *Name, *PlatformName.GetPlainNameString());
+	LibName.ToLowerInline();
+	FString Path = FPaths::ProjectContentDir() / LibName + TEXT(".zip");
 	
+	if ( IFileManager::Get().FileExists(*Path) )
+	{
+		DebugFile = FMetalShaderDebugCache::Get().GetDebugFile(Path);
+	}
+#endif
 }
 
 FMetalShaderLibrary::~FMetalShaderLibrary()
diff --git a/Engine/Source/Runtime/Apple/MetalRHI/Public/MetalResources.h b/Engine/Source/Runtime/Apple/MetalRHI/Public/MetalResources.h
--- a/Engine/Source/Runtime/Apple/MetalRHI/Public/MetalResources.h
+++ b/Engine/Source/Runtime/Apple/MetalRHI/Public/MetalResources.h
@@ -1104,6 +1104,9 @@ private:
 	
 private:
 	TArray<mtlpp::Library> Library;
+#if !UE_BUILD_SHIPPING
+	class FMetalShaderDebugZipFile* DebugFile;
+#endif
 	FMetalShaderMap Map;
 	FString ShaderLibraryFilename;
 };
