Add missing !UE_BUILD_SHIPPING to fix the shipping build.

diff --git a/Engine/Source/Runtime/Apple/MetalRHI/Private/MetalShaders.cpp b/Engine/Source/Runtime/Apple/MetalRHI/Private/MetalShaders.cpp
--- a/Engine/Source/Runtime/Apple/MetalRHI/Private/MetalShaders.cpp
+++ b/Engine/Source/Runtime/Apple/MetalRHI/Private/MetalShaders.cpp
@@ -546,10 +546,12 @@ void TMetalBaseShader<BaseResourceType, ShaderType>::Init(const TArray<uint8>& I
             memcpy(&CodeSize, UnSourceLen, sizeof(uint32));
 			bHasShaderSource = false;
         }
+#if !UE_BUILD_SHIPPING
 		else if(bForceTextShaders)
 		{
 			GlslCodeNSString = [FMetalShaderDebugCache::Get().GetShaderCode(SourceLen, SourceCRC).GetPtr() retain];
 		}
+#endif
         if (bForceTextShaders)
         {
             bHasShaderSource = (GetSourceCode() != nil);
