Make sure we don't access the GPU profiler after it has been destructed when shutting down MetalRHI.

diff --git a/Engine/Source/Runtime/Apple/MetalRHI/Private/MetalPipeline.cpp b/Engine/Source/Runtime/Apple/MetalRHI/Private/MetalPipeline.cpp
--- a/Engine/Source/Runtime/Apple/MetalRHI/Private/MetalPipeline.cpp
+++ b/Engine/Source/Runtime/Apple/MetalRHI/Private/MetalPipeline.cpp
@@ -855,10 +855,10 @@ static FMetalShaderPipeline* CreateMTLRenderPipeline(bool const bSync, FMetalGra
 				}
 #endif
 #if METAL_DEBUG_OPTIONS
-				if (GetMetalDeviceContext().GetCommandQueue().GetRuntimeDebuggingLevel() >= EMetalDebugLevelFastValidation METAL_STATISTIC(|| GetMetalDeviceContext().GetCommandQueue().GetStatistics()))
+				if (GetMetalDeviceContext().GetCommandQueue().GetRuntimeDebuggingLevel() >= EMetalDebugLevelFastValidation METAL_STATISTICS_ONLY(|| GetMetalDeviceContext().GetCommandQueue().GetStatistics()))
 				{
 					mtlpp::AutoReleasedComputePipelineReflection Reflection;
-					ComputeOption = mtlpp::PipelineOption::ArgumentInfo|mtlpp::PipelineOption::BufferTypeInfo METAL_STATISTIC(|NSUInteger(EMTLPipelineStats));
+					ComputeOption = mtlpp::PipelineOption::ArgumentInfo|mtlpp::PipelineOption::BufferTypeInfo METAL_STATISTICS_ONLY(|NSUInteger(EMTLPipelineStats));
 					Pipeline->ComputePipelineState = Device.NewComputePipelineState(ComputePipelineDesc, (mtlpp::PipelineOption)ComputeOption, &Reflection, &AutoError);
 					Pipeline->ComputePipelineReflection = Reflection;
 				}
@@ -982,9 +982,9 @@ static FMetalShaderPipeline* CreateMTLRenderPipeline(bool const bSync, FMetalGra
 #if METAL_DEBUG_OPTIONS
 		mtlpp::AutoReleasedRenderPipelineReflection OutReflection;
 		Reflection = &OutReflection;
-        if (GetMetalDeviceContext().GetCommandQueue().GetRuntimeDebuggingLevel() >= EMetalDebugLevelFastValidation METAL_STATISTIC(|| GetMetalDeviceContext().GetCommandQueue().GetStatistics()))
+        if (GetMetalDeviceContext().GetCommandQueue().GetRuntimeDebuggingLevel() >= EMetalDebugLevelFastValidation METAL_STATISTICS_ONLY(|| GetMetalDeviceContext().GetCommandQueue().GetStatistics()))
         {
-        	RenderOption = mtlpp::PipelineOption::ArgumentInfo|mtlpp::PipelineOption::BufferTypeInfo METAL_STATISTIC(|NSUInteger(EMTLPipelineStats));
+        	RenderOption = mtlpp::PipelineOption::ArgumentInfo|mtlpp::PipelineOption::BufferTypeInfo METAL_STATISTICS_ONLY(|NSUInteger(EMTLPipelineStats));
         }
 #endif
 
diff --git a/Engine/Source/Runtime/Apple/MetalRHI/Private/MetalRHIPrivate.h b/Engine/Source/Runtime/Apple/MetalRHI/Private/MetalRHIPrivate.h
--- a/Engine/Source/Runtime/Apple/MetalRHI/Private/MetalRHIPrivate.h
+++ b/Engine/Source/Runtime/Apple/MetalRHI/Private/MetalRHIPrivate.h
@@ -124,9 +124,11 @@ extern bool GMetalSupportsTileShaders;
 #endif
 
 #if METAL_STATISTICS
-#define METAL_STATISTIC(Code) Code
+#define METAL_STATISTIC(Code) if (GIsMetalInitialized) { Code; }
+#define METAL_STATISTICS_ONLY(Code) Code
 #else
 #define METAL_STATISTIC(Code)
+#define METAL_STATISTICS_ONLY(Code)
 #endif
 
 #define UNREAL_TO_METAL_BUFFER_INDEX(Index) ((MaxMetalStreams - 1) - Index)
diff --git a/Engine/Source/Runtime/Apple/MetalRHI/Private/MetalShaders.cpp b/Engine/Source/Runtime/Apple/MetalRHI/Private/MetalShaders.cpp
--- a/Engine/Source/Runtime/Apple/MetalRHI/Private/MetalShaders.cpp
+++ b/Engine/Source/Runtime/Apple/MetalRHI/Private/MetalShaders.cpp
@@ -924,12 +924,12 @@ FMetalShaderPipeline* FMetalComputeShader::GetPipeline()
 		
 		METAL_GPUPROFILE(FScopedMetalCPUStats CPUStat(FString::Printf(TEXT("NewComputePipeline: %d_%d"), SourceLen, SourceCRC)));
     #if METAL_DEBUG_OPTIONS
-        if (GetMetalDeviceContext().GetCommandQueue().GetRuntimeDebuggingLevel() >= EMetalDebugLevelFastValidation METAL_STATISTIC(|| GetMetalDeviceContext().GetCommandQueue().GetStatistics()))
+        if (GetMetalDeviceContext().GetCommandQueue().GetRuntimeDebuggingLevel() >= EMetalDebugLevelFastValidation METAL_STATISTICS_ONLY(|| GetMetalDeviceContext().GetCommandQueue().GetStatistics()))
         {
 			ns::AutoReleasedError ComputeError;
             mtlpp::AutoReleasedComputePipelineReflection ComputeReflection;
 			
-			NSUInteger ComputeOption = mtlpp::PipelineOption::ArgumentInfo|mtlpp::PipelineOption::BufferTypeInfo METAL_STATISTIC(|NSUInteger(EMTLPipelineStats));
+			NSUInteger ComputeOption = mtlpp::PipelineOption::ArgumentInfo|mtlpp::PipelineOption::BufferTypeInfo METAL_STATISTICS_ONLY(|NSUInteger(EMTLPipelineStats));
 			Kernel = GetMetalDeviceContext().GetDevice().NewComputePipelineState(Descriptor, mtlpp::PipelineOption(ComputeOption), &ComputeReflection, &ComputeError);
 			Error = ComputeError;
 			Reflection = ComputeReflection;
