Potential fix for iOS latency issues

diff --git a/Engine/Source/Runtime/Apple/AudioMixerAudioUnit/Private/AudioMixerPlatformAudioUnit.h b/Engine/Source/Runtime/Apple/AudioMixerAudioUnit/Private/AudioMixerPlatformAudioUnit.h
--- a/Engine/Source/Runtime/Apple/AudioMixerAudioUnit/Private/AudioMixerPlatformAudioUnit.h
+++ b/Engine/Source/Runtime/Apple/AudioMixerAudioUnit/Private/AudioMixerPlatformAudioUnit.h
@@ -43,7 +43,8 @@ namespace Audio
         virtual int32 GetNumFrames(const int32 InNumReqestedFrames) override;
         virtual void ResumeContext() override;
         virtual void SuspendContext() override;
-        //~ End IAudioMixerPlatformInterface
+		virtual bool DisablePCMAudioCaching() const override { return true; }
+		//~ End IAudioMixerPlatformInterface
         
 	private:
 		AudioStreamBasicDescription OutputFormat;
