Back out changelist 6599192

diff --git a/Engine/Source/Runtime/Apple/AudioMixerAudioUnit/Private/AudioMixerPlatformAudioUnit.h b/Engine/Source/Runtime/Apple/AudioMixerAudioUnit/Private/AudioMixerPlatformAudioUnit.h
--- a/Engine/Source/Runtime/Apple/AudioMixerAudioUnit/Private/AudioMixerPlatformAudioUnit.h
+++ b/Engine/Source/Runtime/Apple/AudioMixerAudioUnit/Private/AudioMixerPlatformAudioUnit.h
@@ -43,8 +43,7 @@ namespace Audio
         virtual int32 GetNumFrames(const int32 InNumReqestedFrames) override;
         virtual void ResumeContext() override;
         virtual void SuspendContext() override;
-		virtual bool DisablePCMAudioCaching() const override { return true; }
-		//~ End IAudioMixerPlatformInterface
+        //~ End IAudioMixerPlatformInterface
         
 	private:
 		AudioStreamBasicDescription OutputFormat;
