Very primitive Metal implementation of RHICopyTexture - can only cope with 2D & cube textures at the moment.

diff --git a/Engine/Source/Runtime/Apple/MetalRHI/Private/MetalTexture.cpp b/Engine/Source/Runtime/Apple/MetalRHI/Private/MetalTexture.cpp
--- a/Engine/Source/Runtime/Apple/MetalRHI/Private/MetalTexture.cpp
+++ b/Engine/Source/Runtime/Apple/MetalRHI/Private/MetalTexture.cpp
@@ -2824,3 +2824,88 @@ void FMetalDynamicRHI::RHICopySubTextureRegion(FTexture2DRHIParamRef SourceTextu
 	}
 }
 
+
+
+void FMetalRHICommandContext::RHICopyTexture(FTextureRHIParamRef SourceTextureRHI, FTextureRHIParamRef DestTextureRHI, const FRHICopyTextureInfo& CopyInfo)
+{
+	if (!SourceTextureRHI || !DestTextureRHI || SourceTextureRHI == DestTextureRHI)
+	{
+		// no need to do anything (silently ignored)
+		return;
+	}
+	
+	RHITransitionResources(EResourceTransitionAccess::EReadable, &SourceTextureRHI, 1);
+	
+	@autoreleasepool {
+		check(SourceTextureRHI);
+		check(DestTextureRHI);
+		
+		// Can only copy 2d faces in Metal.
+		check(CopyInfo.Size.Z == 1);
+		
+		// Can only copy between 2D texture types - anything with depth won't work.
+		check(CopyInfo.SourcePosition.Z == 0);
+		check(CopyInfo.DestPosition.Z == 0);
+		
+		if(SourceTextureRHI->GetFormat() == DestTextureRHI->GetFormat())
+		{
+			FMetalSurface* MetalSrcTexture = GetMetalSurfaceFromRHITexture(SourceTextureRHI);
+			FMetalSurface* MetalDestTexture = GetMetalSurfaceFromRHITexture(DestTextureRHI);
+			
+			FIntVector Size = (CopyInfo.Size != FIntVector::ZeroValue) ? CopyInfo.Size : FIntVector(MetalSrcTexture->Texture.GetWidth(), MetalSrcTexture->Texture.GetHeight(), MetalSrcTexture->Texture.GetDepth());
+			
+			mtlpp::Origin SourceOrigin(CopyInfo.SourcePosition.X, CopyInfo.SourcePosition.Y, CopyInfo.SourcePosition.Z);
+			mtlpp::Origin DestinationOrigin(CopyInfo.DestPosition.X, CopyInfo.DestPosition.Y, CopyInfo.DestPosition.Z);
+			
+			for (uint32 SliceIndex = 0; SliceIndex < CopyInfo.NumSlices; ++SliceIndex)
+			{
+				uint32 SourceSliceIndex = CopyInfo.SourceSliceIndex + SliceIndex;
+				uint32 DestSliceIndex = CopyInfo.DestSliceIndex + SliceIndex;
+
+				for (uint32 MipIndex = 0; MipIndex < CopyInfo.NumMips; ++MipIndex)
+				{
+					uint32 SourceMipIndex = CopyInfo.SourceMipIndex + MipIndex;
+					uint32 DestMipIndex = CopyInfo.DestMipIndex + MipIndex;
+					mtlpp::Size SourceSize(FMath::Max(CopyInfo.Size.X >> MipIndex, 1), FMath::Max(CopyInfo.Size.Y >> MipIndex, 1), FMath::Max(CopyInfo.Size.Z >> MipIndex, 1));
+					
+					// Account for create with TexCreate_SRGB flag which could make these different
+					if(MetalSrcTexture->Texture.GetPixelFormat() == MetalDestTexture->Texture.GetPixelFormat())
+					{
+						GetInternalContext().CopyFromTextureToTexture(MetalSrcTexture->Texture, SourceSliceIndex, SourceMipIndex, SourceOrigin,SourceSize,MetalDestTexture->Texture, DestSliceIndex, DestMipIndex, DestinationOrigin);
+					}
+					else
+					{
+						// Linear and sRGB mismatch then try to go via metal buffer
+						// Modified clone of logic from MetalRenderTarget.cpp
+						uint32 BytesPerPixel = (MetalSrcTexture->PixelFormat != PF_DepthStencil) ? GPixelFormats[MetalSrcTexture->PixelFormat].BlockBytes : 1;
+						const uint32 Stride = BytesPerPixel * SourceSize.width;
+						const uint32 Alignment = PLATFORM_MAC ? 1u : 64u;
+						const uint32 AlignedStride = ((Stride - 1) & ~(Alignment - 1)) + Alignment;
+						const uint32 BytesPerImage = AlignedStride *  SourceSize.height;
+						
+						FMetalBuffer Buffer = GetMetalDeviceContext().CreatePooledBuffer(FMetalPooledBufferArgs(GetInternalContext().GetDevice(), BytesPerImage, mtlpp::StorageMode::Shared));
+						
+						check(Buffer);
+						
+						mtlpp::BlitOption Options = mtlpp::BlitOption::None;
+#if !PLATFORM_MAC
+						if (MetalSrcTexture->Surface.Texture.GetPixelFormat() >= mtlpp::PixelFormat::PVRTC_RGB_2BPP && MetalSrcTexture->Surface.Texture.GetPixelFormat() <= mtlpp::PixelFormat::PVRTC_RGBA_4BPP_sRGB)
+						{
+							Options = mtlpp::BlitOption::RowLinearPVRTC;
+						}
+#endif
+						GetInternalContext().CopyFromTextureToBuffer(MetalSrcTexture->Texture, SourceSliceIndex, SourceMipIndex, SourceOrigin, SourceSize, Buffer, 0, AlignedStride, BytesPerImage, Options);
+						GetInternalContext().CopyFromBufferToTexture(Buffer, 0, Stride, BytesPerImage, SourceSize, MetalDestTexture->Texture, DestSliceIndex, DestMipIndex, DestinationOrigin, Options);
+						
+						GetMetalDeviceContext().ReleaseBuffer(Buffer);
+					}
+				}
+			}
+		}
+		else
+		{
+			UE_LOG(LogMetal, Warning, TEXT("RHICopyTexture Source <-> Destination texture format mismatch"));
+		}
+	}
+}
+
diff --git a/Engine/Source/Runtime/Apple/MetalRHI/Public/MetalRHIContext.h b/Engine/Source/Runtime/Apple/MetalRHI/Public/MetalRHIContext.h
--- a/Engine/Source/Runtime/Apple/MetalRHI/Public/MetalRHIContext.h
+++ b/Engine/Source/Runtime/Apple/MetalRHI/Public/MetalRHIContext.h
@@ -47,6 +47,8 @@ public:
 	/** Clears a UAV to the multi-component value provided. */
 	virtual void RHIClearTinyUAV(FUnorderedAccessViewRHIParamRef UnorderedAccessViewRHI, const uint32* Values) final override;
 	
+	virtual void RHICopyTexture(FTextureRHIParamRef SourceTextureRHI, FTextureRHIParamRef DestTextureRHI, const FRHICopyTextureInfo& CopyInfo) final override;
+	
 	/**
 	 * Resolves from one texture to another.
 	 * @param SourceTexture - texture to resolve from, 0 is silenty ignored
