Looks like RHICopyTexture should work for 2DArray and probably CubeArray, but 3D textures might not - until we actually try it on one I can't be sure. Can probably use texture-views to avoid the buffer path now that we don't have to support older iOS/macOS versions that lack that feature.

diff --git a/Engine/Source/Runtime/Apple/MetalRHI/Private/MetalTexture.cpp b/Engine/Source/Runtime/Apple/MetalRHI/Private/MetalTexture.cpp
--- a/Engine/Source/Runtime/Apple/MetalRHI/Private/MetalTexture.cpp
+++ b/Engine/Source/Runtime/Apple/MetalRHI/Private/MetalTexture.cpp
@@ -2840,13 +2840,6 @@ void FMetalRHICommandContext::RHICopyTexture(FTextureRHIParamRef SourceTextureRH
 		check(SourceTextureRHI);
 		check(DestTextureRHI);
 		
-		// Can only copy 2d faces in Metal.
-		check(CopyInfo.Size.Z == 1);
-		
-		// Can only copy between 2D texture types - anything with depth won't work.
-		check(CopyInfo.SourcePosition.Z == 0);
-		check(CopyInfo.DestPosition.Z == 0);
-		
 		if(SourceTextureRHI->GetFormat() == DestTextureRHI->GetFormat())
 		{
 			FMetalSurface* MetalSrcTexture = GetMetalSurfaceFromRHITexture(SourceTextureRHI);
@@ -2882,8 +2875,9 @@ void FMetalRHICommandContext::RHICopyTexture(FTextureRHIParamRef SourceTextureRH
 						const uint32 Alignment = PLATFORM_MAC ? 1u : 64u;
 						const uint32 AlignedStride = ((Stride - 1) & ~(Alignment - 1)) + Alignment;
 						const uint32 BytesPerImage = AlignedStride *  SourceSize.height;
+						const uint32 DataSize = BytesPerImage * SourceSize.depth;
 						
-						FMetalBuffer Buffer = GetMetalDeviceContext().CreatePooledBuffer(FMetalPooledBufferArgs(GetInternalContext().GetDevice(), BytesPerImage, mtlpp::StorageMode::Shared));
+						FMetalBuffer Buffer = GetMetalDeviceContext().CreatePooledBuffer(FMetalPooledBufferArgs(GetInternalContext().GetDevice(), DataSize, mtlpp::StorageMode::Shared));
 						
 						check(Buffer);
 						
