Use texture views in RHICopySubTextureRegion  and RHICopyTexture to avoid a slower path that copies via a temporary buffer.

diff --git a/Engine/Source/Runtime/Apple/MetalRHI/Private/MetalTexture.cpp b/Engine/Source/Runtime/Apple/MetalRHI/Private/MetalTexture.cpp
--- a/Engine/Source/Runtime/Apple/MetalRHI/Private/MetalTexture.cpp
+++ b/Engine/Source/Runtime/Apple/MetalRHI/Private/MetalTexture.cpp
@@ -2785,10 +2785,25 @@ void FMetalDynamicRHI::RHICopySubTextureRegion(FTexture2DRHIParamRef SourceTextu
 			check(DestinationSize.width == SourceSize.width);
 			check(DestinationSize.height == SourceSize.height);
 			
+			FMetalTexture SrcTexture;
+			mtlpp::TextureUsage Usage = MetalSrcTexture->Surface.Texture.GetUsage();
+			if(Usage & mtlpp::TextureUsage::PixelFormatView)
+			{
+				ns::Range Slices(0, MetalSrcTexture->Surface.Texture.GetArrayLength() * (MetalSrcTexture->Surface.bIsCubemap ? 6 : 1));
+				if(MetalSrcTexture->Surface.Texture.GetPixelFormat() != MetalDestTexture->Surface.Texture.GetPixelFormat())
+				{
+					SrcTexture = MetalSrcTexture->Surface.Texture.NewTextureView(MetalDestTexture->Surface.Texture.GetPixelFormat(), MetalSrcTexture->Surface.Texture.GetTextureType(), ns::Range(0, MetalSrcTexture->Surface.Texture.GetMipmapLevelCount()), Slices);
+				}
+			}
+			if (!SrcTexture)
+			{
+				SrcTexture = MetalSrcTexture->Surface.Texture;
+			}
+			
 			// Account for create with TexCreate_SRGB flag which could make these different
-			if(MetalSrcTexture->Surface.Texture.GetPixelFormat() == MetalDestTexture->Surface.Texture.GetPixelFormat())
+			if(SrcTexture.GetPixelFormat() == MetalDestTexture->Surface.Texture.GetPixelFormat())
 			{
-				ImmediateContext.GetInternalContext().CopyFromTextureToTexture(MetalSrcTexture->Surface.Texture, 0, 0, SourceOrigin,SourceSize,MetalDestTexture->Surface.Texture, 0, 0, DestinationOrigin);
+				ImmediateContext.GetInternalContext().CopyFromTextureToTexture(SrcTexture, 0, 0, SourceOrigin,SourceSize,MetalDestTexture->Surface.Texture, 0, 0, DestinationOrigin);
 			}
 			else
 			{
@@ -2816,6 +2831,11 @@ void FMetalDynamicRHI::RHICopySubTextureRegion(FTexture2DRHIParamRef SourceTextu
 				
 				GetMetalDeviceContext().ReleaseBuffer(Buffer);
 			}
+			
+			if (SrcTexture != MetalSrcTexture->Surface.Texture)
+			{
+				SafeReleaseMetalTexture(SrcTexture);
+			}
 		}
 		else
 		{
@@ -2849,6 +2869,21 @@ void FMetalRHICommandContext::RHICopyTexture(FTextureRHIParamRef SourceTextureRH
 			
 			mtlpp::Origin SourceOrigin(CopyInfo.SourcePosition.X, CopyInfo.SourcePosition.Y, CopyInfo.SourcePosition.Z);
 			mtlpp::Origin DestinationOrigin(CopyInfo.DestPosition.X, CopyInfo.DestPosition.Y, CopyInfo.DestPosition.Z);
+
+			FMetalTexture SrcTexture;
+			mtlpp::TextureUsage Usage = MetalSrcTexture->Texture.GetUsage();
+			if(Usage & mtlpp::TextureUsage::PixelFormatView)
+			{
+				ns::Range Slices(0, MetalSrcTexture->Texture.GetArrayLength() * (MetalSrcTexture->bIsCubemap ? 6 : 1));
+				if(MetalSrcTexture->Texture.GetPixelFormat() != MetalDestTexture->Texture.GetPixelFormat())
+				{
+					SrcTexture = MetalSrcTexture->Texture.NewTextureView(MetalDestTexture->Texture.GetPixelFormat(), MetalSrcTexture->Texture.GetTextureType(), ns::Range(0, MetalSrcTexture->Texture.GetMipmapLevelCount()), Slices);
+				}
+			}
+			if (!SrcTexture)
+			{
+				SrcTexture = MetalSrcTexture->Texture;
+			}
 			
 			for (uint32 SliceIndex = 0; SliceIndex < CopyInfo.NumSlices; ++SliceIndex)
 			{
@@ -2862,9 +2897,9 @@ void FMetalRHICommandContext::RHICopyTexture(FTextureRHIParamRef SourceTextureRH
 					mtlpp::Size SourceSize(FMath::Max(CopyInfo.Size.X >> MipIndex, 1), FMath::Max(CopyInfo.Size.Y >> MipIndex, 1), FMath::Max(CopyInfo.Size.Z >> MipIndex, 1));
 					
 					// Account for create with TexCreate_SRGB flag which could make these different
-					if(MetalSrcTexture->Texture.GetPixelFormat() == MetalDestTexture->Texture.GetPixelFormat())
+					if(SrcTexture.GetPixelFormat() == MetalDestTexture->Texture.GetPixelFormat())
 					{
-						GetInternalContext().CopyFromTextureToTexture(MetalSrcTexture->Texture, SourceSliceIndex, SourceMipIndex, SourceOrigin,SourceSize,MetalDestTexture->Texture, DestSliceIndex, DestMipIndex, DestinationOrigin);
+						GetInternalContext().CopyFromTextureToTexture(SrcTexture, SourceSliceIndex, SourceMipIndex, SourceOrigin,SourceSize,MetalDestTexture->Texture, DestSliceIndex, DestMipIndex, DestinationOrigin);
 					}
 					else
 					{
@@ -2895,6 +2930,11 @@ void FMetalRHICommandContext::RHICopyTexture(FTextureRHIParamRef SourceTextureRH
 					}
 				}
 			}
+			
+			if (SrcTexture != MetalSrcTexture->Texture)
+			{
+				SafeReleaseMetalTexture(SrcTexture);
+			}
 		}
 		else
 		{
