Remove all the obsolete optional features from EMetalFeatures that are now universally available (or which depend solely on macOS vs. iOS) and all the associated code that branched based on them. This makes iOS and macOS more consistent as well as removing chunks of code that was no longer necessary.

diff --git a/Engine/Source/Runtime/Apple/MetalRHI/Private/MetalCommandEncoder.cpp b/Engine/Source/Runtime/Apple/MetalRHI/Private/MetalCommandEncoder.cpp
--- a/Engine/Source/Runtime/Apple/MetalRHI/Private/MetalCommandEncoder.cpp
+++ b/Engine/Source/Runtime/Apple/MetalRHI/Private/MetalCommandEncoder.cpp
@@ -107,8 +107,6 @@ void FMetalCommandEncoder::Reset(void)
 		RenderPassDesc = nil;
 	}
 	
-	static bool bDeferredStoreActions = CommandList.GetCommandQueue().SupportsFeature(EMetalFeaturesDeferredStoreActions);
-	if (bDeferredStoreActions)
 	{
 		for (uint32 i = 0; i < MaxSimultaneousRenderTargets; i++)
 		{
@@ -504,8 +502,7 @@ TRefCountPtr<FMetalFence> FMetalCommandEncoder::EndEncoding(void)
 			if (RenderCommandEncoder)
 			{
 				check(!bSupportsFences || EncoderFence || !CommandList.IsImmediate());
-				static bool bDeferredStoreActions = CommandList.GetCommandQueue().SupportsFeature(EMetalFeaturesDeferredStoreActions);
-				if (bDeferredStoreActions && ParallelRenderCommandEncoder.GetPtr() == nil)
+				if (ParallelRenderCommandEncoder.GetPtr() == nil)
 				{
 					check(RenderPassDesc);
 					
@@ -611,8 +608,6 @@ TRefCountPtr<FMetalFence> FMetalCommandEncoder::EndEncoding(void)
 
 			if (ParallelRenderCommandEncoder && IsImmediate())
 			{
-				static bool bDeferredStoreActions = CommandList.GetCommandQueue().SupportsFeature(EMetalFeaturesDeferredStoreActions);
-				if (bDeferredStoreActions)
 				{
 					check(RenderPassDesc);
 					
@@ -1125,9 +1120,6 @@ void FMetalCommandEncoder::SetRenderPassDescriptor(mtlpp::RenderPassDescriptor R
 	{
 		SafeReleaseMetalRenderPassDescriptor(RenderPassDesc);
 		RenderPassDesc = RenderPass;
-		
-		static bool bDeferredStoreActions = CommandList.GetCommandQueue().SupportsFeature(EMetalFeaturesDeferredStoreActions);
-		if (bDeferredStoreActions)
 		{
 			for (uint32 i = 0; i < MaxSimultaneousRenderTargets; i++)
 			{
@@ -1156,8 +1148,6 @@ void FMetalCommandEncoder::SetRenderPassDescriptor(mtlpp::RenderPassDescriptor R
 void FMetalCommandEncoder::SetRenderPassStoreActions(mtlpp::StoreAction const* const ColorStore, mtlpp::StoreAction const DepthStore, mtlpp::StoreAction const StencilStore)
 {
 	check(RenderPassDesc);
-	static bool bDeferredStoreActions = CommandList.GetCommandQueue().SupportsFeature(EMetalFeaturesDeferredStoreActions);
-	if (bDeferredStoreActions)
 	{
 		for (uint32 i = 0; i < MaxSimultaneousRenderTargets; i++)
 		{
diff --git a/Engine/Source/Runtime/Apple/MetalRHI/Private/MetalCommandQueue.cpp b/Engine/Source/Runtime/Apple/MetalRHI/Private/MetalCommandQueue.cpp
--- a/Engine/Source/Runtime/Apple/MetalRHI/Private/MetalCommandQueue.cpp
+++ b/Engine/Source/Runtime/Apple/MetalRHI/Private/MetalCommandQueue.cpp
@@ -61,14 +61,10 @@ FMetalCommandQueue::FMetalCommandQueue(mtlpp::Device InDevice, uint32 const MaxN
 	NSOperatingSystemVersion Vers = [[NSProcessInfo processInfo] operatingSystemVersion];
 	if(Vers.majorVersion >= 9)
 	{
-		Features = EMetalFeaturesSeparateStencil | EMetalFeaturesSetBufferOffset | EMetalFeaturesResourceOptions | EMetalFeaturesDepthStencilBlitOptions | EMetalFeaturesShaderVersions | EMetalFeaturesSetBytes;
+		Features = EMetalFeaturesSetBufferOffset | EMetalFeaturesSetBytes;
 
 #if PLATFORM_TVOS
         Features &= ~(EMetalFeaturesSetBytes);
-		if(Device.SupportsFeatureSet(mtlpp::FeatureSet::tvOS_GPUFamily1_v2))
-		{
-			Features |= EMetalFeaturesStencilView | EMetalFeaturesGraphicsUAVs | EMetalFeaturesFunctionConstants | EMetalFeaturesMemoryLessResources;
-		}
 		
 		if(Device.SupportsFeatureSet(mtlpp::FeatureSet::tvOS_GPUFamily2_v1))
 		{
@@ -77,21 +73,14 @@ FMetalCommandQueue::FMetalCommandQueue(mtlpp::Device InDevice, uint32 const MaxN
 		
 		if(Vers.majorVersion > 10)
 		{
-			Features |= EMetalFeaturesGPUCommandBufferTimes;
-			Features |= EMetalFeaturesLinearTextures;
 			Features |= EMetalFeaturesPrivateBufferSubAllocation;
-			Features |= EMetalFeaturesDeferredStoreActions | EMetalFeaturesCombinedDepthStencil;
 			
 			if(Vers.majorVersion >= 11)
 			{
 				Features |= EMetalFeaturesGPUCaptureManager | EMetalFeaturesBufferSubAllocation | EMetalFeaturesParallelRenderEncoders | EMetalFeaturesPipelineBufferMutability;
 				
 				GMetalFColorVertexFormat = mtlpp::VertexFormat::UChar4Normalized_BGRA;
-				
-				if (MaxShaderVersion >= 3)
-				{
-					Features |= EMetalFeaturesLinearTextureUAVs;
-				}
+
 				if (Vers.majorVersion >= 12)
 				{
 					Features |= EMetalFeaturesMaxThreadsPerThreadgroup;
@@ -113,8 +102,6 @@ FMetalCommandQueue::FMetalCommandQueue(mtlpp::Device InDevice, uint32 const MaxN
 		
 		if(Device.SupportsFeatureSet(mtlpp::FeatureSet::iOS_GPUFamily3_v2) || Device.SupportsFeatureSet(mtlpp::FeatureSet::iOS_GPUFamily2_v3) || Device.SupportsFeatureSet(mtlpp::FeatureSet::iOS_GPUFamily1_v3))
 		{
-			Features |= EMetalFeaturesStencilView | EMetalFeaturesFunctionConstants | EMetalFeaturesGraphicsUAVs | EMetalFeaturesMemoryLessResources;
-			
 			if (FParse::Param(FCommandLine::Get(),TEXT("metalfence")))
 			{
 				Features |= EMetalFeaturesFences;
@@ -133,37 +120,32 @@ FMetalCommandQueue::FMetalCommandQueue(mtlpp::Device InDevice, uint32 const MaxN
 		
 		if(Vers.majorVersion > 10 || (Vers.majorVersion == 10 && Vers.minorVersion >= 3))
         {
-			Features |= EMetalFeaturesGPUCommandBufferTimes;
-
-			Features |= EMetalFeaturesLinearTextures;
-			// InjectCurves() does not work with this
-			//Features |= EMetalFeaturesEfficientBufferBlits;
+			// Turning the below option on will allocate more buffer memory which isn't generally desirable on iOS
+			// Features |= EMetalFeaturesEfficientBufferBlits;
+			
+			// These options are fine however as thye just change how we allocate small buffers
             Features |= EMetalFeaturesBufferSubAllocation;
 			Features |= EMetalFeaturesPrivateBufferSubAllocation;
 			
-			if(Device.SupportsFeatureSet(mtlpp::FeatureSet::iOS_GPUFamily3_v2) || Device.SupportsFeatureSet(mtlpp::FeatureSet::iOS_GPUFamily2_v3) || Device.SupportsFeatureSet(mtlpp::FeatureSet::iOS_GPUFamily1_v3))
-			{
-				Features |= EMetalFeaturesDeferredStoreActions | EMetalFeaturesCombinedDepthStencil;
-			}
-			
 			if(Vers.majorVersion >= 11)
 			{
 				GMetalFColorVertexFormat = mtlpp::VertexFormat::UChar4Normalized_BGRA;
 				
 				Features |= EMetalFeaturesPresentMinDuration | EMetalFeaturesGPUCaptureManager | EMetalFeaturesBufferSubAllocation | EMetalFeaturesParallelRenderEncoders | EMetalFeaturesPipelineBufferMutability;
 				
-				// Turn on Linear Texture UAVs! Avoids the need to have function-constants which reduces initial runtime shader compile time
-				if (MaxShaderVersion >= 3)
-				{
-					Features |= EMetalFeaturesLinearTextureUAVs;
-				}
-				
 				// Turn on Texture Buffers! These are faster on the GPU as we don't need to do out-of-bounds tests but require Metal 2.1 and macOS 10.14
 				if (Vers.majorVersion >= 12)
 				{
 					Features |= EMetalFeaturesMaxThreadsPerThreadgroup;
-					Features |= EMetalFeaturesFences;
-					Features |= EMetalFeaturesHeaps;
+                    if (!FParse::Param(FCommandLine::Get(),TEXT("nometalfence")))
+                    {
+                        Features |= EMetalFeaturesFences;
+                    }
+                    
+                    if (!FParse::Param(FCommandLine::Get(),TEXT("nometalheap")))
+                    {
+                        Features |= EMetalFeaturesHeaps;
+                    }
 					
 					if (MaxShaderVersion >= 4)
 					{
@@ -180,24 +162,22 @@ FMetalCommandQueue::FMetalCommandQueue(mtlpp::Device InDevice, uint32 const MaxN
 	}
 	else if(Vers.majorVersion == 8 && Vers.minorVersion >= 3)
 	{
-		Features = EMetalFeaturesSeparateStencil | EMetalFeaturesSetBufferOffset;
+		Features = EMetalFeaturesSetBufferOffset;
 	}
 #else // Assume that Mac & other platforms all support these from the start. They can diverge later.
 	const bool bIsNVIDIA = [Device.GetName().GetPtr() rangeOfString:@"Nvidia" options:NSCaseInsensitiveSearch].location != NSNotFound;
-	Features = EMetalFeaturesSeparateStencil | EMetalFeaturesDepthClipMode | EMetalFeaturesResourceOptions | EMetalFeaturesDepthStencilBlitOptions | EMetalFeaturesCountingQueries | EMetalFeaturesBaseVertexInstance | EMetalFeaturesIndirectBuffer | EMetalFeaturesLayeredRendering | EMetalFeaturesShaderVersions | EMetalFeaturesCombinedDepthStencil | EMetalFeaturesCubemapArrays;
+	Features = EMetalFeaturesCountingQueries | EMetalFeaturesBaseVertexInstance | EMetalFeaturesIndirectBuffer | EMetalFeaturesLayeredRendering | EMetalFeaturesCubemapArrays;
 	if (!bIsNVIDIA)
 	{
 		Features |= EMetalFeaturesSetBufferOffset;
 	}
 	if (Device.SupportsFeatureSet(mtlpp::FeatureSet::macOS_GPUFamily1_v2))
     {
-        Features |= EMetalFeaturesStencilView | EMetalFeaturesDepth16 | EMetalFeaturesTessellation | EMetalFeaturesFunctionConstants | EMetalFeaturesGraphicsUAVs | EMetalFeaturesDeferredStoreActions | EMetalFeaturesMSAADepthResolve | EMetalFeaturesMSAAStoreAndResolve;
+        Features |= EMetalFeaturesTessellation | EMetalFeaturesMSAADepthResolve | EMetalFeaturesMSAAStoreAndResolve;
         
         // Assume that set*Bytes only works on macOS Sierra and above as no-one has tested it anywhere else.
 		Features |= EMetalFeaturesSetBytes;
 		
-		Features |= EMetalFeaturesLinearTextures;
-		
 		FString DeviceName(Device.GetName());
 		// On earlier OS versions Intel Broadwell couldn't suballocate properly
 		if (!(DeviceName.Contains(TEXT("Intel")) && (DeviceName.Contains(TEXT("5300")) || DeviceName.Contains(TEXT("6000")) || DeviceName.Contains(TEXT("6100")))) || FPlatformMisc::MacOSXVersionCompare(10,14,0) >= 0)
@@ -226,12 +206,6 @@ FMetalCommandQueue::FMetalCommandQueue(mtlpp::Device InDevice, uint32 const MaxN
 			}
 		}
 
-		// Turn on Linear Texture UAVs! Avoids the need to have function-constants which reduces initial runtime shader compile time
-		if (MaxShaderVersion >= 3 && FPlatformMisc::MacOSXVersionCompare(10,13,5) >= 0)
-		{
-			Features |= EMetalFeaturesLinearTextureUAVs;
-		}
-
 		// Turn on Texture Buffers! These are faster on the GPU as we don't need to do out-of-bounds tests but require Metal 2.1 and macOS 10.14
 		if (FPlatformMisc::MacOSXVersionCompare(10,14,0) >= 0)
 		{
@@ -271,7 +245,7 @@ FMetalCommandQueue::FMetalCommandQueue(mtlpp::Device InDevice, uint32 const MaxN
     
     if(Device.SupportsFeatureSet(mtlpp::FeatureSet::macOS_GPUFamily1_v3) && FPlatformMisc::MacOSXVersionCompare(10,13,0) >= 0)
     {
-        Features |= EMetalFeaturesMultipleViewports | EMetalFeaturesGPUCommandBufferTimes | EMetalFeaturesPipelineBufferMutability | EMetalFeaturesGPUCaptureManager | EMetalFeaturesAbsoluteTimeQueries | EMetalFeaturesSupportsVSyncToggle;
+        Features |= EMetalFeaturesMultipleViewports | EMetalFeaturesPipelineBufferMutability | EMetalFeaturesGPUCaptureManager;
 		
 		if (FParse::Param(FCommandLine::Get(),TEXT("metalfence")))
 		{
@@ -288,12 +262,6 @@ FMetalCommandQueue::FMetalCommandQueue(mtlpp::Device InDevice, uint32 const MaxN
 			Features |= EMetalFeaturesIABs;
 		}
     }
-	else
-	// Time query emulation breaks on AMD < 10.13 - disable by default until they can explain why, should work everywhere else.
-	if (!FString(Device.GetName()).Contains(TEXT("AMD")) || FParse::Param(FCommandLine::Get(),TEXT("metaltimequery")))
-	{
-		Features |= EMetalFeaturesAbsoluteTimeQueries;
-	}
 	
 	GMetalManagedUniformBuffers = FParse::Param(FCommandLine::Get(),TEXT("metalmanagedubs"));
 #endif
@@ -339,17 +307,13 @@ FMetalCommandQueue::FMetalCommandQueue(mtlpp::Device InDevice, uint32 const MaxN
 	PermittedOptions = 0;
 	PermittedOptions |= mtlpp::ResourceOptions::CpuCacheModeDefaultCache;
 	PermittedOptions |= mtlpp::ResourceOptions::CpuCacheModeWriteCombined;
-	if (Features & EMetalFeaturesResourceOptions)
 	{
 		PermittedOptions |= mtlpp::ResourceOptions::StorageModeShared;
 		PermittedOptions |= mtlpp::ResourceOptions::StorageModePrivate;
 #if PLATFORM_MAC
 		PermittedOptions |= mtlpp::ResourceOptions::StorageModeManaged;
 #else
-		if (Features & EMetalFeaturesMemoryLessResources)
-		{
-			PermittedOptions |= mtlpp::ResourceOptions::StorageModeMemoryless;
-		}
+		PermittedOptions |= mtlpp::ResourceOptions::StorageModeMemoryless;
 #endif
 		// You can't use HazardUntracked under the validation layer due to bugs in the layer when trying to create linear-textures/texture-buffers
 		if ((Features & EMetalFeaturesFences) && !(Features & EMetalFeaturesValidation))
@@ -456,12 +420,7 @@ FMetalFence* FMetalCommandQueue::CreateFence(ns::String const& Label) const
 				FMetalDebugFence* Fence = (FMetalDebugFence*)InnerFence.GetPtr();
 				Fence.label = String;
 			}
-			else
 	#endif
-			if(InnerFence && String)
-			{
-				InnerFence.SetLabel(String);
-			}
 		}
 		return InternalFence;
 	}
@@ -493,7 +452,7 @@ mtlpp::ResourceOptions FMetalCommandQueue::GetCompatibleResourceOptions(mtlpp::R
 {
 	NSUInteger NewOptions = (Options & PermittedOptions);
 #if PLATFORM_IOS // Swizzle Managed to Shared for iOS - we can do this as they are equivalent, unlike Shared -> Managed on Mac.
-	if ((Features & EMetalFeaturesResourceOptions) && (Options & (1 /*mtlpp::StorageMode::Managed*/ << mtlpp::ResourceStorageModeShift)))
+	if ((Options & (1 /*mtlpp::StorageMode::Managed*/ << mtlpp::ResourceStorageModeShift)))
 	{
 		NewOptions |= mtlpp::ResourceOptions::StorageModeShared;
 	}
diff --git a/Engine/Source/Runtime/Apple/MetalRHI/Private/MetalCommandQueue.h b/Engine/Source/Runtime/Apple/MetalRHI/Private/MetalCommandQueue.h
--- a/Engine/Source/Runtime/Apple/MetalRHI/Private/MetalCommandQueue.h
+++ b/Engine/Source/Runtime/Apple/MetalRHI/Private/MetalCommandQueue.h
@@ -17,92 +17,58 @@ class FMetalCommandList;
  */
 typedef NS_OPTIONS(uint64, EMetalFeatures)
 {
-	/** Support for separate front & back stencil ref. values */
-	EMetalFeaturesSeparateStencil = 1 << 0,
 	/** Support for specifying an update to the buffer offset only */
-	EMetalFeaturesSetBufferOffset = 1 << 1,
-	/** Support for specifying the depth clip mode */
-	EMetalFeaturesDepthClipMode = 1 << 2,
-	/** Support for specifying resource usage & memory options */
-	EMetalFeaturesResourceOptions = 1 << 3,
-	/** Supports texture->buffer blit options for depth/stencil blitting */
-	EMetalFeaturesDepthStencilBlitOptions = 1 << 4,
-    /** Supports creating a native stencil texture view from a depth/stencil texture */
-    EMetalFeaturesStencilView = 1 << 5,
-    /** Supports a depth-16 pixel format */
-    EMetalFeaturesDepth16 = 1 << 6,
+	EMetalFeaturesSetBufferOffset = 1 << 0,
 	/** Supports NSUInteger counting visibility queries */
-	EMetalFeaturesCountingQueries = 1 << 7,
+	EMetalFeaturesCountingQueries = 1 << 1,
 	/** Supports base vertex/instance for draw calls */
-	EMetalFeaturesBaseVertexInstance = 1 << 8,
+	EMetalFeaturesBaseVertexInstance = 1 << 2,
 	/** Supports indirect buffers for draw calls */
-	EMetalFeaturesIndirectBuffer = 1 << 9,
+	EMetalFeaturesIndirectBuffer = 1 << 3,
 	/** Supports layered rendering */
-	EMetalFeaturesLayeredRendering = 1 << 10,
+	EMetalFeaturesLayeredRendering = 1 << 4,
 	/** Support for specifying small buffers as byte arrays */
-	EMetalFeaturesSetBytes = 1 << 11,
-	/** Supports different shader standard versions */
-	EMetalFeaturesShaderVersions = 1 << 12,
+	EMetalFeaturesSetBytes = 1 << 5,
 	/** Supports tessellation rendering */
-	EMetalFeaturesTessellation = 1 << 13,
-	/** Supports arbitrary buffer/texture writes from graphics shaders */
-	EMetalFeaturesGraphicsUAVs = 1 << 14,
+	EMetalFeaturesTessellation = 1 << 6,
 	/** Supports framework-level validation */
-	EMetalFeaturesValidation = 1 << 15,
-	/** Supports absolute-time emulation using command-buffer completion handlers */
-	EMetalFeaturesAbsoluteTimeQueries = 1 << 16,
+	EMetalFeaturesValidation = 1 << 7,
 	/** Supports detailed statistics */
-	EMetalFeaturesStatistics= 1 << 17,
-	/** Supports memory-less texture resources */
-	EMetalFeaturesMemoryLessResources = 1 << 18,
+	EMetalFeaturesStatistics = 1 << 8,
 	/** Supports the explicit MTLHeap APIs */
-	EMetalFeaturesHeaps = 1 << 19,
+	EMetalFeaturesHeaps = 1 << 9,
 	/** Supports the explicit MTLFence APIs */
-	EMetalFeaturesFences = 1 << 20,
-	/** Supports deferred store action speficication */
-	EMetalFeaturesDeferredStoreActions = 1 << 21,
+	EMetalFeaturesFences = 1 << 10,
 	/** Supports MSAA Depth Resolves */
-	EMetalFeaturesMSAADepthResolve = 1 << 22,
+	EMetalFeaturesMSAADepthResolve = 1 << 11,
 	/** Supports Store & Resolve in a single store action */
-	EMetalFeaturesMSAAStoreAndResolve = 1 << 23,
+	EMetalFeaturesMSAAStoreAndResolve = 1 << 12,
 	/** Supports framework GPU frame capture */
-	EMetalFeaturesGPUTrace = 1 << 24,
-	/** Supports combined depth-stencil formats */
-	EMetalFeaturesCombinedDepthStencil = 1 << 25,
+	EMetalFeaturesGPUTrace = 1 << 13,
 	/** Supports the use of cubemap arrays */
-	EMetalFeaturesCubemapArrays = 1 << 26,
-	/** Supports the creation of texture-views using buffers as the backing store */
-	EMetalFeaturesLinearTextures = 1 << 27,
-	/** Supports the creation of texture-views for UAVs using buffers as the backing store */
-	EMetalFeaturesLinearTextureUAVs = 1 << 28,
+	EMetalFeaturesCubemapArrays = 1 << 14,
 	/** Supports the specification of multiple viewports and scissor rects */
-	EMetalFeaturesMultipleViewports = 1 << 29,
-	/** Supports accurate GPU times for commandbuffer start/end */
-    EMetalFeaturesGPUCommandBufferTimes = 1 << 30,
+	EMetalFeaturesMultipleViewports = 1 << 15,
     /** Supports minimum on-glass duration for drawables */
-    EMetalFeaturesPresentMinDuration = 1llu << 31llu,
+    EMetalFeaturesPresentMinDuration = 1llu << 16llu,
     /** Supports programmatic frame capture API */
-    EMetalFeaturesGPUCaptureManager = 1llu << 32llu,
-	/** Supports toggling V-Sync on & off */
-	EMetalFeaturesSupportsVSyncToggle = 1llu << 33llu,
-	/** Supports function-constants for runtime shader specialisation */
-	EMetalFeaturesFunctionConstants = 1llu << 34llu,
+    EMetalFeaturesGPUCaptureManager = 1llu << 17llu,
 	/** Supports efficient buffer-blits */
-	EMetalFeaturesEfficientBufferBlits = 1llu << 35llu,
+	EMetalFeaturesEfficientBufferBlits = 1llu << 18llu,
 	/** Supports any kind of buffer sub-allocation */
-	EMetalFeaturesBufferSubAllocation = 1llu << 36llu,
+	EMetalFeaturesBufferSubAllocation = 1llu << 19llu,
 	/** Supports private buffer sub-allocation */
-	EMetalFeaturesPrivateBufferSubAllocation = 1llu << 37llu,
+	EMetalFeaturesPrivateBufferSubAllocation = 1llu << 20llu,
 	/** Supports texture buffers */
-	EMetalFeaturesTextureBuffers = 1llu << 38llu,
+	EMetalFeaturesTextureBuffers = 1llu << 21llu,
 	/** Supports max. compute threads per threadgroup */
-	EMetalFeaturesMaxThreadsPerThreadgroup = 1llu << 39llu,
+	EMetalFeaturesMaxThreadsPerThreadgroup = 1llu << 22llu,
 	/** Supports parallel render encoders */
-	EMetalFeaturesParallelRenderEncoders = 1llu << 40llu,
+	EMetalFeaturesParallelRenderEncoders = 1llu << 23llu,
 	/** Supports indirect argument buffers */
-	EMetalFeaturesIABs = 1llu << 41llu,
+	EMetalFeaturesIABs = 1llu << 24llu,
 	/** Supports specifying the mutability of buffers bound to PSOs */
-	EMetalFeaturesPipelineBufferMutability = 1llu << 42llu,
+	EMetalFeaturesPipelineBufferMutability = 1llu << 25llu,
 };
 
 /**
diff --git a/Engine/Source/Runtime/Apple/MetalRHI/Private/MetalCommands.cpp b/Engine/Source/Runtime/Apple/MetalRHI/Private/MetalCommands.cpp
--- a/Engine/Source/Runtime/Apple/MetalRHI/Private/MetalCommands.cpp
+++ b/Engine/Source/Runtime/Apple/MetalRHI/Private/MetalCommands.cpp
@@ -640,7 +640,6 @@ void FMetalRHICommandContext::RHISetRenderTargetsAndClear(const FRHISetRenderTar
 	@autoreleasepool {
 	bool bHasTarget = (RenderTargetsInfo.DepthStencilRenderTarget.Texture != nullptr);
 	FMetalContext* Manager = Context;
-	if (Context->GetCommandQueue().SupportsFeature(EMetalFeaturesGraphicsUAVs))
 	{
 		for (uint32 i = 0; !bHasTarget && i < RenderTargetsInfo.NumUAVs; i++)
 		{
@@ -650,10 +649,6 @@ void FMetalRHICommandContext::RHISetRenderTargetsAndClear(const FRHISetRenderTar
 			}
 		}
 	}
-	else
-	{
-		checkf(RenderTargetsInfo.NumUAVs == 0, TEXT("Calling SetRenderTargets with UAVs is not supported in this Metal standard"));
-	}
 	
 	for (uint32 i = 0; bHasTarget == false && i < RenderTargetsInfo.NumColorRenderTargets; i++)
 	{
diff --git a/Engine/Source/Runtime/Apple/MetalRHI/Private/MetalContext.cpp b/Engine/Source/Runtime/Apple/MetalRHI/Private/MetalContext.cpp
--- a/Engine/Source/Runtime/Apple/MetalRHI/Private/MetalContext.cpp
+++ b/Engine/Source/Runtime/Apple/MetalRHI/Private/MetalContext.cpp
@@ -1137,10 +1137,6 @@ void FMetalContext::TransitionResources(FUnorderedAccessViewRHIParamRef* InUAVs,
 				if (Surface != nullptr && Surface->Texture)
 				{
 					RenderPass.TransitionResources(Surface->Texture);
-					if (Surface->StencilTexture)
-					{
-						RenderPass.TransitionResources(Surface->StencilTexture);
-					}
 					if (Surface->MSAATexture)
 					{
 						RenderPass.TransitionResources(Surface->MSAATexture);
@@ -1161,10 +1157,6 @@ void FMetalContext::TransitionResources(FTextureRHIParamRef* InTextures, int32 N
 			if (Surface != nullptr && Surface->Texture)
 			{
 				RenderPass.TransitionResources(Surface->Texture);
-				if (Surface->StencilTexture)
-				{
-					RenderPass.TransitionResources(Surface->StencilTexture);
-				}
 				if (Surface->MSAATexture)
 				{
 					RenderPass.TransitionResources(Surface->MSAATexture);
diff --git a/Engine/Source/Runtime/Apple/MetalRHI/Private/MetalProfiler.cpp b/Engine/Source/Runtime/Apple/MetalRHI/Private/MetalProfiler.cpp
--- a/Engine/Source/Runtime/Apple/MetalRHI/Private/MetalProfiler.cpp
+++ b/Engine/Source/Runtime/Apple/MetalRHI/Private/MetalProfiler.cpp
@@ -126,16 +126,9 @@ mtlpp::CommandBufferHandler FMetalEventNode::Start(void)
 {
 	mtlpp::CommandBufferHandler Block = [this](mtlpp::CommandBuffer const& CompletedBuffer)
 	{
-		if (FMetalCommandQueue::SupportsFeature(EMetalFeaturesGPUCommandBufferTimes))
-		{
-			const CFTimeInterval GpuTimeSeconds = CompletedBuffer.GetGpuStartTime();
-			const double CyclesPerSecond = 1.0 / FPlatformTime::GetSecondsPerCycle();
-			StartTime = GpuTimeSeconds * CyclesPerSecond;
-		}
-		else
-		{
-			StartTime = mach_absolute_time();
-		}
+		const CFTimeInterval GpuTimeSeconds = CompletedBuffer.GetGpuStartTime();
+		const double CyclesPerSecond = 1.0 / FPlatformTime::GetSecondsPerCycle();
+		StartTime = GpuTimeSeconds * CyclesPerSecond;
 	};
     return Block_copy(Block);
 }
@@ -164,17 +157,6 @@ void FMetalEventNode::StopTiming()
 
 			if (bRoot)
 			{
-				// We have a different mechanism for the overall frametime that works even with empty encoders and that doesn't report any GPU idle time between frames, we only use the fallback code below on older OSes.
-				if (!FMetalCommandQueue::SupportsFeature(EMetalFeaturesGPUCommandBufferTimes))
-				{
-					uint32 Time = FMath::TruncToInt( double(GetTiming()) / double(FPlatformTime::GetSecondsPerCycle()) );
-#if PLATFORM_IOS
-                    FPlatformAtomics::AtomicStore_Relaxed((int32*)&GGPUFrameTime, (int32)0);
-#else
-					FPlatformAtomics::AtomicStore_Relaxed((int32*)&GGPUFrameTime, (int32)Time);
-#endif
-				}
-
 				if(!bFullProfiling)
 				{
 					delete this;
@@ -199,31 +181,13 @@ mtlpp::CommandBufferHandler FMetalEventNode::Stop(void)
 {
 	mtlpp::CommandBufferHandler Block = [this](mtlpp::CommandBuffer const& CompletedBuffer)
 	{
-		if (FMetalCommandQueue::SupportsFeature(EMetalFeaturesGPUCommandBufferTimes))
-		{
-			// This is still used by ProfileGPU
-			const CFTimeInterval GpuTimeSeconds = CompletedBuffer.GetGpuEndTime();
-			const double CyclesPerSecond = 1.0 / FPlatformTime::GetSecondsPerCycle();
-			EndTime = GpuTimeSeconds * CyclesPerSecond;
-		}
-		else
-		{
-			EndTime = mach_absolute_time();
-		}
+		// This is still used by ProfileGPU
+		const CFTimeInterval GpuTimeSeconds = CompletedBuffer.GetGpuEndTime();
+		const double CyclesPerSecond = 1.0 / FPlatformTime::GetSecondsPerCycle();
+		EndTime = GpuTimeSeconds * CyclesPerSecond;
 		
 		if(bRoot)
 		{
-			// But we have a different mechanism for the overall frametime that works even with empty encoders and that doesn't report any GPU idle time between frames, we only use the fallback code below on older OSes.
-			if (!FMetalCommandQueue::SupportsFeature(EMetalFeaturesGPUCommandBufferTimes))
-			{
-				uint32 Time = FMath::TruncToInt( double(GetTiming()) / double(FPlatformTime::GetSecondsPerCycle()) );
-#if PLATFORM_IOS
-                FPlatformAtomics::AtomicStore_Relaxed((int32*)&GGPUFrameTime, (int32)0);
-#else
-				FPlatformAtomics::AtomicStore_Relaxed((int32*)&GGPUFrameTime, (int32)Time);
-#endif
-			}
-			
 			if(!bFullProfiling)
 			{
 				delete this;
@@ -394,68 +358,59 @@ void FMetalGPUProfiler::RecordFrame()
 {
 	// Note: We do not call RecordCommandBuffer here.
 	// This is only applied by FMetalCommandList::Commit, which will also apply RecordCommandBuffer.
-	if (FMetalCommandQueue::SupportsFeature(EMetalFeaturesGPUCommandBufferTimes))
-	{
-		const double CyclesPerSecond = 1.0 / FPlatformTime::GetSecondsPerCycle();
-		FrameGPUTimeCycles[FrameTimeGPUIndex] = uint64(CyclesPerSecond * RunningFrameTimeSeconds);
-        // for now set the GPU time to 0 until we figure out why it isn't always accurate
+	const double CyclesPerSecond = 1.0 / FPlatformTime::GetSecondsPerCycle();
+	FrameGPUTimeCycles[FrameTimeGPUIndex] = uint64(CyclesPerSecond * RunningFrameTimeSeconds);
+	// for now set the GPU time to 0 until we figure out why it isn't always accurate
 #if PLATFORM_IOS
-        FPlatformAtomics::AtomicStore_Relaxed((int32*)&GGPUFrameTime, (int32)0);
+	FPlatformAtomics::AtomicStore_Relaxed((int32*)&GGPUFrameTime, (int32)0);
 #else
-		FPlatformAtomics::AtomicStore_Relaxed((int32*)&GGPUFrameTime, int32(FrameGPUTimeCycles[FrameTimeGPUIndex]));
+	FPlatformAtomics::AtomicStore_Relaxed((int32*)&GGPUFrameTime, int32(FrameGPUTimeCycles[FrameTimeGPUIndex]));
 #endif
-		
+	
 #if STATS
-		FPlatformAtomics::AtomicStore_Relaxed(&GMetalGPUWorkTime, FrameGPUTimeCycles[FrameTimeGPUIndex]);
-		
-		int64 TimeAsCycles = int64(FrameEndGPUCycles[FrameTimeGPUIndex] - FrameStartGPUCycles[FrameTimeGPUIndex] - FrameGPUTimeCycles[FrameTimeGPUIndex]);
-		FrameIdleTimeCycles[FrameTimeGPUIndex] = TimeAsCycles;
-		FPlatformAtomics::AtomicStore_Relaxed(&GMetalGPUIdleTime, TimeAsCycles);
+	FPlatformAtomics::AtomicStore_Relaxed(&GMetalGPUWorkTime, FrameGPUTimeCycles[FrameTimeGPUIndex]);
+	
+	int64 TimeAsCycles = int64(FrameEndGPUCycles[FrameTimeGPUIndex] - FrameStartGPUCycles[FrameTimeGPUIndex] - FrameGPUTimeCycles[FrameTimeGPUIndex]);
+	FrameIdleTimeCycles[FrameTimeGPUIndex] = TimeAsCycles;
+	FPlatformAtomics::AtomicStore_Relaxed(&GMetalGPUIdleTime, TimeAsCycles);
 #endif //STATS
-		
-		IncrementFrameIndex();
-	}
+	
+	IncrementFrameIndex();
 }
 
 void FMetalGPUProfiler::RecordPresent(const mtlpp::CommandBuffer& Buffer)
 {
-	if (FMetalCommandQueue::SupportsFeature(EMetalFeaturesGPUCommandBufferTimes))
-	{
-		const CFTimeInterval GpuStartTimeSeconds = Buffer.GetGpuStartTime();
-		const CFTimeInterval GpuEndTimeSeconds = Buffer.GetGpuEndTime();
-		const double CyclesPerSecond = 1.0 / FPlatformTime::GetSecondsPerCycle();
-		uint64 StartTimeCycles = uint64(GpuStartTimeSeconds * CyclesPerSecond);
-		uint64 EndTimeCycles = uint64(GpuEndTimeSeconds * CyclesPerSecond);
-		int64 Time = int64(EndTimeCycles - StartTimeCycles);
-		FramePresentTimeCycles[FrameTimeGPUIndex] = Time;
-		FPlatformAtomics::AtomicStore_Relaxed(&GMetalPresentTime, Time);
-	}
+	const CFTimeInterval GpuStartTimeSeconds = Buffer.GetGpuStartTime();
+	const CFTimeInterval GpuEndTimeSeconds = Buffer.GetGpuEndTime();
+	const double CyclesPerSecond = 1.0 / FPlatformTime::GetSecondsPerCycle();
+	uint64 StartTimeCycles = uint64(GpuStartTimeSeconds * CyclesPerSecond);
+	uint64 EndTimeCycles = uint64(GpuEndTimeSeconds * CyclesPerSecond);
+	int64 Time = int64(EndTimeCycles - StartTimeCycles);
+	FramePresentTimeCycles[FrameTimeGPUIndex] = Time;
+	FPlatformAtomics::AtomicStore_Relaxed(&GMetalPresentTime, Time);
 }
 
 void FMetalGPUProfiler::RecordCommandBuffer(const mtlpp::CommandBuffer& Buffer)
 {
-    if (FMetalCommandQueue::SupportsFeature(EMetalFeaturesGPUCommandBufferTimes))
-    {
-		const double CyclesPerSecond = 1.0 / FPlatformTime::GetSecondsPerCycle();
-		
-		const CFTimeInterval GpuEndTimeSeconds   = Buffer.GetGpuEndTime();
-		const CFTimeInterval GpuStartTimeSeconds = Buffer.GetGpuStartTime();
-		
-		uint64 StartTimeCycles = uint64(GpuStartTimeSeconds * CyclesPerSecond);
-		uint64 EndTimeCycles   = uint64(GpuEndTimeSeconds * CyclesPerSecond);
-		
-		{
-			uint64 Existing = FrameEndGPUCycles[FrameTimeGPUIndex];
-			FrameEndGPUCycles[FrameTimeGPUIndex] = Existing > 0 ? FMath::Max(Existing, EndTimeCycles) : EndTimeCycles;
-		}
-		
-		{
-			uint64 Existing = FrameStartGPUCycles[FrameTimeGPUIndex];
-			FrameStartGPUCycles[FrameTimeGPUIndex] = Existing > 0 ? FMath::Min(Existing, StartTimeCycles) : StartTimeCycles;
-		}
-		
-		RunningFrameTimeSeconds += (GpuEndTimeSeconds - GpuStartTimeSeconds);
-    }
+	const double CyclesPerSecond = 1.0 / FPlatformTime::GetSecondsPerCycle();
+	
+	const CFTimeInterval GpuEndTimeSeconds   = Buffer.GetGpuEndTime();
+	const CFTimeInterval GpuStartTimeSeconds = Buffer.GetGpuStartTime();
+	
+	uint64 StartTimeCycles = uint64(GpuStartTimeSeconds * CyclesPerSecond);
+	uint64 EndTimeCycles   = uint64(GpuEndTimeSeconds * CyclesPerSecond);
+	
+	{
+		uint64 Existing = FrameEndGPUCycles[FrameTimeGPUIndex];
+		FrameEndGPUCycles[FrameTimeGPUIndex] = Existing > 0 ? FMath::Max(Existing, EndTimeCycles) : EndTimeCycles;
+	}
+	
+	{
+		uint64 Existing = FrameStartGPUCycles[FrameTimeGPUIndex];
+		FrameStartGPUCycles[FrameTimeGPUIndex] = Existing > 0 ? FMath::Min(Existing, StartTimeCycles) : StartTimeCycles;
+	}
+	
+	RunningFrameTimeSeconds += (GpuEndTimeSeconds - GpuStartTimeSeconds);
 }
 	
 int32 FMetalGPUProfiler::FrameTimeGPUIndex = 0;
@@ -1123,13 +1078,6 @@ void FMetalCommandBufferStats::Start(mtlpp::CommandBuffer const& Buffer)
 	
 	GPUStartTime = 0;
 	GPUEndTime = 0;
-	
-	if (!FMetalCommandQueue::SupportsFeature(EMetalFeaturesGPUCommandBufferTimes))
-	{
-		CmdBuffer.AddScheduledHandler(^(const mtlpp::CommandBuffer &) {
-			GPUStartTime = FPlatformTime::ToMilliseconds64(mach_absolute_time()) * 1000.0;
-		});
-	}
 }
 
 void FMetalCommandBufferStats::End(mtlpp::CommandBuffer const& Buffer)
@@ -1141,18 +1089,12 @@ void FMetalCommandBufferStats::End(mtlpp::CommandBuffer const& Buffer)
 	
 	bool const bTracing = FMetalProfiler::GetProfiler() && FMetalProfiler::GetProfiler()->TracingEnabled();
 	CmdBuffer.AddCompletedHandler(^(const mtlpp::CommandBuffer & InnerBuffer) {
-		if (FMetalCommandQueue::SupportsFeature(EMetalFeaturesGPUCommandBufferTimes))
-		{
-			const CFTimeInterval GpuTimeSeconds = InnerBuffer.GetGpuStartTime();
-			GPUStartTime = GpuTimeSeconds * 1000000.0;
-			
-			const CFTimeInterval GpuEndTimeSeconds = InnerBuffer.GetGpuEndTime();
-			GPUEndTime = GpuEndTimeSeconds * 1000000.0;
-		}
-		else
-		{
-			GPUEndTime = FPlatformTime::ToMilliseconds64(mach_absolute_time()) * 1000.0;
-		}
+		const CFTimeInterval GpuTimeSeconds = InnerBuffer.GetGpuStartTime();
+		GPUStartTime = GpuTimeSeconds * 1000000.0;
+		
+		const CFTimeInterval GpuEndTimeSeconds = InnerBuffer.GetGpuEndTime();
+		GPUEndTime = GpuEndTimeSeconds * 1000000.0;
+	
 		if (bTracing)
 		{
 			FMetalProfiler::GetProfiler()->AddCommandBuffer(this);
diff --git a/Engine/Source/Runtime/Apple/MetalRHI/Private/MetalQuery.cpp b/Engine/Source/Runtime/Apple/MetalRHI/Private/MetalQuery.cpp
--- a/Engine/Source/Runtime/Apple/MetalRHI/Private/MetalQuery.cpp
+++ b/Engine/Source/Runtime/Apple/MetalRHI/Private/MetalQuery.cpp
@@ -300,15 +300,8 @@ FRenderQueryRHIRef FMetalDynamicRHI::RHICreateRenderQuery_RenderThread(class FRH
 FRenderQueryRHIRef FMetalDynamicRHI::RHICreateRenderQuery(ERenderQueryType QueryType)
 {
 	@autoreleasepool {
-	FRenderQueryRHIRef Query;
-	// AMD have subtleties to their completion handler routines that mean we don't seem able to reliably wait on command-buffers
-	// until after a drawable present...
-	static bool const bSupportsTimeQueries = GetMetalDeviceContext().GetCommandQueue().SupportsFeature(EMetalFeaturesAbsoluteTimeQueries);
-	if (QueryType != RQT_AbsoluteTime || bSupportsTimeQueries)
-	{
-		Query = new FMetalRenderQuery(QueryType);
-	}
-	return Query;
+		FRenderQueryRHIRef Query = new FMetalRenderQuery(QueryType);
+		return Query;
 	}
 }
 
diff --git a/Engine/Source/Runtime/Apple/MetalRHI/Private/MetalRHI.cpp b/Engine/Source/Runtime/Apple/MetalRHI/Private/MetalRHI.cpp
--- a/Engine/Source/Runtime/Apple/MetalRHI/Private/MetalRHI.cpp
+++ b/Engine/Source/Runtime/Apple/MetalRHI/Private/MetalRHI.cpp
@@ -665,16 +665,9 @@ FMetalDynamicRHI::FMetalDynamicRHI(ERHIFeatureLevel::Type RequestedFeatureLevel)
 		GPixelFormats[PF_FloatR11G11B10		].Supported			= true;
 	}
 	
-	if (FMetalCommandQueue::SupportsFeature(EMetalFeaturesStencilView) && FMetalCommandQueue::SupportsFeature(EMetalFeaturesCombinedDepthStencil) && !FParse::Param(FCommandLine::Get(),TEXT("metalforceseparatedepthstencil")))
-	{
-		GPixelFormats[PF_DepthStencil		].PlatformFormat	= (uint32)mtlpp::PixelFormat::Depth32Float_Stencil8;
-		GPixelFormats[PF_DepthStencil		].BlockBytes		= 4;
-	}
-	else
-	{
-		GPixelFormats[PF_DepthStencil		].PlatformFormat	= (uint32)mtlpp::PixelFormat::Depth32Float;
-		GPixelFormats[PF_DepthStencil		].BlockBytes		= 4;
-	}
+	GPixelFormats[PF_DepthStencil		].PlatformFormat	= (uint32)mtlpp::PixelFormat::Depth32Float_Stencil8;
+	GPixelFormats[PF_DepthStencil		].BlockBytes		= 4;
+
 	GPixelFormats[PF_DepthStencil		].Supported			= true;
 	GPixelFormats[PF_ShadowDepth		].PlatformFormat	= (uint32)mtlpp::PixelFormat::Depth32Float;
 	GPixelFormats[PF_ShadowDepth		].BlockBytes		= 4;
diff --git a/Engine/Source/Runtime/Apple/MetalRHI/Private/MetalRHIContext.cpp b/Engine/Source/Runtime/Apple/MetalRHI/Private/MetalRHIContext.cpp
--- a/Engine/Source/Runtime/Apple/MetalRHI/Private/MetalRHIContext.cpp
+++ b/Engine/Source/Runtime/Apple/MetalRHI/Private/MetalRHIContext.cpp
@@ -172,7 +172,6 @@ void FMetalRHICommandContext::RHIBeginRenderPass(const FRHIRenderPassInfo& InInf
 
 	@autoreleasepool {
 	bool bHasTarget = (RenderTargetsInfo.DepthStencilRenderTarget.Texture != nullptr);
-	if (Context->GetCommandQueue().SupportsFeature(EMetalFeaturesGraphicsUAVs))
 	{
 		for (uint32 i = 0; !bHasTarget && i < RenderTargetsInfo.NumUAVs; i++)
 		{
@@ -182,10 +181,6 @@ void FMetalRHICommandContext::RHIBeginRenderPass(const FRHIRenderPassInfo& InInf
 			}
 		}
 	}
-	else
-	{
-		checkf(RenderTargetsInfo.NumUAVs == 0, TEXT("Calling SetRenderTargets with UAVs is not supported in this Metal standard"));
-	}
 	
 	for (uint32 i = 0; bHasTarget == false && i < RenderTargetsInfo.NumColorRenderTargets; i++)
 	{
diff --git a/Engine/Source/Runtime/Apple/MetalRHI/Private/MetalRenderPass.cpp b/Engine/Source/Runtime/Apple/MetalRHI/Private/MetalRenderPass.cpp
--- a/Engine/Source/Runtime/Apple/MetalRHI/Private/MetalRenderPass.cpp
+++ b/Engine/Source/Runtime/Apple/MetalRHI/Private/MetalRenderPass.cpp
@@ -835,17 +835,10 @@ void FMetalRenderPass::CopyFromTextureToBuffer(FMetalTexture const& Texture, uin
 	check(Encoder.GetPtr());
 	
 	METAL_GPUPROFILE(FMetalProfiler::GetProfiler()->EncodeBlit(CurrentEncoder.GetCommandBufferStats(), __FUNCTION__));
-	if (CmdList.GetCommandQueue().SupportsFeature(EMetalFeaturesDepthStencilBlitOptions))
 	{
 		MTLPP_VALIDATE(mtlpp::BlitCommandEncoder, Encoder, SafeGetRuntimeDebuggingLevel() >= EMetalDebugLevelValidation, Copy(Texture, sourceSlice, sourceLevel, sourceOrigin, sourceSize, toBuffer, destinationOffset, destinationBytesPerRow, destinationBytesPerImage, options));
 		METAL_DEBUG_LAYER(EMetalDebugLevelFastValidation, CurrentEncoder.GetBlitCommandEncoderDebugging().Copy(Texture, sourceSlice, sourceLevel, sourceOrigin, sourceSize, toBuffer, destinationOffset, destinationBytesPerRow, destinationBytesPerImage, options));
 	}
-	else
-	{
-		check(options == mtlpp::BlitOption::None);
-		MTLPP_VALIDATE(mtlpp::BlitCommandEncoder, Encoder, SafeGetRuntimeDebuggingLevel() >= EMetalDebugLevelValidation, Copy(Texture, sourceSlice, sourceLevel, sourceOrigin, sourceSize, toBuffer, destinationOffset, destinationBytesPerRow, destinationBytesPerImage));
-		METAL_DEBUG_LAYER(EMetalDebugLevelFastValidation, CurrentEncoder.GetBlitCommandEncoderDebugging().Copy(Texture, sourceSlice, sourceLevel, sourceOrigin, sourceSize, toBuffer, destinationOffset, destinationBytesPerRow, destinationBytesPerImage));
-	}
 	ConditionalSubmit();
 }
 
diff --git a/Engine/Source/Runtime/Apple/MetalRHI/Private/MetalRenderTarget.cpp b/Engine/Source/Runtime/Apple/MetalRHI/Private/MetalRenderTarget.cpp
--- a/Engine/Source/Runtime/Apple/MetalRHI/Private/MetalRenderTarget.cpp
+++ b/Engine/Source/Runtime/Apple/MetalRHI/Private/MetalRenderTarget.cpp
@@ -605,7 +605,7 @@ void FMetalDynamicRHI::RHIReadSurfaceData(FTextureRHIParamRef TextureRHI, FIntRe
 		SCOPE_CYCLE_COUNTER(STAT_MetalTexturePageOffTime);
 		
 		FMetalTexture TempTexture = nil;
-		if (ImmediateContext.Context->GetCommandQueue().SupportsFeature(EMetalFeaturesResourceOptions) && Texture.GetStorageMode() == mtlpp::StorageMode::Private)
+		if (Texture.GetStorageMode() == mtlpp::StorageMode::Private)
 		{
 #if PLATFORM_MAC
 			mtlpp::StorageMode StorageMode = mtlpp::StorageMode::Managed;
@@ -679,7 +679,7 @@ void FMetalDynamicRHI::RHIReadSurfaceData(FTextureRHIParamRef TextureRHI, FIntRe
 			{
 				ImmediateContext.Context->CopyFromTextureToBuffer(Texture, 0, 0, Region.origin, Region.size, Buffer, 0, AlignedStride, BytesPerImage, mtlpp::BlitOption::None);
 			}
-			else if (GetMetalDeviceContext().SupportsFeature(EMetalFeaturesDepthStencilBlitOptions))
+			else
 			{
 				if (!InFlags.GetOutputStencil())
 				{
@@ -690,11 +690,6 @@ void FMetalDynamicRHI::RHIReadSurfaceData(FTextureRHIParamRef TextureRHI, FIntRe
 					ImmediateContext.Context->CopyFromTextureToBuffer(Texture, 0, 0, Region.origin, Region.size, Buffer, 0, AlignedStride, BytesPerImage, mtlpp::BlitOption::StencilFromDepthStencil);
 				}
 			}
-			else
-			{
-				// not supported yet
-				NOT_SUPPORTED("RHIReadSurfaceData Format");
-			}
 			
 			//kick the current command buffer.
 			ImmediateContext.Context->SubmitCommandBufferAndWait();
diff --git a/Engine/Source/Runtime/Apple/MetalRHI/Private/MetalShaders.cpp b/Engine/Source/Runtime/Apple/MetalRHI/Private/MetalShaders.cpp
--- a/Engine/Source/Runtime/Apple/MetalRHI/Private/MetalShaders.cpp
+++ b/Engine/Source/Runtime/Apple/MetalRHI/Private/MetalShaders.cpp
@@ -408,7 +408,6 @@ void TMetalBaseShader<BaseResourceType, ShaderType>::Init(const TArray<uint8>& I
 #endif
 			CompileOptions.SetPreprocessorMacros(PreprocessorMacros);
 #endif
-			if (GetMetalDeviceContext().SupportsFeature(EMetalFeaturesShaderVersions))
 			{
 				if (Header.Version < 3)
 				{
@@ -442,10 +441,6 @@ void TMetalBaseShader<BaseResourceType, ShaderType>::Init(const TArray<uint8>& I
 			[GlslCodeNSString retain];
 		}
 		
-		// Make sure that the current device can actually run with function constants otherwise bad things will happen!
-		UE_CLOG(bHasFunctionConstants && !FMetalCommandQueue::SupportsFeature(EMetalFeaturesFunctionConstants), LogMetal, Error, TEXT("Metal shader has function constants but current OS/device does not support them."));
-		UE_CLOG(bHasFunctionConstants && !FMetalCommandQueue::SupportsFeature(EMetalFeaturesFunctionConstants), LogMetal, Fatal, TEXT("%s"), *FString(GetSourceCode()));
-		
         GetCompiledFunction(true);
     }
 	UniformBuffersCopyInfo = Header.UniformBuffersCopyInfo;
diff --git a/Engine/Source/Runtime/Apple/MetalRHI/Private/MetalStateCache.cpp b/Engine/Source/Runtime/Apple/MetalRHI/Private/MetalStateCache.cpp
--- a/Engine/Source/Runtime/Apple/MetalRHI/Private/MetalStateCache.cpp
+++ b/Engine/Source/Runtime/Apple/MetalRHI/Private/MetalStateCache.cpp
@@ -399,9 +399,6 @@ bool FMetalStateCache::SetRenderTargetsInfo(FRHISetRenderTargetsInfo const& InRe
 	{
 		bool bNeedsClear = false;
 		
-		// Deferred store actions make life a bit easier...
-		static bool bSupportsDeferredStore = GetMetalDeviceContext().GetCommandQueue().SupportsFeature(EMetalFeaturesDeferredStoreActions);
-		
 		//Create local store action states if we support deferred store
 		mtlpp::StoreAction NewColorStore[MaxSimultaneousRenderTargets];
 		for (uint32 i = 0; i < MaxSimultaneousRenderTargets; ++i)
@@ -537,7 +534,7 @@ bool FMetalStateCache::SetRenderTargetsInfo(FRHISetRenderTargetsInfo const& InRe
 					// set up an MSAA attachment
 					ColorAttachment.SetTexture(Surface.MSAATexture);
 					NewColorStore[RenderTargetIndex] = GetMetalRTStoreAction(ERenderTargetStoreAction::EMultisampleResolve);
-					ColorAttachment.SetStoreAction(bSupportsDeferredStore && GRHIDeviceId > 2 ? mtlpp::StoreAction::Unknown : NewColorStore[RenderTargetIndex]);
+					ColorAttachment.SetStoreAction(GRHIDeviceId > 2 ? mtlpp::StoreAction::Unknown : NewColorStore[RenderTargetIndex]);
 					ColorAttachment.SetResolveTexture(Surface.MSAAResolveTexture ? Surface.MSAAResolveTexture : Surface.Texture);
 					SampleCount = Surface.MSAATexture.GetSampleCount();
 #if PLATFORM_IOS
@@ -561,7 +558,7 @@ bool FMetalStateCache::SetRenderTargetsInfo(FRHISetRenderTargetsInfo const& InRe
 					// set up non-MSAA attachment
 					ColorAttachment.SetTexture(Surface.Texture);
 					NewColorStore[RenderTargetIndex] = GetMetalRTStoreAction(HighLevelStoreAction);
-					ColorAttachment.SetStoreAction(bSupportsDeferredStore ? mtlpp::StoreAction::Unknown : NewColorStore[RenderTargetIndex]);
+					ColorAttachment.SetStoreAction(mtlpp::StoreAction::Unknown);
                     SampleCount = 1;
 				}
 				
@@ -804,7 +801,7 @@ bool FMetalStateCache::SetRenderTargetsInfo(FRHISetRenderTargetsInfo const& InRe
                 //needed to quiet the metal validation that runs when you end renderpass. (it requires some kind of 'resolve' for an msaa target)
 				//But with deferredstore we don't set the real one until submit time.
 				NewDepthStore = !Surface.MSAATexture || bSupportsMSAADepthResolve ? GetMetalRTStoreAction(HighLevelStoreAction) : mtlpp::StoreAction::DontCare;
-				DepthAttachment.SetStoreAction(bSupportsDeferredStore && Surface.MSAATexture && GRHIDeviceId > 2 ? mtlpp::StoreAction::Unknown : NewDepthStore);
+				DepthAttachment.SetStoreAction(Surface.MSAATexture && GRHIDeviceId > 2 ? mtlpp::StoreAction::Unknown : NewDepthStore);
 				DepthAttachment.SetClearDepth(DepthClearValue);
 				check(SampleCount > 0);
 
@@ -836,63 +833,55 @@ bool FMetalStateCache::SetRenderTargetsInfo(FRHISetRenderTargetsInfo const& InRe
 			
 			bool const bCombinedDepthStencilUsingDepth = (StencilTexture && StencilTexture.GetPixelFormat() != mtlpp::PixelFormat::Stencil8 && RenderTargetsInfo.DepthStencilRenderTarget.GetDepthStencilAccess().IsUsingDepth());
 			bool const bUsingStencil = RenderTargetsInfo.DepthStencilRenderTarget.GetDepthStencilAccess().IsUsingStencil() || (bUsingValidation && bCombinedDepthStencilUsingDepth);
-			if (StencilTexture && bUsingStencil && (FMetalCommandQueue::SupportsFeature(EMetalFeaturesCombinedDepthStencil) || !bDepthStencilSampleCountMismatchFixup))
+			if (StencilTexture && bUsingStencil)
 			{
-                if (!FMetalCommandQueue::SupportsFeature(EMetalFeaturesCombinedDepthStencil) && bDepthStencilSampleCountMismatchFixup)
-                {
-                    checkf(!RenderTargetsInfo.DepthStencilRenderTarget.GetDepthStencilAccess().IsStencilWrite(), TEXT("Stencil write not allowed as we don't have a proper stencil to use."));
-                }
-                else
-                {
-					mtlpp::RenderPassStencilAttachmentDescriptor StencilAttachment;
-                    
-                    StencilFormatKey = Surface.FormatKey;
-        
-                    // set up the stencil attachment
-                    StencilAttachment.SetTexture(StencilTexture);
-                    StencilAttachment.SetLoadAction(GetMetalRTLoadAction(RenderTargetsInfo.DepthStencilRenderTarget.StencilLoadAction));
-                    
-                    bNeedsClear |= (StencilAttachment.GetLoadAction() == mtlpp::LoadAction::Clear);
-					
-					ERenderTargetStoreAction HighLevelStoreAction = RenderTargetsInfo.DepthStencilRenderTarget.GetStencilStoreAction();
-					if (bUsingStencil && (HighLevelStoreAction == ERenderTargetStoreAction::ENoAction || bDepthStencilSampleCountMismatchFixup))
-					{
-						HighLevelStoreAction = ERenderTargetStoreAction::EStore;
-					}
-					
+				mtlpp::RenderPassStencilAttachmentDescriptor StencilAttachment;
+				
+				StencilFormatKey = Surface.FormatKey;
+	
+				// set up the stencil attachment
+				StencilAttachment.SetTexture(StencilTexture);
+				StencilAttachment.SetLoadAction(GetMetalRTLoadAction(RenderTargetsInfo.DepthStencilRenderTarget.StencilLoadAction));
+				
+				bNeedsClear |= (StencilAttachment.GetLoadAction() == mtlpp::LoadAction::Clear);
+				
+				ERenderTargetStoreAction HighLevelStoreAction = RenderTargetsInfo.DepthStencilRenderTarget.GetStencilStoreAction();
+				if (bUsingStencil && (HighLevelStoreAction == ERenderTargetStoreAction::ENoAction || bDepthStencilSampleCountMismatchFixup))
+				{
+					HighLevelStoreAction = ERenderTargetStoreAction::EStore;
+				}
+				
 #if PLATFORM_IOS
-					if (StencilTexture.GetStorageMode() == mtlpp::StorageMode::Memoryless)
-					{
-						HighLevelStoreAction = ERenderTargetStoreAction::ENoAction;
-						StencilAttachment.SetLoadAction(mtlpp::LoadAction::Clear);
-					}
+				if (StencilTexture.GetStorageMode() == mtlpp::StorageMode::Memoryless)
+				{
+					HighLevelStoreAction = ERenderTargetStoreAction::ENoAction;
+					StencilAttachment.SetLoadAction(mtlpp::LoadAction::Clear);
+				}
 #endif
-					
-					// For the case where Depth+Stencil is MSAA we can't Resolve depth and Store stencil - we can only Resolve + DontCare or StoreResolve + Store (on newer H/W and iOS).
-					// We only allow use of StoreResolve in the Desktop renderers as the mobile renderer does not and should not assume hardware support for it.
-					NewStencilStore = (StencilTexture.GetSampleCount() == 1  || GetMetalRTStoreAction(ERenderTargetStoreAction::EMultisampleResolve) == mtlpp::StoreAction::StoreAndMultisampleResolve) ? GetMetalRTStoreAction(HighLevelStoreAction) : mtlpp::StoreAction::DontCare;
-					StencilAttachment.SetStoreAction(bSupportsDeferredStore && StencilTexture.GetSampleCount() > 1 && GRHIDeviceId > 2 ? mtlpp::StoreAction::Unknown : NewStencilStore);
-                    StencilAttachment.SetClearStencil(StencilClearValue);
+				
+				// For the case where Depth+Stencil is MSAA we can't Resolve depth and Store stencil - we can only Resolve + DontCare or StoreResolve + Store (on newer H/W and iOS).
+				// We only allow use of StoreResolve in the Desktop renderers as the mobile renderer does not and should not assume hardware support for it.
+				NewStencilStore = (StencilTexture.GetSampleCount() == 1  || GetMetalRTStoreAction(ERenderTargetStoreAction::EMultisampleResolve) == mtlpp::StoreAction::StoreAndMultisampleResolve) ? GetMetalRTStoreAction(HighLevelStoreAction) : mtlpp::StoreAction::DontCare;
+				StencilAttachment.SetStoreAction(StencilTexture.GetSampleCount() > 1 && GRHIDeviceId > 2 ? mtlpp::StoreAction::Unknown : NewStencilStore);
+				StencilAttachment.SetClearStencil(StencilClearValue);
 
-                    if (SampleCount == 0)
-                    {
-                        SampleCount = StencilAttachment.GetTexture().GetSampleCount();
-                    }
-                    
-                    bHasValidRenderTarget = true;
-                    
-                    // @todo Stencil writes that need to persist must use ERenderTargetStoreAction::EStore on iOS.
-                    // We should probably be using deferred store actions so that we can safely lazily instantiate encoders.
-                    bCanRestartRenderPass &= (SampleCount <= 1) && ((RenderTargetsInfo.DepthStencilRenderTarget.Texture == FallbackDepthStencilSurface) || ((StencilAttachment.GetLoadAction() == mtlpp::LoadAction::Load) && (1 || !RenderTargetsInfo.DepthStencilRenderTarget.GetDepthStencilAccess().IsStencilWrite() || (RenderTargetsInfo.DepthStencilRenderTarget.GetStencilStoreAction() == ERenderTargetStoreAction::EStore))));
-                    
-                    // and assign it
-                    RenderPass.SetStencilAttachment(StencilAttachment);
-                }
+				if (SampleCount == 0)
+				{
+					SampleCount = StencilAttachment.GetTexture().GetSampleCount();
+				}
+				
+				bHasValidRenderTarget = true;
+				
+				// @todo Stencil writes that need to persist must use ERenderTargetStoreAction::EStore on iOS.
+				// We should probably be using deferred store actions so that we can safely lazily instantiate encoders.
+				bCanRestartRenderPass &= (SampleCount <= 1) && ((RenderTargetsInfo.DepthStencilRenderTarget.Texture == FallbackDepthStencilSurface) || ((StencilAttachment.GetLoadAction() == mtlpp::LoadAction::Load) && (1 || !RenderTargetsInfo.DepthStencilRenderTarget.GetDepthStencilAccess().IsStencilWrite() || (RenderTargetsInfo.DepthStencilRenderTarget.GetStencilStoreAction() == ERenderTargetStoreAction::EStore))));
+				
+				// and assign it
+				RenderPass.SetStencilAttachment(StencilAttachment);
 			}
 		}
 		
 		//Update deferred store states if required otherwise they're already set directly on the Metal Attachement Descriptors
-		if (bSupportsDeferredStore)
 		{
 			for (uint32 i = 0; i < MaxSimultaneousRenderTargets; ++i)
 			{
@@ -1150,8 +1139,6 @@ void FMetalStateCache::SetGraphicsPipelineState(FMetalGraphicsPipelineState* Sta
 			}
 		}
 		
-		static bool bSupportsUAVs = GetMetalDeviceContext().GetCommandQueue().SupportsFeature(EMetalFeaturesGraphicsUAVs);
-		if (bSupportsUAVs)
 		{
 			for (uint32 i = 0; i < RenderTargetsInfo.NumUAVs; i++)
 			{
@@ -1275,64 +1262,57 @@ bool FMetalStateCache::NeedsToSetRenderTarget(const FRHISetRenderTargetsInfo& In
 		if (InRenderTargetsInfo.DepthStencilRenderTarget.Texture && (InRenderTargetsInfo.DepthStencilRenderTarget.DepthStoreAction > RenderTargetsInfo.DepthStencilRenderTarget.DepthStoreAction || InRenderTargetsInfo.DepthStencilRenderTarget.GetStencilStoreAction() > RenderTargetsInfo.DepthStencilRenderTarget.GetStencilStoreAction()))
 		{
 			// Don't break the encoder if we can just change the store actions.
-			if (FMetalCommandQueue::SupportsFeature(EMetalFeaturesDeferredStoreActions))
+			mtlpp::StoreAction NewDepthStore = DepthStore;
+			mtlpp::StoreAction NewStencilStore = StencilStore;
+			if (InRenderTargetsInfo.DepthStencilRenderTarget.DepthStoreAction > RenderTargetsInfo.DepthStencilRenderTarget.DepthStoreAction)
 			{
-				mtlpp::StoreAction NewDepthStore = DepthStore;
-				mtlpp::StoreAction NewStencilStore = StencilStore;
-				if (InRenderTargetsInfo.DepthStencilRenderTarget.DepthStoreAction > RenderTargetsInfo.DepthStencilRenderTarget.DepthStoreAction)
+				if (RenderPassDesc.GetDepthAttachment().GetTexture())
 				{
-					if (RenderPassDesc.GetDepthAttachment().GetTexture())
-					{
-						FMetalSurface& Surface = *GetMetalSurfaceFromRHITexture(RenderTargetsInfo.DepthStencilRenderTarget.Texture);
-						
-						const uint32 DepthSampleCount = (Surface.MSAATexture ? Surface.MSAATexture.GetSampleCount() : Surface.Texture.GetSampleCount());
-						bool const bDepthStencilSampleCountMismatchFixup = (SampleCount != DepthSampleCount);
+					FMetalSurface& Surface = *GetMetalSurfaceFromRHITexture(RenderTargetsInfo.DepthStencilRenderTarget.Texture);
+					
+					const uint32 DepthSampleCount = (Surface.MSAATexture ? Surface.MSAATexture.GetSampleCount() : Surface.Texture.GetSampleCount());
+					bool const bDepthStencilSampleCountMismatchFixup = (SampleCount != DepthSampleCount);
 
-						ERenderTargetStoreAction HighLevelStoreAction = (Surface.MSAATexture && !bDepthStencilSampleCountMismatchFixup) ? ERenderTargetStoreAction::EMultisampleResolve : RenderTargetsInfo.DepthStencilRenderTarget.DepthStoreAction;
-						
+					ERenderTargetStoreAction HighLevelStoreAction = (Surface.MSAATexture && !bDepthStencilSampleCountMismatchFixup) ? ERenderTargetStoreAction::EMultisampleResolve : RenderTargetsInfo.DepthStencilRenderTarget.DepthStoreAction;
+					
 #if PLATFORM_IOS
-						FMetalTexture& Tex = Surface.MSAATexture ? Surface.MSAATexture : Surface.Texture;
-						if (Tex.GetStorageMode() == mtlpp::StorageMode::Memoryless)
-						{
-							HighLevelStoreAction = ERenderTargetStoreAction::ENoAction;
-						}
-#endif
-						
-						NewDepthStore = GetMetalRTStoreAction(HighLevelStoreAction);
-					}
-					else
+					FMetalTexture& Tex = Surface.MSAATexture ? Surface.MSAATexture : Surface.Texture;
+					if (Tex.GetStorageMode() == mtlpp::StorageMode::Memoryless)
 					{
-						bAllChecksPassed = false;
+						HighLevelStoreAction = ERenderTargetStoreAction::ENoAction;
 					}
+#endif
+					
+					NewDepthStore = GetMetalRTStoreAction(HighLevelStoreAction);
 				}
-				
-				if (InRenderTargetsInfo.DepthStencilRenderTarget.GetStencilStoreAction() > RenderTargetsInfo.DepthStencilRenderTarget.GetStencilStoreAction())
+				else
 				{
-					if (RenderPassDesc.GetStencilAttachment().GetTexture())
-					{
-						NewStencilStore = GetMetalRTStoreAction(RenderTargetsInfo.DepthStencilRenderTarget.GetStencilStoreAction());
+					bAllChecksPassed = false;
+				}
+			}
+			
+			if (InRenderTargetsInfo.DepthStencilRenderTarget.GetStencilStoreAction() > RenderTargetsInfo.DepthStencilRenderTarget.GetStencilStoreAction())
+			{
+				if (RenderPassDesc.GetStencilAttachment().GetTexture())
+				{
+					NewStencilStore = GetMetalRTStoreAction(RenderTargetsInfo.DepthStencilRenderTarget.GetStencilStoreAction());
 #if PLATFORM_IOS
-						if (RenderPassDesc.GetStencilAttachment().GetTexture().GetStorageMode() == mtlpp::StorageMode::Memoryless)
-						{
-							NewStencilStore = GetMetalRTStoreAction(ERenderTargetStoreAction::ENoAction);
-						}
-#endif
-					}
-					else
+					if (RenderPassDesc.GetStencilAttachment().GetTexture().GetStorageMode() == mtlpp::StorageMode::Memoryless)
 					{
-						bAllChecksPassed = false;
+						NewStencilStore = GetMetalRTStoreAction(ERenderTargetStoreAction::ENoAction);
 					}
+#endif
 				}
-				
-				if (bAllChecksPassed)
+				else
 				{
-					DepthStore = NewDepthStore;
-					StencilStore = NewStencilStore;
+					bAllChecksPassed = false;
 				}
 			}
-			else
+			
+			if (bAllChecksPassed)
 			{
-				bAllChecksPassed = false;
+				DepthStore = NewDepthStore;
+				StencilStore = NewStencilStore;
 			}
 		}
 	}
@@ -1860,9 +1840,6 @@ bool FMetalStateCache::PrepareToRestart(bool const bCurrentApplied)
 	{
 		if (SampleCount <= 1)
 		{
-			// Deferred store actions make life a bit easier...
-			static bool bSupportsDeferredStore = GetMetalDeviceContext().GetCommandQueue().SupportsFeature(EMetalFeaturesDeferredStoreActions);
-			
 			FRHISetRenderTargetsInfo Info = GetRenderTargetsInfo();
 			
 			ERenderTargetLoadAction DepthLoadAction = Info.DepthStencilRenderTarget.DepthLoadAction;
@@ -1947,9 +1924,7 @@ void FMetalStateCache::SetRenderStoreActions(FMetalCommandEncoder& CommandEncode
 {
 	check(CommandEncoder.IsRenderCommandEncoderActive())
 	{
-		// Deferred store actions make life a bit easier...
-		static bool bSupportsDeferredStore = GetMetalDeviceContext().GetCommandQueue().SupportsFeature(EMetalFeaturesDeferredStoreActions);
-		if (bConditionalSwitch && bSupportsDeferredStore)
+		if (bConditionalSwitch)
 		{
 			ns::Array<mtlpp::RenderPassColorAttachmentDescriptor> ColorAttachments = RenderPassDesc.GetColorAttachments();
 			for (int32 RenderTargetIndex = 0; RenderTargetIndex < RenderTargetsInfo.NumColorRenderTargets; RenderTargetIndex++)
diff --git a/Engine/Source/Runtime/Apple/MetalRHI/Private/MetalTexture.cpp b/Engine/Source/Runtime/Apple/MetalRHI/Private/MetalTexture.cpp
--- a/Engine/Source/Runtime/Apple/MetalRHI/Private/MetalTexture.cpp
+++ b/Engine/Source/Runtime/Apple/MetalRHI/Private/MetalTexture.cpp
@@ -214,13 +214,6 @@ void FMetalSurface::PrepareTextureView()
 		check(!bTextureView);
 		check(ImageSurfaceRef == nullptr);
 		
-		if (StencilTexture && (StencilTexture != Texture))
-		{
-			FMetalTexture OldStencilTexture = StencilTexture;
-			StencilTexture = Reallocate(StencilTexture, mtlpp::TextureUsage::PixelFormatView);
-			SafeReleaseMetalTexture(this, OldStencilTexture, bTextureView, ImageSurfaceRef != nullptr);
-		}
-		
 		check(Texture);
 		const bool bMSAATextureIsTexture = MSAATexture == Texture;
 		const bool bMSAAResolveTextureIsTexture = MSAAResolveTexture == Texture;
@@ -313,19 +306,6 @@ void FMetalSurface::ReplaceTexture(FMetalContext& Context, FMetalTexture Current
 		SafeReleaseMetalTexture(OldTexture);
 		Texture = NewTexture;
 	}
-	if (StencilTexture && StencilTexture == CurrentTexture)
-	{
-		if (StencilTexture && (StencilTexture != OldTexture))
-		{
-			FMetalTexture OldStencilTexture = StencilTexture;
-			SafeReleaseMetalTexture(OldStencilTexture);
-			StencilTexture = NewTexture;
-		}
-		else if(StencilTexture)
-		{
-			StencilTexture = NewTexture;
-		}
-	}
 	if (MSAATexture && MSAATexture == CurrentTexture)
 	{
 		if (MSAATexture && (MSAATexture != OldTexture))
@@ -362,8 +342,6 @@ void FMetalSurface::ReplaceTexture(FMetalContext& Context, FMetalTexture Current
 		}
 		SRV->TextureView->Texture = nil;
 		
-		SRV->TextureView->StencilTexture = nil;
-		
 		SRV->TextureView->MSAATexture = nil;
 		
 		if (SRV->Format == PF_Unknown)
@@ -385,13 +363,6 @@ void FMetalSurface::MakeAliasable(void)
 	static bool bSupportsHeaps = GetMetalDeviceContext().SupportsFeature(EMetalFeaturesHeaps);
 	if (bSupportsHeaps && Texture.GetStorageMode() == mtlpp::StorageMode::Private && Texture.GetHeap())
 	{
-		if (StencilTexture && (StencilTexture != Texture) && !StencilTexture.IsAliasable())
-		{
-			StencilTexture.MakeAliasable();
-#if STATS || ENABLE_LOW_LEVEL_MEM_TRACKER
-			MetalLLM::LogAliasTexture(StencilTexture);
-#endif
-		}
 		if (MSAATexture && (MSAATexture != Texture) && !MSAATexture.IsAliasable())
 		{
 			MSAATexture.MakeAliasable();
@@ -420,16 +391,6 @@ void FMetalSurface::MakeUnAliasable(void)
 		FMetalTexture OldTexture = Texture;
 		Texture = Reallocate(Texture, mtlpp::TextureUsage::Unknown);
 		SafeReleaseMetalTexture(this, OldTexture);
-		if (StencilTexture && (StencilTexture != OldTexture))
-		{
-			FMetalTexture OldStencilTexture = StencilTexture;
-			StencilTexture = Reallocate(StencilTexture, mtlpp::TextureUsage::Unknown);
-			SafeReleaseMetalTexture(this, OldStencilTexture);
-		}
-		else if(StencilTexture)
-		{
-			StencilTexture = Texture;
-		}
 		if (MSAATexture && (MSAATexture != OldTexture))
 		{
 			FMetalTexture OldMSAATexture = MSAATexture;
@@ -449,7 +410,6 @@ void FMetalSurface::MakeUnAliasable(void)
 			}
 			SRV->TextureView->Texture = nil;
 			
-			SRV->TextureView->StencilTexture = nil;
 			SRV->TextureView->MSAATexture = nil;
 			
 			if (SRV->Format == PF_Unknown)
@@ -480,7 +440,7 @@ void FMetalSurface::Init(FMetalSurface& Source, NSRange MipRange)
 	
 	// Recreate the texture to enable MTLTextureUsagePixelFormatView which must be off unless we definitely use this feature or we are throwing ~4% performance vs. Windows on the floor.
 	mtlpp::TextureUsage Usage = Source.Texture.GetUsage();
-	if(!(Usage & mtlpp::TextureUsage::PixelFormatView) && (Source.PixelFormat != PF_DepthStencil) && GetMetalDeviceContext().SupportsFeature(EMetalFeaturesResourceOptions) && !bUseSourceTex)
+	if(!(Usage & mtlpp::TextureUsage::PixelFormatView) && (Source.PixelFormat != PF_DepthStencil) && !bUseSourceTex)
 	{
 		Source.PrepareTextureView();
 	}
@@ -521,7 +481,7 @@ void FMetalSurface::Init(FMetalSurface& Source, NSRange MipRange, EPixelFormat F
 	
 	// Recreate the texture to enable MTLTextureUsagePixelFormatView which must be off unless we definitely use this feature or we are throwing ~4% performance vs. Windows on the floor.
 	mtlpp::TextureUsage Usage = Source.Texture.GetUsage();
-	if(!(Usage & mtlpp::TextureUsage::PixelFormatView) && (Source.PixelFormat == PF_DepthStencil && Format == PF_X24_G8 && GetMetalDeviceContext().SupportsFeature(EMetalFeaturesStencilView)) && !bUseSourceTex)
+	if(!(Usage & mtlpp::TextureUsage::PixelFormatView) && (Source.PixelFormat == PF_DepthStencil && Format == PF_X24_G8) && !bUseSourceTex)
 	{
 		Source.PrepareTextureView();
 	}
@@ -531,7 +491,7 @@ void FMetalSurface::Init(FMetalSurface& Source, NSRange MipRange, EPixelFormat F
 	{
 		Texture = Source.Texture.NewTextureView(MetalFormat, Source.Texture.GetTextureType(), ns::Range(MipRange.location, MipRange.length), Slices);
 	}
-	else if (Source.PixelFormat == PF_DepthStencil && Format == PF_X24_G8 && GetMetalDeviceContext().SupportsFeature(EMetalFeaturesStencilView) && (mtlpp::PixelFormat)GPixelFormats[PF_DepthStencil].PlatformFormat != mtlpp::PixelFormat::Depth32Float)
+	else if (Source.PixelFormat == PF_DepthStencil && Format == PF_X24_G8)
 	{
 		switch (Source.Texture.GetPixelFormat())
 		{
@@ -550,42 +510,11 @@ void FMetalSurface::Init(FMetalSurface& Source, NSRange MipRange, EPixelFormat F
 		
 		Texture = Source.Texture.NewTextureView(MetalFormat, Source.Texture.GetTextureType(), ns::Range(MipRange.location, MipRange.length), Slices);
 	}
-	else if(Source.PixelFormat == PF_DepthStencil && Format == PF_X24_G8)
-	{
-		check(Source.Type == RRT_Texture2D); // Only 2D Textures can be X24_G8 SRVs
-		
-		mtlpp::PixelFormat DepthStencilFormat = Source.Texture ? Source.Texture.GetPixelFormat() : mtlpp::PixelFormat::Invalid;
-		
-		switch(DepthStencilFormat)
-		{
-			case mtlpp::PixelFormat::Stencil8:
-			{
-				Texture = Source.Texture;
-				break;
-			}
-			case mtlpp::PixelFormat::Depth32Float:
-			{
-				check(Source.StencilTexture);
-				Texture = Source.StencilTexture;
-				break;
-			}
-			default:
-			{
-				check(false);
-				break;
-			}
-		}
-	}
 	else
 	{
 		Texture = Source.Texture;
 	}
 	
-	if(Source.StencilTexture && !StencilTexture)
-	{
-		StencilTexture = Texture;
-	}
-	
 	const uint32 BlockSizeX = GPixelFormats[PixelFormat].BlockSizeX;
 	const uint32 BlockSizeY = GPixelFormats[PixelFormat].BlockSizeY;
 	const uint32 BlockBytes = GPixelFormats[PixelFormat].BlockBytes;
@@ -600,7 +529,6 @@ FMetalSurface::FMetalSurface(FMetalSurface& Source, NSRange MipRange)
 , Texture(nil)
 , MSAATexture(nil)
 , MSAAResolveTexture(nil)
-, StencilTexture(nil)
 , SizeX(Source.SizeX)
 , SizeY(Source.SizeY)
 , SizeZ(Source.SizeZ)
@@ -623,7 +551,6 @@ FMetalSurface::FMetalSurface(FMetalSurface& Source, NSRange const MipRange, EPix
 , Texture(nil)
 , MSAATexture(nil)
 , MSAAResolveTexture(nil)
-, StencilTexture(nil)
 , SizeX(Source.SizeX)
 , SizeY(Source.SizeY)
 , SizeZ(Source.SizeZ)
@@ -731,7 +658,6 @@ FMetalSurface::FMetalSurface(ERHIResourceType ResourceType, EPixelFormat Format,
 , Texture(nil)
 , MSAATexture(nil)
 , MSAAResolveTexture(nil)
-, StencilTexture(nil)
 , SizeX(InSizeX)
 , SizeY(InSizeY)
 , SizeZ(InSizeZ)
@@ -753,10 +679,7 @@ FMetalSurface::FMetalSurface(ERHIResourceType ResourceType, EPixelFormat Format,
 		GetMetalPixelFormatKey(mtlpp::PixelFormat::Depth32Float_Stencil8);
 #if PLATFORM_MAC
 		GetMetalPixelFormatKey(mtlpp::PixelFormat::Depth24Unorm_Stencil8);
-		if (GetMetalDeviceContext().SupportsFeature(EMetalFeaturesDepth16))
-		{
-			GetMetalPixelFormatKey(mtlpp::PixelFormat::Depth16Unorm);
-		}
+		GetMetalPixelFormatKey(mtlpp::PixelFormat::Depth16Unorm);
 #endif
 	}
 	
@@ -846,7 +769,6 @@ FMetalSurface::FMetalSurface(ERHIResourceType ResourceType, EPixelFormat Format,
 	}
 	Desc.SetMipmapLevelCount(NumMips);
 	
-	if(GetMetalDeviceContext().SupportsFeature(EMetalFeaturesResourceOptions))
 	{
 		Desc.SetUsage(ConvertFlagsToUsage(Flags));
 		
@@ -877,7 +799,7 @@ FMetalSurface::FMetalSurface(ERHIResourceType ResourceType, EPixelFormat Format,
 		{
 			check(!(Flags & TexCreate_CPUReadback));
 #if PLATFORM_IOS
-			if (FMetalCommandQueue::SupportsFeature(EMetalFeaturesMemoryLessResources) && !(Flags & (TexCreate_ShaderResource|TexCreate_UAV)))
+			if (!(Flags & (TexCreate_ShaderResource|TexCreate_UAV)))
 			{
 				Desc.SetStorageMode(mtlpp::StorageMode::Memoryless);
 				Desc.SetResourceOptions(mtlpp::ResourceOptions::StorageModeMemoryless);
@@ -1014,7 +936,7 @@ FMetalSurface::FMetalSurface(ERHIResourceType ResourceType, EPixelFormat Format,
 			
 			bool bMemoryless = false;
 #if PLATFORM_IOS
-			if (FMetalCommandQueue::SupportsFeature(EMetalFeaturesMemoryLessResources) && GMaxRHIShaderPlatform == SP_METAL)
+			if (GMaxRHIShaderPlatform == SP_METAL)
 			{
 				bMemoryless = true;
 				Desc.SetStorageMode(mtlpp::StorageMode::Memoryless);
@@ -1072,21 +994,8 @@ FMetalSurface::FMetalSurface(ERHIResourceType ResourceType, EPixelFormat Format,
 	// create a stencil buffer if needed
 	if (Format == PF_DepthStencil)
 	{
-		if(MTLFormat == mtlpp::PixelFormat::Depth32Float)
-		{
-			Desc.SetPixelFormat(mtlpp::PixelFormat::Stencil8);
-			StencilTexture = GetMetalDeviceContext().CreateTexture(this, Desc);
-			
-			// 1 byte per texel
-			TotalTextureSize += SizeX * SizeY;
-		}
-		else
-		{
-			StencilTexture = Texture;
-			
-			// 1 byte per texel
-			TotalTextureSize += SizeX * SizeY;
-		}
+		// 1 byte per texel
+		TotalTextureSize += SizeX * SizeY;
 	}
 	
 	// track memory usage
@@ -1218,14 +1127,6 @@ FMetalSurface::~FMetalSurface()
 		}
 	}
 	
-	if (StencilTexture.GetPtr())
-	{
-		if (StencilTexture.GetPtr() != Texture.GetPtr())
-		{
-			SafeReleaseMetalTexture(this, StencilTexture, bTextureView, false);
-		}
-	}
-	
 	if (!(Flags & TexCreate_Presentable) && Texture.GetPtr())
 	{
 		SafeReleaseMetalTexture(this, Texture, bTextureView, (ImageSurfaceRef != nullptr));
@@ -1234,7 +1135,6 @@ FMetalSurface::~FMetalSurface()
 	MSAATexture = nil;
 	MSAAResolveTexture = nil;
 	Texture = nil;
-	StencilTexture = nil;
 	
 	// track memory usage
 	FMetalDeferredStats* Block = [FMetalDeferredStats new];
@@ -1264,14 +1164,11 @@ FMetalSurface::~FMetalSurface()
 
 FMetalBuffer FMetalSurface::AllocSurface(uint32 MipIndex, uint32 ArrayIndex, EResourceLockMode LockMode, uint32& DestStride)
 {
-	// Whether the device supports resource options, so we don't access invalid properties on older versions of iOS
-	bool const bSupportsResourceOptions = GetMetalDeviceContext().SupportsFeature(EMetalFeaturesResourceOptions);
-	
 	// get size and stride
 	uint32 MipBytes = GetMipSize(MipIndex, &DestStride, false);
 	
 	// allocate some temporary memory
-	mtlpp::ResourceOptions ResMode = GetMetalDeviceContext().GetCommandQueue().GetCompatibleResourceOptions((mtlpp::ResourceOptions)(mtlpp::ResourceOptions::StorageModeShared | (bSupportsResourceOptions && !(PLATFORM_MAC && PixelFormat == PF_G8 && (Flags & TexCreate_SRGB)) ? mtlpp::ResourceOptions::CpuCacheModeWriteCombined : 0)));
+	mtlpp::ResourceOptions ResMode = GetMetalDeviceContext().GetCommandQueue().GetCompatibleResourceOptions((mtlpp::ResourceOptions)(mtlpp::ResourceOptions::StorageModeShared | (!(PLATFORM_MAC && PixelFormat == PF_G8 && (Flags & TexCreate_SRGB)) ? mtlpp::ResourceOptions::CpuCacheModeWriteCombined : 0)));
 	
 	FMetalBuffer Buffer = GetMetalDeviceContext().GetResourceHeap().CreateBuffer(MipBytes, BufferOffsetAlignment, ResMode);
 	
@@ -1295,9 +1192,6 @@ void FMetalSurface::UpdateSurface(FMetalBuffer& Buffer, uint32 MipIndex, uint32
 #endif
 	check(Buffer);
 	
-	// Whether the device supports resource options, so we don't access invalid properties on older versions of iOS
-	bool const bSupportsResourceOptions = GetMetalDeviceContext().SupportsFeature(EMetalFeaturesResourceOptions);
-	
 	uint32 Stride;
 	uint32 BytesPerImage = GetMipSize(MipIndex, &Stride, true);
 	
@@ -1335,7 +1229,7 @@ void FMetalSurface::UpdateSurface(FMetalBuffer& Buffer, uint32 MipIndex, uint32
 	}
 #endif
 	
-	if(bSupportsResourceOptions && Texture.GetStorageMode() == mtlpp::StorageMode::Private)
+	if(Texture.GetStorageMode() == mtlpp::StorageMode::Private)
 	{
 		SCOPED_AUTORELEASE_POOL;
 		
@@ -1417,9 +1311,6 @@ void FMetalSurface::UpdateSurface(FMetalBuffer& Buffer, uint32 MipIndex, uint32
 
 void* FMetalSurface::Lock(uint32 MipIndex, uint32 ArrayIndex, EResourceLockMode LockMode, uint32& DestStride)
 {
-	// Whether the device supports resource options, so we don't access invalid properties on older versions of iOS
-	bool const bSupportsResourceOptions = GetMetalDeviceContext().SupportsFeature(EMetalFeaturesResourceOptions);
-	
 	// get size and stride
 	uint32 MipBytes = GetMipSize(MipIndex, &DestStride, false);
 	
@@ -1447,7 +1338,7 @@ void* FMetalSurface::Lock(uint32 MipIndex, uint32 ArrayIndex, EResourceLockMode
 				Region = mtlpp::Region(0, 0, 0, FMath::Max<uint32>(SizeX >> MipIndex, 1), FMath::Max<uint32>(SizeY >> MipIndex, 1), FMath::Max<uint32>(SizeZ >> MipIndex, 1));
 			}
 			
-			if (bSupportsResourceOptions && Texture.GetStorageMode() == mtlpp::StorageMode::Private)
+			if (Texture.GetStorageMode() == mtlpp::StorageMode::Private)
 			{
 				GetMetalDeviceContext().CopyFromTextureToBuffer(Texture, ArrayIndex, MipIndex, Region.origin, Region.size, LockedMemory[MipIndex], 0, DestStride, MipBytes, mtlpp::BlitOption::None);
 				
@@ -2081,7 +1972,7 @@ static void InternalUpdateTexture2D(FMetalContext& Context, FTexture2DRHIParamRe
 	
 	mtlpp::Region Region = mtlpp::Region(UpdateRegion.DestX, UpdateRegion.DestY, UpdateRegion.Width, UpdateRegion.Height);
 	
-	if(GetMetalDeviceContext().SupportsFeature(EMetalFeaturesResourceOptions) && Tex.GetStorageMode() == mtlpp::StorageMode::Private)
+	if(Tex.GetStorageMode() == mtlpp::StorageMode::Private)
 	{
 		SCOPED_AUTORELEASE_POOL;
 		
@@ -2155,7 +2046,7 @@ void FMetalDynamicRHI::RHIUpdateTexture2D(FTexture2DRHIParamRef TextureRHI, uint
 	{
 		FMetalTexture2D* Texture = ResourceCast(TextureRHI);
 		FMetalTexture Tex = Texture->Surface.Texture;
-		bool const bUseIntermediateMetalBuffer = GetMetalDeviceContext().SupportsFeature(EMetalFeaturesResourceOptions) && Tex.GetStorageMode() == mtlpp::StorageMode::Private;
+		bool const bUseIntermediateMetalBuffer = Tex.GetStorageMode() == mtlpp::StorageMode::Private;
 		
 		if(bUseIntermediateMetalBuffer)
 		{
@@ -2214,7 +2105,7 @@ static void InternalUpdateTexture3D(FMetalContext& Context, FTexture3DRHIParamRe
 	
 	mtlpp::Region Region = mtlpp::Region(UpdateRegion.DestX, UpdateRegion.DestY, UpdateRegion.DestZ, UpdateRegion.Width, UpdateRegion.Height, UpdateRegion.Depth);
 	
-	if(GetMetalDeviceContext().SupportsFeature(EMetalFeaturesResourceOptions) && Tex.GetStorageMode() == mtlpp::StorageMode::Private)
+	if(Tex.GetStorageMode() == mtlpp::StorageMode::Private)
 	{
 		const uint32 BytesPerImage = SourceRowPitch * UpdateRegion.Height;
 		mtlpp::BlitOption Options = mtlpp::BlitOption::None;
@@ -2322,12 +2213,10 @@ void FMetalDynamicRHI::RHIUpdateTexture3D(FTexture3DRHIParamRef TextureRHI,uint3
 		FMetalTexture3D* Texture = ResourceCast(TextureRHI);
 		FMetalTexture Tex = Texture->Surface.Texture;
 		
-		bool const bSupportsResourceOptions = GetMetalDeviceContext().SupportsFeature(EMetalFeaturesResourceOptions);
-		
 #if PLATFORM_MAC
 		checkf(!(Texture->GetFormat() == PF_G8 && (Texture->GetFlags() & TexCreate_SRGB)), TEXT("MetalRHI does not support PF_G8_sRGB on 3D, array or cube textures as it requires manual, CPU-side expansion to RGBA8_sRGB which is expensive!"));
 #endif
-		if(bSupportsResourceOptions && Tex.GetStorageMode() == mtlpp::StorageMode::Private)
+		if(Tex.GetStorageMode() == mtlpp::StorageMode::Private)
 		{
 			SCOPED_AUTORELEASE_POOL;
 
@@ -2521,17 +2410,6 @@ void FMetalDynamicRHI::RHIBindDebugLabelName(FTextureRHIParamRef TextureRHI, con
 		{
 			Surf->MSAATexture.SetLabel(FString(Name).GetNSString());
 		}
-		if(Surf->StencilTexture)
-		{
-			if(Surf->StencilTexture != Surf->Texture)
-			{
-				Surf->StencilTexture.SetLabel([NSString stringWithFormat:@"%@StencilSRV", FString(Name).GetNSString()]);
-			}
-			else
-			{
-				Surf->StencilTexture.SetLabel(FString(Name).GetNSString());
-			}
-		}
 	}
 }
 
diff --git a/Engine/Source/Runtime/Apple/MetalRHI/Private/MetalUAV.cpp b/Engine/Source/Runtime/Apple/MetalRHI/Private/MetalUAV.cpp
--- a/Engine/Source/Runtime/Apple/MetalRHI/Private/MetalUAV.cpp
+++ b/Engine/Source/Runtime/Apple/MetalRHI/Private/MetalUAV.cpp
@@ -29,8 +29,6 @@ FMetalShaderResourceView::~FMetalShaderResourceView()
 		{
 			TextureView->Texture = nil;
 			
-			TextureView->StencilTexture = nil;
-			
 			TextureView->MSAATexture = nil;
 		}
 		delete TextureView;
@@ -44,7 +42,6 @@ FMetalShaderResourceView::~FMetalShaderResourceView()
 ns::AutoReleased<FMetalTexture> FMetalShaderResourceView::GetLinearTexture(bool const bUAV)
 {
 	ns::AutoReleased<FMetalTexture> NewLinearTexture;
-	if (FMetalCommandQueue::SupportsFeature(EMetalFeaturesLinearTextures) && (!bUAV || FMetalCommandQueue::SupportsFeature(EMetalFeaturesLinearTextureUAVs)))
 	{
 		if (IsValidRef(SourceVertexBuffer))
 		{
@@ -69,7 +66,7 @@ FUnorderedAccessViewRHIRef FMetalDynamicRHI::RHICreateUnorderedAccessView_Render
 {
 	FMetalSurface* Surface = (FMetalSurface*)Texture->GetTextureBaseRHI();
 	FMetalTexture Tex = Surface->Texture;
-	if (ImmediateContext.Context->GetCommandQueue().SupportsFeature(EMetalFeaturesResourceOptions) && !(Tex.GetUsage() & mtlpp::TextureUsage::PixelFormatView))
+	if (!(Tex.GetUsage() & mtlpp::TextureUsage::PixelFormatView))
 	{
 		FScopedRHIThreadStaller StallRHIThread(RHICmdList);
 		return GDynamicRHI->RHICreateUnorderedAccessView(Texture, MipLevel);
@@ -153,8 +150,6 @@ FUnorderedAccessViewRHIRef FMetalDynamicRHI::RHICreateUnorderedAccessView(FVerte
 	SRV->SourceIndexBuffer = nullptr;
 	SRV->SourceStructuredBuffer = nullptr;
 	SRV->Format = Format;
-		
-	if (FMetalCommandQueue::SupportsFeature(EMetalFeaturesLinearTextureUAVs))
 	{
 		check(VertexBuffer->GetUsage() & BUF_UnorderedAccess);
 		VertexBuffer->CreateLinearTexture((EPixelFormat)Format, VertexBuffer);
@@ -179,8 +174,6 @@ FUnorderedAccessViewRHIRef FMetalDynamicRHI::RHICreateUnorderedAccessView(FIndex
 		SRV->SourceIndexBuffer = IndexBuffer;
 		SRV->SourceStructuredBuffer = nullptr;
 		SRV->Format = Format;
-		
-		if (FMetalCommandQueue::SupportsFeature(EMetalFeaturesLinearTextureUAVs))
 		{
 			check(IndexBuffer->GetUsage() & BUF_UnorderedAccess);
 			IndexBuffer->CreateLinearTexture((EPixelFormat)Format, IndexBuffer);
@@ -198,7 +191,7 @@ FShaderResourceViewRHIRef FMetalDynamicRHI::RHICreateShaderResourceView_RenderTh
 {
 	FMetalTexture2D* Texture = ResourceCast(Texture2DRHI);
 	FMetalTexture Tex = Texture->Surface.Texture;
-	if (ImmediateContext.Context->GetCommandQueue().SupportsFeature(EMetalFeaturesResourceOptions) && !(Tex.GetUsage() & mtlpp::TextureUsage::PixelFormatView))
+	if (!(Tex.GetUsage() & mtlpp::TextureUsage::PixelFormatView))
 	{
 		FScopedRHIThreadStaller StallRHIThread(RHICmdList);
 		return GDynamicRHI->RHICreateShaderResourceView(Texture2DRHI, MipLevel);
@@ -213,7 +206,7 @@ FShaderResourceViewRHIRef FMetalDynamicRHI::RHICreateShaderResourceView_RenderTh
 {
 	FMetalTexture2D* Texture = ResourceCast(Texture2DRHI);
 	FMetalTexture Tex = Texture->Surface.Texture;
-	if (ImmediateContext.Context->GetCommandQueue().SupportsFeature(EMetalFeaturesResourceOptions) && !(Tex.GetUsage() & mtlpp::TextureUsage::PixelFormatView))
+	if (!(Tex.GetUsage() & mtlpp::TextureUsage::PixelFormatView))
 	{
 		FScopedRHIThreadStaller StallRHIThread(RHICmdList);
 		return GDynamicRHI->RHICreateShaderResourceView(Texture2DRHI, MipLevel, NumMipLevels, Format);
@@ -228,7 +221,7 @@ FShaderResourceViewRHIRef FMetalDynamicRHI::RHICreateShaderResourceView_RenderTh
 {
 	FMetalTexture3D* Texture = ResourceCast(Texture3DRHI);
 	FMetalTexture Tex = Texture->Surface.Texture;
-	if (ImmediateContext.Context->GetCommandQueue().SupportsFeature(EMetalFeaturesResourceOptions) && !(Tex.GetUsage() & mtlpp::TextureUsage::PixelFormatView))
+	if (!(Tex.GetUsage() & mtlpp::TextureUsage::PixelFormatView))
 	{
 		FScopedRHIThreadStaller StallRHIThread(RHICmdList);
 		return GDynamicRHI->RHICreateShaderResourceView(Texture3DRHI, MipLevel);
@@ -243,7 +236,7 @@ FShaderResourceViewRHIRef FMetalDynamicRHI::RHICreateShaderResourceView_RenderTh
 {
 	FMetalTexture2DArray* Texture = ResourceCast(Texture2DArrayRHI);
 	FMetalTexture Tex = Texture->Surface.Texture;
-	if (ImmediateContext.Context->GetCommandQueue().SupportsFeature(EMetalFeaturesResourceOptions) && !(Tex.GetUsage() & mtlpp::TextureUsage::PixelFormatView))
+	if (!(Tex.GetUsage() & mtlpp::TextureUsage::PixelFormatView))
 	{
 		FScopedRHIThreadStaller StallRHIThread(RHICmdList);
 		return GDynamicRHI->RHICreateShaderResourceView(Texture2DArrayRHI, MipLevel);
@@ -258,7 +251,7 @@ FShaderResourceViewRHIRef FMetalDynamicRHI::RHICreateShaderResourceView_RenderTh
 {
 	FMetalTextureCube* Texture = ResourceCast(TextureCubeRHI);
 	FMetalTexture Tex = Texture->Surface.Texture;
-	if (ImmediateContext.Context->GetCommandQueue().SupportsFeature(EMetalFeaturesResourceOptions) && !(Tex.GetUsage() & mtlpp::TextureUsage::PixelFormatView))
+	if (!(Tex.GetUsage() & mtlpp::TextureUsage::PixelFormatView))
 	{
 		FScopedRHIThreadStaller StallRHIThread(RHICmdList);
 		return GDynamicRHI->RHICreateShaderResourceView(TextureCubeRHI, MipLevel);
@@ -328,8 +321,6 @@ FShaderResourceViewRHIRef FMetalDynamicRHI::RHICreateShaderResourceView(FVertexB
 	SRV->SourceStructuredBuffer = nullptr;
 	SRV->Format = Format;
 	SRV->Stride = Stride;
-	
-	if (FMetalCommandQueue::SupportsFeature(EMetalFeaturesLinearTextures))
 	{
 		check(Stride == GPixelFormats[Format].BlockBytes);
 		check(VertexBuffer->GetUsage() & BUF_ShaderResource);
@@ -352,8 +343,6 @@ FShaderResourceViewRHIRef FMetalDynamicRHI::RHICreateShaderResourceView(FIndexBu
 	SRV->TextureView = nullptr;
 	SRV->SourceStructuredBuffer = nullptr;
 	SRV->Format = (Buffer->IndexType == mtlpp::IndexType::UInt16) ? PF_R16_UINT : PF_R32_UINT;
-	
-	if (FMetalCommandQueue::SupportsFeature(EMetalFeaturesLinearTextures))
 	{
 		Buffer->CreateLinearTexture((EPixelFormat)SRV->Format, Buffer);
 	}
diff --git a/Engine/Source/Runtime/Apple/MetalRHI/Private/MetalVertexBuffer.cpp b/Engine/Source/Runtime/Apple/MetalRHI/Private/MetalVertexBuffer.cpp
--- a/Engine/Source/Runtime/Apple/MetalRHI/Private/MetalVertexBuffer.cpp
+++ b/Engine/Source/Runtime/Apple/MetalRHI/Private/MetalVertexBuffer.cpp
@@ -146,8 +146,7 @@ FMetalRHIBuffer::FMetalRHIBuffer(uint32 InSize, uint32 InUsage, ERHIResourceType
 					AllocSize = Align(AllocSize + 512, 1024);
 				}
 				
-				if ((FMetalCommandQueue::SupportsFeature(EMetalFeaturesLinearTextures) && (InUsage & (BUF_ShaderResource)))
-				|| (FMetalCommandQueue::SupportsFeature(EMetalFeaturesLinearTextureUAVs) && (InUsage & (BUF_UnorderedAccess))))
+				if (InUsage & (BUF_ShaderResource|BUF_UnorderedAccess))
 				{
 					uint32 NumElements = AllocSize;
 					uint32 SizeX = NumElements;
@@ -248,7 +247,7 @@ void FMetalRHIBuffer::Alloc(uint32 InSize, EResourceLockMode LockMode)
 
         METAL_INC_DWORD_STAT_BY(Type, MemAlloc, InSize);
         
-		if (FMetalCommandQueue::SupportsFeature(EMetalFeaturesLinearTextures) && (Usage & (BUF_UnorderedAccess|BUF_ShaderResource)))
+		if ((Usage & (BUF_UnorderedAccess|BUF_ShaderResource)))
 		{
 			for (TPair<EPixelFormat, FMetalTexture>& Pair : LinearTextures)
 			{
@@ -274,7 +273,7 @@ void FMetalRHIBuffer::Alloc(uint32 InSize, EResourceLockMode LockMode)
 FMetalTexture FMetalRHIBuffer::AllocLinearTexture(EPixelFormat Format)
 {
 	METAL_LLM_BUFFER_SCOPE(Type);
-	if (FMetalCommandQueue::SupportsFeature(EMetalFeaturesLinearTextures) && (Usage & (BUF_UnorderedAccess|BUF_ShaderResource)))
+	if ((Usage & (BUF_UnorderedAccess|BUF_ShaderResource)))
 	{
 		mtlpp::PixelFormat MTLFormat = (mtlpp::PixelFormat)GMetalBufferFormats[Format].LinearTextureFormat;
 		
@@ -369,7 +368,7 @@ struct FMetalRHICommandCreateLinearTexture : public FRHICommand<FMetalRHICommand
 ns::AutoReleased<FMetalTexture> FMetalRHIBuffer::CreateLinearTexture(EPixelFormat Format, FRHIResource* InParent)
 {
 	ns::AutoReleased<FMetalTexture> Texture;
-	if (FMetalCommandQueue::SupportsFeature(EMetalFeaturesLinearTextures) && (Usage & (BUF_UnorderedAccess|BUF_ShaderResource)) && GMetalBufferFormats[Format].LinearTextureFormat != mtlpp::PixelFormat::Invalid)
+	if ((Usage & (BUF_UnorderedAccess|BUF_ShaderResource)) && GMetalBufferFormats[Format].LinearTextureFormat != mtlpp::PixelFormat::Invalid)
 	{
 		if (IsRunningRHIInSeparateThread() && !IsInRHIThread() && !FRHICommandListExecutor::GetImmediateCommandList().Bypass())
 		{
@@ -398,7 +397,7 @@ ns::AutoReleased<FMetalTexture> FMetalRHIBuffer::CreateLinearTexture(EPixelForma
 ns::AutoReleased<FMetalTexture> FMetalRHIBuffer::GetLinearTexture(EPixelFormat Format)
 {
 	ns::AutoReleased<FMetalTexture> Texture;
-	if (FMetalCommandQueue::SupportsFeature(EMetalFeaturesLinearTextures) && (Usage & (BUF_UnorderedAccess|BUF_ShaderResource)) && GMetalBufferFormats[Format].LinearTextureFormat != mtlpp::PixelFormat::Invalid)
+	if ((Usage & (BUF_UnorderedAccess|BUF_ShaderResource)) && GMetalBufferFormats[Format].LinearTextureFormat != mtlpp::PixelFormat::Invalid)
 	{
 		FMetalTexture* ExistingTexture = LinearTextures.Find(Format);
 		if (ExistingTexture)
diff --git a/Engine/Source/Runtime/Apple/MetalRHI/Private/MetalViewport.cpp b/Engine/Source/Runtime/Apple/MetalRHI/Private/MetalViewport.cpp
--- a/Engine/Source/Runtime/Apple/MetalRHI/Private/MetalViewport.cpp
+++ b/Engine/Source/Runtime/Apple/MetalRHI/Private/MetalViewport.cpp
@@ -404,7 +404,6 @@ void FMetalViewport::Present(FMetalCommandQueue& CommandQueue, bool bLockToVsync
 	NSNumber* ScreenId = [View.window.screen.deviceDescription objectForKey:@"NSScreenNumber"];
 	DisplayID = ScreenId.unsignedIntValue;
 	bIsLiveResize = View.inLiveResize;
-	if (FMetalCommandQueue::SupportsFeature(EMetalFeaturesSupportsVSyncToggle))
 	{
 		FCAMetalLayer* CurrentLayer = (FCAMetalLayer*)[View layer];
 		static bool sVSyncSafe = FPlatformMisc::MacOSXVersionCompare(10,13,4) >= 0;
diff --git a/Engine/Source/Runtime/Apple/MetalRHI/Public/MetalResources.h b/Engine/Source/Runtime/Apple/MetalRHI/Public/MetalResources.h
--- a/Engine/Source/Runtime/Apple/MetalRHI/Public/MetalResources.h
+++ b/Engine/Source/Runtime/Apple/MetalRHI/Public/MetalResources.h
@@ -503,7 +503,6 @@ public:
 	// iOS A9+ where depth resolve is available
 	// iOS < A9 where depth resolve is unavailable.
 	FMetalTexture MSAAResolveTexture;
-	FMetalTexture StencilTexture;
 	uint32 SizeX, SizeY, SizeZ;
 	bool bIsCubemap;
 	int32 volatile Written;
