iOS compile fixes.

diff --git a/Engine/Source/Runtime/Apple/MetalRHI/Private/MetalStateCache.cpp b/Engine/Source/Runtime/Apple/MetalRHI/Private/MetalStateCache.cpp
--- a/Engine/Source/Runtime/Apple/MetalRHI/Private/MetalStateCache.cpp
+++ b/Engine/Source/Runtime/Apple/MetalRHI/Private/MetalStateCache.cpp
@@ -710,9 +710,7 @@ bool FMetalStateCache::SetRenderTargetsInfo(FRHISetRenderTargetsInfo const& InRe
 					switch(DepthStencilFormat)
 					{
 						case mtlpp::PixelFormat::Depth32Float:
-#if !PLATFORM_MAC
-							StencilTexture = (DepthStencilPixelFormat == PF_DepthStencil) ? Surface.StencilTexture : nil;
-#endif
+							StencilTexture =  nil;
 							break;
 						case mtlpp::PixelFormat::Stencil8:
 							StencilTexture = DepthTexture;
diff --git a/Engine/Source/Runtime/Apple/MetalRHI/Private/MetalTexture.cpp b/Engine/Source/Runtime/Apple/MetalRHI/Private/MetalTexture.cpp
--- a/Engine/Source/Runtime/Apple/MetalRHI/Private/MetalTexture.cpp
+++ b/Engine/Source/Runtime/Apple/MetalRHI/Private/MetalTexture.cpp
@@ -2796,7 +2796,7 @@ void FMetalRHICommandContext::RHICopyTexture(FTextureRHIParamRef SourceTextureRH
 						
 						mtlpp::BlitOption Options = mtlpp::BlitOption::None;
 #if !PLATFORM_MAC
-						if (MetalSrcTexture->Surface.Texture.GetPixelFormat() >= mtlpp::PixelFormat::PVRTC_RGB_2BPP && MetalSrcTexture->Surface.Texture.GetPixelFormat() <= mtlpp::PixelFormat::PVRTC_RGBA_4BPP_sRGB)
+						if (MetalSrcTexture->Texture.GetPixelFormat() >= mtlpp::PixelFormat::PVRTC_RGB_2BPP && MetalSrcTexture->Texture.GetPixelFormat() <= mtlpp::PixelFormat::PVRTC_RGBA_4BPP_sRGB)
 						{
 							Options = mtlpp::BlitOption::RowLinearPVRTC;
 						}
