Fix the ValidateVersion function so that it actually does what it should do even with newer shader platforms.

diff --git a/Engine/Source/Runtime/Apple/MetalRHI/Private/MetalShaders.cpp b/Engine/Source/Runtime/Apple/MetalRHI/Private/MetalShaders.cpp
--- a/Engine/Source/Runtime/Apple/MetalRHI/Private/MetalShaders.cpp
+++ b/Engine/Source/Runtime/Apple/MetalRHI/Private/MetalShaders.cpp
@@ -145,8 +145,6 @@ mtlpp::LanguageVersion ValidateVersion(uint8 Version)
 		TEXT("Metal 2.1"),
 	};
 	
-	Version = FMath::Min(Version, (uint8)3);
-	
 	mtlpp::LanguageVersion Result = mtlpp::LanguageVersion::Version1_1;
 	switch(Version)
 	{
@@ -179,9 +177,9 @@ mtlpp::LanguageVersion ValidateVersion(uint8 Version)
 #if PLATFORM_MAC
 		Args.Add(TEXT("RequiredOS"), FText::FromString(FString::Printf(TEXT("macOS %d.%d.%d"), MetalMacOSVersions[Version][0], MetalMacOSVersions[Version][1], MetalMacOSVersions[Version][2])));
 #else
-		Args.Add(TEXT("RequiredOS"), FText::FromString(FString::Printf(TEXT("macOS %d.%d.%d"), MetaliOSVersions[Version][0], MetaliOSVersions[Version][1], MetaliOSVersions[Version][2])));
+		Args.Add(TEXT("RequiredOS"), FText::FromString(FString::Printf(TEXT("%d.%d.%d"), MetaliOSVersions[Version][0], MetaliOSVersions[Version][1], MetaliOSVersions[Version][2])));
 #endif
-		FText LocalizedMsg = FText::Format(NSLOCTEXT("MetalRHI", "ShaderVersionUnsupported", "The current OS version does not support {Version} required by the project. You must upgrade to {RequiredOS} to run this project."),Args);
+		FText LocalizedMsg = FText::Format(NSLOCTEXT("MetalRHI", "ShaderVersionUnsupported", "The current OS version does not support {ShaderVersion} required by the project. You must upgrade to {RequiredOS} to run this project."),Args);
 		
 		FText Title = NSLOCTEXT("MetalRHI", "ShaderVersionUnsupportedTitle", "Shader Version Unsupported");
 		FMessageDialog::Open(EAppMsgType::Ok, LocalizedMsg, &Title);
