Fix handling of shader versions for text shaders up to Metal 2.1 - need a better way of dealing with this that adapts more easily to new versions.

diff --git a/Engine/Source/Runtime/Apple/MetalRHI/Private/MetalShaders.cpp b/Engine/Source/Runtime/Apple/MetalRHI/Private/MetalShaders.cpp
--- a/Engine/Source/Runtime/Apple/MetalRHI/Private/MetalShaders.cpp
+++ b/Engine/Source/Runtime/Apple/MetalRHI/Private/MetalShaders.cpp
@@ -406,20 +406,34 @@ void TMetalBaseShader<BaseResourceType, ShaderType>::Init(const TArray<uint8>& I
 #endif
 			CompileOptions.SetPreprocessorMacros(PreprocessorMacros);
 #endif
+			
+			mtlpp::LanguageVersion MetalVersion;
+			switch(Header.Version)
 			{
-				if (Header.Version < 3)
-				{
-	#if PLATFORM_MAC
-					CompileOptions.SetLanguageVersion(Header.Version == 0 ? mtlpp::LanguageVersion::Version1_1 : (mtlpp::LanguageVersion)((1 << 16) + FMath::Min(Header.Version, (uint8)2u)));
-	#else
-					CompileOptions.SetLanguageVersion((mtlpp::LanguageVersion)((1 << 16) + FMath::Min(Header.Version, (uint8)2u)));
-	#endif
-				}
-				else if (Header.Version == 3)
-				{
-					CompileOptions.SetLanguageVersion((mtlpp::LanguageVersion)(2 << 16));
-				}
+				case 4:
+					MetalVersion = mtlpp::LanguageVersion::Version2_1;
+					break;
+				case 3:
+					MetalVersion = mtlpp::LanguageVersion::Version2_0;
+					break;
+				case 2:
+					MetalVersion = mtlpp::LanguageVersion::Version1_2;
+					break;
+				case 1:
+					MetalVersion = mtlpp::LanguageVersion::Version1_1;
+					break;
+				case 0:
+#if PLATFORM_MAC
+					MetalVersion = mtlpp::LanguageVersion::Version1_1;
+#else
+					MetalVersion = mtlpp::LanguageVersion::Version1_0;
+#endif
+					break;
+				default:
+					UE_LOG(LogRHI, Fatal, TEXT("Failed to create shader with unknown version %d: %s"), Header.Version, *FString(NewShaderString));
+					break;
 			}
+			CompileOptions.SetLanguageVersion(MetalVersion);
 			
 			ns::AutoReleasedError Error;
 			Library = GetMetalDeviceContext().GetDevice().NewLibrary(NewShaderString, CompileOptions, &Error);
