Separate static buffers from dynamic buffers when sub-allocating MTLBuffers to reduce fragmentation - though level streaming tends to result in quite a lot of fragmentation even for static buffers. Also change the sub-buffer heap to use broader allocation ranges so that we fill empty space better - make sure to use a sensible min. alignment on each platform rather than the minimum allocation size as that also reduces wastage. It doesn't seem like a best-fit allocation strategy is much better than first-fit so we'll stick with first-fit for now. If we must do better than this then we'll have to resort to defragmenting the buffers and rewiring all the pointers which is a bit more engineering work and will come with its own costs.

diff --git a/Engine/Source/Runtime/Apple/MetalRHI/Private/MetalBuffer.cpp b/Engine/Source/Runtime/Apple/MetalRHI/Private/MetalBuffer.cpp
--- a/Engine/Source/Runtime/Apple/MetalRHI/Private/MetalBuffer.cpp
+++ b/Engine/Source/Runtime/Apple/MetalRHI/Private/MetalBuffer.cpp
@@ -872,7 +872,7 @@ FMetalBuffer FMetalSubBufferRing::NewBuffer(NSUInteger Size, uint32 Alignment)
 	// Allocate on first use
 	if(!Buffer.IsValid())
 	{
-		Buffer = MakeShared<FMetalRingBufferRef, ESPMode::ThreadSafe>(GetMetalDeviceContext().GetResourceHeap().CreateBuffer(InitialSize, MinAlign, Options, true));
+		Buffer = MakeShared<FMetalRingBufferRef, ESPMode::ThreadSafe>(GetMetalDeviceContext().GetResourceHeap().CreateBuffer(InitialSize, MinAlign, BUF_Dynamic, Options, true));
 	}
 	
 	if(Buffer->LastRead <= WriteHead)
@@ -891,7 +891,7 @@ FMetalBuffer FMetalSubBufferRing::NewBuffer(NSUInteger Size, uint32 Alignment)
 		else if (Storage == mtlpp::StorageMode::Managed)
 		{
 			Submit();
-			Buffer = MakeShared<FMetalRingBufferRef, ESPMode::ThreadSafe>(GetMetalDeviceContext().GetResourceHeap().CreateBuffer(Buffer->Buffer.GetLength(), MinAlign, Options, true));
+			Buffer = MakeShared<FMetalRingBufferRef, ESPMode::ThreadSafe>(GetMetalDeviceContext().GetResourceHeap().CreateBuffer(Buffer->Buffer.GetLength(), MinAlign, BUF_Dynamic, Options, true));
 			WriteHead = 0;
 			CommitHead = 0;
 			SubmitHead = 0;
@@ -912,7 +912,7 @@ FMetalBuffer FMetalSubBufferRing::NewBuffer(NSUInteger Size, uint32 Alignment)
 		
 		Submit();
 		
-		Buffer = MakeShared<FMetalRingBufferRef, ESPMode::ThreadSafe>(GetMetalDeviceContext().GetResourceHeap().CreateBuffer(NewBufferSize, MinAlign, Options, true));
+		Buffer = MakeShared<FMetalRingBufferRef, ESPMode::ThreadSafe>(GetMetalDeviceContext().GetResourceHeap().CreateBuffer(NewBufferSize, MinAlign, BUF_Dynamic, Options, true));
 		WriteHead = 0;
 		CommitHead = 0;
 		SubmitHead = 0;
@@ -949,7 +949,7 @@ void FMetalSubBufferRing::Shrink()
 			
 			UE_LOG(LogMetal, Verbose, TEXT("Shrinking RingBuffer from %u to %u as max. usage is %u at frame %lld]"), (uint32)Buffer->Buffer.GetLength(), (uint32)ThreeQuarterSize, (uint32)FrameMax, GFrameNumberRenderThread);
 			
-			Buffer = MakeShared<FMetalRingBufferRef, ESPMode::ThreadSafe>(GetMetalDeviceContext().GetResourceHeap().CreateBuffer(ThreeQuarterSize, MinAlign, Options, true));
+			Buffer = MakeShared<FMetalRingBufferRef, ESPMode::ThreadSafe>(GetMetalDeviceContext().GetResourceHeap().CreateBuffer(ThreeQuarterSize, MinAlign, BUF_Dynamic, Options, true));
 			
 			WriteHead = 0;
 			CommitHead = 0;
@@ -1092,7 +1092,7 @@ FMetalBuffer FMetalBufferPoolPolicyData::CreateResource(CreationArguments Args)
 
 FMetalBufferPoolPolicyData::CreationArguments FMetalBufferPoolPolicyData::GetCreationArguments(FMetalBuffer const& Resource)
 {
-	return FMetalBufferPoolPolicyData::CreationArguments(Resource.GetDevice(), Resource.GetLength(), Resource.GetStorageMode());
+	return FMetalBufferPoolPolicyData::CreationArguments(Resource.GetDevice(), Resource.GetLength(), 0, Resource.GetStorageMode());
 }
 
 void FMetalBufferPoolPolicyData::FreeResource(FMetalBuffer& Resource)
@@ -1342,7 +1342,7 @@ mtlpp::Heap FMetalResourceHeap::GetTextureHeap(mtlpp::TextureDescriptor Desc, mt
 	return Result;
 }
 
-FMetalBuffer FMetalResourceHeap::CreateBuffer(uint32 Size, uint32 Alignment, mtlpp::ResourceOptions Options, bool bForceUnique)
+FMetalBuffer FMetalResourceHeap::CreateBuffer(uint32 Size, uint32 Alignment, uint32 Flags, mtlpp::ResourceOptions Options, bool bForceUnique)
 {
 	LLM_SCOPE_METAL(ELLMTagMetal::Buffers);
 	LLM_PLATFORM_SCOPE_METAL(ELLMTagMetal::Buffers);
@@ -1351,6 +1351,8 @@ FMetalBuffer FMetalResourceHeap::CreateBuffer(uint32 Size, uint32 Alignment, mtl
 	static bool bSupportsBufferSubAllocation = FMetalCommandQueue::SupportsFeature(EMetalFeaturesBufferSubAllocation);
 	bForceUnique |= (!bSupportsBufferSubAllocation && !bSupportsHeaps);
 	
+	uint32 Usage = (Flags & BUF_Static) ? UsageStatic : UsageDynamic;
+	
 	FMetalBuffer Buffer;
 	uint32 BlockSize = Align(Size, Alignment);
 	mtlpp::StorageMode StorageMode = (mtlpp::StorageMode)(((NSUInteger)Options & mtlpp::ResourceStorageModeMask) >> mtlpp::ResourceStorageModeShift);
@@ -1386,7 +1388,7 @@ FMetalBuffer FMetalResourceHeap::CreateBuffer(uint32 Size, uint32 Alignment, mtl
 				 }
 				 else
 				{
-					Buffer = ManagedBuffers.CreatePooledResource(FMetalPooledBufferArgs(Queue->GetDevice(), BlockSize, StorageMode));
+					Buffer = ManagedBuffers.CreatePooledResource(FMetalPooledBufferArgs(Queue->GetDevice(), BlockSize, Flags, StorageMode));
 					DEC_MEMORY_STAT_BY(STAT_MetalBufferUnusedMemory, Buffer.GetLength());
 					DEC_MEMORY_STAT_BY(STAT_MetalPooledBufferUnusedMemory, Buffer.GetLength());
 				}
@@ -1405,7 +1407,7 @@ FMetalBuffer FMetalResourceHeap::CreateBuffer(uint32 Size, uint32 Alignment, mtl
 					FScopeLock Lock(&Mutex);
 					
 					uint32 i = GetMagazineIndex(BlockSize);
-					TArray<FMetalSubBufferMagazine*>& Heaps = SmallBuffers[Storage][i];
+					TArray<FMetalSubBufferMagazine*>& Heaps = SmallBuffers[Usage][Storage][i];
 					
 					FMetalSubBufferMagazine* Found = nullptr;
 					for (FMetalSubBufferMagazine* Heap : Heaps)
@@ -1420,7 +1422,7 @@ FMetalBuffer FMetalResourceHeap::CreateBuffer(uint32 Size, uint32 Alignment, mtl
 					if (!Found)
 					{
 						Found = new FMetalSubBufferMagazine(MagazineAllocSizes[i], MagazineSizes[i], mtlpp::ResourceOptions((NSUInteger)Options & (mtlpp::ResourceStorageModeMask|mtlpp::ResourceHazardTrackingModeMask)));
-						SmallBuffers[Storage][i].Add(Found);
+						SmallBuffers[Usage][Storage][i].Add(Found);
 					}
 					check(Found);
 					
@@ -1432,7 +1434,7 @@ FMetalBuffer FMetalResourceHeap::CreateBuffer(uint32 Size, uint32 Alignment, mtl
 					FScopeLock Lock(&Mutex);
 					
 					uint32 i = GetHeapIndex(BlockSize);
-					TArray<FMetalSubBufferHeap*>& Heaps = BufferHeaps[Storage][i];
+					TArray<FMetalSubBufferHeap*>& Heaps = BufferHeaps[Usage][Storage][i];
 					
 					FMetalSubBufferHeap* Found = nullptr;
 					for (FMetalSubBufferHeap* Heap : Heaps)
@@ -1446,8 +1448,9 @@ FMetalBuffer FMetalResourceHeap::CreateBuffer(uint32 Size, uint32 Alignment, mtl
 					
 					if (!Found)
 					{
-						Found = new FMetalSubBufferHeap(HeapAllocSizes[i], HeapSizes[i], mtlpp::ResourceOptions((NSUInteger)Options & (mtlpp::ResourceStorageModeMask|mtlpp::ResourceHazardTrackingModeMask)), Mutex);
-						BufferHeaps[Storage][i].Add(Found);
+						uint32 MinAlign = PLATFORM_MAC ? 256 : 64;
+						Found = new FMetalSubBufferHeap(HeapAllocSizes[i], MinAlign, mtlpp::ResourceOptions((NSUInteger)Options & (mtlpp::ResourceStorageModeMask|mtlpp::ResourceHazardTrackingModeMask)), Mutex);
+						BufferHeaps[Usage][Storage][i].Add(Found);
 					}
 					check(Found);
 					
@@ -1457,7 +1460,7 @@ FMetalBuffer FMetalResourceHeap::CreateBuffer(uint32 Size, uint32 Alignment, mtl
 				else
 				{
 					FScopeLock Lock(&Mutex);
-					Buffer = Buffers[Storage].CreatePooledResource(FMetalPooledBufferArgs(Queue->GetDevice(), BlockSize, StorageMode));
+					Buffer = Buffers[Storage].CreatePooledResource(FMetalPooledBufferArgs(Queue->GetDevice(), BlockSize, Flags, StorageMode));
 					DEC_MEMORY_STAT_BY(STAT_MetalBufferUnusedMemory, Buffer.GetLength());
 					DEC_MEMORY_STAT_BY(STAT_MetalPooledBufferUnusedMemory, Buffer.GetLength());
 				}
@@ -1575,30 +1578,33 @@ void FMetalResourceHeap::ReleaseTexture(FMetalSurface* Surface, FMetalTexture& T
 void FMetalResourceHeap::Compact(bool const bForce)
 {
 	FScopeLock Lock(&Mutex);
-	for (uint32 t = 0; t < NumAllocTypes; t++)
+	for (uint32 u = 0; u < NumUsageTypes; u++)
 	{
-		for (uint32 i = 0; i < NumMagazineSizes; i++)
+		for (uint32 t = 0; t < NumAllocTypes; t++)
 		{
-			for (auto It = SmallBuffers[t][i].CreateIterator(); It; ++It)
+			for (uint32 i = 0; i < NumMagazineSizes; i++)
 			{
-				FMetalSubBufferMagazine* Data = *It;
-				if (Data->NumCurrentAllocations() == 0 || bForce)
+				for (auto It = SmallBuffers[u][t][i].CreateIterator(); It; ++It)
 				{
-					It.RemoveCurrent();
-					delete Data;
+					FMetalSubBufferMagazine* Data = *It;
+					if (Data->NumCurrentAllocations() == 0 || bForce)
+					{
+						It.RemoveCurrent();
+						delete Data;
+					}
 				}
 			}
-		}
-		
-		for (uint32 i = 0; i < NumHeapSizes; i++)
-		{
-			for (auto It = BufferHeaps[t][i].CreateIterator(); It; ++It)
+			
+			for (uint32 i = 0; i < NumHeapSizes; i++)
 			{
-				FMetalSubBufferHeap* Data = *It;
-				if (Data->NumCurrentAllocations() == 0 || bForce)
+				for (auto It = BufferHeaps[u][t][i].CreateIterator(); It; ++It)
 				{
-					It.RemoveCurrent();
-					delete Data;
+					FMetalSubBufferHeap* Data = *It;
+					if (Data->NumCurrentAllocations() == 0 || bForce)
+					{
+						It.RemoveCurrent();
+						delete Data;
+					}
 				}
 			}
 		}
@@ -1673,15 +1679,10 @@ uint32 FMetalResourceHeap::MagazineSizes[FMetalResourceHeap::NumMagazineSizes] =
 	1024,
 	2048,
 	4096,
+	8192,
 };
 
 uint32 FMetalResourceHeap::HeapSizes[FMetalResourceHeap::NumHeapSizes] = {
-	16384,
-	32768,
-	65536,
-	131072,
-	262144,
-	524288,
 	1048576,
 	2097152,
 };
@@ -1696,17 +1697,12 @@ uint32 FMetalResourceHeap::MagazineAllocSizes[FMetalResourceHeap::NumMagazineSiz
 	16384,
 	16384,
 	16384,
+	32768,
 };
 
 uint32 FMetalResourceHeap::HeapAllocSizes[FMetalResourceHeap::NumHeapSizes] = {
-	131072,
-	131072,
-	524288,
-	524288,
-	1048576,
 	2097152,
 	4194304,
-	4194304,
 };
 
 uint32 FMetalResourceHeap::HeapTextureHeapSizes[FMetalResourceHeap::NumTextureHeapSizes] = {
diff --git a/Engine/Source/Runtime/Apple/MetalRHI/Private/MetalBuffer.h b/Engine/Source/Runtime/Apple/MetalRHI/Private/MetalBuffer.h
--- a/Engine/Source/Runtime/Apple/MetalRHI/Private/MetalBuffer.h
+++ b/Engine/Source/Runtime/Apple/MetalRHI/Private/MetalBuffer.h
@@ -11,17 +11,19 @@
 
 struct FMetalPooledBufferArgs
 {
-	FMetalPooledBufferArgs() : Device(nil), Size(0), Storage(mtlpp::StorageMode::Shared) {}
+	FMetalPooledBufferArgs() : Device(nil), Size(0), Flags(0), Storage(mtlpp::StorageMode::Shared) {}
 	
-	FMetalPooledBufferArgs(mtlpp::Device InDevice, uint32 InSize, mtlpp::StorageMode InStorage)
+	FMetalPooledBufferArgs(mtlpp::Device InDevice, uint32 InSize, uint32 InFlags, mtlpp::StorageMode InStorage)
 	: Device(InDevice)
 	, Size(InSize)
+	, Flags(InFlags)
 	, Storage(InStorage)
 	{
 	}
 	
 	mtlpp::Device Device;
 	uint32 Size;
+	uint32 Flags;
 	mtlpp::StorageMode Storage;
 };
 
@@ -345,17 +347,12 @@ class FMetalResourceHeap
 		Size1024,
 		Size2048,
 		Size4096,
+		Size8192,
 		NumMagazineSizes
 	};
 	
 	enum HeapSize
 	{
-		Size16k,
-		Size32k,
-		Size64k,
-		Size128k,
-		Size256k,
-		Size512k,
 		Size1Mb,
 		Size2Mb,
 		NumHeapSizes
@@ -382,6 +379,13 @@ class FMetalResourceHeap
 		NumAllocTypes = 2
 	};
 	
+	enum UsageTypes
+	{
+		UsageStatic,
+		UsageDynamic,
+		NumUsageTypes = 2
+	};
+	
 	enum EMetalHeapTextureUsage
 	{
 		/** Regular texture resource */
@@ -398,7 +402,7 @@ public:
 	
 	void Init(FMetalCommandQueue& Queue);
 	
-	FMetalBuffer CreateBuffer(uint32 Size, uint32 Alignment, mtlpp::ResourceOptions Options, bool bForceUnique = false);
+	FMetalBuffer CreateBuffer(uint32 Size, uint32 Alignment, uint32 Flags, mtlpp::ResourceOptions Options, bool bForceUnique = false);
 	FMetalTexture CreateTexture(mtlpp::TextureDescriptor Desc, FMetalSurface* Surface);
 	
 	void ReleaseBuffer(FMetalBuffer& Buffer);
@@ -424,11 +428,11 @@ private:
 	FMetalCommandQueue* Queue;
 	
 	/** Small allocations (<= 4KB) are made from magazine allocators that use sub-ranges of a buffer */
-	TArray<FMetalSubBufferMagazine*> SmallBuffers[NumAllocTypes][NumMagazineSizes];
+	TArray<FMetalSubBufferMagazine*> SmallBuffers[NumUsageTypes][NumAllocTypes][NumMagazineSizes];
 
 	/** Typical allocations (4KB - 4MB) are made from heap allocators that use sub-ranges of a buffer */
 	/** There are two alignment categories for heaps - 16b for Vertes/Index data and 256b for constant data (macOS-only) */
-	TArray<FMetalSubBufferHeap*> BufferHeaps[NumAllocTypes][NumHeapSizes];
+	TArray<FMetalSubBufferHeap*> BufferHeaps[NumUsageTypes][NumAllocTypes][NumHeapSizes];
 	
 	/** Larger buffers (up-to 32MB) that are subject to bucketing & pooling rather than sub-allocation */
 	FMetalBufferPool Buffers[NumAllocTypes];
diff --git a/Engine/Source/Runtime/Apple/MetalRHI/Private/MetalContext.cpp b/Engine/Source/Runtime/Apple/MetalRHI/Private/MetalContext.cpp
--- a/Engine/Source/Runtime/Apple/MetalRHI/Private/MetalContext.cpp
+++ b/Engine/Source/Runtime/Apple/MetalRHI/Private/MetalContext.cpp
@@ -767,7 +767,7 @@ FMetalTexture FMetalDeviceContext::CreateTexture(FMetalSurface* Surface, mtlpp::
 
 FMetalBuffer FMetalDeviceContext::CreatePooledBuffer(FMetalPooledBufferArgs const& Args)
 {
-	FMetalBuffer Buffer = Heap.CreateBuffer(Args.Size, BufferOffsetAlignment, GetCommandQueue().GetCompatibleResourceOptions((mtlpp::ResourceOptions)(BUFFER_CACHE_MODE | mtlpp::ResourceOptions::HazardTrackingModeUntracked | ((NSUInteger)Args.Storage << mtlpp::ResourceStorageModeShift))));
+	FMetalBuffer Buffer = Heap.CreateBuffer(Args.Size, BufferOffsetAlignment, Args.Flags, GetCommandQueue().GetCompatibleResourceOptions((mtlpp::ResourceOptions)(BUFFER_CACHE_MODE | mtlpp::ResourceOptions::HazardTrackingModeUntracked | ((NSUInteger)Args.Storage << mtlpp::ResourceStorageModeShift))));
 	check(Buffer && Buffer.GetPtr());
 #if METAL_DEBUG_OPTIONS
 	if (GMetalResourcePurgeOnDelete && !Buffer.GetHeap())
diff --git a/Engine/Source/Runtime/Apple/MetalRHI/Private/MetalRenderPass.cpp b/Engine/Source/Runtime/Apple/MetalRHI/Private/MetalRenderPass.cpp
--- a/Engine/Source/Runtime/Apple/MetalRHI/Private/MetalRenderPass.cpp
+++ b/Engine/Source/Runtime/Apple/MetalRHI/Private/MetalRenderPass.cpp
@@ -591,19 +591,19 @@ void FMetalRenderPass::DrawPatches(uint32 PrimitiveType,FMetalBuffer const& Inde
 		FMetalBuffer hullShaderOutputBuffer = nil;
 		if(hullShaderOutputBufferSize)
 		{
-			hullShaderOutputBuffer = deviceContext.CreatePooledBuffer(FMetalPooledBufferArgs(device, hullShaderOutputBufferSize, mtlpp::StorageMode::Private));
+			hullShaderOutputBuffer = deviceContext.CreatePooledBuffer(FMetalPooledBufferArgs(device, hullShaderOutputBufferSize, BUF_Dynamic, mtlpp::StorageMode::Private));
 		}
 		
 		FMetalBuffer hullConstShaderOutputBuffer = nil;
 		if(hullConstShaderOutputBufferSize)
 		{
-			hullConstShaderOutputBuffer = deviceContext.CreatePooledBuffer(FMetalPooledBufferArgs(device, hullConstShaderOutputBufferSize, mtlpp::StorageMode::Private));
+			hullConstShaderOutputBuffer = deviceContext.CreatePooledBuffer(FMetalPooledBufferArgs(device, hullConstShaderOutputBufferSize, BUF_Dynamic, mtlpp::StorageMode::Private));
 		}
 		
 		FMetalBuffer tessellationFactorBuffer = nil;
 		if(tessellationFactorBufferSize)
 		{
-			tessellationFactorBuffer = deviceContext.CreatePooledBuffer(FMetalPooledBufferArgs(device, tessellationFactorBufferSize, mtlpp::StorageMode::Private));
+			tessellationFactorBuffer = deviceContext.CreatePooledBuffer(FMetalPooledBufferArgs(device, tessellationFactorBufferSize, BUF_Dynamic, mtlpp::StorageMode::Private));
 		}
 	
 		auto& computeEncoder = PrologueEncoder.GetComputeCommandEncoder();
diff --git a/Engine/Source/Runtime/Apple/MetalRHI/Private/MetalRenderTarget.cpp b/Engine/Source/Runtime/Apple/MetalRHI/Private/MetalRenderTarget.cpp
--- a/Engine/Source/Runtime/Apple/MetalRHI/Private/MetalRenderTarget.cpp
+++ b/Engine/Source/Runtime/Apple/MetalRHI/Private/MetalRenderTarget.cpp
@@ -670,7 +670,7 @@ void FMetalDynamicRHI::RHIReadSurfaceData(FTextureRHIParamRef TextureRHI, FIntRe
 		const uint32 Alignment = PLATFORM_MAC ? 1u : 64u; // Mac permits natural row alignment (tightly-packed) but iOS does not.
 		const uint32 AlignedStride = ((Stride - 1) & ~(Alignment - 1)) + Alignment;
 		const uint32 BytesPerImage = AlignedStride * SizeY;
-		FMetalBuffer Buffer = ((FMetalDeviceContext*)ImmediateContext.Context)->CreatePooledBuffer(FMetalPooledBufferArgs(ImmediateContext.Context->GetDevice(), BytesPerImage, mtlpp::StorageMode::Shared));
+		FMetalBuffer Buffer = ((FMetalDeviceContext*)ImmediateContext.Context)->CreatePooledBuffer(FMetalPooledBufferArgs(ImmediateContext.Context->GetDevice(), BytesPerImage, BUF_Dynamic, mtlpp::StorageMode::Shared));
 		{
 			// Synchronise the texture with the CPU
 			SCOPE_CYCLE_COUNTER(STAT_MetalTexturePageOffTime);
@@ -766,7 +766,7 @@ void FMetalDynamicRHI::RHIReadSurfaceFloatData(FTextureRHIParamRef TextureRHI, F
 	const uint32 AlignedStride = ((Stride - 1) & ~(Alignment - 1)) + Alignment;
 	const uint32 BytesPerImage = AlignedStride  * SizeY;
 	int32 FloatBGRADataSize = BytesPerImage;
-	FMetalBuffer Buffer = ((FMetalDeviceContext*)ImmediateContext.Context)->CreatePooledBuffer(FMetalPooledBufferArgs(ImmediateContext.Context->GetDevice(), FloatBGRADataSize, mtlpp::StorageMode::Shared));
+	FMetalBuffer Buffer = ((FMetalDeviceContext*)ImmediateContext.Context)->CreatePooledBuffer(FMetalPooledBufferArgs(ImmediateContext.Context->GetDevice(), FloatBGRADataSize, BUF_Dynamic, mtlpp::StorageMode::Shared));
 	{
 		// Synchronise the texture with the CPU
 		SCOPE_CYCLE_COUNTER(STAT_MetalTexturePageOffTime);
@@ -832,7 +832,7 @@ void FMetalDynamicRHI::RHIRead3DSurfaceFloatData(FTextureRHIParamRef TextureRHI,
 	const uint32 AlignedStride = ((Stride - 1) & ~(Alignment - 1)) + Alignment;
 	const uint32 BytesPerImage = AlignedStride  * SizeY;
 	int32 FloatBGRADataSize = BytesPerImage * SizeZ;
-	FMetalBuffer Buffer = ((FMetalDeviceContext*)ImmediateContext.Context)->CreatePooledBuffer(FMetalPooledBufferArgs(ImmediateContext.Context->GetDevice(), FloatBGRADataSize, mtlpp::StorageMode::Shared));
+	FMetalBuffer Buffer = ((FMetalDeviceContext*)ImmediateContext.Context)->CreatePooledBuffer(FMetalPooledBufferArgs(ImmediateContext.Context->GetDevice(), FloatBGRADataSize, BUF_Dynamic, mtlpp::StorageMode::Shared));
 	{
 		// Synchronise the texture with the CPU
 		SCOPE_CYCLE_COUNTER(STAT_MetalTexturePageOffTime);
diff --git a/Engine/Source/Runtime/Apple/MetalRHI/Private/MetalTexture.cpp b/Engine/Source/Runtime/Apple/MetalRHI/Private/MetalTexture.cpp
--- a/Engine/Source/Runtime/Apple/MetalRHI/Private/MetalTexture.cpp
+++ b/Engine/Source/Runtime/Apple/MetalRHI/Private/MetalTexture.cpp
@@ -892,7 +892,7 @@ FMetalSurface::FMetalSurface(ERHIResourceType ResourceType, EPixelFormat Format,
 		{
 			mtlpp::Device Device = GetMetalDeviceContext().GetDevice();
 			mtlpp::SizeAndAlign SizeAlign = Device.HeapTextureSizeAndAlign(Desc);
-			FMetalPooledBufferArgs Args(Device, SizeAlign.Size, mtlpp::StorageMode::Private);
+			FMetalPooledBufferArgs Args(Device, SizeAlign.Size, BUF_Dynamic, mtlpp::StorageMode::Private);
 			FMetalBuffer Buffer = GetMetalDeviceContext().CreatePooledBuffer(Args);
 			Texture = Buffer.NewTexture(Desc, 0, SizeAlign.Size);
 		}
@@ -1170,7 +1170,7 @@ FMetalBuffer FMetalSurface::AllocSurface(uint32 MipIndex, uint32 ArrayIndex, ERe
 	// allocate some temporary memory
 	mtlpp::ResourceOptions ResMode = GetMetalDeviceContext().GetCommandQueue().GetCompatibleResourceOptions((mtlpp::ResourceOptions)(mtlpp::ResourceOptions::StorageModeShared | (!(PLATFORM_MAC && PixelFormat == PF_G8 && (Flags & TexCreate_SRGB)) ? mtlpp::ResourceOptions::CpuCacheModeWriteCombined : 0)));
 	
-	FMetalBuffer Buffer = GetMetalDeviceContext().GetResourceHeap().CreateBuffer(MipBytes, BufferOffsetAlignment, ResMode);
+	FMetalBuffer Buffer = GetMetalDeviceContext().GetResourceHeap().CreateBuffer(MipBytes, BufferOffsetAlignment, BUF_Dynamic, ResMode);
 	
 #if PLATFORM_MAC
 	// Expand R8_sRGB into RGBA8_sRGB for Mac.
@@ -1938,7 +1938,7 @@ static FMetalBuffer InternalCopyTexture2DUpdateRegion(FTexture2DRHIParamRef Text
 	if(Texture->GetFormat() == PF_G8 && (Texture->GetFlags() & TexCreate_SRGB))
 	{
 		const uint32 BufferSize = UpdateRegion.Height * UpdateRegion.Width * sizeof(uint32);		
-		Buffer = GetMetalDeviceContext().CreatePooledBuffer(FMetalPooledBufferArgs(GetMetalDeviceContext().GetDevice(), BufferSize, mtlpp::StorageMode::Shared));
+		Buffer = GetMetalDeviceContext().CreatePooledBuffer(FMetalPooledBufferArgs(GetMetalDeviceContext().GetDevice(), BufferSize, BUF_Dynamic, mtlpp::StorageMode::Shared));
 		InternalExpandR8ToStandardRGBA((uint32*)Buffer.GetContents(), UpdateRegion, InOutSourcePitch, SourceData);
 	}
 
@@ -1946,7 +1946,7 @@ static FMetalBuffer InternalCopyTexture2DUpdateRegion(FTexture2DRHIParamRef Text
 #endif
 	{
 		const uint32 BufferSize = UpdateRegion.Height * InSourcePitch;
-		Buffer = GetMetalDeviceContext().CreatePooledBuffer(FMetalPooledBufferArgs(GetMetalDeviceContext().GetDevice(), BufferSize, mtlpp::StorageMode::Shared));
+		Buffer = GetMetalDeviceContext().CreatePooledBuffer(FMetalPooledBufferArgs(GetMetalDeviceContext().GetDevice(), BufferSize, BUF_Dynamic, mtlpp::StorageMode::Shared));
 
 		// Limit copy to line by line by update region pitch otherwise we can go off the end of source data on the last row
 		uint8* pDestRow = (uint8*)Buffer.GetContents();
@@ -2150,7 +2150,7 @@ struct FMetalDynamicRHIUpdateTexture3DCommand final : public FRHICommand<FMetalD
 		FMetalTexture Tex = Texture->Surface.Texture;
 		const uint32 BufferSize = UpdateRegion.Height * UpdateRegion.Depth* SourceRowPitch;
 		
-		Buffer = GetMetalDeviceContext().CreatePooledBuffer(FMetalPooledBufferArgs(GetMetalDeviceContext().GetDevice(), BufferSize, mtlpp::StorageMode::Shared));
+		Buffer = GetMetalDeviceContext().CreatePooledBuffer(FMetalPooledBufferArgs(GetMetalDeviceContext().GetDevice(), BufferSize, BUF_Dynamic, mtlpp::StorageMode::Shared));
 		InternalCopyTexture3DUpdateRegionData(Texture, UpdateRegion, SourceRowPitch, SourceDepthPitch, SourceData, (uint8*)Buffer.GetContents());
 	}
 	
@@ -2221,7 +2221,7 @@ void FMetalDynamicRHI::RHIUpdateTexture3D(FTexture3DRHIParamRef TextureRHI,uint3
 			SCOPED_AUTORELEASE_POOL;
 
 			const uint32 BufferSize = UpdateRegion.Height * UpdateRegion.Depth * SourceRowPitch;
-			FMetalBuffer IntermediateBuffer = GetMetalDeviceContext().CreatePooledBuffer(FMetalPooledBufferArgs(GetMetalDeviceContext().GetDevice(), BufferSize, mtlpp::StorageMode::Shared));
+			FMetalBuffer IntermediateBuffer = GetMetalDeviceContext().CreatePooledBuffer(FMetalPooledBufferArgs(GetMetalDeviceContext().GetDevice(), BufferSize, BUF_Dynamic, mtlpp::StorageMode::Shared));
 			InternalCopyTexture3DUpdateRegionData(TextureRHI, UpdateRegion, SourceRowPitch, SourceDepthPitch, SourceData, (uint8*)IntermediateBuffer.GetContents());
 			InternalUpdateTexture3D(ImmediateContext.GetInternalContext(), TextureRHI, MipIndex, UpdateRegion, SourceRowPitch, SourceDepthPitch, IntermediateBuffer);
 			GetMetalDeviceContext().ReleaseBuffer(IntermediateBuffer);
@@ -2693,7 +2693,7 @@ void FMetalDynamicRHI::RHICopySubTextureRegion(FTexture2DRHIParamRef SourceTextu
 				const uint32 AlignedStride = ((Stride - 1) & ~(Alignment - 1)) + Alignment;
 				const uint32 BytesPerImage = AlignedStride *  SourceSize.height;
 				
-				FMetalBuffer Buffer = GetMetalDeviceContext().CreatePooledBuffer(FMetalPooledBufferArgs(ImmediateContext.Context->GetDevice(), BytesPerImage, mtlpp::StorageMode::Shared));
+				FMetalBuffer Buffer = GetMetalDeviceContext().CreatePooledBuffer(FMetalPooledBufferArgs(ImmediateContext.Context->GetDevice(), BytesPerImage, BUF_Dynamic, mtlpp::StorageMode::Shared));
 				
 				check(Buffer);
 				
@@ -2790,7 +2790,7 @@ void FMetalRHICommandContext::RHICopyTexture(FTextureRHIParamRef SourceTextureRH
 						const uint32 BytesPerImage = AlignedStride *  SourceSize.height;
 						const uint32 DataSize = BytesPerImage * SourceSize.depth;
 						
-						FMetalBuffer Buffer = GetMetalDeviceContext().CreatePooledBuffer(FMetalPooledBufferArgs(GetInternalContext().GetDevice(), DataSize, mtlpp::StorageMode::Shared));
+						FMetalBuffer Buffer = GetMetalDeviceContext().CreatePooledBuffer(FMetalPooledBufferArgs(GetInternalContext().GetDevice(), DataSize, BUF_Dynamic, mtlpp::StorageMode::Shared));
 						
 						check(Buffer);
 						
diff --git a/Engine/Source/Runtime/Apple/MetalRHI/Private/MetalUAV.cpp b/Engine/Source/Runtime/Apple/MetalRHI/Private/MetalUAV.cpp
--- a/Engine/Source/Runtime/Apple/MetalRHI/Private/MetalUAV.cpp
+++ b/Engine/Source/Runtime/Apple/MetalRHI/Private/MetalUAV.cpp
@@ -533,7 +533,7 @@ void FMetalRHICommandContext::RHIClearTinyUAV(FUnorderedAccessViewRHIParamRef Un
 			FMetalBuffer Temp = nil;
 			bool bBufferPooled = false;
 			
-			FMetalPooledBufferArgs Args(GetMetalDeviceContext().GetDevice(), AlignedSize, mtlpp::StorageMode::Shared);
+			FMetalPooledBufferArgs Args(GetMetalDeviceContext().GetDevice(), AlignedSize, BUF_Dynamic, mtlpp::StorageMode::Shared);
 			Temp = GetMetalDeviceContext().CreatePooledBuffer(Args);
 			bBufferPooled = true;
 			
@@ -888,7 +888,7 @@ void FMetalRHICommandContext::RHICopyToStagingBuffer(FVertexBufferRHIParamRef So
 			{
 				SafeReleaseMetalBuffer(ReadbackBuffer);
 			}
-			FMetalPooledBufferArgs ArgsCPU(GetMetalDeviceContext().GetDevice(), NumBytes, mtlpp::StorageMode::Shared);
+			FMetalPooledBufferArgs ArgsCPU(GetMetalDeviceContext().GetDevice(), NumBytes, BUF_Dynamic, mtlpp::StorageMode::Shared);
 			ReadbackBuffer = GetMetalDeviceContext().CreatePooledBuffer(ArgsCPU);
 		}
 
diff --git a/Engine/Source/Runtime/Apple/MetalRHI/Private/MetalUniformBuffer.cpp b/Engine/Source/Runtime/Apple/MetalRHI/Private/MetalUniformBuffer.cpp
--- a/Engine/Source/Runtime/Apple/MetalRHI/Private/MetalUniformBuffer.cpp
+++ b/Engine/Source/Runtime/Apple/MetalRHI/Private/MetalUniformBuffer.cpp
@@ -420,7 +420,7 @@ void FMetalUniformBuffer::InitIAB()
 			Desc.SetAccess(mtlpp::ArgumentAccess::ReadOnly);
 			Desc.SetDataType(mtlpp::DataType::Pointer);
 			
-			FMetalPooledBufferArgs Args(GetMetalDeviceContext().GetDevice(), BufferSizes.Num() * sizeof(uint32), BUFFER_STORAGE_MODE);
+			FMetalPooledBufferArgs Args(GetMetalDeviceContext().GetDevice(), BufferSizes.Num() * sizeof(uint32), BUF_Dynamic, BUFFER_STORAGE_MODE);
 			NewIAB->IndirectArgumentBufferSideTable = GetMetalDeviceContext().CreatePooledBuffer(Args);
 			
 			FMemory::Memcpy(NewIAB->IndirectArgumentBufferSideTable.GetContents(), BufferSizes.GetData(), BufferSizes.Num() * sizeof(uint32));
@@ -447,7 +447,7 @@ void FMetalUniformBuffer::InitIAB()
 		
 		mtlpp::ArgumentEncoder Encoder = FMetalArgumentEncoderCache::Get().CreateEncoder(Arguments);
 		
-		NewIAB->IndirectArgumentBuffer = GetMetalDeviceContext().GetResourceHeap().CreateBuffer(Encoder.GetEncodedLength(), 16, mtlpp::ResourceOptions(BUFFER_CACHE_MODE | ((NSUInteger)BUFFER_STORAGE_MODE << mtlpp::ResourceStorageModeShift)), true);
+		NewIAB->IndirectArgumentBuffer = GetMetalDeviceContext().GetResourceHeap().CreateBuffer(Encoder.GetEncodedLength(), 16, BUF_Dynamic, mtlpp::ResourceOptions(BUFFER_CACHE_MODE | ((NSUInteger)BUFFER_STORAGE_MODE << mtlpp::ResourceStorageModeShift)), true);
 		
 		Encoder.SetArgumentBuffer(NewIAB->IndirectArgumentBuffer, 0);
 		
diff --git a/Engine/Source/Runtime/Apple/MetalRHI/Private/MetalVertexBuffer.cpp b/Engine/Source/Runtime/Apple/MetalRHI/Private/MetalVertexBuffer.cpp
--- a/Engine/Source/Runtime/Apple/MetalRHI/Private/MetalVertexBuffer.cpp
+++ b/Engine/Source/Runtime/Apple/MetalRHI/Private/MetalVertexBuffer.cpp
@@ -241,7 +241,7 @@ void FMetalRHIBuffer::Alloc(uint32 InSize, EResourceLockMode LockMode)
 	{
         check(LockMode != RLM_ReadOnly);
 		mtlpp::StorageMode Mode = (bUsePrivateMem ? mtlpp::StorageMode::Private : BUFFER_STORAGE_MODE);
-		FMetalPooledBufferArgs Args(GetMetalDeviceContext().GetDevice(), InSize, Mode);
+		FMetalPooledBufferArgs Args(GetMetalDeviceContext().GetDevice(), InSize, Usage, Mode);
 		Buffer = GetMetalDeviceContext().CreatePooledBuffer(Args);
 		check(Buffer && Buffer.GetPtr());
 
@@ -262,7 +262,7 @@ void FMetalRHIBuffer::Alloc(uint32 InSize, EResourceLockMode LockMode)
 	
 	if (bUsePrivateMem && !CPUBuffer)
 	{
-		FMetalPooledBufferArgs ArgsCPU(GetMetalDeviceContext().GetDevice(), InSize, mtlpp::StorageMode::Shared);
+		FMetalPooledBufferArgs ArgsCPU(GetMetalDeviceContext().GetDevice(), InSize, BUF_Dynamic, mtlpp::StorageMode::Shared);
 		CPUBuffer = GetMetalDeviceContext().CreatePooledBuffer(ArgsCPU);
 		check(CPUBuffer && CPUBuffer.GetPtr());
         METAL_INC_DWORD_STAT_BY(Type, MemAlloc, InSize);
