Permit defragging of FMetalSubBufferHeap by moving allocations from newer to older buffers at a specified number of bytes per-frame (rhi.Metal.HeapBufferBytesToCompact). Only allocations 'owned' by an RHI object can be defragged - temporary allocations for texture streaming can't be - they should only last a frame anyway. This isn't especially complicated and might require a frame-number based heuristic to prevent constant defragging due to short-lived allocations.

diff --git a/Engine/Source/Runtime/Apple/MetalRHI/Private/MetalBuffer.cpp b/Engine/Source/Runtime/Apple/MetalRHI/Private/MetalBuffer.cpp
--- a/Engine/Source/Runtime/Apple/MetalRHI/Private/MetalBuffer.cpp
+++ b/Engine/Source/Runtime/Apple/MetalRHI/Private/MetalBuffer.cpp
@@ -12,6 +12,15 @@ DECLARE_MEMORY_STAT(TEXT("Unused Pooled Buffer Memory"), STAT_MetalPooledBufferU
 DECLARE_MEMORY_STAT(TEXT("Unused Magazine Buffer Memory"), STAT_MetalMagazineBufferUnusedMemory, STATGROUP_MetalRHI);
 DECLARE_MEMORY_STAT(TEXT("Unused Heap Buffer Memory"), STAT_MetalHeapBufferUnusedMemory, STATGROUP_MetalRHI);
 
+
+
+static int32 GMetalHeapBufferBytesToCompact = 0;
+static FAutoConsoleVariableRef CVarMetalHeapBufferBytesToCompact(
+	TEXT("rhi.Metal.HeapBufferBytesToCompact"),
+	GMetalHeapBufferBytesToCompact,
+	TEXT("When enabled (> 0) this will force MetalRHI to compact the given number of bytes each frame into older buffer heaps from newer ones in order to defragment memory and reduce wastage.\n")
+	TEXT("(Off by default (0))"));
+
 #if METAL_DEBUG_OPTIONS
 extern int32 GMetalBufferScribble;
 #endif
@@ -135,6 +144,14 @@ void FMetalBuffer::Release()
 	}
 }
 
+void FMetalBuffer::SetOwner(class FMetalRHIBuffer* Owner)
+{
+	if (Heap)
+	{
+		Heap->SetOwner(ns::Range(GetOffset(), GetLength()), Owner);
+	}
+}
+
 FMetalSubBufferHeap::FMetalSubBufferHeap(NSUInteger Size, NSUInteger Alignment, mtlpp::ResourceOptions Options, FCriticalSection& InPoolMutex)
 : PoolMutex(InPoolMutex)
 , OutstandingAllocs(0)
@@ -189,9 +206,36 @@ FMetalSubBufferHeap::~FMetalSubBufferHeap()
 	}
 }
 
+void FMetalSubBufferHeap::SetOwner(ns::Range const& Range, FMetalRHIBuffer* Owner)
+{
+	FScopeLock Lock(&PoolMutex);
+	for (uint32 i = 0; i < AllocRanges.Num(); i++)
+	{
+		if (AllocRanges[i].Range.Location == Range.Location)
+		{
+			check(AllocRanges[i].Range.Length == Range.Length);
+			check(AllocRanges[i].Owner == nullptr || Owner == nullptr);
+			AllocRanges[i].Owner = Owner;
+			break;
+		}
+	}
+}
+
 void FMetalSubBufferHeap::FreeRange(ns::Range const& Range)
 {
 	FPlatformAtomics::InterlockedDecrement(&OutstandingAllocs);
+	{
+		FScopeLock Lock(&PoolMutex);
+		for (uint32 i = 0; i < AllocRanges.Num(); i++)
+		{
+			if (AllocRanges[i].Range.Location == Range.Location)
+			{
+				check(AllocRanges[i].Range.Length == Range.Length);
+				AllocRanges.RemoveAt(i);
+				break;
+			}
+		}
+	}
 	if (ParentHeap)
 	{
 		SET_MEMORY_STAT(STAT_MetalBufferUnusedMemory, ParentHeap.GetSize() - ParentHeap.GetUsedSize());
@@ -425,6 +469,12 @@ FMetalBuffer FMetalSubBufferHeap::NewBuffer(NSUInteger length)
 #endif
 					
 					Result = FMetalBuffer(MTLPP_VALIDATE(mtlpp::Buffer, ParentBuffer, SafeGetRuntimeDebuggingLevel() >= EMetalDebugLevelValidation, NewBuffer(Range)), this);
+					
+					Allocation Alloc;
+					Alloc.Range = Range;
+					Alloc.Resource = Result.GetPtr();
+					Alloc.Owner = nullptr;
+					AllocRanges.Add(Alloc);
 				
 					break;
 				}
@@ -1595,8 +1645,67 @@ void FMetalResourceHeap::Compact(bool const bForce)
 				}
 			}
 			
+			mtlpp::CommandBuffer Buffer;
+			mtlpp::BlitCommandEncoder Encoder;
+			
+			uint32 BytesCompacted = 0;
+			uint32 const BytesToCompact = GMetalHeapBufferBytesToCompact;
+			
 			for (uint32 i = 0; i < NumHeapSizes; i++)
 			{
+				for (uint32 j = 1; j < BufferHeaps[u][t][i].Num(); j++)
+				{
+					FMetalSubBufferHeap* Data = BufferHeaps[u][t][i][j];
+					if (Data->AllocRanges.Num() > 0 && BytesCompacted < BytesToCompact)
+					{
+						for (FMetalSubBufferHeap::Allocation const& Alloc : Data->AllocRanges)
+						{
+							if (Alloc.Owner)
+							{
+								for (uint32 AllocIndex = 0; AllocIndex < j && BytesCompacted < BytesToCompact; AllocIndex++)
+								{
+									FMetalSubBufferHeap* Prev = BufferHeaps[u][t][i][AllocIndex];
+									if (Prev->MaxAvailableSize() >= Alloc.Range.Length)
+									{
+										if (!Encoder)
+										{
+											Buffer = Queue->CreateCommandBuffer();
+											Encoder = Buffer.BlitCommandEncoder();
+										}
+										
+										FMetalBuffer NewBuffer = Prev->NewBuffer(Alloc.Range.Length);
+										check(NewBuffer && NewBuffer.GetPtr());
+										
+										FMetalBuffer SourceResource = FMetalBuffer((id<MTLBuffer>)Alloc.Resource);
+										Encoder.Copy(SourceResource, Alloc.Range.Location, NewBuffer, 0, Alloc.Range.Length);
+										
+										FMetalBuffer PrevBuffer;
+										FMetalBuffer PrevCPUBuffer;
+										if (Alloc.Owner->Buffer.GetPtr() == Alloc.Resource && Alloc.Owner->Buffer.GetOffset() == Alloc.Range.Location)
+										{
+											PrevBuffer = Alloc.Owner->Buffer;
+											Alloc.Owner->Buffer = NewBuffer;
+										}
+										if (Alloc.Owner->CPUBuffer.GetPtr() == Alloc.Resource && Alloc.Owner->CPUBuffer.GetOffset() == Alloc.Range.Location)
+										{
+											PrevCPUBuffer = Alloc.Owner->CPUBuffer;
+											Alloc.Owner->CPUBuffer = NewBuffer;
+										}
+										if (PrevBuffer)
+											SafeReleaseMetalBuffer(PrevBuffer);
+										if (PrevCPUBuffer)
+											SafeReleaseMetalBuffer(PrevCPUBuffer);
+										
+										BytesCompacted += Alloc.Range.Length;
+										break;
+									}
+								}
+							}
+						}
+					}
+				}
+				
+				
 				for (auto It = BufferHeaps[u][t][i].CreateIterator(); It; ++It)
 				{
 					FMetalSubBufferHeap* Data = *It;
@@ -1607,6 +1716,12 @@ void FMetalResourceHeap::Compact(bool const bForce)
 					}
 				}
 			}
+			
+			if (Encoder)
+			{
+				Encoder.EndEncoding();
+				Queue->CommitCommandBuffer(Buffer);
+			}
 		}
 	}
 	
diff --git a/Engine/Source/Runtime/Apple/MetalRHI/Private/MetalBuffer.h b/Engine/Source/Runtime/Apple/MetalRHI/Private/MetalBuffer.h
--- a/Engine/Source/Runtime/Apple/MetalRHI/Private/MetalBuffer.h
+++ b/Engine/Source/Runtime/Apple/MetalRHI/Private/MetalBuffer.h
@@ -29,6 +29,8 @@ struct FMetalPooledBufferArgs
 
 class FMetalSubBufferHeap
 {
+	friend class FMetalResourceHeap;
+	
 public:
 	FMetalSubBufferHeap(NSUInteger Size, NSUInteger Alignment, mtlpp::ResourceOptions, FCriticalSection& PoolMutex);
 	~FMetalSubBufferHeap();
@@ -47,8 +49,17 @@ public:
     FMetalBuffer NewBuffer(NSUInteger length);
     mtlpp::PurgeableState SetPurgeableState(mtlpp::PurgeableState state);
 	void FreeRange(ns::Range const& Range);
+	
+	void SetOwner(ns::Range const& Range, FMetalRHIBuffer* Owner);
 
 private:
+	struct Allocation
+	{
+		ns::Range Range;
+		mtlpp::Buffer::Type Resource;
+		FMetalRHIBuffer* Owner;
+	};
+	
 	FCriticalSection& PoolMutex;
 	int64 volatile OutstandingAllocs;
 	NSUInteger MinAlign;
@@ -56,6 +67,7 @@ private:
 	mtlpp::Buffer ParentBuffer;
 	mutable mtlpp::Heap ParentHeap;
 	TArray<ns::Range> FreeRanges;
+	TArray<Allocation> AllocRanges;
 };
 
 class FMetalSubBufferLinear
diff --git a/Engine/Source/Runtime/Apple/MetalRHI/Private/MetalRHIContext.cpp b/Engine/Source/Runtime/Apple/MetalRHI/Private/MetalRHIContext.cpp
--- a/Engine/Source/Runtime/Apple/MetalRHI/Private/MetalRHIContext.cpp
+++ b/Engine/Source/Runtime/Apple/MetalRHI/Private/MetalRHIContext.cpp
@@ -42,6 +42,7 @@ void SafeReleaseMetalBuffer(FMetalBuffer& Buffer)
 {
 	if(GIsMetalInitialized && GDynamicRHI && Buffer)
 	{
+		Buffer.SetOwner(nullptr);
 		FMetalRHICommandContext* Context = static_cast<FMetalRHICommandContext*>(RHIGetDefaultContext());
 		if(Context)
 		{
diff --git a/Engine/Source/Runtime/Apple/MetalRHI/Private/MetalVertexBuffer.cpp b/Engine/Source/Runtime/Apple/MetalRHI/Private/MetalVertexBuffer.cpp
--- a/Engine/Source/Runtime/Apple/MetalRHI/Private/MetalVertexBuffer.cpp
+++ b/Engine/Source/Runtime/Apple/MetalRHI/Private/MetalVertexBuffer.cpp
@@ -244,6 +244,8 @@ void FMetalRHIBuffer::Alloc(uint32 InSize, EResourceLockMode LockMode)
 		FMetalPooledBufferArgs Args(GetMetalDeviceContext().GetDevice(), InSize, Usage, Mode);
 		Buffer = GetMetalDeviceContext().CreatePooledBuffer(Args);
 		check(Buffer && Buffer.GetPtr());
+		
+		Buffer.SetOwner(this);
 
         METAL_INC_DWORD_STAT_BY(Type, MemAlloc, InSize);
         
@@ -264,6 +266,7 @@ void FMetalRHIBuffer::Alloc(uint32 InSize, EResourceLockMode LockMode)
 	{
 		FMetalPooledBufferArgs ArgsCPU(GetMetalDeviceContext().GetDevice(), InSize, BUF_Dynamic, mtlpp::StorageMode::Shared);
 		CPUBuffer = GetMetalDeviceContext().CreatePooledBuffer(ArgsCPU);
+		CPUBuffer.SetOwner(this);
 		check(CPUBuffer && CPUBuffer.GetPtr());
         METAL_INC_DWORD_STAT_BY(Type, MemAlloc, InSize);
 		check(CPUBuffer.GetLength() >= Buffer.GetLength());
diff --git a/Engine/Source/Runtime/Apple/MetalRHI/Public/MetalResources.h b/Engine/Source/Runtime/Apple/MetalRHI/Public/MetalResources.h
--- a/Engine/Source/Runtime/Apple/MetalRHI/Public/MetalResources.h
+++ b/Engine/Source/Runtime/Apple/MetalRHI/Public/MetalResources.h
@@ -348,6 +348,7 @@ public:
 	inline bool IsPooled() const { return bPooled; }
 	inline bool IsSingleUse() const { return bSingleUse; }
 	inline void MarkSingleUse() { bSingleUse = true; }
+	void SetOwner(class FMetalRHIBuffer* Owner);
 	void Release();
 	
 	friend uint32 GetTypeHash(FMetalBuffer const& Hash)
