Remove some more dead code.

diff --git a/Engine/Source/Runtime/Apple/MetalRHI/Private/MetalRHI.cpp b/Engine/Source/Runtime/Apple/MetalRHI/Private/MetalRHI.cpp
--- a/Engine/Source/Runtime/Apple/MetalRHI/Private/MetalRHI.cpp
+++ b/Engine/Source/Runtime/Apple/MetalRHI/Private/MetalRHI.cpp
@@ -23,12 +23,6 @@
 DEFINE_LOG_CATEGORY(LogMetal)
 
 bool GIsMetalInitialized = false;
-bool GMetalSupportsHeaps = false;
-bool GMetalSupportsIndirectArgumentBuffers = false;
-bool GMetalSupportsTileShaders = false;
-bool GMetalSupportsStoreActionOptions = false;
-bool GMetalSupportsDepthClipMode = false;
-bool GMetalCommandBufferHasStartEndTimeAPI = false;
 
 FMetalBufferFormat GMetalBufferFormats[PF_MAX];
 
@@ -348,24 +342,6 @@ FMetalDynamicRHI::FMetalDynamicRHI(ERHIFeatureLevel::Type RequestedFeatureLevel)
 	}
 	
 #endif
-
-	if (FApplePlatformMisc::IsOSAtLeastVersion((uint32[]){10, 13, 0}, (uint32[]){11, 0, 0}, (uint32[]){11, 0, 0}))
-	{
-		GMetalSupportsIndirectArgumentBuffers = true;
-		GMetalSupportsStoreActionOptions = true;
-	}
-	if (!PLATFORM_MAC && FApplePlatformMisc::IsOSAtLeastVersion((uint32[]){0, 0, 0}, (uint32[]){11, 0, 0}, (uint32[]){11, 0, 0}))
-	{
-		GMetalSupportsTileShaders = true;
-	}
-	if (FApplePlatformMisc::IsOSAtLeastVersion((uint32[]){10, 11, 0}, (uint32[]){11, 0, 0}, (uint32[]){11, 0, 0}))
-	{
-		GMetalSupportsDepthClipMode = true;
-	}
-	if (FApplePlatformMisc::IsOSAtLeastVersion((uint32[]){10, 13, 0}, (uint32[]){10, 3, 0}, (uint32[]){10, 3, 0}))
-	{
-		GMetalCommandBufferHasStartEndTimeAPI = true;
-	}
 		
 	if(
 	   #if PLATFORM_MAC
diff --git a/Engine/Source/Runtime/Apple/MetalRHI/Private/MetalRHIPrivate.h b/Engine/Source/Runtime/Apple/MetalRHI/Private/MetalRHIPrivate.h
--- a/Engine/Source/Runtime/Apple/MetalRHI/Private/MetalRHIPrivate.h
+++ b/Engine/Source/Runtime/Apple/MetalRHI/Private/MetalRHIPrivate.h
@@ -75,12 +75,7 @@ enum EMTLTextureType
 #define METAL_SUPPORTS_CAPTURE_MANAGER (PLATFORM_MAC && __MAC_OS_X_VERSION_MAX_ALLOWED >= 101300) || (!PLATFORM_MAC && __IPHONE_OS_VERSION_MAX_ALLOWED >= 110000) && (__clang_major__ >= 9)
 #define METAL_SUPPORTS_TILE_SHADERS (!PLATFORM_MAC && !PLATFORM_TVOS && __IPHONE_OS_VERSION_MAX_ALLOWED >= 110000) && (__clang_major__ >= 9)
 // In addition to compile-time SDK checks we also need a way to check if these are available on runtime
-extern bool GMetalSupportsIndirectArgumentBuffers;
 extern bool GMetalSupportsCaptureManager;
-extern bool GMetalSupportsTileShaders;
-extern bool GMetalSupportsStoreActionOptions;
-extern bool GMetalSupportsDepthClipMode;
-extern bool GMetalCommandBufferHasStartEndTimeAPI;
 
 struct FMetalBufferFormat
 {
@@ -107,8 +102,6 @@ extern FMetalBufferFormat GMetalBufferFormats[PF_MAX];
 #define METAL_DEBUG_LAYER(Level, Code)
 #endif
 
-extern bool GMetalSupportsTileShaders;
-
 /** Set to 1 to enable GPU events in Xcode frame debugger */
 #ifndef ENABLE_METAL_GPUEVENTS_IN_TEST
 	#define ENABLE_METAL_GPUEVENTS_IN_TEST 0
