Remove some more dead code.

diff --git a/Engine/Source/Runtime/Apple/MetalRHI/Private/MetalCommandQueue.cpp b/Engine/Source/Runtime/Apple/MetalRHI/Private/MetalCommandQueue.cpp
--- a/Engine/Source/Runtime/Apple/MetalRHI/Private/MetalCommandQueue.cpp
+++ b/Engine/Source/Runtime/Apple/MetalRHI/Private/MetalCommandQueue.cpp
@@ -19,7 +19,6 @@
 #pragma mark - Private C++ Statics -
 uint64 FMetalCommandQueue::Features = 0;
 extern mtlpp::VertexFormat GMetalFColorVertexFormat;
-extern bool GMetalManagedUniformBuffers;
 
 #pragma mark - Public C++ Boilerplate -
 
@@ -262,8 +261,6 @@ FMetalCommandQueue::FMetalCommandQueue(mtlpp::Device InDevice, uint32 const MaxN
 			Features |= EMetalFeaturesIABs;
 		}
     }
-	
-	GMetalManagedUniformBuffers = FParse::Param(FCommandLine::Get(),TEXT("metalmanagedubs"));
 #endif
 	
 #if !UE_BUILD_SHIPPING
diff --git a/Engine/Source/Runtime/Apple/MetalRHI/Private/MetalUniformBuffer.cpp b/Engine/Source/Runtime/Apple/MetalRHI/Private/MetalUniformBuffer.cpp
--- a/Engine/Source/Runtime/Apple/MetalRHI/Private/MetalUniformBuffer.cpp
+++ b/Engine/Source/Runtime/Apple/MetalRHI/Private/MetalUniformBuffer.cpp
@@ -29,7 +29,6 @@ struct FMetalRHICommandInitialiseUniformBufferIAB : public FRHICommand<FMetalRHI
 		Buffer->InitIAB();
 	}
 };
-bool GMetalManagedUniformBuffers = false;
 
 struct FMetalArgumentDesc
 {
