Mark pooled Metal resource objects as volatile when they are placed into the pool awaiting reuse, for textures after 2 frames we'll actually reclaim the memory and just keep the MTLTexture around. This should save memory when streaming as it allows the OS to reuse the memory for other allocations. It does not purge Metal resources that are still owned by RHI objects - that would require much more bookkeeping.

diff --git a/Engine/Source/Runtime/Apple/MetalRHI/Private/MetalBuffer.cpp b/Engine/Source/Runtime/Apple/MetalRHI/Private/MetalBuffer.cpp
--- a/Engine/Source/Runtime/Apple/MetalRHI/Private/MetalBuffer.cpp
+++ b/Engine/Source/Runtime/Apple/MetalRHI/Private/MetalBuffer.cpp
@@ -1199,6 +1199,8 @@ FMetalTexture FMetalTexturePool::CreateTexture(mtlpp::Device Device, mtlpp::Text
 	{
 		Texture = *Tex;
 		Pool.Remove(Descriptor);
+		
+		Texture.SetPurgeableState(mtlpp::PurgeableState::NonVolatile);
 	}
 	else
 	{
@@ -1226,6 +1228,8 @@ void FMetalTexturePool::ReleaseTexture(FMetalTexture& Texture)
 	Descriptor.usage = Texture.GetUsage();
 	Descriptor.freedFrame = GFrameNumberRenderThread;
 	
+	Texture.SetPurgeableState(mtlpp::PurgeableState::Volatile);
+	
 	FScopeLock Lock(&PoolMutex);
 	Pool.Add(Descriptor, Texture);
 }
@@ -1245,6 +1249,10 @@ void FMetalTexturePool::Drain(bool const bForce)
 			{
 				It.RemoveCurrent();
 			}
+			else if ((GFrameNumberRenderThread - It->Key.freedFrame) >= PurgeAfterNumFrames)
+			{
+				It->Value.SetPurgeableState(mtlpp::PurgeableState::Empty);
+			}
 		}
 	}
 }
@@ -1439,6 +1447,7 @@ FMetalBuffer FMetalResourceHeap::CreateBuffer(uint32 Size, uint32 Alignment, uin
 				 else
 				{
 					Buffer = ManagedBuffers.CreatePooledResource(FMetalPooledBufferArgs(Queue->GetDevice(), BlockSize, Flags, StorageMode));
+					Buffer.SetPurgeableState(mtlpp::PurgeableState::NonVolatile);
 					DEC_MEMORY_STAT_BY(STAT_MetalBufferUnusedMemory, Buffer.GetLength());
 					DEC_MEMORY_STAT_BY(STAT_MetalPooledBufferUnusedMemory, Buffer.GetLength());
 				}
@@ -1511,6 +1520,7 @@ FMetalBuffer FMetalResourceHeap::CreateBuffer(uint32 Size, uint32 Alignment, uin
 				{
 					FScopeLock Lock(&Mutex);
 					Buffer = Buffers[Storage].CreatePooledResource(FMetalPooledBufferArgs(Queue->GetDevice(), BlockSize, Flags, StorageMode));
+					Buffer.SetPurgeableState(mtlpp::PurgeableState::NonVolatile);
 					DEC_MEMORY_STAT_BY(STAT_MetalBufferUnusedMemory, Buffer.GetLength());
 					DEC_MEMORY_STAT_BY(STAT_MetalPooledBufferUnusedMemory, Buffer.GetLength());
 				}
@@ -1552,6 +1562,8 @@ void FMetalResourceHeap::ReleaseBuffer(FMetalBuffer& Buffer)
 		INC_MEMORY_STAT_BY(STAT_MetalBufferUnusedMemory, Buffer.GetLength());
 		INC_MEMORY_STAT_BY(STAT_MetalPooledBufferUnusedMemory, Buffer.GetLength());
 		
+		Buffer.SetPurgeableState(mtlpp::PurgeableState::Volatile);
+		
 		switch (StorageMode)
 		{
 	#if PLATFORM_MAC
diff --git a/Engine/Source/Runtime/Apple/MetalRHI/Private/MetalBuffer.h b/Engine/Source/Runtime/Apple/MetalRHI/Private/MetalBuffer.h
--- a/Engine/Source/Runtime/Apple/MetalRHI/Private/MetalBuffer.h
+++ b/Engine/Source/Runtime/Apple/MetalRHI/Private/MetalBuffer.h
@@ -264,6 +264,7 @@ class FMetalTexturePool
 {
 	enum
 	{
+		PurgeAfterNumFrames = 2, /* Textures must be reused fairly rapidly but after this number of frames we reclaim the memory, even though the object persists */
 		CullAfterNumFrames = 3, /* Textures must be reused fairly rapidly or we bin them as they are much larger than buffers */
 	};
 public:
