Various fixes to the Metal buffer memory reporting.

diff --git a/Engine/Source/Runtime/Apple/MetalRHI/Private/MetalBuffer.cpp b/Engine/Source/Runtime/Apple/MetalRHI/Private/MetalBuffer.cpp
--- a/Engine/Source/Runtime/Apple/MetalRHI/Private/MetalBuffer.cpp
+++ b/Engine/Source/Runtime/Apple/MetalRHI/Private/MetalBuffer.cpp
@@ -8,9 +8,17 @@
 #include "MetalCommandBuffer.h"
 #include "MetalProfiler.h"
 
+
+DECLARE_MEMORY_STAT(TEXT("Used Device Buffer Memory"), STAT_MetalDeviceBufferMemory, STATGROUP_MetalRHI);
+DECLARE_MEMORY_STAT(TEXT("Used Pooled Buffer Memory"), STAT_MetalPooledBufferMemory, STATGROUP_MetalRHI);
+DECLARE_MEMORY_STAT(TEXT("Used Magazine Buffer Memory"), STAT_MetalMagazineBufferMemory, STATGROUP_MetalRHI);
+DECLARE_MEMORY_STAT(TEXT("Used Heap Buffer Memory"), STAT_MetalHeapBufferMemory, STATGROUP_MetalRHI);
+DECLARE_MEMORY_STAT(TEXT("Used Linear Buffer Memory"), STAT_MetalLinearBufferMemory, STATGROUP_MetalRHI);
+
 DECLARE_MEMORY_STAT(TEXT("Unused Pooled Buffer Memory"), STAT_MetalPooledBufferUnusedMemory, STATGROUP_MetalRHI);
 DECLARE_MEMORY_STAT(TEXT("Unused Magazine Buffer Memory"), STAT_MetalMagazineBufferUnusedMemory, STATGROUP_MetalRHI);
 DECLARE_MEMORY_STAT(TEXT("Unused Heap Buffer Memory"), STAT_MetalHeapBufferUnusedMemory, STATGROUP_MetalRHI);
+DECLARE_MEMORY_STAT(TEXT("Unused Linear Buffer Memory"), STAT_MetalLinearBufferUnusedMemory, STATGROUP_MetalRHI);
 
 
 
@@ -238,8 +246,9 @@ void FMetalSubBufferHeap::FreeRange(ns::Range const& Range)
 	}
 	if (ParentHeap)
 	{
-		SET_MEMORY_STAT(STAT_MetalBufferUnusedMemory, ParentHeap.GetSize() - ParentHeap.GetUsedSize());
-		SET_MEMORY_STAT(STAT_MetalHeapBufferUnusedMemory, ParentHeap.GetSize() - ParentHeap.GetUsedSize());
+		INC_MEMORY_STAT_BY(STAT_MetalBufferUnusedMemory, Range.Length);
+		INC_MEMORY_STAT_BY(STAT_MetalHeapBufferUnusedMemory, Range.Length);
+		DEC_MEMORY_STAT_BY(STAT_MetalHeapBufferMemory, Range.Length);
 	}
 	else
 	{
@@ -292,6 +301,7 @@ void FMetalSubBufferHeap::FreeRange(ns::Range const& Range)
 			
 			INC_MEMORY_STAT_BY(STAT_MetalBufferUnusedMemory, Range.Length);
 			INC_MEMORY_STAT_BY(STAT_MetalHeapBufferUnusedMemory, Range.Length);
+			DEC_MEMORY_STAT_BY(STAT_MetalHeapBufferMemory, Range.Length);
 		
 #if METAL_DEBUG_OPTIONS
 			uint64 LostSize = GetSize() - UsedSize;
@@ -428,8 +438,10 @@ FMetalBuffer FMetalSubBufferHeap::NewBuffer(NSUInteger length)
 #if STATS || ENABLE_LOW_LEVEL_MEM_TRACKER
 		MetalLLM::LogAllocBuffer(GetMetalDeviceContext().GetDevice(), Result);
 #endif
-		SET_MEMORY_STAT(STAT_MetalBufferUnusedMemory, ParentHeap.GetSize() - ParentHeap.GetUsedSize());
-		SET_MEMORY_STAT(STAT_MetalHeapBufferUnusedMemory, ParentHeap.GetSize() - ParentHeap.GetUsedSize());
+		
+		DEC_MEMORY_STAT_BY(STAT_MetalBufferUnusedMemory, Result.GetLength());
+		DEC_MEMORY_STAT_BY(STAT_MetalHeapBufferUnusedMemory, Result.GetLength());
+		INC_MEMORY_STAT_BY(STAT_MetalHeapBufferMemory, Result.GetLength());
 	}
 	else
 	{
@@ -449,6 +461,7 @@ FMetalBuffer FMetalSubBufferHeap::NewBuffer(NSUInteger length)
 					
 					DEC_MEMORY_STAT_BY(STAT_MetalBufferUnusedMemory, Range.Length);
 					DEC_MEMORY_STAT_BY(STAT_MetalHeapBufferUnusedMemory, Range.Length);
+					INC_MEMORY_STAT_BY(STAT_MetalHeapBufferMemory, Range.Length);
 				
 					if (Range.Length > Size)
 					{
@@ -518,13 +531,13 @@ FMetalSubBufferLinear::FMetalSubBufferLinear(NSUInteger Size, NSUInteger Alignme
 	MetalLLM::LogAllocBuffer(GetMetalDeviceContext().GetDevice(), ParentBuffer);
 #endif
 	INC_MEMORY_STAT_BY(STAT_MetalBufferUnusedMemory, FullSize);
-	INC_MEMORY_STAT_BY(STAT_MetalHeapBufferUnusedMemory, FullSize);
+	INC_MEMORY_STAT_BY(STAT_MetalLinearBufferUnusedMemory, FullSize);
 }
 
 FMetalSubBufferLinear::~FMetalSubBufferLinear()
 {
 	DEC_MEMORY_STAT_BY(STAT_MetalBufferUnusedMemory, ParentBuffer.GetLength());
-	DEC_MEMORY_STAT_BY(STAT_MetalHeapBufferUnusedMemory, ParentBuffer.GetLength());
+	DEC_MEMORY_STAT_BY(STAT_MetalLinearBufferUnusedMemory, ParentBuffer.GetLength());
 }
 
 void FMetalSubBufferLinear::FreeRange(ns::Range const& Range)
@@ -542,7 +555,8 @@ void FMetalSubBufferLinear::FreeRange(ns::Range const& Range)
 	{
 		FreedSize += Range.Length;
 		INC_MEMORY_STAT_BY(STAT_MetalBufferUnusedMemory, Range.Length);
-		INC_MEMORY_STAT_BY(STAT_MetalHeapBufferUnusedMemory, Range.Length);
+		INC_MEMORY_STAT_BY(STAT_MetalLinearBufferUnusedMemory, Range.Length);
+		DEC_MEMORY_STAT_BY(STAT_MetalLinearBufferMemory, Range.Length);
 		if (FreedSize == UsedSize)
 		{
 			UsedSize = 0;
@@ -613,7 +627,8 @@ FMetalBuffer FMetalSubBufferLinear::NewBuffer(NSUInteger length)
 	{
 		ns::Range Range(NewWriteHead, Size);
 		DEC_MEMORY_STAT_BY(STAT_MetalBufferUnusedMemory, Range.Length);
-		DEC_MEMORY_STAT_BY(STAT_MetalHeapBufferUnusedMemory, Range.Length);
+		DEC_MEMORY_STAT_BY(STAT_MetalLinearBufferUnusedMemory, Range.Length);
+		INC_MEMORY_STAT_BY(STAT_MetalLinearBufferMemory, Range.Length);
 		Result = FMetalBuffer(MTLPP_VALIDATE(mtlpp::Buffer, ParentBuffer, SafeGetRuntimeDebuggingLevel() >= EMetalDebugLevelValidation, NewBuffer(Range)), this);
 		UsedSize += Size;
 		WriteHead = NewWriteHead + Size;
@@ -678,7 +693,7 @@ FMetalSubBufferMagazine::~FMetalSubBufferMagazine()
 	if (ParentHeap)
 	{
 		DEC_MEMORY_STAT_BY(STAT_MetalBufferUnusedMemory, ParentHeap.GetSize());
-		DEC_MEMORY_STAT_BY(STAT_MetalHeapBufferUnusedMemory, ParentHeap.GetSize());
+		DEC_MEMORY_STAT_BY(STAT_MetalMagazineBufferUnusedMemory, ParentHeap.GetSize());
 	}
 	else
 	{
@@ -692,8 +707,9 @@ void FMetalSubBufferMagazine::FreeRange(ns::Range const& Range)
 	FPlatformAtomics::InterlockedDecrement(&OutstandingAllocs);
 	if (ParentHeap)
 	{
-		SET_MEMORY_STAT(STAT_MetalBufferUnusedMemory, ParentHeap.GetSize() - ParentHeap.GetUsedSize());
-		SET_MEMORY_STAT(STAT_MetalHeapBufferUnusedMemory, ParentHeap.GetSize() - ParentHeap.GetUsedSize());
+		INC_MEMORY_STAT_BY(STAT_MetalBufferUnusedMemory, Range.Length);
+		INC_MEMORY_STAT_BY(STAT_MetalMagazineBufferUnusedMemory, Range.Length);
+		DEC_MEMORY_STAT_BY(STAT_MetalMagazineBufferMemory, Range.Length);
 	}
 	else
 	{
@@ -711,6 +727,7 @@ void FMetalSubBufferMagazine::FreeRange(ns::Range const& Range)
 		
 		INC_MEMORY_STAT_BY(STAT_MetalBufferUnusedMemory, Range.Length);
 		INC_MEMORY_STAT_BY(STAT_MetalMagazineBufferUnusedMemory, Range.Length);
+		DEC_MEMORY_STAT_BY(STAT_MetalMagazineBufferMemory, Range.Length);
 	}
 }
 
@@ -823,8 +840,9 @@ FMetalBuffer FMetalSubBufferMagazine::NewBuffer()
 #if STATS || ENABLE_LOW_LEVEL_MEM_TRACKER
 		MetalLLM::LogAllocBuffer(GetMetalDeviceContext().GetDevice(), Result);
 #endif
-		SET_MEMORY_STAT(STAT_MetalBufferUnusedMemory, ParentHeap.GetSize() - ParentHeap.GetUsedSize());
-		SET_MEMORY_STAT(STAT_MetalHeapBufferUnusedMemory, ParentHeap.GetSize() - ParentHeap.GetUsedSize());
+		DEC_MEMORY_STAT_BY(STAT_MetalBufferUnusedMemory, Result.GetLength());
+		DEC_MEMORY_STAT_BY(STAT_MetalMagazineBufferUnusedMemory, Result.GetLength());
+		INC_MEMORY_STAT_BY(STAT_MetalMagazineBufferMemory, Result.GetLength());
 	}
 	else
 	{
@@ -836,6 +854,7 @@ FMetalBuffer FMetalSubBufferMagazine::NewBuffer()
 			FPlatformAtomics::InterlockedAdd(&UsedSize, ((int64)Range->Length));
 			DEC_MEMORY_STAT_BY(STAT_MetalBufferUnusedMemory, Range->Length);
 			DEC_MEMORY_STAT_BY(STAT_MetalMagazineBufferUnusedMemory, Range->Length);
+			INC_MEMORY_STAT_BY(STAT_MetalMagazineBufferMemory, Range->Length);
 			Result = FMetalBuffer(MTLPP_VALIDATE(mtlpp::Buffer, ParentBuffer, SafeGetRuntimeDebuggingLevel() >= EMetalDebugLevelValidation, NewBuffer(*Range)), this);
 			delete Range;
 		}
@@ -1450,6 +1469,7 @@ FMetalBuffer FMetalResourceHeap::CreateBuffer(uint32 Size, uint32 Alignment, uin
 					Buffer.SetPurgeableState(mtlpp::PurgeableState::NonVolatile);
 					DEC_MEMORY_STAT_BY(STAT_MetalBufferUnusedMemory, Buffer.GetLength());
 					DEC_MEMORY_STAT_BY(STAT_MetalPooledBufferUnusedMemory, Buffer.GetLength());
+					INC_MEMORY_STAT_BY(STAT_MetalPooledBufferMemory, Buffer.GetLength());
 				}
 				break;
 			}
@@ -1523,6 +1543,7 @@ FMetalBuffer FMetalResourceHeap::CreateBuffer(uint32 Size, uint32 Alignment, uin
 					Buffer.SetPurgeableState(mtlpp::PurgeableState::NonVolatile);
 					DEC_MEMORY_STAT_BY(STAT_MetalBufferUnusedMemory, Buffer.GetLength());
 					DEC_MEMORY_STAT_BY(STAT_MetalPooledBufferUnusedMemory, Buffer.GetLength());
+					INC_MEMORY_STAT_BY(STAT_MetalPooledBufferMemory, Buffer.GetLength());
 				}
 				break;
 			}
@@ -1540,6 +1561,7 @@ FMetalBuffer FMetalResourceHeap::CreateBuffer(uint32 Size, uint32 Alignment, uin
 #if STATS || ENABLE_LOW_LEVEL_MEM_TRACKER
 		MetalLLM::LogAllocBuffer(Queue->GetDevice(), Buffer);
 #endif
+		INC_MEMORY_STAT_BY(STAT_MetalDeviceBufferMemory, Buffer.GetLength());
 	}
 	
 	if (GMetalBufferZeroFill && Buffer.GetStorageMode() != mtlpp::StorageMode::Private)
@@ -1561,6 +1583,7 @@ void FMetalResourceHeap::ReleaseBuffer(FMetalBuffer& Buffer)
 		
 		INC_MEMORY_STAT_BY(STAT_MetalBufferUnusedMemory, Buffer.GetLength());
 		INC_MEMORY_STAT_BY(STAT_MetalPooledBufferUnusedMemory, Buffer.GetLength());
+		DEC_MEMORY_STAT_BY(STAT_MetalPooledBufferMemory, Buffer.GetLength());
 		
 		Buffer.SetPurgeableState(mtlpp::PurgeableState::Volatile);
 		
@@ -1592,6 +1615,7 @@ void FMetalResourceHeap::ReleaseBuffer(FMetalBuffer& Buffer)
 	}
 	else
 	{
+		DEC_MEMORY_STAT_BY(STAT_MetalDeviceBufferMemory, Buffer.GetLength());
 		Buffer.Release();
 	}
 }
