Fix MTLHeap alignment restrictions not being honoured when allocating from MTLHeap - this doesn't matter on Mac but it is important on iOS if we ever turn them on - and make the query buffer allocate through the resource-heap so we can track the memory.

diff --git a/Engine/Source/Runtime/Apple/MetalRHI/Private/MetalBuffer.cpp b/Engine/Source/Runtime/Apple/MetalRHI/Private/MetalBuffer.cpp
--- a/Engine/Source/Runtime/Apple/MetalRHI/Private/MetalBuffer.cpp
+++ b/Engine/Source/Runtime/Apple/MetalRHI/Private/MetalBuffer.cpp
@@ -423,6 +423,24 @@ NSUInteger FMetalSubBufferHeap::MaxAvailableSize() const
 	}
 }
 
+
+bool FMetalSubBufferHeap::CanAllocateSize(NSUInteger Size) const
+{
+	if (ParentHeap)
+	{
+		NSUInteger Storage = (NSUInteger(GetStorageMode()) << mtlpp::ResourceStorageModeShift);
+		NSUInteger Cache = (NSUInteger(GetCpuCacheMode()) << mtlpp::ResourceCpuCacheModeShift);
+		mtlpp::ResourceOptions Opt = mtlpp::ResourceOptions(Storage | Cache);
+		
+		NSUInteger Align = ParentHeap.GetDevice().HeapBufferSizeAndAlign(Size, Opt).Align;
+		return Size <= ParentHeap.MaxAvailableSizeWithAlignment(Align);
+	}
+	else
+	{
+		return Size <= MaxAvailableSize();
+	}
+}
+
 FMetalBuffer FMetalSubBufferHeap::NewBuffer(NSUInteger length)
 {
 	NSUInteger Size = Align(length, MinAlign);
@@ -646,18 +664,24 @@ mtlpp::PurgeableState FMetalSubBufferLinear::SetPurgeableState(mtlpp::PurgeableS
 
 FMetalSubBufferMagazine::FMetalSubBufferMagazine(NSUInteger Size, NSUInteger ChunkSize, mtlpp::ResourceOptions Options)
 : MinAlign(ChunkSize)
+, BlockSize(ChunkSize)
 , OutstandingAllocs(0)
 , UsedSize(0)
 {
-	NSUInteger FullSize = Align(Size, ChunkSize);
+	static bool bSupportsHeaps = GetMetalDeviceContext().SupportsFeature(EMetalFeaturesHeaps);
+	mtlpp::StorageMode Storage = (mtlpp::StorageMode)((Options & mtlpp::ResourceStorageModeMask) >> mtlpp::ResourceStorageModeShift);
+	if (bSupportsHeaps && Storage == mtlpp::StorageMode::Private)
+	{
+		MinAlign = GetMetalDeviceContext().GetDevice().HeapBufferSizeAndAlign(BlockSize, Options).Align;
+	}
+	
+	NSUInteger FullSize = Align(Size, MinAlign);
 	METAL_GPUPROFILE(FScopedMetalCPUStats CPUStat(FString::Printf(TEXT("AllocBuffer: %llu, %llu"), FullSize, Options)));
 	
-	mtlpp::StorageMode Storage = (mtlpp::StorageMode)((Options & mtlpp::ResourceStorageModeMask) >> mtlpp::ResourceStorageModeShift);
 #if PLATFORM_MAC
 	check(Storage != mtlpp::StorageMode::Managed /* Managed memory cannot be safely suballocated! When you overwrite existing data the GPU buffer is immediately disposed of! */);
 #endif
 
-	static bool bSupportsHeaps = GetMetalDeviceContext().SupportsFeature(EMetalFeaturesHeaps);
 	if (bSupportsHeaps && Storage == mtlpp::StorageMode::Private)
 	{
 		mtlpp::HeapDescriptor Desc;
@@ -805,7 +829,14 @@ NSUInteger FMetalSubBufferMagazine::GetUsedSize() const
 
 NSUInteger FMetalSubBufferMagazine::GetFreeSize() const 
 {
-	return GetSize() - GetUsedSize();
+	if (ParentHeap)
+	{
+		return ParentHeap.MaxAvailableSizeWithAlignment(MinAlign);
+	}
+	else
+	{
+		return GetSize() - GetUsedSize();
+	}
 }
 
 int64 FMetalSubBufferMagazine::NumCurrentAllocations() const
@@ -813,6 +844,11 @@ int64 FMetalSubBufferMagazine::NumCurrentAllocations() const
 	return OutstandingAllocs;
 }
 
+bool FMetalSubBufferMagazine::CanAllocateSize(NSUInteger Size) const
+{
+	return GetFreeSize() >= Size;
+}
+
 void FMetalSubBufferMagazine::SetLabel(const ns::String& label)
 {
 	if (ParentHeap)
@@ -827,7 +863,7 @@ void FMetalSubBufferMagazine::SetLabel(const ns::String& label)
 
 FMetalBuffer FMetalSubBufferMagazine::NewBuffer()
 {
-	NSUInteger Size = MinAlign;
+	NSUInteger Size = BlockSize;
 	FMetalBuffer Result;
 
 	if (ParentHeap)
@@ -1491,7 +1527,7 @@ FMetalBuffer FMetalResourceHeap::CreateBuffer(uint32 Size, uint32 Alignment, uin
 					FMetalSubBufferMagazine* Found = nullptr;
 					for (FMetalSubBufferMagazine* Heap : Heaps)
 					{
-						if (Heap->GetFreeSize() >= BlockSize)
+						if (Heap->CanAllocateSize(BlockSize))
 						{
 							Found = Heap;
 							break;
@@ -1518,7 +1554,7 @@ FMetalBuffer FMetalResourceHeap::CreateBuffer(uint32 Size, uint32 Alignment, uin
 					FMetalSubBufferHeap* Found = nullptr;
 					for (FMetalSubBufferHeap* Heap : Heaps)
 					{
-						if (Heap->MaxAvailableSize() >= BlockSize)
+						if (Heap->CanAllocateSize(BlockSize))
 						{
 							Found = Heap;
 							break;
diff --git a/Engine/Source/Runtime/Apple/MetalRHI/Private/MetalBuffer.h b/Engine/Source/Runtime/Apple/MetalRHI/Private/MetalBuffer.h
--- a/Engine/Source/Runtime/Apple/MetalRHI/Private/MetalBuffer.h
+++ b/Engine/Source/Runtime/Apple/MetalRHI/Private/MetalBuffer.h
@@ -43,6 +43,7 @@ public:
     NSUInteger     GetUsedSize() const;
 	NSUInteger	 MaxAvailableSize() const;
 	int64     NumCurrentAllocations() const;
+	bool	 CanAllocateSize(NSUInteger Size) const;
 
     void SetLabel(const ns::String& label);
 	
@@ -113,6 +114,7 @@ public:
     NSUInteger     GetUsedSize() const;
 	NSUInteger	 GetFreeSize() const;
 	int64     NumCurrentAllocations() const;
+	bool	 CanAllocateSize(NSUInteger Size) const;
 
     void SetLabel(const ns::String& label);
 	void FreeRange(ns::Range const& Range);
@@ -122,6 +124,7 @@ public:
 
 private:
 	NSUInteger MinAlign;
+	NSUInteger BlockSize;
 	int64 volatile OutstandingAllocs;
 	int64 volatile UsedSize;
 	mtlpp::Buffer ParentBuffer;
diff --git a/Engine/Source/Runtime/Apple/MetalRHI/Private/MetalQuery.cpp b/Engine/Source/Runtime/Apple/MetalRHI/Private/MetalQuery.cpp
--- a/Engine/Source/Runtime/Apple/MetalRHI/Private/MetalQuery.cpp
+++ b/Engine/Source/Runtime/Apple/MetalRHI/Private/MetalQuery.cpp
@@ -57,7 +57,8 @@ FMetalQueryBuffer* FMetalQueryBufferPool::GetCurrentQueryBuffer()
 			LLM_PLATFORM_SCOPE_METAL(ELLMTagMetal::Buffers);
 
 			METAL_GPUPROFILE(FScopedMetalCPUStats CPUStat(FString::Printf(TEXT("AllocBuffer: %llu, %llu"), EQueryBufferMaxSize, mtlpp::ResourceOptions::StorageModeShared)));
-			Buffer = FMetalBuffer(MTLPP_VALIDATE(mtlpp::Device, Context->GetDevice(), SafeGetRuntimeDebuggingLevel() >= EMetalDebugLevelValidation, NewBuffer(EQueryBufferMaxSize, GetMetalDeviceContext().GetCommandQueue().GetCompatibleResourceOptions((mtlpp::ResourceOptions)(BUFFER_CACHE_MODE | mtlpp::ResourceOptions::HazardTrackingModeUntracked | mtlpp::ResourceOptions::StorageModeShared)))), false);
+			
+			Buffer = GetMetalDeviceContext().GetResourceHeap().CreateBuffer(EQueryBufferMaxSize, 16, BUF_Dynamic, Context->GetCommandQueue().GetCompatibleResourceOptions((mtlpp::ResourceOptions)(BUFFER_CACHE_MODE | mtlpp::ResourceOptions::HazardTrackingModeUntracked | mtlpp::ResourceOptions::StorageModeShared)), true);
 			FMemory::Memzero((((uint8*)Buffer.GetContents())), EQueryBufferMaxSize);
 
 #if STATS || ENABLE_LOW_LEVEL_MEM_TRACKER
