Make MetalRHI work in terms of FRHIRenderPassInfo rather than the older FRHISetRenderTargetsInfo structure - not taking advantage of that in anyway yet, this just ports the internals over.

diff --git a/Engine/Source/Runtime/Apple/MetalRHI/Private/MetalCommands.cpp b/Engine/Source/Runtime/Apple/MetalRHI/Private/MetalCommands.cpp
--- a/Engine/Source/Runtime/Apple/MetalRHI/Private/MetalCommands.cpp
+++ b/Engine/Source/Runtime/Apple/MetalRHI/Private/MetalCommands.cpp
@@ -638,27 +638,50 @@ void FMetalRHICommandContext::RHISetRenderTargets(uint32 NumSimultaneousRenderTa
 void FMetalRHICommandContext::RHISetRenderTargetsAndClear(const FRHISetRenderTargetsInfo& RenderTargetsInfo)
 {
 	@autoreleasepool {
+		
+	FRHIRenderPassInfo PassInfo;
 	bool bHasTarget = (RenderTargetsInfo.DepthStencilRenderTarget.Texture != nullptr);
 	FMetalContext* Manager = Context;
 	{
-		for (uint32 i = 0; !bHasTarget && i < RenderTargetsInfo.NumUAVs; i++)
+		PassInfo.NumUAVs = RenderTargetsInfo.NumUAVs;
+		for (uint32 i = 0; i < RenderTargetsInfo.NumUAVs; i++)
 		{
 			if (IsValidRef(RenderTargetsInfo.UnorderedAccessView[i]))
 			{
+				PassInfo.UAVs[i] = RenderTargetsInfo.UnorderedAccessView[i];
 				bHasTarget = true;
 			}
 		}
 	}
 	
-	for (uint32 i = 0; bHasTarget == false && i < RenderTargetsInfo.NumColorRenderTargets; i++)
+	for (uint32 i = 0; i < RenderTargetsInfo.NumColorRenderTargets; i++)
 	{
-		bHasTarget = (RenderTargetsInfo.ColorRenderTarget[i].Texture != nullptr);
+		if (RenderTargetsInfo.ColorRenderTarget[i].Texture)
+		{
+			PassInfo.ColorRenderTargets[i].RenderTarget = RenderTargetsInfo.ColorRenderTarget[i].Texture;
+			PassInfo.ColorRenderTargets[i].ArraySlice = RenderTargetsInfo.ColorRenderTarget[i].ArraySliceIndex;
+			PassInfo.ColorRenderTargets[i].MipIndex = RenderTargetsInfo.ColorRenderTarget[i].MipIndex;
+			PassInfo.ColorRenderTargets[i].Action = MakeRenderTargetActions(RenderTargetsInfo.ColorRenderTarget[i].LoadAction, RenderTargetsInfo.ColorRenderTarget[i].StoreAction);
+			bHasTarget = (RenderTargetsInfo.ColorRenderTarget[i].Texture != nullptr);
+			PassInfo.bIsMSAA |= PassInfo.ColorRenderTargets[i].RenderTarget->GetNumSamples() > 1;
+		}
 	}
+		
+	if (RenderTargetsInfo.DepthStencilRenderTarget.Texture)
+	{
+		PassInfo.DepthStencilRenderTarget.DepthStencilTarget = RenderTargetsInfo.DepthStencilRenderTarget.Texture;
+		PassInfo.DepthStencilRenderTarget.ExclusiveDepthStencil = RenderTargetsInfo.DepthStencilRenderTarget.GetDepthStencilAccess();
+		PassInfo.DepthStencilRenderTarget.Action = MakeDepthStencilTargetActions(MakeRenderTargetActions(RenderTargetsInfo.DepthStencilRenderTarget.DepthLoadAction, RenderTargetsInfo.DepthStencilRenderTarget.DepthStoreAction), MakeRenderTargetActions(RenderTargetsInfo.DepthStencilRenderTarget.StencilLoadAction, RenderTargetsInfo.DepthStencilRenderTarget.GetStencilStoreAction()));
+		PassInfo.bIsMSAA |= RenderTargetsInfo.DepthStencilRenderTarget.Texture->GetNumSamples() > 1;
+	}
+		
+	PassInfo.NumOcclusionQueries = UINT16_MAX;
+	PassInfo.bOcclusionQueries = true;
 
 	// Ignore any attempt to "clear" the render-targets as that is senseless with the way MetalRHI has to try and coalesce passes.
 	if (bHasTarget)
 	{
-		Manager->SetRenderTargetsInfo(RenderTargetsInfo);
+		Manager->SetRenderPassInfo(PassInfo);
 
 		// Set the viewport to the full size of render target 0.
 		if (RenderTargetsInfo.ColorRenderTarget[0].Texture)
diff --git a/Engine/Source/Runtime/Apple/MetalRHI/Private/MetalContext.cpp b/Engine/Source/Runtime/Apple/MetalRHI/Private/MetalContext.cpp
--- a/Engine/Source/Runtime/Apple/MetalRHI/Private/MetalContext.cpp
+++ b/Engine/Source/Runtime/Apple/MetalRHI/Private/MetalContext.cpp
@@ -1192,7 +1192,7 @@ void FMetalContext::ResetRenderCommandEncoder()
 	
 	StateCache.InvalidateRenderTargets();
 	
-	SetRenderTargetsInfo(StateCache.GetRenderTargetsInfo(), true);
+	SetRenderPassInfo(StateCache.GetRenderPassInfo(), true);
 }
 
 bool FMetalContext::PrepareToDraw(uint32 PrimitiveType, EMetalIndexType IndexType)
@@ -1254,7 +1254,7 @@ bool FMetalContext::PrepareToDraw(uint32 PrimitiveType, EMetalIndexType IndexTyp
 	bool const bNeedsDepthStencilWrite = (IsValidRef(CurrentPSO->PixelShader) && (CurrentPSO->PixelShader->Bindings.InOutMask & 0x8000));
 	
 	// @todo Improve the way we handle binding a dummy depth/stencil so we can get pure UAV raster operations...
-	bool const bNeedsDepthStencilForUAVRaster = (StateCache.GetRenderTargetsInfo().NumColorRenderTargets == 0 && StateCache.GetRenderTargetsInfo().NumUAVs > 0);
+	bool const bNeedsDepthStencilForUAVRaster = (StateCache.GetRenderPassInfo().GetNumColorRenderTargets() == 0 && StateCache.GetRenderPassInfo().NumUAVs > 0);
 	
 	bool const bBindDepthStencilForWrite = bNeedsDepthStencilWrite && !StateCache.HasValidDepthStencilSurface();
 	bool const bBindDepthStencilForUAVRaster = bNeedsDepthStencilForUAVRaster && !StateCache.HasValidDepthStencilSurface();
@@ -1285,7 +1285,7 @@ bool FMetalContext::PrepareToDraw(uint32 PrimitiveType, EMetalIndexType IndexTyp
 			FBSize = CGSizeMake(StateCache.GetViewport(0).width, StateCache.GetViewport(0).height);
 		}
 		
-		FRHISetRenderTargetsInfo Info = StateCache.GetRenderTargetsInfo();
+		FRHIRenderPassInfo Info = StateCache.GetRenderPassInfo();
 		
 		FTexture2DRHIRef FallbackDepthStencilSurface = StateCache.CreateFallbackDepthStencilSurface(FBSize.width, FBSize.height);
 		check(IsValidRef(FallbackDepthStencilSurface));
@@ -1293,25 +1293,27 @@ bool FMetalContext::PrepareToDraw(uint32 PrimitiveType, EMetalIndexType IndexTyp
 		if (bBindDepthStencilForWrite)
 		{
 			check(!bBindDepthStencilForUAVRaster);
-			Info.DepthStencilRenderTarget.Texture = FallbackDepthStencilSurface;
+			Info.DepthStencilRenderTarget.DepthStencilTarget = FallbackDepthStencilSurface;
 		}
 		else
 		{
 			check(bBindDepthStencilForUAVRaster);
-			Info.DepthStencilRenderTarget = FRHIDepthRenderTargetView(FallbackDepthStencilSurface, ERenderTargetLoadAction::ELoad, ERenderTargetStoreAction::ENoAction, FExclusiveDepthStencil::DepthRead_StencilRead);
+			Info.DepthStencilRenderTarget.DepthStencilTarget = FallbackDepthStencilSurface;
+			Info.DepthStencilRenderTarget.ExclusiveDepthStencil = FExclusiveDepthStencil::DepthRead_StencilRead;
+			Info.DepthStencilRenderTarget.Action = MakeDepthStencilTargetActions(MakeRenderTargetActions(ERenderTargetLoadAction::ELoad, ERenderTargetStoreAction::ENoAction), MakeRenderTargetActions(ERenderTargetLoadAction::ELoad, ERenderTargetStoreAction::ENoAction));
 		}
 		
 		// Ensure that we make it a Clear/Store -> Load/Store for the colour targets or we might render incorrectly
-		for (uint32 i = 0; i < Info.NumColorRenderTargets; i++)
+		for (uint32 i = 0; i < Info.GetNumColorRenderTargets(); i++)
 		{
-			if (Info.ColorRenderTarget[i].LoadAction != ERenderTargetLoadAction::ELoad)
+			if (GetLoadAction(Info.ColorRenderTargets[i].Action) != ERenderTargetLoadAction::ELoad)
 			{
-				check(Info.ColorRenderTarget[i].StoreAction == ERenderTargetStoreAction::EStore || Info.ColorRenderTarget[i].StoreAction == ERenderTargetStoreAction::EMultisampleResolve);
-				Info.ColorRenderTarget[i].LoadAction = ERenderTargetLoadAction::ELoad;
+				check(GetStoreAction(Info.ColorRenderTargets[i].Action) == ERenderTargetStoreAction::EStore || GetStoreAction(Info.ColorRenderTargets[i].Action) == ERenderTargetStoreAction::EMultisampleResolve);
+				Info.ColorRenderTargets[i].Action = MakeRenderTargetActions(ERenderTargetLoadAction::ELoad, GetStoreAction(Info.ColorRenderTargets[i].Action));
 			}
 		}
 		
-		if (StateCache.SetRenderTargetsInfo(Info, StateCache.GetVisibilityResultsBuffer(), true))
+		if (StateCache.SetRenderPassInfo(Info, StateCache.GetVisibilityResultsBuffer(), true))
 		{
 			RenderPass.RestartRenderPass(StateCache.GetRenderPassDescriptor());
 		}
@@ -1326,13 +1328,13 @@ bool FMetalContext::PrepareToDraw(uint32 PrimitiveType, EMetalIndexType IndexTyp
 	}
 	else if (!bNeedsDepthStencilWrite && !bNeedsDepthStencilForUAVRaster && StateCache.GetFallbackDepthStencilBound())
 	{
-		FRHISetRenderTargetsInfo Info = StateCache.GetRenderTargetsInfo();
-		Info.DepthStencilRenderTarget.Texture = nullptr;
+		FRHIRenderPassInfo Info = StateCache.GetRenderPassInfo();
+		Info.DepthStencilRenderTarget.DepthStencilTarget = nullptr;
 		
 		RenderPass.EndRenderPass();
 		
 		StateCache.SetRenderTargetsActive(false);
-		StateCache.SetRenderTargetsInfo(Info, StateCache.GetVisibilityResultsBuffer(), true);
+		StateCache.SetRenderPassInfo(Info, StateCache.GetVisibilityResultsBuffer(), true);
 		
 		RenderPass.BeginRenderPass(StateCache.GetRenderPassDescriptor());
 		
@@ -1345,7 +1347,7 @@ bool FMetalContext::PrepareToDraw(uint32 PrimitiveType, EMetalIndexType IndexTyp
 	return true;
 }
 
-void FMetalContext::SetRenderTargetsInfo(const FRHISetRenderTargetsInfo& RenderTargetsInfo, bool const bRestart)
+void FMetalContext::SetRenderPassInfo(const FRHIRenderPassInfo& RenderTargetsInfo, bool const bRestart)
 {
 	if (CommandList.IsParallel())
 	{
@@ -1359,10 +1361,10 @@ void FMetalContext::SetRenderTargetsInfo(const FRHISetRenderTargetsInfo& RenderT
 		
 		for (uint32 RenderTargetIndex = 0; RenderTargetIndex < MaxSimultaneousRenderTargets; RenderTargetIndex++)
 		{
-			if (RenderTargetIndex < RenderTargetsInfo.NumColorRenderTargets && RenderTargetsInfo.ColorRenderTarget[RenderTargetIndex].Texture != nullptr)
+			if (RenderTargetIndex < RenderTargetsInfo.GetNumColorRenderTargets() && RenderTargetsInfo.ColorRenderTargets[RenderTargetIndex].RenderTarget != nullptr)
 			{
-				const FRHIRenderTargetView& RenderTargetView = RenderTargetsInfo.ColorRenderTarget[RenderTargetIndex];
-				if(RenderTargetView.LoadAction == ERenderTargetLoadAction::EClear)
+				const FRHIRenderPassInfo::FColorEntry& RenderTargetView = RenderTargetsInfo.ColorRenderTargets[RenderTargetIndex];
+				if(GetLoadAction(RenderTargetView.Action) == ERenderTargetLoadAction::EClear)
 				{
 					bClearInParallelBuffer = true;
 				}
@@ -1374,13 +1376,13 @@ void FMetalContext::SetRenderTargetsInfo(const FRHISetRenderTargetsInfo& RenderT
 			UE_LOG(LogMetal, Warning, TEXT("One or more render targets bound for clear during parallel encoding: this will not behave as expected because each command-buffer will clear the target of the previous contents."));
 		}
 		
-		if (RenderTargetsInfo.DepthStencilRenderTarget.Texture != nullptr)
+		if (RenderTargetsInfo.DepthStencilRenderTarget.DepthStencilTarget != nullptr)
 		{
-			if(RenderTargetsInfo.DepthStencilRenderTarget.DepthLoadAction == ERenderTargetLoadAction::EClear)
+			if(GetLoadAction(GetDepthActions(RenderTargetsInfo.DepthStencilRenderTarget.Action)) == ERenderTargetLoadAction::EClear)
 			{
 				UE_LOG(LogMetal, Warning, TEXT("Depth-target bound for clear during parallel encoding: this will not behave as expected because each command-buffer will clear the target of the previous contents."));
 			}
-			if(RenderTargetsInfo.DepthStencilRenderTarget.StencilLoadAction == ERenderTargetLoadAction::EClear)
+			if(GetLoadAction(GetStencilActions(RenderTargetsInfo.DepthStencilRenderTarget.Action)) == ERenderTargetLoadAction::EClear)
 			{
 				UE_LOG(LogMetal, Warning, TEXT("Stencil-target bound for clear during parallel encoding: this will not behave as expected because each command-buffer will clear the target of the previous contents."));
 			}
@@ -1392,25 +1394,29 @@ void FMetalContext::SetRenderTargetsInfo(const FRHISetRenderTargetsInfo& RenderT
 	if (IsFeatureLevelSupported( GMaxRHIShaderPlatform, ERHIFeatureLevel::ES3_1 ))
 	{
 		// @todo Improve the way we handle binding a dummy depth/stencil so we can get pure UAV raster operations...
-		const bool bNeedsDepthStencilForUAVRaster = RenderTargetsInfo.NumColorRenderTargets == 0 && RenderTargetsInfo.NumUAVs > 0 && !RenderTargetsInfo.DepthStencilRenderTarget.Texture;
+		const bool bNeedsDepthStencilForUAVRaster = RenderTargetsInfo.GetNumColorRenderTargets() == 0 && RenderTargetsInfo.NumUAVs > 0 && !RenderTargetsInfo.DepthStencilRenderTarget.DepthStencilTarget;
 
 		if (bNeedsDepthStencilForUAVRaster)
 		{
-			FRHISetRenderTargetsInfo Info = RenderTargetsInfo;
+			FRHIRenderPassInfo Info = RenderTargetsInfo;
 			CGSize FBSize = CGSizeMake(StateCache.GetViewport(0).width, StateCache.GetViewport(0).height);
 			FTexture2DRHIRef FallbackDepthStencilSurface = StateCache.CreateFallbackDepthStencilSurface(FBSize.width, FBSize.height);
 			check(IsValidRef(FallbackDepthStencilSurface));
+
+			Info.DepthStencilRenderTarget.DepthStencilTarget = FallbackDepthStencilSurface;
+			Info.DepthStencilRenderTarget.ResolveTarget = nullptr;
+			Info.DepthStencilRenderTarget.ExclusiveDepthStencil = FExclusiveDepthStencil::DepthRead_StencilRead;
 #if PLATFORM_MAC
-			Info.DepthStencilRenderTarget = FRHIDepthRenderTargetView(FallbackDepthStencilSurface, ERenderTargetLoadAction::ENoAction, ERenderTargetStoreAction::ENoAction, FExclusiveDepthStencil::DepthRead_StencilRead);
+			Info.DepthStencilRenderTarget.Action = MakeDepthStencilTargetActions(MakeRenderTargetActions(ERenderTargetLoadAction::ENoAction, ERenderTargetStoreAction::ENoAction), MakeRenderTargetActions(ERenderTargetLoadAction::ENoAction, ERenderTargetStoreAction::ENoAction));
 #else
-			Info.DepthStencilRenderTarget = FRHIDepthRenderTargetView(FallbackDepthStencilSurface, ERenderTargetLoadAction::EClear, ERenderTargetStoreAction::ENoAction, FExclusiveDepthStencil::DepthRead_StencilRead);
+			Info.DepthStencilRenderTarget.Action = MakeDepthStencilTargetActions(MakeRenderTargetActions(ERenderTargetLoadAction::EClear, ERenderTargetStoreAction::ENoAction), MakeRenderTargetActions(ERenderTargetLoadAction::EClear, ERenderTargetStoreAction::ENoAction));
 #endif
 
 			if (QueryBuffer->GetCurrentQueryBuffer() != StateCache.GetVisibilityResultsBuffer())
 			{
 				RenderPass.EndRenderPass();
 			}
-			bSet = StateCache.SetRenderTargetsInfo(Info, QueryBuffer->GetCurrentQueryBuffer(), bRestart);
+			bSet = StateCache.SetRenderPassInfo(Info, QueryBuffer->GetCurrentQueryBuffer(), bRestart);
 		}
 		else
 		{
@@ -1418,7 +1424,7 @@ void FMetalContext::SetRenderTargetsInfo(const FRHISetRenderTargetsInfo& RenderT
 			{
 				RenderPass.EndRenderPass();
 			}
-			bSet = StateCache.SetRenderTargetsInfo(RenderTargetsInfo, QueryBuffer->GetCurrentQueryBuffer(), bRestart);
+			bSet = StateCache.SetRenderPassInfo(RenderTargetsInfo, QueryBuffer->GetCurrentQueryBuffer(), bRestart);
 		}
 	}
 	else
@@ -1427,7 +1433,7 @@ void FMetalContext::SetRenderTargetsInfo(const FRHISetRenderTargetsInfo& RenderT
 		{
 			RenderPass.EndRenderPass();
 		}
-		bSet = StateCache.SetRenderTargetsInfo(RenderTargetsInfo, NULL, bRestart);
+		bSet = StateCache.SetRenderPassInfo(RenderTargetsInfo, NULL, bRestart);
 	}
 	
 	if (bSet && StateCache.GetHasValidRenderTarget())
@@ -1629,7 +1635,7 @@ void FMetalDeviceContext::BeginParallelRenderCommandEncoding(uint32 Num)
 	FPlatformAtomics::AtomicStore(&NumParallelContextsInPass, (int32)Num);
 }
 
-void FMetalDeviceContext::SetParallelRenderPassDescriptor(FRHISetRenderTargetsInfo const& TargetInfo)
+void FMetalDeviceContext::SetParallelRenderPassDescriptor(FRHIRenderPassInfo const& TargetInfo)
 {
 	FScopeLock Lock(&FreeListMutex);
 
@@ -1638,7 +1644,7 @@ void FMetalDeviceContext::SetParallelRenderPassDescriptor(FRHISetRenderTargetsIn
 		RenderPass.Begin(EndFence);
 		EndFence = nullptr;
 		StateCache.InvalidateRenderTargets();
-		SetRenderTargetsInfo(TargetInfo, false);
+		SetRenderPassInfo(TargetInfo, false);
 	}
 }
 
diff --git a/Engine/Source/Runtime/Apple/MetalRHI/Private/MetalContext.h b/Engine/Source/Runtime/Apple/MetalRHI/Private/MetalContext.h
--- a/Engine/Source/Runtime/Apple/MetalRHI/Private/MetalContext.h
+++ b/Engine/Source/Runtime/Apple/MetalRHI/Private/MetalContext.h
@@ -51,7 +51,7 @@ public:
 	/**
 	 * Set the color, depth and stencil render targets, and then make the new command buffer/encoder
 	 */
-	void SetRenderTargetsInfo(const FRHISetRenderTargetsInfo& RenderTargetsInfo, bool const bRestart = false);
+	void SetRenderPassInfo(const FRHIRenderPassInfo& RenderTargetsInfo, bool const bRestart = false);
 	
 	/**
 	 * Allocate from a dynamic ring buffer - by default align to the allowed alignment for offset field when setting buffers
@@ -215,7 +215,7 @@ public:
 	uint32 GetNumActiveContexts(void) const;
 
 	void BeginParallelRenderCommandEncoding(uint32 Num);
-	void SetParallelRenderPassDescriptor(FRHISetRenderTargetsInfo const& TargetInfo);
+	void SetParallelRenderPassDescriptor(FRHIRenderPassInfo const& TargetInfo);
 	mtlpp::RenderCommandEncoder GetParallelRenderCommandEncoder(uint32 Index, mtlpp::ParallelRenderCommandEncoder& ParallelEncoder, mtlpp::CommandBuffer& CommandBuffer);
 	void EndParallelRenderCommandEncoding(void);
 	
diff --git a/Engine/Source/Runtime/Apple/MetalRHI/Private/MetalRHIContext.cpp b/Engine/Source/Runtime/Apple/MetalRHI/Private/MetalRHIContext.cpp
--- a/Engine/Source/Runtime/Apple/MetalRHI/Private/MetalRHIContext.cpp
+++ b/Engine/Source/Runtime/Apple/MetalRHI/Private/MetalRHIContext.cpp
@@ -140,6 +140,9 @@ FMetalRHIImmediateCommandContext::FMetalRHIImmediateCommandContext(class FMetalP
 
 void FMetalRHICommandContext::RHIBeginRenderPass(const FRHIRenderPassInfo& InInfo, const TCHAR* InName)
 {
+	@autoreleasepool {
+	bool bHasTarget = (InInfo.DepthStencilRenderTarget.DepthStencilTarget != nullptr || InInfo.GetNumColorRenderTargets() > 0 || InInfo.NumUAVs > 0);
+	
 	if (InInfo.bGeneratingMips)
 	{
 		FRHITexture* Textures[MaxSimultaneousRenderTargets];
@@ -168,36 +171,16 @@ void FMetalRHICommandContext::RHIBeginRenderPass(const FRHIRenderPassInfo& InInf
 		Context->GetCommandList().SetParallelIndex(0, 0);
 	}
 
-	FRHISetRenderTargetsInfo RenderTargetsInfo;
-	InInfo.ConvertToRenderTargetsInfo(RenderTargetsInfo);
-
-	@autoreleasepool {
-	bool bHasTarget = (RenderTargetsInfo.DepthStencilRenderTarget.Texture != nullptr);
-	{
-		for (uint32 i = 0; !bHasTarget && i < RenderTargetsInfo.NumUAVs; i++)
-		{
-			if (IsValidRef(RenderTargetsInfo.UnorderedAccessView[i]))
-			{
-				bHasTarget = true;
-			}
-		}
-	}
-	
-	for (uint32 i = 0; bHasTarget == false && i < RenderTargetsInfo.NumColorRenderTargets; i++)
-	{
-		bHasTarget = (RenderTargetsInfo.ColorRenderTarget[i].Texture != nullptr);
-	}
-
 	// Ignore any attempt to "clear" the render-targets as that is senseless with the way MetalRHI has to try and coalesce passes.
 	if (bHasTarget)
 	{
-		Context->SetRenderTargetsInfo(RenderTargetsInfo);
+		Context->SetRenderPassInfo(InInfo);
 
 		// Set the viewport to the full size of render target 0.
-		if (RenderTargetsInfo.ColorRenderTarget[0].Texture)
+		if (InInfo.ColorRenderTargets[0].RenderTarget)
 		{
-			const FRHIRenderTargetView& RenderTargetView = RenderTargetsInfo.ColorRenderTarget[0];
-			FMetalSurface* RenderTarget = GetMetalSurfaceFromRHITexture(RenderTargetView.Texture);
+			const FRHIRenderPassInfo::FColorEntry& RenderTargetView = InInfo.ColorRenderTargets[0];
+			FMetalSurface* RenderTarget = GetMetalSurfaceFromRHITexture(RenderTargetView.RenderTarget);
 
 			uint32 Width = FMath::Max((uint32)(RenderTarget->Texture.GetWidth() >> RenderTargetView.MipIndex), (uint32)1);
 			uint32 Height = FMath::Max((uint32)(RenderTarget->Texture.GetHeight() >> RenderTargetView.MipIndex), (uint32)1);
diff --git a/Engine/Source/Runtime/Apple/MetalRHI/Private/MetalRenderPass.cpp b/Engine/Source/Runtime/Apple/MetalRHI/Private/MetalRenderPass.cpp
--- a/Engine/Source/Runtime/Apple/MetalRHI/Private/MetalRenderPass.cpp
+++ b/Engine/Source/Runtime/Apple/MetalRHI/Private/MetalRenderPass.cpp
@@ -1692,7 +1692,7 @@ void FMetalRenderPass::ConditionalSubmit()
 	bool bCanForceSubmit = State.CanRestartRenderPass();
 
 #if METAL_DEBUG_OPTIONS
-	FRHISetRenderTargetsInfo CurrentRenderTargets = State.GetRenderTargetsInfo();
+	FRHIRenderPassInfo CurrentRenderTargets = State.GetRenderPassInfo();
 	
 	// Force a command-encoder when GMetalRuntimeDebugLevel is enabled to help track down intermittent command-buffer failures.
 	if (GMetalCommandBufferCommitThreshold > 0 && NumOutstandingOps >= GMetalCommandBufferCommitThreshold && CmdList.GetCommandQueue().GetRuntimeDebuggingLevel() >= EMetalDebugLevelConditionalSubmit)
@@ -1704,14 +1704,13 @@ void FMetalRenderPass::ConditionalSubmit()
 			const bool bIsMSAAActive = State.GetHasValidRenderTarget() && State.GetSampleCount() != 1;
 			bCanChangeRT = !bIsMSAAActive;
 			
-			for (int32 RenderTargetIndex = 0; bCanChangeRT && RenderTargetIndex < CurrentRenderTargets.NumColorRenderTargets; RenderTargetIndex++)
+			for (int32 RenderTargetIndex = 0; bCanChangeRT && RenderTargetIndex < CurrentRenderTargets.GetNumColorRenderTargets(); RenderTargetIndex++)
 			{
-				FRHIRenderTargetView& RenderTargetView = CurrentRenderTargets.ColorRenderTarget[RenderTargetIndex];
+				FRHIRenderPassInfo::FColorEntry& RenderTargetView = CurrentRenderTargets.ColorRenderTargets[RenderTargetIndex];
 				
-				if (RenderTargetView.StoreAction != ERenderTargetStoreAction::EMultisampleResolve)
+				if (GetStoreAction(RenderTargetView.Action) != ERenderTargetStoreAction::EMultisampleResolve)
 				{
-					RenderTargetView.LoadAction = ERenderTargetLoadAction::ELoad;
-					RenderTargetView.StoreAction = ERenderTargetStoreAction::EStore;
+					RenderTargetView.Action = MakeRenderTargetActions(ERenderTargetLoadAction::ELoad, ERenderTargetStoreAction::EStore);
 				}
 				else
 				{
@@ -1719,11 +1718,12 @@ void FMetalRenderPass::ConditionalSubmit()
 				}
 			}
 			
-			if (bCanChangeRT && CurrentRenderTargets.DepthStencilRenderTarget.Texture)
+			if (bCanChangeRT && CurrentRenderTargets.DepthStencilRenderTarget.DepthStencilTarget)
 			{
-				if (CurrentRenderTargets.DepthStencilRenderTarget.DepthStoreAction != ERenderTargetStoreAction::EMultisampleResolve && CurrentRenderTargets.DepthStencilRenderTarget.GetStencilStoreAction() != ERenderTargetStoreAction::EMultisampleResolve)
+				if (GetStoreAction(GetDepthActions(CurrentRenderTargets.DepthStencilRenderTarget.Action)) != ERenderTargetStoreAction::EMultisampleResolve && GetStoreAction(GetStencilActions(CurrentRenderTargets.DepthStencilRenderTarget.Action)) != ERenderTargetStoreAction::EMultisampleResolve)
 				{
-					CurrentRenderTargets.DepthStencilRenderTarget = FRHIDepthRenderTargetView(CurrentRenderTargets.DepthStencilRenderTarget.Texture, ERenderTargetLoadAction::ELoad, ERenderTargetStoreAction::EStore);
+					ERenderTargetActions Actions = MakeRenderTargetActions(ERenderTargetLoadAction::ELoad, ERenderTargetStoreAction::EStore);
+					CurrentRenderTargets.DepthStencilRenderTarget.Action = MakeDepthStencilTargetActions(Actions, Actions);
 				}
 				else
 				{
@@ -1752,11 +1752,11 @@ void FMetalRenderPass::ConditionalSubmit()
 			State.InvalidateRenderTargets();
 			if (IsFeatureLevelSupported( GMaxRHIShaderPlatform, ERHIFeatureLevel::SM4 ))
 			{
-				bSet = State.SetRenderTargetsInfo(CurrentRenderTargets, State.GetVisibilityResultsBuffer(), false);
+				bSet = State.SetRenderPassInfo(CurrentRenderTargets, State.GetVisibilityResultsBuffer(), false);
 			}
 			else
 			{
-				bSet = State.SetRenderTargetsInfo(CurrentRenderTargets, NULL, false);
+				bSet = State.SetRenderPassInfo(CurrentRenderTargets, NULL, false);
 			}
 			
 			if (bSet)
diff --git a/Engine/Source/Runtime/Apple/MetalRHI/Private/MetalStateCache.cpp b/Engine/Source/Runtime/Apple/MetalRHI/Private/MetalStateCache.cpp
--- a/Engine/Source/Runtime/Apple/MetalRHI/Private/MetalStateCache.cpp
+++ b/Engine/Source/Runtime/Apple/MetalRHI/Private/MetalStateCache.cpp
@@ -175,7 +175,7 @@ FMetalStateCache::FMetalStateCache(bool const bInImmediate)
 		ColorStore[i] = mtlpp::StoreAction::Unknown;
 	}
 	
-	FMemory::Memzero(RenderTargetsInfo);
+	FMemory::Memzero(RenderPassInfo);
 	FMemory::Memzero(DirtyUniformBuffers);
 }
 
@@ -230,7 +230,7 @@ void FMetalStateCache::Reset(void)
 	ActiveViewports = 0;
 	ActiveScissors = 0;
 	
-	FMemory::Memzero(RenderTargetsInfo);
+	FMemory::Memzero(RenderPassInfo);
 	bIsRenderTargetActive = false;
 	bHasValidRenderTarget = false;
 	bHasValidColorTarget = false;
@@ -390,7 +390,7 @@ void FMetalStateCache::SetComputeShader(FMetalComputeShader* InComputeShader)
 	}
 }
 
-bool FMetalStateCache::SetRenderTargetsInfo(FRHISetRenderTargetsInfo const& InRenderTargets, FMetalQueryBuffer* QueryBuffer, bool const bRestart)
+bool FMetalStateCache::SetRenderPassInfo(FRHIRenderPassInfo const& InRenderTargets, FMetalQueryBuffer* QueryBuffer, bool const bRestart)
 {
 	bool bNeedsSet = false;
 	
@@ -410,7 +410,7 @@ bool FMetalStateCache::SetRenderTargetsInfo(FRHISetRenderTargetsInfo const& InRe
 		mtlpp::StoreAction NewStencilStore = mtlpp::StoreAction::Unknown;
 		
 		// back this up for next frame
-		RenderTargetsInfo = InRenderTargets;
+		RenderPassInfo = InRenderTargets;
 		
 		// at this point, we need to fully set up an encoder/command buffer, so make a new one (autoreleased)
 		mtlpp::RenderPassDescriptor RenderPass = FMetalRenderPassDescriptorPool::Get().CreateDescriptor();
@@ -453,17 +453,20 @@ bool FMetalStateCache::SetRenderTargetsInfo(FRHISetRenderTargetsInfo const& InRe
 		
 		ns::Array<mtlpp::RenderPassColorAttachmentDescriptor> Attachements = RenderPass.GetColorAttachments();
 		
+		uint32 NumColorRenderTargets = RenderPassInfo.GetNumColorRenderTargets();
+		
 		for (uint32 RenderTargetIndex = 0; RenderTargetIndex < MaxSimultaneousRenderTargets; RenderTargetIndex++)
 		{
 			// default to invalid
 			uint8 FormatKey = 0;
-			// only try to set it if it was one that was set (ie less than RenderTargetsInfo.NumColorRenderTargets)
-			if (RenderTargetIndex < RenderTargetsInfo.NumColorRenderTargets && RenderTargetsInfo.ColorRenderTarget[RenderTargetIndex].Texture != nullptr)
+			// only try to set it if it was one that was set (ie less than RenderPassInfo.NumColorRenderTargets)
+			if (RenderTargetIndex < NumColorRenderTargets && RenderPassInfo.ColorRenderTargets[RenderTargetIndex].RenderTarget != nullptr)
 			{
-				const FRHIRenderTargetView& RenderTargetView = RenderTargetsInfo.ColorRenderTarget[RenderTargetIndex];
-				ColorTargets[RenderTargetIndex] = RenderTargetView.Texture;
+				const FRHIRenderPassInfo::FColorEntry& RenderTargetView = RenderPassInfo.ColorRenderTargets[RenderTargetIndex];
+				ColorTargets[RenderTargetIndex] = RenderTargetView.RenderTarget;
+				ResolveTargets[RenderTargetIndex] = RenderTargetView.ResolveTarget;
 				
-				FMetalSurface& Surface = *GetMetalSurfaceFromRHITexture(RenderTargetView.Texture);
+				FMetalSurface& Surface = *GetMetalSurfaceFromRHITexture(RenderTargetView.RenderTarget);
 				FormatKey = Surface.FormatKey;
 				
 				uint32 Width = FMath::Max((uint32)(Surface.SizeX >> RenderTargetView.MipIndex), (uint32)1);
@@ -498,7 +501,7 @@ bool FMetalStateCache::SetRenderTargetsInfo(FRHISetRenderTargetsInfo const& InRe
 				check (Surface.Texture);
 	
 				// user code generally passes -1 as a default, but we need 0
-				uint32 ArraySliceIndex = RenderTargetView.ArraySliceIndex == 0xFFFFFFFF ? 0 : RenderTargetView.ArraySliceIndex;
+				uint32 ArraySliceIndex = RenderTargetView.ArraySlice == 0xFFFFFFFF ? 0 : RenderTargetView.ArraySlice;
 				if (Surface.bIsCubemap)
 				{
 					ArraySliceIndex = GetMetalCubeFace((ECubeFace)ArraySliceIndex);
@@ -509,7 +512,7 @@ bool FMetalStateCache::SetRenderTargetsInfo(FRHISetRenderTargetsInfo const& InRe
 					case RRT_Texture2DArray:
 					case RRT_Texture3D:
 					case RRT_TextureCube:
-						if(RenderTargetView.ArraySliceIndex == 0xFFFFFFFF)
+						if(RenderTargetView.ArraySlice == 0xFFFFFFFF)
 						{
 							ArrayTargets |= (1 << RenderTargetIndex);
 							ArrayRenderLayers = FMath::Min(ArrayRenderLayers, Surface.GetNumFaces());
@@ -526,8 +529,8 @@ bool FMetalStateCache::SetRenderTargetsInfo(FRHISetRenderTargetsInfo const& InRe
 	
 				mtlpp::RenderPassColorAttachmentDescriptor ColorAttachment = Attachements[RenderTargetIndex];
 	
-				ERenderTargetStoreAction HighLevelStoreAction = RenderTargetView.StoreAction;
-				ERenderTargetLoadAction HighLevelLoadAction = RenderTargetView.LoadAction;
+				ERenderTargetStoreAction HighLevelStoreAction = GetStoreAction(RenderTargetView.Action);
+				ERenderTargetLoadAction HighLevelLoadAction = GetLoadAction(RenderTargetView.Action);
 				
 				if (Surface.MSAATexture)
 				{
@@ -544,7 +547,7 @@ bool FMetalStateCache::SetRenderTargetsInfo(FRHISetRenderTargetsInfo const& InRe
 					}
 #endif
 					// only allow one MRT with msaa
-					checkf(RenderTargetsInfo.NumColorRenderTargets == 1, TEXT("Only expected one MRT when using MSAA"));
+					checkf(NumColorRenderTargets == 1, TEXT("Only expected one MRT when using MSAA"));
 				}
 				else
 				{
@@ -577,14 +580,14 @@ bool FMetalStateCache::SetRenderTargetsInfo(FRHISetRenderTargetsInfo const& InRe
 				
 				bNeedsClear |= (ColorAttachment.GetLoadAction() == mtlpp::LoadAction::Clear);
 				
-				const FClearValueBinding& ClearValue = RenderTargetsInfo.ColorRenderTarget[RenderTargetIndex].Texture->GetClearBinding();
+				const FClearValueBinding& ClearValue = RenderPassInfo.ColorRenderTargets[RenderTargetIndex].RenderTarget->GetClearBinding();
 				if (ClearValue.ColorBinding == EClearBinding::EColorBound)
 				{
 					const FLinearColor& ClearColor = ClearValue.GetClearColor();
 					ColorAttachment.SetClearColor(mtlpp::ClearColor(ClearColor.R, ClearColor.G, ClearColor.B, ClearColor.A));
 				}
 
-				bCanRestartRenderPass &= (SampleCount <= 1) && (ColorAttachment.GetLoadAction() == mtlpp::LoadAction::Load) && (RenderTargetView.StoreAction == ERenderTargetStoreAction::EStore);
+				bCanRestartRenderPass &= (SampleCount <= 1) && (ColorAttachment.GetLoadAction() == mtlpp::LoadAction::Load) && (HighLevelStoreAction == ERenderTargetStoreAction::EStore);
 	
 				bHasValidRenderTarget = true;
 				bHasValidColorTarget = true;
@@ -592,6 +595,7 @@ bool FMetalStateCache::SetRenderTargetsInfo(FRHISetRenderTargetsInfo const& InRe
 			else
 			{
 				ColorTargets[RenderTargetIndex].SafeRelease();
+				ResolveTargets[RenderTargetIndex].SafeRelease();
 			}
 		}
 		
@@ -627,9 +631,9 @@ bool FMetalStateCache::SetRenderTargetsInfo(FRHISetRenderTargetsInfo const& InRe
 		uint8 StencilFormatKey = 0;
 		
 		// setup depth and/or stencil
-		if (RenderTargetsInfo.DepthStencilRenderTarget.Texture != nullptr)
+		if (RenderPassInfo.DepthStencilRenderTarget.DepthStencilTarget != nullptr)
 		{
-			FMetalSurface& Surface = *GetMetalSurfaceFromRHITexture(RenderTargetsInfo.DepthStencilRenderTarget.Texture);
+			FMetalSurface& Surface = *GetMetalSurfaceFromRHITexture(RenderPassInfo.DepthStencilRenderTarget.DepthStencilTarget);
 			
 			switch(Surface.Type)
 			{
@@ -669,7 +673,7 @@ bool FMetalStateCache::SetRenderTargetsInfo(FRHISetRenderTargetsInfo const& InRe
 				FrameBufferSize.height = FMath::Min(FrameBufferSize.height, (CGFloat)Surface.SizeY);
 			}
 			
-			EPixelFormat DepthStencilPixelFormat = RenderTargetsInfo.DepthStencilRenderTarget.Texture->GetFormat();
+			EPixelFormat DepthStencilPixelFormat = RenderPassInfo.DepthStencilRenderTarget.DepthStencilTarget->GetFormat();
 			
 			FMetalTexture DepthTexture = nil;
 			FMetalTexture StencilTexture = nil;
@@ -739,7 +743,7 @@ bool FMetalStateCache::SetRenderTargetsInfo(FRHISetRenderTargetsInfo const& InRe
 			
 			float DepthClearValue = 0.0f;
 			uint32 StencilClearValue = 0;
-			const FClearValueBinding& ClearValue = RenderTargetsInfo.DepthStencilRenderTarget.Texture->GetClearBinding();
+			const FClearValueBinding& ClearValue = RenderPassInfo.DepthStencilRenderTarget.DepthStencilTarget->GetClearBinding();
 			if (ClearValue.ColorBinding == EClearBinding::EDepthStencilBound)
 			{
 				ClearValue.GetDepthStencil(DepthClearValue, StencilClearValue);
@@ -751,21 +755,25 @@ bool FMetalStateCache::SetRenderTargetsInfo(FRHISetRenderTargetsInfo const& InRe
 
             static bool const bUsingValidation = FMetalCommandQueue::SupportsFeature(EMetalFeaturesValidation) && !FApplePlatformMisc::IsOSAtLeastVersion((uint32[]){10, 14, 0}, (uint32[]){12, 0, 0}, (uint32[]){12, 0, 0});
             
-            bool const bCombinedDepthStencilUsingStencil = (DepthTexture && (mtlpp::PixelFormat)DepthTexture.GetPixelFormat() != mtlpp::PixelFormat::Depth32Float && RenderTargetsInfo.DepthStencilRenderTarget.GetDepthStencilAccess().IsUsingStencil());			
-			bool const bUsingDepth = (RenderTargetsInfo.DepthStencilRenderTarget.GetDepthStencilAccess().IsUsingDepth() || (bUsingValidation && bCombinedDepthStencilUsingStencil));
+            bool const bCombinedDepthStencilUsingStencil = (DepthTexture && (mtlpp::PixelFormat)DepthTexture.GetPixelFormat() != mtlpp::PixelFormat::Depth32Float && RenderPassInfo.DepthStencilRenderTarget.ExclusiveDepthStencil.IsUsingStencil());
+			bool const bUsingDepth = (RenderPassInfo.DepthStencilRenderTarget.ExclusiveDepthStencil.IsUsingDepth() || (bUsingValidation && bCombinedDepthStencilUsingStencil));
 			if (DepthTexture && bUsingDepth)
 			{
 				mtlpp::RenderPassDepthAttachmentDescriptor DepthAttachment;
 				
 				DepthFormatKey = Surface.FormatKey;
-	
+				
+				ERenderTargetActions DepthActions = GetDepthActions(RenderPassInfo.DepthStencilRenderTarget.Action);
+				ERenderTargetLoadAction DepthLoadAction = GetLoadAction(DepthActions);
+				ERenderTargetStoreAction DepthStoreAction = GetStoreAction(DepthActions);
+
 				// set up the depth attachment
 				DepthAttachment.SetTexture(DepthTexture);
-				DepthAttachment.SetLoadAction(GetMetalRTLoadAction(RenderTargetsInfo.DepthStencilRenderTarget.DepthLoadAction));
+				DepthAttachment.SetLoadAction(GetMetalRTLoadAction(DepthLoadAction));
 				
 				bNeedsClear |= (DepthAttachment.GetLoadAction() == mtlpp::LoadAction::Clear);
 				
-				ERenderTargetStoreAction HighLevelStoreAction = (Surface.MSAATexture && !bDepthStencilSampleCountMismatchFixup) ? ERenderTargetStoreAction::EMultisampleResolve : RenderTargetsInfo.DepthStencilRenderTarget.DepthStoreAction;
+				ERenderTargetStoreAction HighLevelStoreAction = (Surface.MSAATexture && !bDepthStencilSampleCountMismatchFixup) ? ERenderTargetStoreAction::EMultisampleResolve : DepthStoreAction;
 				if (bUsingDepth && (HighLevelStoreAction == ERenderTargetStoreAction::ENoAction || bDepthStencilSampleCountMismatchFixup))
 				{
 					if (DepthSampleCount > 1)
@@ -818,9 +826,9 @@ bool FMetalStateCache::SetRenderTargetsInfo(FRHISetRenderTargetsInfo const& InRe
 				}
 				
 				bHasValidRenderTarget = true;
-				bFallbackDepthStencilBound = (RenderTargetsInfo.DepthStencilRenderTarget.Texture == FallbackDepthStencilSurface);
+				bFallbackDepthStencilBound = (RenderPassInfo.DepthStencilRenderTarget.DepthStencilTarget == FallbackDepthStencilSurface);
 
-				bCanRestartRenderPass &= (SampleCount <= 1) && ((RenderTargetsInfo.DepthStencilRenderTarget.Texture == FallbackDepthStencilSurface) || ((DepthAttachment.GetLoadAction() == mtlpp::LoadAction::Load) && (!RenderTargetsInfo.DepthStencilRenderTarget.GetDepthStencilAccess().IsDepthWrite() || (RenderTargetsInfo.DepthStencilRenderTarget.DepthStoreAction == ERenderTargetStoreAction::EStore))));
+				bCanRestartRenderPass &= (SampleCount <= 1) && ((RenderPassInfo.DepthStencilRenderTarget.DepthStencilTarget == FallbackDepthStencilSurface) || ((DepthAttachment.GetLoadAction() == mtlpp::LoadAction::Load) && (!RenderPassInfo.DepthStencilRenderTarget.ExclusiveDepthStencil.IsDepthWrite() || (DepthStoreAction == ERenderTargetStoreAction::EStore))));
 				
 				// and assign it
 				RenderPass.SetDepthAttachment(DepthAttachment);
@@ -829,21 +837,25 @@ bool FMetalStateCache::SetRenderTargetsInfo(FRHISetRenderTargetsInfo const& InRe
             //if we're dealing with a samplecount mismatch we just bail on stencil entirely as stencil
             //doesn't have an autoresolve target to use.
 			
-			bool const bCombinedDepthStencilUsingDepth = (StencilTexture && StencilTexture.GetPixelFormat() != mtlpp::PixelFormat::Stencil8 && RenderTargetsInfo.DepthStencilRenderTarget.GetDepthStencilAccess().IsUsingDepth());
-			bool const bUsingStencil = RenderTargetsInfo.DepthStencilRenderTarget.GetDepthStencilAccess().IsUsingStencil() || (bUsingValidation && bCombinedDepthStencilUsingDepth);
+			bool const bCombinedDepthStencilUsingDepth = (StencilTexture && StencilTexture.GetPixelFormat() != mtlpp::PixelFormat::Stencil8 && RenderPassInfo.DepthStencilRenderTarget.ExclusiveDepthStencil.IsUsingDepth());
+			bool const bUsingStencil = RenderPassInfo.DepthStencilRenderTarget.ExclusiveDepthStencil.IsUsingStencil() || (bUsingValidation && bCombinedDepthStencilUsingDepth);
 			if (StencilTexture && bUsingStencil)
 			{
 				mtlpp::RenderPassStencilAttachmentDescriptor StencilAttachment;
 				
 				StencilFormatKey = Surface.FormatKey;
+				
+				ERenderTargetActions StencilActions = GetStencilActions(RenderPassInfo.DepthStencilRenderTarget.Action);
+				ERenderTargetLoadAction StencilLoadAction = GetLoadAction(StencilActions);
+				ERenderTargetStoreAction StencilStoreAction = GetStoreAction(StencilActions);
 	
 				// set up the stencil attachment
 				StencilAttachment.SetTexture(StencilTexture);
-				StencilAttachment.SetLoadAction(GetMetalRTLoadAction(RenderTargetsInfo.DepthStencilRenderTarget.StencilLoadAction));
+				StencilAttachment.SetLoadAction(GetMetalRTLoadAction(StencilLoadAction));
 				
 				bNeedsClear |= (StencilAttachment.GetLoadAction() == mtlpp::LoadAction::Clear);
 				
-				ERenderTargetStoreAction HighLevelStoreAction = RenderTargetsInfo.DepthStencilRenderTarget.GetStencilStoreAction();
+				ERenderTargetStoreAction HighLevelStoreAction = StencilStoreAction;
 				if (bUsingStencil && (HighLevelStoreAction == ERenderTargetStoreAction::ENoAction || bDepthStencilSampleCountMismatchFixup))
 				{
 					HighLevelStoreAction = ERenderTargetStoreAction::EStore;
@@ -872,7 +884,7 @@ bool FMetalStateCache::SetRenderTargetsInfo(FRHISetRenderTargetsInfo const& InRe
 				
 				// @todo Stencil writes that need to persist must use ERenderTargetStoreAction::EStore on iOS.
 				// We should probably be using deferred store actions so that we can safely lazily instantiate encoders.
-				bCanRestartRenderPass &= (SampleCount <= 1) && ((RenderTargetsInfo.DepthStencilRenderTarget.Texture == FallbackDepthStencilSurface) || ((StencilAttachment.GetLoadAction() == mtlpp::LoadAction::Load) && (1 || !RenderTargetsInfo.DepthStencilRenderTarget.GetDepthStencilAccess().IsStencilWrite() || (RenderTargetsInfo.DepthStencilRenderTarget.GetStencilStoreAction() == ERenderTargetStoreAction::EStore))));
+				bCanRestartRenderPass &= (SampleCount <= 1) && ((RenderPassInfo.DepthStencilRenderTarget.DepthStencilTarget == FallbackDepthStencilSurface) || ((StencilAttachment.GetLoadAction() == mtlpp::LoadAction::Load) && (1 || !RenderPassInfo.DepthStencilRenderTarget.ExclusiveDepthStencil.IsStencilWrite() || (StencilStoreAction == ERenderTargetStoreAction::EStore))));
 				
 				// and assign it
 				RenderPass.SetStencilAttachment(StencilAttachment);
@@ -901,11 +913,13 @@ bool FMetalStateCache::SetRenderTargetsInfo(FRHISetRenderTargetsInfo const& InRe
 		if (bHasValidRenderTarget)
 		{
 			// Retain and/or release the depth-stencil surface in case it is a temporary surface for a draw call that writes to depth without a depth/stencil buffer bound.
-			DepthStencilSurface = RenderTargetsInfo.DepthStencilRenderTarget.Texture;
+			DepthStencilSurface = RenderPassInfo.DepthStencilRenderTarget.DepthStencilTarget;
+			DepthStencilResolve = RenderPassInfo.DepthStencilRenderTarget.ResolveTarget;
 		}
 		else
 		{
 			DepthStencilSurface.SafeRelease();
+			DepthStencilResolve.SafeRelease();
 		}
 		
 		RenderPassDesc = RenderPass;
@@ -1138,11 +1152,11 @@ void FMetalStateCache::SetGraphicsPipelineState(FMetalGraphicsPipelineState* Sta
 		}
 		
 		{
-			for (uint32 i = 0; i < RenderTargetsInfo.NumUAVs; i++)
+			for (uint32 i = 0; i < RenderPassInfo.NumUAVs; i++)
 			{
-				if (IsValidRef(RenderTargetsInfo.UnorderedAccessView[i]))
+				if (IsValidRef(RenderPassInfo.UAVs[i]))
 				{
-					FMetalUnorderedAccessView* UAV = ResourceCast(RenderTargetsInfo.UnorderedAccessView[i].GetReference());
+					FMetalUnorderedAccessView* UAV = ResourceCast(RenderPassInfo.UAVs[i].GetReference());
 					SetShaderUnorderedAccessView(SF_Pixel, i, UAV);
 				}
 			}
@@ -1204,26 +1218,29 @@ void FMetalStateCache::ConditionalUpdateBackBuffer(FMetalSurface& Surface)
 	}
 }
 
-bool FMetalStateCache::NeedsToSetRenderTarget(const FRHISetRenderTargetsInfo& InRenderTargetsInfo)
+bool FMetalStateCache::NeedsToSetRenderTarget(const FRHIRenderPassInfo& InRenderPassInfo)
 {
 	// see if our new Info matches our previous Info
+	uint32 CurrentNumColorRenderTargets = RenderPassInfo.GetNumColorRenderTargets();
+	uint32 NewNumColorRenderTargets = InRenderPassInfo.GetNumColorRenderTargets();
 	
 	// basic checks
-	bool bAllChecksPassed = GetHasValidRenderTarget() && bIsRenderTargetActive && InRenderTargetsInfo.NumColorRenderTargets == RenderTargetsInfo.NumColorRenderTargets && InRenderTargetsInfo.NumUAVs == RenderTargetsInfo.NumUAVs &&
-		(InRenderTargetsInfo.DepthStencilRenderTarget.Texture == RenderTargetsInfo.DepthStencilRenderTarget.Texture);
+	bool bAllChecksPassed = GetHasValidRenderTarget() && bIsRenderTargetActive && CurrentNumColorRenderTargets == NewNumColorRenderTargets && InRenderPassInfo.NumUAVs == RenderPassInfo.NumUAVs &&
+		(InRenderPassInfo.DepthStencilRenderTarget.DepthStencilTarget == RenderPassInfo.DepthStencilRenderTarget.DepthStencilTarget);
 
 	// now check each color target if the basic tests passe
 	if (bAllChecksPassed)
 	{
-		for (int32 RenderTargetIndex = 0; RenderTargetIndex < InRenderTargetsInfo.NumColorRenderTargets; RenderTargetIndex++)
+		for (int32 RenderTargetIndex = 0; RenderTargetIndex < NewNumColorRenderTargets; RenderTargetIndex++)
 		{
-			const FRHIRenderTargetView& RenderTargetView = InRenderTargetsInfo.ColorRenderTarget[RenderTargetIndex];
-			const FRHIRenderTargetView& PreviousRenderTargetView = RenderTargetsInfo.ColorRenderTarget[RenderTargetIndex];
+			const FRHIRenderPassInfo::FColorEntry& RenderTargetView = InRenderPassInfo.ColorRenderTargets[RenderTargetIndex];
+			const FRHIRenderPassInfo::FColorEntry& PreviousRenderTargetView = RenderPassInfo.ColorRenderTargets[RenderTargetIndex];
 
 			// handle simple case of switching textures or mip/slice
-			if (RenderTargetView.Texture != PreviousRenderTargetView.Texture ||
+			if (RenderTargetView.RenderTarget != PreviousRenderTargetView.RenderTarget ||
+				RenderTargetView.ResolveTarget != PreviousRenderTargetView.ResolveTarget ||
 				RenderTargetView.MipIndex != PreviousRenderTargetView.MipIndex ||
-				RenderTargetView.ArraySliceIndex != PreviousRenderTargetView.ArraySliceIndex)
+				RenderTargetView.ArraySlice != PreviousRenderTargetView.ArraySlice)
 			{
 				bAllChecksPassed = false;
 				break;
@@ -1235,7 +1252,7 @@ bool FMetalStateCache::NeedsToSetRenderTarget(const FRHISetRenderTargetsInfo& In
 			//    If we switch to Clear, we have to always switch to a new RT to force the clear
 			//    If we switch to DontCare, there's definitely no need to switch
 			//    If we switch *from* Clear then we must change target as we *don't* want to clear again.
-            if (RenderTargetView.LoadAction == ERenderTargetLoadAction::EClear)
+            if (GetLoadAction(RenderTargetView.Action) == ERenderTargetLoadAction::EClear)
             {
                 bAllChecksPassed = false;
                 break;
@@ -1252,26 +1269,26 @@ bool FMetalStateCache::NeedsToSetRenderTarget(const FRHISetRenderTargetsInfo& In
             //			}
         }
         
-        if (InRenderTargetsInfo.DepthStencilRenderTarget.Texture && (InRenderTargetsInfo.DepthStencilRenderTarget.DepthLoadAction == ERenderTargetLoadAction::EClear || InRenderTargetsInfo.DepthStencilRenderTarget.StencilLoadAction == ERenderTargetLoadAction::EClear))
+        if (InRenderPassInfo.DepthStencilRenderTarget.DepthStencilTarget && (GetLoadAction(GetDepthActions(InRenderPassInfo.DepthStencilRenderTarget.Action)) == ERenderTargetLoadAction::EClear || GetLoadAction(GetStencilActions(InRenderPassInfo.DepthStencilRenderTarget.Action)) == ERenderTargetLoadAction::EClear))
         {
             bAllChecksPassed = false;
 		}
 		
-		if (InRenderTargetsInfo.DepthStencilRenderTarget.Texture && (InRenderTargetsInfo.DepthStencilRenderTarget.DepthStoreAction > RenderTargetsInfo.DepthStencilRenderTarget.DepthStoreAction || InRenderTargetsInfo.DepthStencilRenderTarget.GetStencilStoreAction() > RenderTargetsInfo.DepthStencilRenderTarget.GetStencilStoreAction()))
+		if (InRenderPassInfo.DepthStencilRenderTarget.DepthStencilTarget && (GetStoreAction(GetDepthActions(InRenderPassInfo.DepthStencilRenderTarget.Action)) > GetStoreAction(GetDepthActions(RenderPassInfo.DepthStencilRenderTarget.Action)) || GetStoreAction(GetStencilActions(InRenderPassInfo.DepthStencilRenderTarget.Action)) > GetStoreAction(GetStencilActions(RenderPassInfo.DepthStencilRenderTarget.Action))))
 		{
 			// Don't break the encoder if we can just change the store actions.
 			mtlpp::StoreAction NewDepthStore = DepthStore;
 			mtlpp::StoreAction NewStencilStore = StencilStore;
-			if (InRenderTargetsInfo.DepthStencilRenderTarget.DepthStoreAction > RenderTargetsInfo.DepthStencilRenderTarget.DepthStoreAction)
+			if (GetStoreAction(GetDepthActions(InRenderPassInfo.DepthStencilRenderTarget.Action)) > GetStoreAction(GetDepthActions(RenderPassInfo.DepthStencilRenderTarget.Action)))
 			{
 				if (RenderPassDesc.GetDepthAttachment().GetTexture())
 				{
-					FMetalSurface& Surface = *GetMetalSurfaceFromRHITexture(RenderTargetsInfo.DepthStencilRenderTarget.Texture);
+					FMetalSurface& Surface = *GetMetalSurfaceFromRHITexture(RenderPassInfo.DepthStencilRenderTarget.DepthStencilTarget);
 					
 					const uint32 DepthSampleCount = (Surface.MSAATexture ? Surface.MSAATexture.GetSampleCount() : Surface.Texture.GetSampleCount());
 					bool const bDepthStencilSampleCountMismatchFixup = (SampleCount != DepthSampleCount);
 
-					ERenderTargetStoreAction HighLevelStoreAction = (Surface.MSAATexture && !bDepthStencilSampleCountMismatchFixup) ? ERenderTargetStoreAction::EMultisampleResolve : RenderTargetsInfo.DepthStencilRenderTarget.DepthStoreAction;
+					ERenderTargetStoreAction HighLevelStoreAction = (Surface.MSAATexture && !bDepthStencilSampleCountMismatchFixup) ? ERenderTargetStoreAction::EMultisampleResolve : GetStoreAction(GetDepthActions(RenderPassInfo.DepthStencilRenderTarget.Action));
 					
 #if PLATFORM_IOS
 					FMetalTexture& Tex = Surface.MSAATexture ? Surface.MSAATexture : Surface.Texture;
@@ -1289,11 +1306,11 @@ bool FMetalStateCache::NeedsToSetRenderTarget(const FRHISetRenderTargetsInfo& In
 				}
 			}
 			
-			if (InRenderTargetsInfo.DepthStencilRenderTarget.GetStencilStoreAction() > RenderTargetsInfo.DepthStencilRenderTarget.GetStencilStoreAction())
+			if (GetStoreAction(GetStencilActions(InRenderPassInfo.DepthStencilRenderTarget.Action)) > GetStoreAction(GetStencilActions(RenderPassInfo.DepthStencilRenderTarget.Action)))
 			{
 				if (RenderPassDesc.GetStencilAttachment().GetTexture())
 				{
-					NewStencilStore = GetMetalRTStoreAction(RenderTargetsInfo.DepthStencilRenderTarget.GetStencilStoreAction());
+					NewStencilStore = GetMetalRTStoreAction(GetStoreAction(GetStencilActions(RenderPassInfo.DepthStencilRenderTarget.Action)));
 #if PLATFORM_IOS
 					if (RenderPassDesc.GetStencilAttachment().GetTexture().GetStorageMode() == mtlpp::StorageMode::Memoryless)
 					{
@@ -1317,8 +1334,8 @@ bool FMetalStateCache::NeedsToSetRenderTarget(const FRHISetRenderTargetsInfo& In
 
 	// if we are setting them to nothing, then this is probably end of frame, and we can't make a framebuffer
 	// with nothng, so just abort this (only need to check on single MRT case)
-	if (InRenderTargetsInfo.NumColorRenderTargets == 1 && InRenderTargetsInfo.ColorRenderTarget[0].Texture == nullptr &&
-		InRenderTargetsInfo.DepthStencilRenderTarget.Texture == nullptr)
+	if (NewNumColorRenderTargets == 1 && InRenderPassInfo.ColorRenderTargets[0].RenderTarget == nullptr &&
+		InRenderPassInfo.DepthStencilRenderTarget.DepthStencilTarget == nullptr)
 	{
 		bAllChecksPassed = true;
 	}
@@ -1835,59 +1852,60 @@ bool FMetalStateCache::PrepareToRestart(bool const bCurrentApplied)
 	{
 		if (SampleCount <= 1)
 		{
-			FRHISetRenderTargetsInfo Info = GetRenderTargetsInfo();
+			FRHIRenderPassInfo Info = GetRenderPassInfo();
 			
-			ERenderTargetLoadAction DepthLoadAction = Info.DepthStencilRenderTarget.DepthLoadAction;
-			ERenderTargetStoreAction DepthStoreAction = Info.DepthStencilRenderTarget.DepthStoreAction;
-			ERenderTargetLoadAction StencilLoadAction = Info.DepthStencilRenderTarget.StencilLoadAction;
-			ERenderTargetStoreAction StencilStoreAction = Info.DepthStencilRenderTarget.GetStencilStoreAction();
-			bool bClearDepth = Info.bClearDepth;
-			bool bClearStencil = Info.bClearStencil;
-
-			if (Info.DepthStencilRenderTarget.Texture)
+			ERenderTargetActions DepthActions = GetDepthActions(Info.DepthStencilRenderTarget.Action);
+			ERenderTargetActions StencilActions = GetStencilActions(Info.DepthStencilRenderTarget.Action);
+			ERenderTargetLoadAction DepthLoadAction = GetLoadAction(DepthActions);
+			ERenderTargetStoreAction DepthStoreAction = GetStoreAction(DepthActions);
+			ERenderTargetLoadAction StencilLoadAction = GetLoadAction(StencilActions);
+			ERenderTargetStoreAction StencilStoreAction = GetStoreAction(StencilActions);
+
+			if (Info.DepthStencilRenderTarget.DepthStencilTarget)
 			{
-				if (bCurrentApplied || Info.DepthStencilRenderTarget.DepthLoadAction != ERenderTargetLoadAction::EClear)
+				if (bCurrentApplied || DepthLoadAction != ERenderTargetLoadAction::EClear)
 				{
 					DepthLoadAction = ERenderTargetLoadAction::ELoad;
-					bClearDepth = false;
 				}
-				if (Info.DepthStencilRenderTarget.GetDepthStencilAccess().IsDepthWrite())
+				if (Info.DepthStencilRenderTarget.ExclusiveDepthStencil.IsDepthWrite())
 				{
 					DepthStoreAction = ERenderTargetStoreAction::EStore;
 				}
 
-				if (bCurrentApplied || Info.DepthStencilRenderTarget.StencilLoadAction != ERenderTargetLoadAction::EClear)
+				if (bCurrentApplied || StencilLoadAction != ERenderTargetLoadAction::EClear)
 				{
 					StencilLoadAction = ERenderTargetLoadAction::ELoad;
-					bClearStencil = false;
 				}
-				if (Info.DepthStencilRenderTarget.GetDepthStencilAccess().IsStencilWrite())
+				if (Info.DepthStencilRenderTarget.ExclusiveDepthStencil.IsStencilWrite())
 				{
 					StencilStoreAction = ERenderTargetStoreAction::EStore;
 				}
 				
-				Info.DepthStencilRenderTarget = FRHIDepthRenderTargetView(Info.DepthStencilRenderTarget.Texture, DepthLoadAction, DepthStoreAction, StencilLoadAction, StencilStoreAction, Info.DepthStencilRenderTarget.GetDepthStencilAccess());
-				Info.bClearDepth = bClearDepth;
-				Info.bClearStencil = bClearStencil;
+				DepthActions = MakeRenderTargetActions(DepthLoadAction, DepthStoreAction);
+				StencilActions = MakeRenderTargetActions(StencilLoadAction, StencilStoreAction);
+				Info.DepthStencilRenderTarget.Action = MakeDepthStencilTargetActions(DepthActions, StencilActions);
 			}
 			
-			for (int32 RenderTargetIndex = 0; RenderTargetIndex < Info.NumColorRenderTargets; RenderTargetIndex++)
+			for (int32 RenderTargetIndex = 0; RenderTargetIndex < Info.GetNumColorRenderTargets(); RenderTargetIndex++)
 			{
-				FRHIRenderTargetView& RenderTargetView = Info.ColorRenderTarget[RenderTargetIndex];
-				if (!bCurrentApplied && RenderTargetView.LoadAction == ERenderTargetLoadAction::EClear)
+				FRHIRenderPassInfo::FColorEntry& RenderTargetView = Info.ColorRenderTargets[RenderTargetIndex];
+				ERenderTargetLoadAction LoadAction = GetLoadAction(RenderTargetView.Action);
+				ERenderTargetStoreAction StoreAction = GetStoreAction(RenderTargetView.Action);
+				
+				if (!bCurrentApplied && LoadAction == ERenderTargetLoadAction::EClear)
 				{
-					RenderTargetView.StoreAction == ERenderTargetStoreAction::EStore;
+					StoreAction == ERenderTargetStoreAction::EStore;
 				}
 				else
 				{
-					RenderTargetView.LoadAction = ERenderTargetLoadAction::ELoad;
-					Info.bClearColor = false;
+					LoadAction = ERenderTargetLoadAction::ELoad;
 				}
-				check(RenderTargetView.Texture == nil || RenderTargetView.StoreAction == ERenderTargetStoreAction::EStore);
+				RenderTargetView.Action = MakeRenderTargetActions(LoadAction, StoreAction);
+				check(RenderTargetView.RenderTarget == nil || GetStoreAction(RenderTargetView.Action) == ERenderTargetStoreAction::EStore);
 			}
 			
 			InvalidateRenderTargets();
-			return SetRenderTargetsInfo(Info, GetVisibilityResultsBuffer(), true) && CanRestartRenderPass();
+			return SetRenderPassInfo(Info, GetVisibilityResultsBuffer(), true) && CanRestartRenderPass();
 		}
 		else
 		{
@@ -1922,17 +1940,17 @@ void FMetalStateCache::SetRenderStoreActions(FMetalCommandEncoder& CommandEncode
 		if (bConditionalSwitch)
 		{
 			ns::Array<mtlpp::RenderPassColorAttachmentDescriptor> ColorAttachments = RenderPassDesc.GetColorAttachments();
-			for (int32 RenderTargetIndex = 0; RenderTargetIndex < RenderTargetsInfo.NumColorRenderTargets; RenderTargetIndex++)
+			for (int32 RenderTargetIndex = 0; RenderTargetIndex < RenderPassInfo.GetNumColorRenderTargets(); RenderTargetIndex++)
 			{
-				FRHIRenderTargetView& RenderTargetView = RenderTargetsInfo.ColorRenderTarget[RenderTargetIndex];
-				if(RenderTargetView.Texture != nil)
+				FRHIRenderPassInfo::FColorEntry& RenderTargetView = RenderPassInfo.ColorRenderTargets[RenderTargetIndex];
+				if(RenderTargetView.RenderTarget != nil)
 				{
 					const bool bMultiSampled = (ColorAttachments[RenderTargetIndex].GetTexture().GetSampleCount() > 1);
 					ColorStore[RenderTargetIndex] = GetConditionalMetalRTStoreAction(bMultiSampled);
 				}
 			}
 			
-			if (RenderTargetsInfo.DepthStencilRenderTarget.Texture)
+			if (RenderPassInfo.DepthStencilRenderTarget.DepthStencilTarget)
 			{
 				const bool bMultiSampled = RenderPassDesc.GetDepthAttachment().GetTexture() && (RenderPassDesc.GetDepthAttachment().GetTexture().GetSampleCount() > 1);
 				DepthStore = GetConditionalMetalRTStoreAction(bMultiSampled);
diff --git a/Engine/Source/Runtime/Apple/MetalRHI/Private/MetalStateCache.h b/Engine/Source/Runtime/Apple/MetalRHI/Private/MetalStateCache.h
--- a/Engine/Source/Runtime/Apple/MetalRHI/Private/MetalStateCache.h
+++ b/Engine/Source/Runtime/Apple/MetalRHI/Private/MetalStateCache.h
@@ -43,7 +43,7 @@ public:
 	void SetBlendFactor(FLinearColor const& InBlendFactor);
 	void SetStencilRef(uint32 const InStencilRef);
 	void SetComputeShader(FMetalComputeShader* InComputeShader);
-	bool SetRenderTargetsInfo(FRHISetRenderTargetsInfo const& InRenderTargets, FMetalQueryBuffer* QueryBuffer, bool const bRestart);
+	bool SetRenderPassInfo(FRHIRenderPassInfo const& InRenderTargets, FMetalQueryBuffer* QueryBuffer, bool const bRestart);
 	void InvalidateRenderTargets(void);
 	void SetRenderTargetsActive(bool const bActive);
 	void SetViewport(const mtlpp::Viewport& InViewport);
@@ -120,8 +120,8 @@ public:
 	FMetalGraphicsPipelineState* GetGraphicsPSO() const { return GraphicsPSO; }
 	FMetalComputeShader* GetComputeShader() const { return ComputeShader; }
 	CGSize GetFrameBufferSize() const { return FrameBufferSize; }
-	FRHISetRenderTargetsInfo const& GetRenderTargetsInfo() const { return RenderTargetsInfo; }
-	int32 GetNumRenderTargets() { return bHasValidColorTarget ? RenderTargetsInfo.NumColorRenderTargets : -1; }
+	FRHIRenderPassInfo const& GetRenderPassInfo() const { return RenderPassInfo; }
+	int32 GetNumRenderTargets() { return bHasValidColorTarget ? RenderPassInfo.GetNumColorRenderTargets() : -1; }
 	bool GetHasValidRenderTarget() const { return bHasValidRenderTarget; }
 	bool GetHasValidColorTarget() const { return bHasValidColorTarget; }
 	const mtlpp::Viewport& GetViewport(uint32 const Index) const { check(Index < ML_MaxViewports); return Viewport[Index]; }
@@ -131,7 +131,7 @@ public:
 	uint32 GetDirtyUniformBuffers(EShaderFrequency const Freq) const { return DirtyUniformBuffers[Freq]; }
 	FMetalQueryBuffer* GetVisibilityResultsBuffer() const { return VisibilityResults; }
 	bool GetScissorRectEnabled() const { return bScissorRectEnabled; }
-	bool NeedsToSetRenderTarget(const FRHISetRenderTargetsInfo& RenderTargetsInfo);
+	bool NeedsToSetRenderTarget(const FRHIRenderPassInfo& RenderPassInfo);
 	bool HasValidDepthStencilSurface() const { return IsValidRef(DepthStencilSurface); }
 	EMetalIndexType GetIndexType() const { return IndexType; }
     bool GetUsingTessellation() const { return bUsingTessellation; }
@@ -270,9 +270,11 @@ private:
 	uint32 ActiveViewports;
 	uint32 ActiveScissors;
 	
-	FRHISetRenderTargetsInfo RenderTargetsInfo;
+	FRHIRenderPassInfo RenderPassInfo;
 	FTextureRHIRef ColorTargets[MaxSimultaneousRenderTargets];
+	FTextureRHIRef ResolveTargets[MaxSimultaneousRenderTargets];
 	FTextureRHIRef DepthStencilSurface;
+	FTextureRHIRef DepthStencilResolve;
 	/** A fallback depth-stencil surface for draw calls that write to depth without a depth-stencil surface bound. */
 	FTexture2DRHIRef FallbackDepthStencilSurface;
 	mtlpp::RenderPassDescriptor RenderPassDesc;
