Support separate MSAA and resolve targets on Metal MRT for iOS devices - this means the Forward renderer can support MSAA the same way as on macOS. It isn't optimal for iOS GPUs but it at least gets it to run without validation errors - though it doesn't yet render correctly. However now that Metal is using the RenderPassInfo descriptor we can start making changes to fix and improve it.

diff --git a/Engine/Source/Runtime/Apple/MetalRHI/Private/MetalCommandQueue.h b/Engine/Source/Runtime/Apple/MetalRHI/Private/MetalCommandQueue.h
--- a/Engine/Source/Runtime/Apple/MetalRHI/Private/MetalCommandQueue.h
+++ b/Engine/Source/Runtime/Apple/MetalRHI/Private/MetalCommandQueue.h
@@ -137,7 +137,7 @@ public:
 	* @returns True if RHISupportsSeparateMSAAAndResolveTextures will be true.  
 	* Currently Mac only.
 	*/
-	static inline bool SupportsSeparateMSAAAndResolveTarget() { return PLATFORM_MAC != 0; }
+	static inline bool SupportsSeparateMSAAAndResolveTarget() { return (PLATFORM_MAC != 0 || GMaxRHIFeatureLevel >= ERHIFeatureLevel::SM5); }
 
 #pragma mark - Public Debug Support -
 
diff --git a/Engine/Source/Runtime/Apple/MetalRHI/Private/MetalTexture.cpp b/Engine/Source/Runtime/Apple/MetalRHI/Private/MetalTexture.cpp
--- a/Engine/Source/Runtime/Apple/MetalRHI/Private/MetalTexture.cpp
+++ b/Engine/Source/Runtime/Apple/MetalRHI/Private/MetalTexture.cpp
@@ -799,7 +799,7 @@ FMetalSurface::FMetalSurface(ERHIResourceType ResourceType, EPixelFormat Format,
 		{
 			check(!(Flags & TexCreate_CPUReadback));
 #if PLATFORM_IOS
-			if (!(Flags & (TexCreate_ShaderResource|TexCreate_UAV)))
+			if (!(Flags & (TexCreate_ShaderResource|TexCreate_UAV)) && (GMaxRHIFeatureLevel < ERHIFeatureLevel::SM5))
 			{
 				Desc.SetStorageMode(mtlpp::StorageMode::Memoryless);
 				Desc.SetResourceOptions(mtlpp::ResourceOptions::StorageModeMemoryless);
@@ -936,7 +936,7 @@ FMetalSurface::FMetalSurface(ERHIResourceType ResourceType, EPixelFormat Format,
 			
 			bool bMemoryless = false;
 #if PLATFORM_IOS
-			if (GMaxRHIShaderPlatform == SP_METAL)
+			if (GMaxRHIFeatureLevel < ERHIFeatureLevel::SM5)
 			{
 				bMemoryless = true;
 				Desc.SetStorageMode(mtlpp::StorageMode::Memoryless);
diff --git a/Engine/Source/Runtime/RHI/Public/RHI.h b/Engine/Source/Runtime/RHI/Public/RHI.h
--- a/Engine/Source/Runtime/RHI/Public/RHI.h
+++ b/Engine/Source/Runtime/RHI/Public/RHI.h
@@ -141,12 +141,6 @@ inline bool RHISupportsMSAA(EShaderPlatform Platform)
 	return 
 		//@todo-rco: Fix when iOS OpenGL supports MSAA
 		Platform != SP_OPENGL_ES2_IOS
-		// @todo marksatt Metal on macOS 10.12 and earlier (or Intel on any macOS < 10.13.2) don't reliably support our MSAA usage & custom resolve.
-#if PLATFORM_MAC
-		&& IsMetalPlatform(Platform) && (FPlatformMisc::MacOSXVersionCompare(10, 13, 0) >= 0) && (!IsRHIDeviceIntel() || FPlatformMisc::MacOSXVersionCompare(10, 13, 2) >= 0)
-#endif
-		// @todo marksatt iOS Desktop Forward needs more work internally
-		&& Platform != SP_METAL_MRT
 		// @todo optimise MSAA for XboxOne, currently uses significant eRAM.
 		&& Platform != SP_XBOXONE_D3D12;
 }
diff --git a/Engine/Source/Runtime/RHI/Public/RHIDefinitions.h b/Engine/Source/Runtime/RHI/Public/RHIDefinitions.h
--- a/Engine/Source/Runtime/RHI/Public/RHIDefinitions.h
+++ b/Engine/Source/Runtime/RHI/Public/RHIDefinitions.h
@@ -1072,7 +1072,7 @@ inline bool RHINeedsToSwitchVerticalAxis(EShaderPlatform Platform)
 inline bool RHISupportsSeparateMSAAAndResolveTextures(const EShaderPlatform Platform)
 {
 	// Metal mobile devices, Vulkan and Android ES2/3.1 need to handle MSAA and resolve textures internally (unless RHICreateTexture2D was changed to take an optional resolve target)
-	const bool bMobileMetalDevice = (Platform == SP_METAL || Platform == SP_METAL_TVOS || Platform == SP_METAL_MRT || Platform == SP_METAL_MRT_TVOS);
+	const bool bMobileMetalDevice = (Platform == SP_METAL || Platform == SP_METAL_TVOS);
 	return !bMobileMetalDevice && !IsVulkanPlatform(Platform) && !IsAndroidOpenGLESPlatform(Platform);
 }
 
