Dupe #4891942
Expose support for CopyTexture in MetalRHI.

diff --git a/Engine/Source/Runtime/Apple/MetalRHI/Private/MetalRHI.cpp b/Engine/Source/Runtime/Apple/MetalRHI/Private/MetalRHI.cpp
--- a/Engine/Source/Runtime/Apple/MetalRHI/Private/MetalRHI.cpp
+++ b/Engine/Source/Runtime/Apple/MetalRHI/Private/MetalRHI.cpp
@@ -343,6 +343,8 @@ FMetalDynamicRHI::FMetalDynamicRHI(ERHIFeatureLevel::Type RequestedFeatureLevel)
 	
 #endif
 		
+	GRHISupportsCopyToTextureMultipleMips = true;
+		
 	if(
 	   #if PLATFORM_MAC
 	   (Device.SupportsFeatureSet(mtlpp::FeatureSet::macOS_GPUFamily1_v3) && FPlatformMisc::MacOSXVersionCompare(10,13,0) >= 0)
