Add R32G32_UINT format
This will be required for some incoming work usng runtime texture compression

diff --git a/Engine/Source/Runtime/Apple/MetalRHI/Private/MetalRHI.cpp b/Engine/Source/Runtime/Apple/MetalRHI/Private/MetalRHI.cpp
--- a/Engine/Source/Runtime/Apple/MetalRHI/Private/MetalRHI.cpp
+++ b/Engine/Source/Runtime/Apple/MetalRHI/Private/MetalRHI.cpp
@@ -566,6 +566,7 @@ FMetalDynamicRHI::FMetalDynamicRHI(ERHIFeatureLevel::Type RequestedFeatureLevel)
 	GMetalBufferFormats[PF_ETC2_RGBA			] = { mtlpp::PixelFormat::Invalid, EMetalBufferFormat::Unknown };
 	GMetalBufferFormats[PF_R32G32B32A32_UINT	] = { mtlpp::PixelFormat::RGBA32Uint, EMetalBufferFormat::RGBA32Uint };
 	GMetalBufferFormats[PF_R16G16_UINT			] = { mtlpp::PixelFormat::RG16Uint, EMetalBufferFormat::RG16Uint };
+	GMetalBufferFormats[PF_R32G32_UINT			] = { mtlpp::PixelFormat::RG32Uint, EMetalBufferFormat::RG32Uint };
 	GMetalBufferFormats[PF_ASTC_4x4             ] = { mtlpp::PixelFormat::Invalid, EMetalBufferFormat::Unknown };
 	GMetalBufferFormats[PF_ASTC_6x6             ] = { mtlpp::PixelFormat::Invalid, EMetalBufferFormat::Unknown };
 	GMetalBufferFormats[PF_ASTC_8x8             ] = { mtlpp::PixelFormat::Invalid, EMetalBufferFormat::Unknown };
@@ -593,6 +594,7 @@ FMetalDynamicRHI::FMetalDynamicRHI(ERHIFeatureLevel::Type RequestedFeatureLevel)
 	GPixelFormats[PF_G16				].PlatformFormat	= (uint32)mtlpp::PixelFormat::R16Unorm;
 	GPixelFormats[PF_R32G32B32A32_UINT	].PlatformFormat	= (uint32)mtlpp::PixelFormat::RGBA32Uint;
 	GPixelFormats[PF_R16G16_UINT		].PlatformFormat	= (uint32)mtlpp::PixelFormat::RG16Uint;
+	GPixelFormats[PF_R32G32_UINT		].PlatformFormat	= (uint32)mtlpp::PixelFormat::RG32Uint;
 		
 #if PLATFORM_IOS
     GPixelFormats[PF_DXT1				].PlatformFormat	= (uint32)mtlpp::PixelFormat::Invalid;
diff --git a/Engine/Source/Runtime/Apple/MetalRHI/Private/MetalUAV.cpp b/Engine/Source/Runtime/Apple/MetalRHI/Private/MetalUAV.cpp
--- a/Engine/Source/Runtime/Apple/MetalRHI/Private/MetalUAV.cpp
+++ b/Engine/Source/Runtime/Apple/MetalRHI/Private/MetalUAV.cpp
@@ -584,6 +584,7 @@ void FMetalRHICommandContext::RHIClearTinyUAV(FUnorderedAccessViewRHIParamRef Un
 					break;
 				}
 				case PF_G32R32F:
+				case PF_R32G32_UINT:
 				{
 					UE_LOG(LogMetal, Warning, TEXT("UAV pattern fill for format: %d is untested"), Format);
 					Pattern[0] = Values[0];
diff --git a/Engine/Source/Runtime/Core/Public/PixelFormat.h b/Engine/Source/Runtime/Core/Public/PixelFormat.h
--- a/Engine/Source/Runtime/Core/Public/PixelFormat.h
+++ b/Engine/Source/Runtime/Core/Public/PixelFormat.h
@@ -79,7 +79,8 @@ enum EPixelFormat
 	PF_PLATFORM_HDR_1		=65,	// Reserved.
 	PF_PLATFORM_HDR_2		=66,	// Reserved.
 	PF_NV12					=67,
-	PF_MAX					=68,
+	PF_R32G32_UINT          =68,
+	PF_MAX					=69,
 };
 #define FOREACH_ENUM_EPIXELFORMAT(op) \
 	op(PF_Unknown) \
@@ -149,5 +150,6 @@ enum EPixelFormat
 	op(PF_PLATFORM_HDR_0) \
 	op(PF_PLATFORM_HDR_1) \
 	op(PF_PLATFORM_HDR_2) \
-	op(PF_NV12)
+	op(PF_NV12) \
+	op(PF_R32G32_UINT)
 
diff --git a/Engine/Source/Runtime/CoreUObject/Public/UObject/NoExportTypes.h b/Engine/Source/Runtime/CoreUObject/Public/UObject/NoExportTypes.h
--- a/Engine/Source/Runtime/CoreUObject/Public/UObject/NoExportTypes.h
+++ b/Engine/Source/Runtime/CoreUObject/Public/UObject/NoExportTypes.h
@@ -219,6 +219,7 @@ enum EPixelFormat
 	PF_PLATFORM_HDR_1,
 	PF_PLATFORM_HDR_2,
 	PF_NV12,
+	PF_R32G32_UINT,
 	PF_MAX,
 };
 
diff --git a/Engine/Source/Runtime/D3D12RHI/Private/D3D12RHI.cpp b/Engine/Source/Runtime/D3D12RHI/Private/D3D12RHI.cpp
--- a/Engine/Source/Runtime/D3D12RHI/Private/D3D12RHI.cpp
+++ b/Engine/Source/Runtime/D3D12RHI/Private/D3D12RHI.cpp
@@ -150,6 +150,7 @@ FD3D12DynamicRHI::FD3D12DynamicRHI(TArray<FD3D12Adapter*>& ChosenAdaptersIn) :
 	GPixelFormats[PF_R8G8			].PlatformFormat = DXGI_FORMAT_R8G8_UNORM;
 	GPixelFormats[PF_R32G32B32A32_UINT].PlatformFormat = DXGI_FORMAT_R32G32B32A32_UINT;
 	GPixelFormats[PF_R16G16_UINT	].PlatformFormat = DXGI_FORMAT_R16G16_UINT;
+	GPixelFormats[PF_R32G32_UINT	].PlatformFormat = DXGI_FORMAT_R32G32_UINT;
 
 	GPixelFormats[PF_BC6H			].PlatformFormat = DXGI_FORMAT_BC6H_UF16;
 	GPixelFormats[PF_BC7			].PlatformFormat = DXGI_FORMAT_BC7_TYPELESS;
diff --git a/Engine/Source/Runtime/OpenGLDrv/Private/OpenGLDevice.cpp b/Engine/Source/Runtime/OpenGLDrv/Private/OpenGLDevice.cpp
--- a/Engine/Source/Runtime/OpenGLDrv/Private/OpenGLDevice.cpp
+++ b/Engine/Source/Runtime/OpenGLDrv/Private/OpenGLDevice.cpp
@@ -976,6 +976,11 @@ static void InitRHICapabilitiesForGL()
 				// The user should check for support for PF_G16R16 and implement a fallback if it's not supported!
 				SetupTextureFormat(PF_G16R16, FOpenGLTextureFormat(GL_RG16, GL_RG16, GL_RG, GL_UNSIGNED_SHORT, false, false));
 			}
+			if (FOpenGL::SupportsRG32UI())
+			{
+				// The user should check for support for PF_G32R32 and implement a fallback if it's not supported!
+				SetupTextureFormat(PF_R32G32_UINT, FOpenGLTextureFormat(GL_RG32UI, GL_RG32UI, GL_RG, GL_UNSIGNED_INT, false, false));
+			}
 		}
 		else
 		{
diff --git a/Engine/Source/Runtime/OpenGLDrv/Public/OpenGL.h b/Engine/Source/Runtime/OpenGLDrv/Public/OpenGL.h
--- a/Engine/Source/Runtime/OpenGLDrv/Public/OpenGL.h
+++ b/Engine/Source/Runtime/OpenGLDrv/Public/OpenGL.h
@@ -120,6 +120,7 @@ public:
 	static FORCEINLINE bool SupportsColorBufferFloat()					{ return true; }
 	static FORCEINLINE bool SupportsColorBufferHalfFloat()				{ return true; }
 	static FORCEINLINE bool	SupportsRG16UI()							{ return true; }
+	static FORCEINLINE bool	SupportsRG32UI()							{ return true; }
 	static FORCEINLINE bool SupportsR11G11B10F()						{ return true; }
 	static FORCEINLINE bool SupportsVolumeTextureRendering()			{ return false; }
 	static FORCEINLINE bool SupportsShaderFramebufferFetch()			{ return false; }
diff --git a/Engine/Source/Runtime/OpenGLDrv/Public/OpenGLES2.h b/Engine/Source/Runtime/OpenGLDrv/Public/OpenGLES2.h
--- a/Engine/Source/Runtime/OpenGLDrv/Public/OpenGLES2.h
+++ b/Engine/Source/Runtime/OpenGLDrv/Public/OpenGLES2.h
@@ -115,6 +115,7 @@ struct FOpenGLES2 : public FOpenGLBase
 	static FORCEINLINE bool SupportsColorBufferFloat()					{ return bSupportsColorBufferFloat; }
 	static FORCEINLINE bool SupportsColorBufferHalfFloat()				{ return bSupportsColorBufferHalfFloat; }
 	static FORCEINLINE bool	SupportsRG16UI()							{ return false; }
+	static FORCEINLINE bool	SupportsRG32UI()							{ return false; }
 	static FORCEINLINE bool SupportsR11G11B10F()						{ return false; }
 	static FORCEINLINE bool SupportsShaderFramebufferFetch()			{ return bSupportsShaderFramebufferFetch; }
 	static FORCEINLINE bool SupportsShaderDepthStencilFetch()			{ return bSupportsShaderDepthStencilFetch; }
diff --git a/Engine/Source/Runtime/OpenGLDrv/Public/OpenGLESDeferred.h b/Engine/Source/Runtime/OpenGLDrv/Public/OpenGLESDeferred.h
--- a/Engine/Source/Runtime/OpenGLDrv/Public/OpenGLESDeferred.h
+++ b/Engine/Source/Runtime/OpenGLDrv/Public/OpenGLESDeferred.h
@@ -75,6 +75,7 @@ struct FOpenGLESDeferred : public FOpenGLBase
 	static FORCEINLINE bool SupportsColorBufferFloat()					{ return bSupportsColorBufferFloat || !bES2Fallback; }
 	static FORCEINLINE bool SupportsColorBufferHalfFloat()				{ return bSupportsColorBufferHalfFloat || !bES2Fallback; }
 	static FORCEINLINE bool	SupportsRG16UI()							{ return bSupportsNvImageFormats && !bES2Fallback; }
+	static FORCEINLINE bool	SupportsRG32UI()							{ return bSupportsNvImageFormats && !bES2Fallback; }
 	static FORCEINLINE bool SupportsR11G11B10F()						{ return bSupportsNvImageFormats && !bES2Fallback; }
 	static FORCEINLINE bool SupportsShaderFramebufferFetch()			{ return bSupportsShaderFramebufferFetch; }
 	static FORCEINLINE bool SupportsShaderDepthStencilFetch()			{ return bSupportsShaderDepthStencilFetch; }
diff --git a/Engine/Source/Runtime/RenderCore/Private/RenderUtils.cpp b/Engine/Source/Runtime/RenderCore/Private/RenderUtils.cpp
--- a/Engine/Source/Runtime/RenderCore/Private/RenderUtils.cpp
+++ b/Engine/Source/Runtime/RenderCore/Private/RenderUtils.cpp
@@ -137,14 +137,15 @@ FPixelFormatInfo	GPixelFormats[PF_MAX] =
 
 	{ TEXT("R16G16B16A16_UINT"),1,			1,			1,			8,			4,				0,				1,				PF_R16G16B16A16_UNORM },
 	{ TEXT("R16G16B16A16_SINT"),1,			1,			1,			8,			4,				0,				1,				PF_R16G16B16A16_SNORM },
-	{ TEXT("PLATFORM_HDR_0"),	0,			0,			0,			0,			0,				0,				0,				PF_PLATFORM_HDR_0 },
-	{ TEXT("PLATFORM_HDR_1"),	0,			0,			0,			0,			0,				0,				0,				PF_PLATFORM_HDR_1 },
-	{ TEXT("PLATFORM_HDR_2"),	0,			0,			0,			0,			0,				0,				0,				PF_PLATFORM_HDR_2 },
+	{ TEXT("PLATFORM_HDR_0"),	0,			0,			0,			0,			0,				0,				0,				PF_PLATFORM_HDR_0   },
+	{ TEXT("PLATFORM_HDR_1"),	0,			0,			0,			0,			0,				0,				0,				PF_PLATFORM_HDR_1   },
+	{ TEXT("PLATFORM_HDR_2"),	0,			0,			0,			0,			0,				0,				0,				PF_PLATFORM_HDR_2   },
 
 	// NV12 contains 2 textures: R8 luminance plane followed by R8G8 1/4 size chrominance plane.
 	// BlockSize/BlockBytes/NumComponents values don't make much sense for this format, so set them all to one.
-	{ TEXT("NV12"),				1,			1,			1,			1,			1,				0,				0,				PF_NV12 },
+	{ TEXT("NV12"),				1,			1,			1,			1,			1,				0,				0,				PF_NV12             },
 
+	{ TEXT("PF_R32G32_UINT"),   1,   		1,			1,			8,			2,				0,				1,				PF_R32G32_UINT      },
 };
 
 static struct FValidatePixelFormats
diff --git a/Engine/Source/Runtime/VulkanRHI/Private/VulkanDevice.cpp b/Engine/Source/Runtime/VulkanRHI/Private/VulkanDevice.cpp
--- a/Engine/Source/Runtime/VulkanRHI/Private/VulkanDevice.cpp
+++ b/Engine/Source/Runtime/VulkanRHI/Private/VulkanDevice.cpp
@@ -562,6 +562,9 @@ void FVulkanDevice::SetupFormats()
 	MapFormatSupport(PF_R16G16B16A16_SINT, VK_FORMAT_R16G16B16A16_SINT);
 	SetComponentMapping(PF_R16G16B16A16_SINT, VK_COMPONENT_SWIZZLE_R, VK_COMPONENT_SWIZZLE_G, VK_COMPONENT_SWIZZLE_B, VK_COMPONENT_SWIZZLE_A);
 
+	MapFormatSupport(PF_R32G32_UINT, VK_FORMAT_R32G32_UINT);
+	SetComponentMapping(PF_R32G32_UINT, VK_COMPONENT_SWIZZLE_R, VK_COMPONENT_SWIZZLE_G, VK_COMPONENT_SWIZZLE_ZERO, VK_COMPONENT_SWIZZLE_ZERO);
+
 	MapFormatSupport(PF_R32G32B32A32_UINT, VK_FORMAT_R32G32B32A32_UINT);
 	SetComponentMapping(PF_R32G32B32A32_UINT, VK_COMPONENT_SWIZZLE_R, VK_COMPONENT_SWIZZLE_G, VK_COMPONENT_SWIZZLE_B, VK_COMPONENT_SWIZZLE_A);
 
diff --git a/Engine/Source/Runtime/Windows/D3D11RHI/Private/D3D11Device.cpp b/Engine/Source/Runtime/Windows/D3D11RHI/Private/D3D11Device.cpp
--- a/Engine/Source/Runtime/Windows/D3D11RHI/Private/D3D11Device.cpp
+++ b/Engine/Source/Runtime/Windows/D3D11RHI/Private/D3D11Device.cpp
@@ -190,7 +190,8 @@ FD3D11DynamicRHI::FD3D11DynamicRHI(IDXGIFactory1* InDXGIFactory1,D3D_FEATURE_LEV
 	GPixelFormats[ PF_R8G8B8A8_SNORM].PlatformFormat	= DXGI_FORMAT_R8G8B8A8_SNORM;
 	GPixelFormats[ PF_R8G8			].PlatformFormat	= DXGI_FORMAT_R8G8_UNORM;
 	GPixelFormats[ PF_R32G32B32A32_UINT].PlatformFormat = DXGI_FORMAT_R32G32B32A32_UINT;
-	GPixelFormats[ PF_R16G16_UINT].PlatformFormat = DXGI_FORMAT_R16G16_UINT;
+	GPixelFormats[ PF_R16G16_UINT   ].PlatformFormat    = DXGI_FORMAT_R16G16_UINT;
+	GPixelFormats[ PF_R32G32_UINT   ].PlatformFormat    = DXGI_FORMAT_R32G32_UINT;
 
 	GPixelFormats[ PF_BC6H			].PlatformFormat	= DXGI_FORMAT_BC6H_UF16;
 	GPixelFormats[ PF_BC7			].PlatformFormat	= DXGI_FORMAT_BC7_TYPELESS;
