Copy of CL 7438644 by [REDACTED_PERSON]

Fix leaking uniform buffers when using the async compute context with Metal.

diff --git a/Engine/Source/Runtime/Apple/MetalRHI/Private/MetalContext.cpp b/Engine/Source/Runtime/Apple/MetalRHI/Private/MetalContext.cpp
--- a/Engine/Source/Runtime/Apple/MetalRHI/Private/MetalContext.cpp
+++ b/Engine/Source/Runtime/Apple/MetalRHI/Private/MetalContext.cpp
@@ -819,7 +819,7 @@ struct FMetalRHICommandUpdateFence final : public FRHICommand<FMetalRHICommandUp
 	void Execute(FRHICommandListBase& CmdList)
 	{
 		GetMetalDeviceContext().SetParallelPassFences(nil, Fence);
-		GetMetalDeviceContext().FinishFrame();
+		GetMetalDeviceContext().FinishFrame(true);
 		GetMetalDeviceContext().BeginParallelRenderCommandEncoding(Num);
 	}
 };
@@ -1085,7 +1085,7 @@ void FMetalContext::InitFrame(bool const bImmediateContext, uint32 Index, uint32
 	StateCache.InvalidateRenderTargets();
 }
 
-void FMetalContext::FinishFrame()
+void FMetalContext::FinishFrame(bool const bImmediateContext)
 {
 	// Ensure that we update the end fence for parallel contexts.
 	RenderPass.Update(EndFence);
@@ -1102,6 +1102,11 @@ void FMetalContext::FinishFrame()
 	// make sure first SetRenderTarget goes through
 	StateCache.InvalidateRenderTargets();
 	
+	if (!bImmediateContext)
+	{
+		StateCache.Reset();
+	}
+
 #if ENABLE_METAL_GPUPROFILE
 	FPlatformTLS::SetTlsValue(CurrentContextTLSSlot, nullptr);
 #endif
@@ -1733,7 +1738,7 @@ public:
 				GetMetalDeviceContext().SetParallelPassFences(Fence, nil);
 			}
 
-			CmdContext->GetInternalContext().FinishFrame();
+			CmdContext->GetInternalContext().FinishFrame(false);
 			GetMetalDeviceContext().EndParallelRenderCommandEncoding();
 
 			CmdContext->GetInternalContext().GetCommandList().Submit(Index, Num);
diff --git a/Engine/Source/Runtime/Apple/MetalRHI/Private/MetalContext.h b/Engine/Source/Runtime/Apple/MetalRHI/Private/MetalContext.h
--- a/Engine/Source/Runtime/Apple/MetalRHI/Private/MetalContext.h
+++ b/Engine/Source/Runtime/Apple/MetalRHI/Private/MetalContext.h
@@ -121,7 +121,7 @@ public:
 	TRefCountPtr<FMetalFence> const& GetParallelPassEndFence(void) const;
 	
 	void InitFrame(bool const bImmediateContext, uint32 Index, uint32 Num);
-	void FinishFrame();
+	void FinishFrame(bool const bImmediateContext);
 
 	// Track Write->Read transitions for TBDR Fragment->Verex fencing
 	void TransitionResources(FRHIUnorderedAccessView** InUAVs, int32 NumUAVs);
diff --git a/Engine/Source/Runtime/Apple/MetalRHI/Private/MetalRHIContext.cpp b/Engine/Source/Runtime/Apple/MetalRHI/Private/MetalRHIContext.cpp
--- a/Engine/Source/Runtime/Apple/MetalRHI/Private/MetalRHIContext.cpp
+++ b/Engine/Source/Runtime/Apple/MetalRHI/Private/MetalRHIContext.cpp
@@ -126,7 +126,7 @@ void FMetalRHIComputeContext::RHISubmitCommandsHint()
 	{
 		Context->InitFrame(false, 0, 0);
 	}
-	Context->FinishFrame();
+	Context->FinishFrame(false);
 	
 #if ENABLE_METAL_GPUPROFILE
 	FMetalContext::MakeCurrent(&GetMetalDeviceContext());
