Retain resource references on Intel GPUs under Metal as the driver is encountering Inalid Resources but we can't reproduce this internally to investigate.

diff --git a/Engine/Source/Runtime/Apple/MetalRHI/Private/MetalCommandQueue.cpp b/Engine/Source/Runtime/Apple/MetalRHI/Private/MetalCommandQueue.cpp
--- a/Engine/Source/Runtime/Apple/MetalRHI/Private/MetalCommandQueue.cpp
+++ b/Engine/Source/Runtime/Apple/MetalRHI/Private/MetalCommandQueue.cpp
@@ -365,7 +365,7 @@ mtlpp::CommandBuffer FMetalCommandQueue::CreateCommandBuffer(void)
 	static bool bUnretainedRefs = FParse::Param(FCommandLine::Get(),TEXT("metalunretained"))
 	|| (!FParse::Param(FCommandLine::Get(),TEXT("metalretainrefs"))
 		&& ([Device.GetName() rangeOfString:@"Nvidia" options:NSCaseInsensitiveSearch].location == NSNotFound)
-		&& ([Device.GetName() rangeOfString:@"Intel" options:NSCaseInsensitiveSearch].location == NSNotFound || FPlatformMisc::MacOSXVersionCompare(10,13,0) >= 0));
+		&& ([Device.GetName() rangeOfString:@"Intel" options:NSCaseInsensitiveSearch].location == NSNotFound));
 #else
 	static bool bUnretainedRefs = !FParse::Param(FCommandLine::Get(),TEXT("metalretainrefs"));
 #endif
