Add assertions that input buffers to Indirect draw & dispatch calls are valid.

diff --git a/Engine/Source/Runtime/Apple/MetalRHI/Private/MetalRenderPass.cpp b/Engine/Source/Runtime/Apple/MetalRHI/Private/MetalRenderPass.cpp
--- a/Engine/Source/Runtime/Apple/MetalRHI/Private/MetalRenderPass.cpp
+++ b/Engine/Source/Runtime/Apple/MetalRHI/Private/MetalRenderPass.cpp
@@ -444,6 +444,7 @@ void FMetalRenderPass::DrawPrimitiveIndirect(uint32 PrimitiveType, FMetalVertexB
 		ConditionalSwitchToRender();
 		check(CurrentEncoder.GetCommandBuffer());
 		check(CurrentEncoder.IsRenderCommandEncoderActive());
+		check(VertexBuffer->Buffer);
 		
 		PrepareToRender(PrimitiveType);
 		
@@ -577,6 +578,8 @@ void FMetalRenderPass::DrawIndexedIndirect(FMetalIndexBuffer* IndexBuffer, uint3
 		ConditionalSwitchToRender();
 		check(CurrentEncoder.GetCommandBuffer());
 		check(CurrentEncoder.IsRenderCommandEncoderActive());
+		check(IndexBuffer->Buffer);
+		check(VertexBuffer->Buffer);
 		
 		// finalize any pending state
 		PrepareToRender(PrimitiveType);
@@ -610,6 +613,8 @@ void FMetalRenderPass::DrawIndexedPrimitiveIndirect(uint32 PrimitiveType,FMetalI
 		ConditionalSwitchToRender();
 		check(CurrentEncoder.GetCommandBuffer());
 		check(CurrentEncoder.IsRenderCommandEncoderActive());
+		check(IndexBuffer->Buffer);
+		check(VertexBuffer->Buffer);
 		
 		PrepareToRender(PrimitiveType);
 		
@@ -888,6 +893,7 @@ void FMetalRenderPass::DispatchIndirect(FMetalVertexBuffer* ArgumentBuffer, uint
 		ConditionalSwitchToAsyncCompute();
 		check(PrologueEncoder.GetCommandBuffer());
 		check(PrologueEncoder.IsComputeCommandEncoderActive());
+		check(ArgumentBuffer->Buffer);
 		
 		PrepareToAsyncDispatch();
 		
