Overhaul of iOS audio mixer platform layer. iOS Audiomixer backend to make buffer parameters observe project settings and decople OS callback from audiomixer callback using a circular buffer.

diff --git a/Engine/Source/Runtime/Apple/AudioMixerAudioUnit/Private/AudioMixerPlatformAudioUnit.cpp b/Engine/Source/Runtime/Apple/AudioMixerAudioUnit/Private/AudioMixerPlatformAudioUnit.cpp
--- a/Engine/Source/Runtime/Apple/AudioMixerAudioUnit/Private/AudioMixerPlatformAudioUnit.cpp
+++ b/Engine/Source/Runtime/Apple/AudioMixerAudioUnit/Private/AudioMixerPlatformAudioUnit.cpp
@@ -14,11 +14,10 @@
 
 #include "VorbisAudioInfo.h"
 #include "OpusAudioInfo.h"
-#include "ADPCMAudioInfo.h"
 #endif // #if PLATFORM_IOS || PLATFORM_TVOS
 /*
-	This implementation only depends on the audio units API which allows it to run on MacOS, iOS and tvOS.
-
+	This implementation only depends on the audio units API which allows it to run on MacOS, iOS and tvOS. 
+	
 	For now just assume an iOS configuration (only 2 left and right channels on a single device)
 */
 
@@ -41,15 +40,15 @@ DEFINE_LOG_CATEGORY(LogAudioMixerAudioUnit);
 namespace Audio
 {
 #if PLATFORM_IOS || PLATFORM_TVOS
-	static const int32 DefaultBufferSize = 4096;
+	static const int32 DefaultBufferSize = 512;
 #else
 	static const int32 DefaultBufferSize = 1024;
 	static const int32 AUBufferSize = 256;
 #endif //#if PLATFORM_IOS || PLATFORM_TVOS
 	static const double DefaultSampleRate = 48000.0;
-
+	
 	static int32 SuspendCounter = 0;
-
+	
 	FMixerPlatformAudioUnit::FMixerPlatformAudioUnit()
 		: bInitialized(false)
 		, bInCallback(false)
@@ -57,6 +56,8 @@ namespace Audio
 		, RemainingBytesInCurrentSubmittedBuffer(0)
 		, BytesPerSubmittedBuffer(0)
 		, GraphSampleRate(DefaultSampleRate)
+		, NumSamplesPerRenderCallback(0)
+		, NumSamplesPerDeviceCallback(0)
 	{
 	}
 
@@ -67,21 +68,23 @@ namespace Audio
 			TeardownHardware();
 		}
 	}
-
+	
 	int32 FMixerPlatformAudioUnit::GetNumFrames(const int32 InNumReqestedFrames)
 	{
 #if PLATFORM_IOS || PLATFORM_TVOS
+		return AlignArbitrary(InNumReqestedFrames, 4);
+		
 		AVAudioSession* AudioSession = [AVAudioSession sharedInstance];
-		double BufferSizeInSec = [AudioSession preferredIOBufferDuration];
+		double BufferSizeInSec = [AudioSession IOBufferDuration];
 		double SampleRate = [AudioSession preferredSampleRate];
-
+		
 		if (BufferSizeInSec == 0.0)
 		{
 			return DefaultBufferSize;
 		}
-
+		
 		int32 NumFrames = (int32)(SampleRate * BufferSizeInSec);
-
+		
 		return NumFrames;
 #else
 	   //On MacOS, we hardcode buffer sizes.
@@ -95,47 +98,53 @@ namespace Audio
 		{
 			return false;
 		}
-
+		
 		OSStatus Status;
-		GraphSampleRate = (double) AudioStreamInfo.DeviceInfo.SampleRate;
-		UInt32 BufferSize = (UInt32) GetNumFrames(OpenStreamParams.NumFrames);
+		GraphSampleRate = (double) InternalPlatformSettings.SampleRate;
+		UInt32 BufferSize = (UInt32) GetNumFrames(InternalPlatformSettings.CallbackBufferFrameSize);
 		const int32 NumChannels = 2;
 
 		if (GraphSampleRate == 0)
 		{
 			GraphSampleRate = DefaultSampleRate;
 		}
-
+	   
 		if (BufferSize == 0)
 		{
 			BufferSize = DefaultBufferSize;
 		}
-
+		
 		BytesPerSubmittedBuffer = BufferSize * NumChannels * sizeof(float);
 		check(BytesPerSubmittedBuffer != 0);
-		UE_LOG(LogAudioMixerAudioUnit, Display, TEXT("Bytes per submitted buffer: %d"), BytesPerSubmittedBuffer);
-
+		
 #if PLATFORM_IOS || PLATFORM_TVOS
 		NSError* error;
-
+		
 		AVAudioSession* AudioSession = [AVAudioSession sharedInstance];
 
-		GraphSampleRate = [AudioSession preferredSampleRate];
+		// this sample rate is currently gotten from AudioSession in GetPlatformSettings, so there should be no issue
 		bool Success = [AudioSession setPreferredSampleRate:GraphSampleRate error:&error];
-
+		
 		if (!Success)
 		{
 			UE_LOG(LogAudioMixerAudioUnit, Display, TEXT("Error setting sample rate."));
 		}
-
+		
 		// By calling setPreferredIOBufferDuration, we indicate that we would prefer that the buffer size not change if possible.
-		Success = [AudioSession setPreferredIOBufferDuration:[AudioSession preferredIOBufferDuration] error: &error];
-
+		float AudioMixerBufferSizeInSec = InternalPlatformSettings.CallbackBufferFrameSize / GraphSampleRate; // todo dont hard code
+		Success = [AudioSession setPreferredIOBufferDuration:AudioMixerBufferSizeInSec error: &error];
+		
+		int32 FinalBufferSize = [AudioSession IOBufferDuration] * GraphSampleRate;
+		int32 FinalPreferredBufferSize = [AudioSession preferredIOBufferDuration] * GraphSampleRate;
+		
+		BytesPerSubmittedBuffer = FinalBufferSize * NumChannels * sizeof(float);
+		check(BytesPerSubmittedBuffer != 0);
+		
 		UE_LOG(LogAudioMixerAudioUnit, Display, TEXT("Device Sample Rate: %f"), GraphSampleRate);
 		check(GraphSampleRate != 0);
-
+		
 		Success = [AudioSession setActive:true error:&error];
-
+		
 		if (!Success)
 		{
 			UE_LOG(LogAudioMixerAudioUnit, Display, TEXT("Error starting audio session."));
@@ -144,28 +153,29 @@ namespace Audio
 		AudioObjectID DeviceAudioObjectID;
 		AudioObjectPropertyAddress DevicePropertyAddress;
 		UInt32 AudioDeviceQuerySize;
-
+		
 		//Get Audio Device ID- this will be used throughout initialization to query the audio hardware.
 		DevicePropertyAddress.mSelector = kAudioHardwarePropertyDefaultOutputDevice;
 		DevicePropertyAddress.mScope = kAudioObjectPropertyScopeGlobal;
 		DevicePropertyAddress.mElement = 0;
 		AudioDeviceQuerySize = sizeof(AudioDeviceID);
 		Status = AudioObjectGetPropertyData(kAudioObjectSystemObject, &DevicePropertyAddress, 0, nullptr, &AudioDeviceQuerySize, &DeviceAudioObjectID);
-
+		
 		if(Status != 0)
 		{
 			UE_LOG(LogAudioMixerAudioUnit, Display, TEXT("ERROR setting sample rate to %f"), GraphSampleRate);
 		}
-
+		
 		Status = AudioObjectGetPropertyData(DeviceAudioObjectID, &DevicePropertyAddress, 0, nullptr, &AudioDeviceQuerySize, &GraphSampleRate);
-
+		
 		if(Status == 0)
 		{
 			UE_LOG(LogAudioMixerAudioUnit, Display, TEXT("Sample Rate: %f"), GraphSampleRate);
 		}
 
 #endif // #if PLATFORM_IOS || PLATFORM_TVOS
-
+		
+		UE_LOG(LogAudioMixerAudioUnit, Display, TEXT("Bytes per submitted buffer: %d"), BytesPerSubmittedBuffer);
 
 		// Linear PCM stream format
 		OutputFormat.mFormatID         = kAudioFormatLinearPCM;
@@ -204,14 +214,14 @@ namespace Audio
 			HandleError(TEXT("Failed to initialize audio output node!"), true);
 			return false;
 		}
-
+		
 		Status = AUGraphOpen(AudioUnitGraph);
 		if (Status != noErr)
 		{
 			HandleError(TEXT("Failed to open audio unit graph"), true);
 			return false;
 		}
-
+		
 		Status = AUGraphNodeInfo(AudioUnitGraph, OutputNode, nullptr, &OutputUnit);
 		if (Status != noErr)
 		{
@@ -242,25 +252,21 @@ namespace Audio
 			HandleError(TEXT("Failed to set output format!"), true);
 			return false;
 		}
-
+		
 #endif //#if PLATFORM_MAC
-
-		AudioStreamInfo.NumOutputFrames = BufferSize;
-
 		AudioStreamInfo.DeviceInfo = GetPlatformDeviceInfo();
-
+		
 		AURenderCallbackStruct InputCallback;
 		InputCallback.inputProc = &AudioRenderCallback;
 		InputCallback.inputProcRefCon = this;
 		Status = AUGraphSetNodeInputCallback(AudioUnitGraph,
-											 OutputNode,
+												 OutputNode,
 											 0,
 											 &InputCallback);
 		UE_CLOG(Status != noErr, LogAudioMixerAudioUnit, Error, TEXT("Failed to set input callback for audio output node"));
 
-		OpenStreamParams.NumFrames = BufferSize;
 		AudioStreamInfo.StreamState = EAudioOutputStreamState::Closed;
-
+		
 		bInitialized = true;
 
 		return true;
@@ -278,7 +284,7 @@ namespace Audio
 		{
 			return true;
 		}
-
+		
 		StopAudioStream();
 		CloseAudioStream();
 
@@ -289,7 +295,7 @@ namespace Audio
 		OutputUnit = nullptr;
 
 		bInitialized = false;
-
+		
 		return true;
 	}
 
@@ -301,7 +307,7 @@ namespace Audio
 	bool FMixerPlatformAudioUnit::GetNumOutputDevices(uint32& OutNumOutputDevices)
 	{
 		OutNumOutputDevices = 1;
-
+		
 		return true;
 	}
 
@@ -314,7 +320,7 @@ namespace Audio
 	bool FMixerPlatformAudioUnit::GetDefaultOutputDeviceIndex(uint32& OutDefaultDeviceIndex) const
 	{
 		OutDefaultDeviceIndex = 0;
-
+		
 		return true;
 	}
 
@@ -324,12 +330,15 @@ namespace Audio
 		{
 			return false;
 		}
-
-		AudioStreamInfo.OutputDeviceIndex = Params.OutputDeviceIndex;
-		AudioStreamInfo.AudioMixer = Params.AudioMixer;
-
 		OpenStreamParams = Params;
-
+		//todo: AudioStreamInfo.SampleRate = OpenStreamParams.SampleRate;
+		AudioStreamInfo.Reset();
+		AudioStreamInfo.OutputDeviceIndex = OpenStreamParams.OutputDeviceIndex;
+		AudioStreamInfo.NumOutputFrames = OpenStreamParams.NumFrames;
+		AudioStreamInfo.NumBuffers = OpenStreamParams.NumBuffers;
+		AudioStreamInfo.AudioMixer = OpenStreamParams.AudioMixer;
+		AudioStreamInfo.DeviceInfo = GetPlatformDeviceInfo();
+		
 		// Initialize the audio unit graph
 		OSStatus Status = AUGraphInitialize(AudioUnitGraph);
 		if (Status != noErr)
@@ -338,6 +347,14 @@ namespace Audio
 			return false;
 		}
 
+		// Set up circular buffer between our rendering buffer size and the device's buffer size.
+		// Since we are only using this circular buffer on a single thread, we do not need to add extra slack.
+		NumSamplesPerRenderCallback = AudioStreamInfo.NumOutputFrames * AudioStreamInfo.DeviceInfo.NumChannels;
+		NumSamplesPerDeviceCallback = InternalPlatformSettings.CallbackBufferFrameSize * AudioStreamInfo.DeviceInfo.NumChannels;
+
+		// initial circular buffer capacity is zero, so this initializes it.
+		GrowCircularBufferIfNeeded(NumSamplesPerRenderCallback, NumSamplesPerDeviceCallback);
+
 		AudioStreamInfo.StreamState = EAudioOutputStreamState::Open;
 
 		return true;
@@ -349,9 +366,9 @@ namespace Audio
 		{
 			return false;
 		}
-
+		
 		AudioStreamInfo.StreamState = EAudioOutputStreamState::Closed;
-
+		
 		return true;
 	}
 
@@ -361,9 +378,9 @@ namespace Audio
 		{
 			return false;
 		}
-
+		
 		BeginGeneratingAudio();
-
+		
 		// This will start the render audio callback
 		OSStatus Status = AUGraphStart(AudioUnitGraph);
 		if (Status != noErr)
@@ -377,37 +394,31 @@ namespace Audio
 
 	bool FMixerPlatformAudioUnit::StopAudioStream()
 	{
-		if(!bInitialized)
+		if(!bInitialized || AudioStreamInfo.StreamState != EAudioOutputStreamState::Running)
 		{
 			return false;
 		}
+		
+		AudioStreamInfo.StreamState = EAudioOutputStreamState::Stopping;
+		
+		AUGraphStop(AudioUnitGraph);
 
-		if (AudioStreamInfo.StreamState != EAudioOutputStreamState::Stopped && AudioStreamInfo.StreamState != EAudioOutputStreamState::Closed)
-		{
-			AUGraphStop(AudioUnitGraph);
-
-			if (AudioStreamInfo.StreamState == EAudioOutputStreamState::Running)
-			{
-				StopGeneratingAudio();
-			}
-
-			check(AudioStreamInfo.StreamState == EAudioOutputStreamState::Stopped);
-		}
-
-
+		AudioStreamInfo.StreamState = EAudioOutputStreamState::Stopped;
+		
 		return true;
 	}
 
 	bool FMixerPlatformAudioUnit::MoveAudioStreamToNewAudioDevice(const FString& InNewDeviceId)
 	{
 		//TODO
+		
 		return false;
 	}
 
 	FAudioPlatformDeviceInfo FMixerPlatformAudioUnit::GetPlatformDeviceInfo() const
 	{
 		FAudioPlatformDeviceInfo DeviceInfo;
-
+		
 	#if PLATFORM_IOS || PLATFORM_TVOS
 		AVAudioSession* AudioSession = [AVAudioSession sharedInstance];
 		double SampleRate = [AudioSession preferredSampleRate];
@@ -421,37 +432,37 @@ namespace Audio
 		DeviceInfo.OutputChannelArray[0] = EAudioMixerChannel::FrontLeft;
 		DeviceInfo.OutputChannelArray[1] = EAudioMixerChannel::FrontRight;
 		DeviceInfo.bIsSystemDefault = true;
-
+		
 		return DeviceInfo;
 	}
 
 	void FMixerPlatformAudioUnit::SubmitBuffer(const uint8* Buffer)
 	{
-		SubmittedBufferPtr = (uint8*) Buffer;
-		SubmittedBytes = 0;
-		RemainingBytesInCurrentSubmittedBuffer = BytesPerSubmittedBuffer;
+		if(!Buffer)
+		{
+			return;
+		}
+		
+		const int32 BytesToSubmitToAudioMixer = NumSamplesPerRenderCallback * sizeof(float);
+		
+		int32 PushResult = CircularOutputBuffer.Push((const int8*)Buffer, BytesToSubmitToAudioMixer);
+		check(PushResult == BytesToSubmitToAudioMixer);
 	}
 
+	
 	FName FMixerPlatformAudioUnit::GetRuntimeFormat(USoundWave* InSoundWave)
 	{
-		static const FName NAME_ADPCM(TEXT("ADPCM"));
-		static const FName NAME_OGG(TEXT("OGG"));
-		static const FName NAME_OPUS(TEXT("OPUS"));
-
 #if PLATFORM_IOS || PLATFORM_TVOS
+		static FName NAME_ADPCM(TEXT("ADPCM"));
 		return NAME_ADPCM;
 #else
-
-		if (InSoundWave->IsSeekableStreaming())
-		{
-			return NAME_ADPCM;
-		}
-
+		static FName NAME_OPUS(TEXT("OPUS"));
+		
 		if (InSoundWave->IsStreaming())
 		{
 			return NAME_OPUS;
 		}
-
+		static FName NAME_OGG(TEXT("OGG"));
 		return NAME_OGG;
 #endif
 	}
@@ -467,17 +478,12 @@ namespace Audio
 		return new FADPCMAudioInfo();
 #else
 		check(InSoundWave);
-
-		if (InSoundWave->IsSeekableStreaming())
-		{
-			return new FADPCMAudioInfo();
-		}
-
+		
 		if (InSoundWave->IsStreaming())
 		{
 			return new FOpusAudioInfo();
 		}
-
+		
 #if WITH_OGGVORBIS
 		static const FName NAME_OGG(TEXT("OGG"));
 		if (FPlatformProperties::RequiresCookedData() ? InSoundWave->HasCompressedData(NAME_OGG) : (InSoundWave->GetCompressedData(NAME_OGG) != nullptr))
@@ -507,15 +513,36 @@ namespace Audio
 
 	FAudioPlatformSettings FMixerPlatformAudioUnit::GetPlatformSettings() const
 	{
-		FAudioPlatformSettings Settings;
-		Settings.NumBuffers = 2;
 #if PLATFORM_IOS || PLATFORM_TVOS
+		InternalPlatformSettings = FAudioPlatformSettings::GetPlatformSettings(TEXT("/Script/IOSRuntimeSettings.IOSRuntimeSettings"));
+#else
+		// parsing project settings for mac has not yet been tested
+		InternalPlatformSettings = FAudioPlatformSettings::GetPlatformSettings(TEXT("/Script/OSXRuntimeSettings.OSXRuntimeSettings")); // #maxtodo: temp, to see if it works post-cook
+#endif // #if PLATFORM_IOS || PLATFORM_TVOS
+		
+#if PLATFORM_IOS || PLATFORM_TVOS
+		// Check for command line overrides
+		FString TempString;
+		
+		// Buffer Size
+		if(FParse::Value(FCommandLine::Get(), TEXT("-ForceIOSAudioMixerBufferSize="), TempString))
+		{
+			InternalPlatformSettings.CallbackBufferFrameSize = FCString::Atoi(*TempString);
+		}
+		
+		// NumBuffers
+		if(FParse::Value(FCommandLine::Get(), TEXT("-ForceIOSAudioMixerNumBuffers="), TempString))
+		{
+			InternalPlatformSettings.NumBuffers = FCString::Atoi(*TempString);
+		}
+		
 		AVAudioSession* AudioSession = [AVAudioSession sharedInstance];
-		double BufferSizeInSec = [AudioSession preferredIOBufferDuration];
+		double PreferredBufferSizeInSec = [AudioSession preferredIOBufferDuration];
+		double BufferSizeInSec = [AudioSession IOBufferDuration];
 		double SampleRate = [AudioSession preferredSampleRate];
-
+		
 		int32 NumFrames;
-
+		
 		if (BufferSizeInSec == 0.0)
 		{
 			NumFrames = DefaultBufferSize;
@@ -524,17 +551,15 @@ namespace Audio
 		{
 			NumFrames = (int32)(SampleRate * BufferSizeInSec);
 		}
-		Settings.CallbackBufferFrameSize = NumFrames;
-		Settings.SampleRate = SampleRate;
-		Settings.MaxChannels = 32;
-
+		InternalPlatformSettings.SampleRate = SampleRate;
+		
 #else
-		Settings.SampleRate = GraphSampleRate;
-		Settings.CallbackBufferFrameSize = DefaultBufferSize;
-
-#endif //#if PLATFORM_IOS || PLATFORM_TVOS
-
-		return Settings;
+		InternalPlatformSettings.SampleRate = GraphSampleRate;
+		InternalPlatformSettings.CallbackBufferFrameSize = DefaultBufferSize;
+		
+#endif // #if PLATFORM_IOS || PLATFORM_TVOS
+		bInternalPlatformSettingsInitialized = true;
+		return InternalPlatformSettings;
 	}
 
 	void FMixerPlatformAudioUnit::ResumeContext()
@@ -547,7 +572,7 @@ namespace Audio
 			bSuspended = false;
 		}
 	}
-
+	
 	void FMixerPlatformAudioUnit::SuspendContext()
 	{
 		if (SuspendCounter == 0)
@@ -568,50 +593,45 @@ namespace Audio
 		}
 	}
 
+	void FMixerPlatformAudioUnit::GrowCircularBufferIfNeeded(const int32 InNumSamplesPerRenderCallback, const int32 InNumSamplesPerDeviceCallback)
+	{
+		const int32 MaxCircularBufferCapacity = 2 * sizeof(float) * FMath::Max<int32>(NumSamplesPerRenderCallback, NumSamplesPerDeviceCallback);
+
+		if (CircularOutputBuffer.GetCapacity() < MaxCircularBufferCapacity)
+		{
+			// SetCapacity also zeros-out data
+			CircularOutputBuffer.SetCapacity(MaxCircularBufferCapacity);
+			UE_LOG(LogAudioMixerAudioUnit, Display, TEXT("Growing iOS circular buffer to %i bytes."), MaxCircularBufferCapacity);
+		}
+	}
+
 	bool FMixerPlatformAudioUnit::PerformCallback(AudioBufferList* OutputBufferData)
 	{
 		bInCallback = true;
 
 		if (AudioStreamInfo.StreamState == EAudioOutputStreamState::Running)
 		{
-			if (!SubmittedBufferPtr)
+			// How many bytes we have left over from previous callback
+			BytesPerSubmittedBuffer = OutputBufferData->mBuffers[0].mDataByteSize;
+			uint8* OutputBufferPtr = (uint8*)OutputBufferData->mBuffers[0].mData;
+			
+			NumSamplesPerDeviceCallback = BytesPerSubmittedBuffer / static_cast<float>(sizeof(float));
+			GrowCircularBufferIfNeeded(NumSamplesPerRenderCallback, NumSamplesPerDeviceCallback);
+
+			// Check to see if the system has requested a larger callback size
+			// (We used a fixed buffer size on Mac, shouldn't change)
+#if PLATFORM_IOS || PLATFORM_TVOS
+			AVAudioSession* AudioSession = [AVAudioSession sharedInstance];
+			double BufferSizeInSec = [AudioSession IOBufferDuration];
+#endif // PLATFORM_IOS || PLATFORM_TVOS
+
+			while (CircularOutputBuffer.Num() < BytesPerSubmittedBuffer)
 			{
 				ReadNextBuffer();
 			}
 
-			// How many bytes we have left over from previous callback
-			int32 SubmittedBufferBytesLeft = BytesPerSubmittedBuffer - SubmittedBytes;
-			int32 OutputBufferBytesLeft = OutputBufferData->mBuffers[0].mDataByteSize;
-			uint8* OutputBufferPtr = (uint8*) OutputBufferData->mBuffers[0].mData;
-			while (OutputBufferBytesLeft > 0)
-			{
-				const int32 BytesToCopy = FMath::Min(SubmittedBufferBytesLeft, OutputBufferBytesLeft);
-
-				FMemory::Memcpy((void*) OutputBufferPtr, SubmittedBufferPtr + SubmittedBytes, BytesToCopy);
-
-				OutputBufferBytesLeft -= BytesToCopy;
-				SubmittedBufferBytesLeft -= BytesToCopy;
-
-				if (SubmittedBufferBytesLeft <= 0)
-				{
-					ReadNextBuffer();
-					SubmittedBytes = 0;
-					SubmittedBufferBytesLeft = BytesPerSubmittedBuffer;
-				}
-				else
-				{
-					SubmittedBytes += BytesToCopy;
-				}
-
-				if (OutputBufferBytesLeft <= 0)
-				{
-					break;
-				}
-				else
-				{
-					OutputBufferPtr += BytesToCopy;
-				}
-			}
+			int32 PopResult = CircularOutputBuffer.Pop((int8*)OutputBufferPtr, BytesPerSubmittedBuffer);
+			check(PopResult == BytesPerSubmittedBuffer);
 		}
 		else
 		{
@@ -632,9 +652,9 @@ namespace Audio
 	{
 		// Get the user data and cast to our FMixerPlatformCoreAudio object
 		FMixerPlatformAudioUnit* me = (FMixerPlatformAudioUnit*) RefCon;
-
+		
 		me->PerformCallback(IOData);
-
+		
 		return noErr;
 	}
 }
diff --git a/Engine/Source/Runtime/AudioMixer/Public/DSP/Dsp.h b/Engine/Source/Runtime/AudioMixer/Public/DSP/Dsp.h
--- a/Engine/Source/Runtime/AudioMixer/Public/DSP/Dsp.h
+++ b/Engine/Source/Runtime/AudioMixer/Public/DSP/Dsp.h
@@ -586,6 +586,12 @@ namespace Audio
 			}
 		}
 
+		// Get the current capacity of the buffer
+		uint32 GetCapacity()
+		{
+			return Capacity;
+		}
+
 		// Get number of samples that can be pushed onto the buffer before it is full.
 		uint32 Remainder()
 		{
