Rolling back CL 7951971 due to mac crash regression

diff --git a/Engine/Source/Runtime/Apple/AudioMixerAudioUnit/Private/AudioMixerPlatformAudioUnit.cpp b/Engine/Source/Runtime/Apple/AudioMixerAudioUnit/Private/AudioMixerPlatformAudioUnit.cpp
--- a/Engine/Source/Runtime/Apple/AudioMixerAudioUnit/Private/AudioMixerPlatformAudioUnit.cpp
+++ b/Engine/Source/Runtime/Apple/AudioMixerAudioUnit/Private/AudioMixerPlatformAudioUnit.cpp
@@ -14,10 +14,11 @@
 
 #include "VorbisAudioInfo.h"
 #include "OpusAudioInfo.h"
+#include "ADPCMAudioInfo.h"
 #endif // #if PLATFORM_IOS || PLATFORM_TVOS
 /*
-	This implementation only depends on the audio units API which allows it to run on MacOS, iOS and tvOS. 
-	
+	This implementation only depends on the audio units API which allows it to run on MacOS, iOS and tvOS.
+
 	For now just assume an iOS configuration (only 2 left and right channels on a single device)
 */
 
@@ -40,15 +41,15 @@ DEFINE_LOG_CATEGORY(LogAudioMixerAudioUnit);
 namespace Audio
 {
 #if PLATFORM_IOS || PLATFORM_TVOS
-	static const int32 DefaultBufferSize = 512;
+	static const int32 DefaultBufferSize = 4096;
 #else
 	static const int32 DefaultBufferSize = 1024;
 	static const int32 AUBufferSize = 256;
 #endif //#if PLATFORM_IOS || PLATFORM_TVOS
 	static const double DefaultSampleRate = 48000.0;
-	
+
 	static int32 SuspendCounter = 0;
-	
+
 	FMixerPlatformAudioUnit::FMixerPlatformAudioUnit()
 		: bInitialized(false)
 		, bInCallback(false)
@@ -56,8 +57,6 @@ namespace Audio
 		, RemainingBytesInCurrentSubmittedBuffer(0)
 		, BytesPerSubmittedBuffer(0)
 		, GraphSampleRate(DefaultSampleRate)
-		, NumSamplesPerRenderCallback(0)
-		, NumSamplesPerDeviceCallback(0)
 	{
 	}
 
@@ -68,23 +67,21 @@ namespace Audio
 			TeardownHardware();
 		}
 	}
-	
+
 	int32 FMixerPlatformAudioUnit::GetNumFrames(const int32 InNumReqestedFrames)
 	{
 #if PLATFORM_IOS || PLATFORM_TVOS
-		return AlignArbitrary(InNumReqestedFrames, 4);
-		
 		AVAudioSession* AudioSession = [AVAudioSession sharedInstance];
-		double BufferSizeInSec = [AudioSession IOBufferDuration];
+		double BufferSizeInSec = [AudioSession preferredIOBufferDuration];
 		double SampleRate = [AudioSession preferredSampleRate];
-		
+
 		if (BufferSizeInSec == 0.0)
 		{
 			return DefaultBufferSize;
 		}
-		
+
 		int32 NumFrames = (int32)(SampleRate * BufferSizeInSec);
-		
+
 		return NumFrames;
 #else
 	   //On MacOS, we hardcode buffer sizes.
@@ -98,53 +95,47 @@ namespace Audio
 		{
 			return false;
 		}
-		
+
 		OSStatus Status;
-		GraphSampleRate = (double) InternalPlatformSettings.SampleRate;
-		UInt32 BufferSize = (UInt32) GetNumFrames(InternalPlatformSettings.CallbackBufferFrameSize);
+		GraphSampleRate = (double) AudioStreamInfo.DeviceInfo.SampleRate;
+		UInt32 BufferSize = (UInt32) GetNumFrames(OpenStreamParams.NumFrames);
 		const int32 NumChannels = 2;
 
 		if (GraphSampleRate == 0)
 		{
 			GraphSampleRate = DefaultSampleRate;
 		}
-	   
+
 		if (BufferSize == 0)
 		{
 			BufferSize = DefaultBufferSize;
 		}
-		
+
 		BytesPerSubmittedBuffer = BufferSize * NumChannels * sizeof(float);
 		check(BytesPerSubmittedBuffer != 0);
-		
+		UE_LOG(LogAudioMixerAudioUnit, Display, TEXT("Bytes per submitted buffer: %d"), BytesPerSubmittedBuffer);
+
 #if PLATFORM_IOS || PLATFORM_TVOS
 		NSError* error;
-		
+
 		AVAudioSession* AudioSession = [AVAudioSession sharedInstance];
 
-		// this sample rate is currently gotten from AudioSession in GetPlatformSettings, so there should be no issue
+		GraphSampleRate = [AudioSession preferredSampleRate];
 		bool Success = [AudioSession setPreferredSampleRate:GraphSampleRate error:&error];
-		
+
 		if (!Success)
 		{
 			UE_LOG(LogAudioMixerAudioUnit, Display, TEXT("Error setting sample rate."));
 		}
-		
+
 		// By calling setPreferredIOBufferDuration, we indicate that we would prefer that the buffer size not change if possible.
-		float AudioMixerBufferSizeInSec = InternalPlatformSettings.CallbackBufferFrameSize / GraphSampleRate; // todo dont hard code
-		Success = [AudioSession setPreferredIOBufferDuration:AudioMixerBufferSizeInSec error: &error];
-		
-		int32 FinalBufferSize = [AudioSession IOBufferDuration] * GraphSampleRate;
-		int32 FinalPreferredBufferSize = [AudioSession preferredIOBufferDuration] * GraphSampleRate;
-		
-		BytesPerSubmittedBuffer = FinalBufferSize * NumChannels * sizeof(float);
-		check(BytesPerSubmittedBuffer != 0);
-		
+		Success = [AudioSession setPreferredIOBufferDuration:[AudioSession preferredIOBufferDuration] error: &error];
+
 		UE_LOG(LogAudioMixerAudioUnit, Display, TEXT("Device Sample Rate: %f"), GraphSampleRate);
 		check(GraphSampleRate != 0);
-		
+
 		Success = [AudioSession setActive:true error:&error];
-		
+
 		if (!Success)
 		{
 			UE_LOG(LogAudioMixerAudioUnit, Display, TEXT("Error starting audio session."));
@@ -153,29 +144,28 @@ namespace Audio
 		AudioObjectID DeviceAudioObjectID;
 		AudioObjectPropertyAddress DevicePropertyAddress;
 		UInt32 AudioDeviceQuerySize;
-		
+
 		//Get Audio Device ID- this will be used throughout initialization to query the audio hardware.
 		DevicePropertyAddress.mSelector = kAudioHardwarePropertyDefaultOutputDevice;
 		DevicePropertyAddress.mScope = kAudioObjectPropertyScopeGlobal;
 		DevicePropertyAddress.mElement = 0;
 		AudioDeviceQuerySize = sizeof(AudioDeviceID);
 		Status = AudioObjectGetPropertyData(kAudioObjectSystemObject, &DevicePropertyAddress, 0, nullptr, &AudioDeviceQuerySize, &DeviceAudioObjectID);
-		
+
 		if(Status != 0)
 		{
 			UE_LOG(LogAudioMixerAudioUnit, Display, TEXT("ERROR setting sample rate to %f"), GraphSampleRate);
 		}
-		
+
 		Status = AudioObjectGetPropertyData(DeviceAudioObjectID, &DevicePropertyAddress, 0, nullptr, &AudioDeviceQuerySize, &GraphSampleRate);
-		
+
 		if(Status == 0)
 		{
 			UE_LOG(LogAudioMixerAudioUnit, Display, TEXT("Sample Rate: %f"), GraphSampleRate);
 		}
 
 #endif // #if PLATFORM_IOS || PLATFORM_TVOS
-		
-		UE_LOG(LogAudioMixerAudioUnit, Display, TEXT("Bytes per submitted buffer: %d"), BytesPerSubmittedBuffer);
+
 
 		// Linear PCM stream format
 		OutputFormat.mFormatID         = kAudioFormatLinearPCM;
@@ -214,14 +204,14 @@ namespace Audio
 			HandleError(TEXT("Failed to initialize audio output node!"), true);
 			return false;
 		}
-		
+
 		Status = AUGraphOpen(AudioUnitGraph);
 		if (Status != noErr)
 		{
 			HandleError(TEXT("Failed to open audio unit graph"), true);
 			return false;
 		}
-		
+
 		Status = AUGraphNodeInfo(AudioUnitGraph, OutputNode, nullptr, &OutputUnit);
 		if (Status != noErr)
 		{
@@ -252,21 +242,25 @@ namespace Audio
 			HandleError(TEXT("Failed to set output format!"), true);
 			return false;
 		}
-		
+
 #endif //#if PLATFORM_MAC
+
+		AudioStreamInfo.NumOutputFrames = BufferSize;
+
 		AudioStreamInfo.DeviceInfo = GetPlatformDeviceInfo();
-		
+
 		AURenderCallbackStruct InputCallback;
 		InputCallback.inputProc = &AudioRenderCallback;
 		InputCallback.inputProcRefCon = this;
 		Status = AUGraphSetNodeInputCallback(AudioUnitGraph,
-												 OutputNode,
+											 OutputNode,
 											 0,
 											 &InputCallback);
 		UE_CLOG(Status != noErr, LogAudioMixerAudioUnit, Error, TEXT("Failed to set input callback for audio output node"));
 
+		OpenStreamParams.NumFrames = BufferSize;
 		AudioStreamInfo.StreamState = EAudioOutputStreamState::Closed;
-		
+
 		bInitialized = true;
 
 		return true;
@@ -284,7 +278,7 @@ namespace Audio
 		{
 			return true;
 		}
-		
+
 		StopAudioStream();
 		CloseAudioStream();
 
@@ -295,7 +289,7 @@ namespace Audio
 		OutputUnit = nullptr;
 
 		bInitialized = false;
-		
+
 		return true;
 	}
 
@@ -307,7 +301,7 @@ namespace Audio
 	bool FMixerPlatformAudioUnit::GetNumOutputDevices(uint32& OutNumOutputDevices)
 	{
 		OutNumOutputDevices = 1;
-		
+
 		return true;
 	}
 
@@ -320,7 +314,7 @@ namespace Audio
 	bool FMixerPlatformAudioUnit::GetDefaultOutputDeviceIndex(uint32& OutDefaultDeviceIndex) const
 	{
 		OutDefaultDeviceIndex = 0;
-		
+
 		return true;
 	}
 
@@ -330,15 +324,12 @@ namespace Audio
 		{
 			return false;
 		}
+
+		AudioStreamInfo.OutputDeviceIndex = Params.OutputDeviceIndex;
+		AudioStreamInfo.AudioMixer = Params.AudioMixer;
+
 		OpenStreamParams = Params;
-		//todo: AudioStreamInfo.SampleRate = OpenStreamParams.SampleRate;
-		AudioStreamInfo.Reset();
-		AudioStreamInfo.OutputDeviceIndex = OpenStreamParams.OutputDeviceIndex;
-		AudioStreamInfo.NumOutputFrames = OpenStreamParams.NumFrames;
-		AudioStreamInfo.NumBuffers = OpenStreamParams.NumBuffers;
-		AudioStreamInfo.AudioMixer = OpenStreamParams.AudioMixer;
-		AudioStreamInfo.DeviceInfo = GetPlatformDeviceInfo();
-		
+
 		// Initialize the audio unit graph
 		OSStatus Status = AUGraphInitialize(AudioUnitGraph);
 		if (Status != noErr)
@@ -347,14 +338,6 @@ namespace Audio
 			return false;
 		}
 
-		// Set up circular buffer between our rendering buffer size and the device's buffer size.
-		// Since we are only using this circular buffer on a single thread, we do not need to add extra slack.
-		NumSamplesPerRenderCallback = AudioStreamInfo.NumOutputFrames * AudioStreamInfo.DeviceInfo.NumChannels;
-		NumSamplesPerDeviceCallback = InternalPlatformSettings.CallbackBufferFrameSize * AudioStreamInfo.DeviceInfo.NumChannels;
-
-		// initial circular buffer capacity is zero, so this initializes it.
-		GrowCircularBufferIfNeeded(NumSamplesPerRenderCallback, NumSamplesPerDeviceCallback);
-
 		AudioStreamInfo.StreamState = EAudioOutputStreamState::Open;
 
 		return true;
@@ -366,9 +349,9 @@ namespace Audio
 		{
 			return false;
 		}
-		
+
 		AudioStreamInfo.StreamState = EAudioOutputStreamState::Closed;
-		
+
 		return true;
 	}
 
@@ -378,9 +361,9 @@ namespace Audio
 		{
 			return false;
 		}
-		
+
 		BeginGeneratingAudio();
-		
+
 		// This will start the render audio callback
 		OSStatus Status = AUGraphStart(AudioUnitGraph);
 		if (Status != noErr)
@@ -394,31 +377,37 @@ namespace Audio
 
 	bool FMixerPlatformAudioUnit::StopAudioStream()
 	{
-		if(!bInitialized || AudioStreamInfo.StreamState != EAudioOutputStreamState::Running)
+		if(!bInitialized)
 		{
 			return false;
 		}
-		
-		AudioStreamInfo.StreamState = EAudioOutputStreamState::Stopping;
-		
-		AUGraphStop(AudioUnitGraph);
 
-		AudioStreamInfo.StreamState = EAudioOutputStreamState::Stopped;
-		
+		if (AudioStreamInfo.StreamState != EAudioOutputStreamState::Stopped && AudioStreamInfo.StreamState != EAudioOutputStreamState::Closed)
+		{
+			AUGraphStop(AudioUnitGraph);
+
+			if (AudioStreamInfo.StreamState == EAudioOutputStreamState::Running)
+			{
+				StopGeneratingAudio();
+			}
+
+			check(AudioStreamInfo.StreamState == EAudioOutputStreamState::Stopped);
+		}
+
+
 		return true;
 	}
 
 	bool FMixerPlatformAudioUnit::MoveAudioStreamToNewAudioDevice(const FString& InNewDeviceId)
 	{
 		//TODO
-		
 		return false;
 	}
 
 	FAudioPlatformDeviceInfo FMixerPlatformAudioUnit::GetPlatformDeviceInfo() const
 	{
 		FAudioPlatformDeviceInfo DeviceInfo;
-		
+
 	#if PLATFORM_IOS || PLATFORM_TVOS
 		AVAudioSession* AudioSession = [AVAudioSession sharedInstance];
 		double SampleRate = [AudioSession preferredSampleRate];
@@ -432,37 +421,37 @@ namespace Audio
 		DeviceInfo.OutputChannelArray[0] = EAudioMixerChannel::FrontLeft;
 		DeviceInfo.OutputChannelArray[1] = EAudioMixerChannel::FrontRight;
 		DeviceInfo.bIsSystemDefault = true;
-		
+
 		return DeviceInfo;
 	}
 
 	void FMixerPlatformAudioUnit::SubmitBuffer(const uint8* Buffer)
 	{
-		if(!Buffer)
-		{
-			return;
-		}
-		
-		const int32 BytesToSubmitToAudioMixer = NumSamplesPerRenderCallback * sizeof(float);
-		
-		int32 PushResult = CircularOutputBuffer.Push((const int8*)Buffer, BytesToSubmitToAudioMixer);
-		check(PushResult == BytesToSubmitToAudioMixer);
+		SubmittedBufferPtr = (uint8*) Buffer;
+		SubmittedBytes = 0;
+		RemainingBytesInCurrentSubmittedBuffer = BytesPerSubmittedBuffer;
 	}
 
-	
 	FName FMixerPlatformAudioUnit::GetRuntimeFormat(USoundWave* InSoundWave)
 	{
+		static const FName NAME_ADPCM(TEXT("ADPCM"));
+		static const FName NAME_OGG(TEXT("OGG"));
+		static const FName NAME_OPUS(TEXT("OPUS"));
+
 #if PLATFORM_IOS || PLATFORM_TVOS
-		static FName NAME_ADPCM(TEXT("ADPCM"));
 		return NAME_ADPCM;
 #else
-		static FName NAME_OPUS(TEXT("OPUS"));
-		
+
+		if (InSoundWave->IsSeekableStreaming())
+		{
+			return NAME_ADPCM;
+		}
+
 		if (InSoundWave->IsStreaming())
 		{
 			return NAME_OPUS;
 		}
-		static FName NAME_OGG(TEXT("OGG"));
+
 		return NAME_OGG;
 #endif
 	}
@@ -478,12 +467,17 @@ namespace Audio
 		return new FADPCMAudioInfo();
 #else
 		check(InSoundWave);
-		
+
+		if (InSoundWave->IsSeekableStreaming())
+		{
+			return new FADPCMAudioInfo();
+		}
+
 		if (InSoundWave->IsStreaming())
 		{
 			return new FOpusAudioInfo();
 		}
-		
+
 #if WITH_OGGVORBIS
 		static const FName NAME_OGG(TEXT("OGG"));
 		if (FPlatformProperties::RequiresCookedData() ? InSoundWave->HasCompressedData(NAME_OGG) : (InSoundWave->GetCompressedData(NAME_OGG) != nullptr))
@@ -513,36 +507,15 @@ namespace Audio
 
 	FAudioPlatformSettings FMixerPlatformAudioUnit::GetPlatformSettings() const
 	{
+		FAudioPlatformSettings Settings;
+		Settings.NumBuffers = 2;
 #if PLATFORM_IOS || PLATFORM_TVOS
-		InternalPlatformSettings = FAudioPlatformSettings::GetPlatformSettings(TEXT("/Script/IOSRuntimeSettings.IOSRuntimeSettings"));
-#else
-		// parsing project settings for mac has not yet been tested
-		InternalPlatformSettings = FAudioPlatformSettings::GetPlatformSettings(TEXT("/Script/OSXRuntimeSettings.OSXRuntimeSettings")); // #maxtodo: temp, to see if it works post-cook
-#endif // #if PLATFORM_IOS || PLATFORM_TVOS
-		
-#if PLATFORM_IOS || PLATFORM_TVOS
-		// Check for command line overrides
-		FString TempString;
-		
-		// Buffer Size
-		if(FParse::Value(FCommandLine::Get(), TEXT("-ForceIOSAudioMixerBufferSize="), TempString))
-		{
-			InternalPlatformSettings.CallbackBufferFrameSize = FCString::Atoi(*TempString);
-		}
-		
-		// NumBuffers
-		if(FParse::Value(FCommandLine::Get(), TEXT("-ForceIOSAudioMixerNumBuffers="), TempString))
-		{
-			InternalPlatformSettings.NumBuffers = FCString::Atoi(*TempString);
-		}
-		
 		AVAudioSession* AudioSession = [AVAudioSession sharedInstance];
-		double PreferredBufferSizeInSec = [AudioSession preferredIOBufferDuration];
-		double BufferSizeInSec = [AudioSession IOBufferDuration];
+		double BufferSizeInSec = [AudioSession preferredIOBufferDuration];
 		double SampleRate = [AudioSession preferredSampleRate];
-		
+
 		int32 NumFrames;
-		
+
 		if (BufferSizeInSec == 0.0)
 		{
 			NumFrames = DefaultBufferSize;
@@ -551,15 +524,17 @@ namespace Audio
 		{
 			NumFrames = (int32)(SampleRate * BufferSizeInSec);
 		}
-		InternalPlatformSettings.SampleRate = SampleRate;
-		
+		Settings.CallbackBufferFrameSize = NumFrames;
+		Settings.SampleRate = SampleRate;
+		Settings.MaxChannels = 32;
+
 #else
-		InternalPlatformSettings.SampleRate = GraphSampleRate;
-		InternalPlatformSettings.CallbackBufferFrameSize = DefaultBufferSize;
-		
-#endif // #if PLATFORM_IOS || PLATFORM_TVOS
-		bInternalPlatformSettingsInitialized = true;
-		return InternalPlatformSettings;
+		Settings.SampleRate = GraphSampleRate;
+		Settings.CallbackBufferFrameSize = DefaultBufferSize;
+
+#endif //#if PLATFORM_IOS || PLATFORM_TVOS
+
+		return Settings;
 	}
 
 	void FMixerPlatformAudioUnit::ResumeContext()
@@ -572,7 +547,7 @@ namespace Audio
 			bSuspended = false;
 		}
 	}
-	
+
 	void FMixerPlatformAudioUnit::SuspendContext()
 	{
 		if (SuspendCounter == 0)
@@ -593,45 +568,50 @@ namespace Audio
 		}
 	}
 
-	void FMixerPlatformAudioUnit::GrowCircularBufferIfNeeded(const int32 InNumSamplesPerRenderCallback, const int32 InNumSamplesPerDeviceCallback)
-	{
-		const int32 MaxCircularBufferCapacity = 2 * sizeof(float) * FMath::Max<int32>(NumSamplesPerRenderCallback, NumSamplesPerDeviceCallback);
-
-		if (CircularOutputBuffer.GetCapacity() < MaxCircularBufferCapacity)
-		{
-			// SetCapacity also zeros-out data
-			CircularOutputBuffer.SetCapacity(MaxCircularBufferCapacity);
-			UE_LOG(LogAudioMixerAudioUnit, Display, TEXT("Growing iOS circular buffer to %i bytes."), MaxCircularBufferCapacity);
-		}
-	}
-
 	bool FMixerPlatformAudioUnit::PerformCallback(AudioBufferList* OutputBufferData)
 	{
 		bInCallback = true;
 
 		if (AudioStreamInfo.StreamState == EAudioOutputStreamState::Running)
 		{
-			// How many bytes we have left over from previous callback
-			BytesPerSubmittedBuffer = OutputBufferData->mBuffers[0].mDataByteSize;
-			uint8* OutputBufferPtr = (uint8*)OutputBufferData->mBuffers[0].mData;
-			
-			NumSamplesPerDeviceCallback = BytesPerSubmittedBuffer / static_cast<float>(sizeof(float));
-			GrowCircularBufferIfNeeded(NumSamplesPerRenderCallback, NumSamplesPerDeviceCallback);
-
-			// Check to see if the system has requested a larger callback size
-			// (We used a fixed buffer size on Mac, shouldn't change)
-#if PLATFORM_IOS || PLATFORM_TVOS
-			AVAudioSession* AudioSession = [AVAudioSession sharedInstance];
-			double BufferSizeInSec = [AudioSession IOBufferDuration];
-#endif // PLATFORM_IOS || PLATFORM_TVOS
-
-			while (CircularOutputBuffer.Num() < BytesPerSubmittedBuffer)
+			if (!SubmittedBufferPtr)
 			{
 				ReadNextBuffer();
 			}
 
-			int32 PopResult = CircularOutputBuffer.Pop((int8*)OutputBufferPtr, BytesPerSubmittedBuffer);
-			check(PopResult == BytesPerSubmittedBuffer);
+			// How many bytes we have left over from previous callback
+			int32 SubmittedBufferBytesLeft = BytesPerSubmittedBuffer - SubmittedBytes;
+			int32 OutputBufferBytesLeft = OutputBufferData->mBuffers[0].mDataByteSize;
+			uint8* OutputBufferPtr = (uint8*) OutputBufferData->mBuffers[0].mData;
+			while (OutputBufferBytesLeft > 0)
+			{
+				const int32 BytesToCopy = FMath::Min(SubmittedBufferBytesLeft, OutputBufferBytesLeft);
+
+				FMemory::Memcpy((void*) OutputBufferPtr, SubmittedBufferPtr + SubmittedBytes, BytesToCopy);
+
+				OutputBufferBytesLeft -= BytesToCopy;
+				SubmittedBufferBytesLeft -= BytesToCopy;
+
+				if (SubmittedBufferBytesLeft <= 0)
+				{
+					ReadNextBuffer();
+					SubmittedBytes = 0;
+					SubmittedBufferBytesLeft = BytesPerSubmittedBuffer;
+				}
+				else
+				{
+					SubmittedBytes += BytesToCopy;
+				}
+
+				if (OutputBufferBytesLeft <= 0)
+				{
+					break;
+				}
+				else
+				{
+					OutputBufferPtr += BytesToCopy;
+				}
+			}
 		}
 		else
 		{
@@ -652,9 +632,9 @@ namespace Audio
 	{
 		// Get the user data and cast to our FMixerPlatformCoreAudio object
 		FMixerPlatformAudioUnit* me = (FMixerPlatformAudioUnit*) RefCon;
-		
+
 		me->PerformCallback(IOData);
-		
+
 		return noErr;
 	}
 }
diff --git a/Engine/Source/Runtime/Apple/AudioMixerAudioUnit/Private/AudioMixerPlatformAudioUnit.h b/Engine/Source/Runtime/Apple/AudioMixerAudioUnit/Private/AudioMixerPlatformAudioUnit.h
--- a/Engine/Source/Runtime/Apple/AudioMixerAudioUnit/Private/AudioMixerPlatformAudioUnit.h
+++ b/Engine/Source/Runtime/Apple/AudioMixerAudioUnit/Private/AudioMixerPlatformAudioUnit.h
@@ -3,7 +3,7 @@
 #pragma once
 
 #include "AudioMixer.h"
-#include "DSP/Dsp.h"
+
 #include <AudioToolbox/AudioToolbox.h>
 #include <AudioUnit/AudioUnit.h>
 
@@ -40,18 +40,11 @@ namespace Audio
 		virtual ICompressedAudioInfo* CreateCompressedAudioInfo(USoundWave* InSoundWave) override;
 		virtual FString GetDefaultDeviceName() override;
 		virtual FAudioPlatformSettings GetPlatformSettings() const override;
-		virtual int32 GetNumFrames(const int32 InNumReqestedFrames) override;
-		virtual void ResumeContext() override;
-		virtual void SuspendContext() override;
-		
-		/** Whether or not the platform supports realtime decompression. */
-		virtual bool SupportsRealtimeDecompression() const override { return true; }
-		
-		/** Whether or not the platform disables caching of decompressed PCM data (i.e. to save memory on fixed memory platforms) */
-		virtual bool DisablePCMAudioCaching() const override { return true; }
-		
-	  //~ End IAudioMixerPlatformInterface
-		
+        virtual int32 GetNumFrames(const int32 InNumReqestedFrames) override;
+        virtual void ResumeContext() override;
+        virtual void SuspendContext() override;
+        //~ End IAudioMixerPlatformInterface
+        
 	private:
 		AudioStreamBasicDescription OutputFormat;
 
@@ -66,27 +59,14 @@ namespace Audio
 		AUGraph     AudioUnitGraph;
 		AUNode      OutputNode;
 		AudioUnit	OutputUnit;
-		uint8*      SubmittedBufferPtr;
-		int32 SubmittedBytes = 0;
-
-		int32       RemainingBytesInCurrentSubmittedBuffer;
-		int32       BytesPerSubmittedBuffer;
-		
-		double GraphSampleRate;
-
-		// We may have to grow the circular buffer capacity since Audio Unit callback size is not guaranteed to be constant
-		// Currently, this just zero's-out and reallocates, so it will pop. (We always keep largest capacity)
-		void GrowCircularBufferIfNeeded(const int32 InNumSamplesPerRenderCallback, const int32 InNumSamplesPerDeviceCallback);
-
-		// This buffer is pushed to and popped from in the SubmitBuffer callback. 
-		// This is required for devices that require frame counts per callback that are not powers of two.
-		Audio::TCircularAudioBuffer<int8> CircularOutputBuffer;
-
-		int32 NumSamplesPerRenderCallback;
-		int32 NumSamplesPerDeviceCallback;
-		mutable bool bInternalPlatformSettingsInitialized{ false };
-		mutable FAudioPlatformSettings InternalPlatformSettings;
-		
+        uint8*      SubmittedBufferPtr;
+        int32 SubmittedBytes = 0;
+
+        int32       RemainingBytesInCurrentSubmittedBuffer;
+        int32       BytesPerSubmittedBuffer;
+        
+        double GraphSampleRate;
+        
 		bool PerformCallback(AudioBufferList* OutputBufferData);
 		void HandleError(const TCHAR* InLogOutput, bool bTeardown = true);
 		static OSStatus AudioRenderCallback(void* RefCon, AudioUnitRenderActionFlags* ActionFlags,
diff --git a/Engine/Source/Runtime/AudioMixer/Public/DSP/Dsp.h b/Engine/Source/Runtime/AudioMixer/Public/DSP/Dsp.h
--- a/Engine/Source/Runtime/AudioMixer/Public/DSP/Dsp.h
+++ b/Engine/Source/Runtime/AudioMixer/Public/DSP/Dsp.h
@@ -586,12 +586,6 @@ namespace Audio
 			}
 		}
 
-		// Get the current capacity of the buffer
-		uint32 GetCapacity()
-		{
-			return Capacity;
-		}
-
 		// Get number of samples that can be pushed onto the buffer before it is full.
 		uint32 Remainder()
 		{
