Non-unity compile fix in MetalRHI.cpp

diff --git a/Engine/Source/Runtime/Apple/MetalRHI/Private/MetalRHI.cpp b/Engine/Source/Runtime/Apple/MetalRHI/Private/MetalRHI.cpp
--- a/Engine/Source/Runtime/Apple/MetalRHI/Private/MetalRHI.cpp
+++ b/Engine/Source/Runtime/Apple/MetalRHI/Private/MetalRHI.cpp
@@ -13,6 +13,7 @@
 #elif PLATFORM_MAC
 #include "Mac/MacApplication.h"
 #include "HAL/PlatformApplicationMisc.h"
+#include "HAL/PlatformFilemanager.h"
 #endif
 #include "MetalProfiler.h"
 #include "GenericPlatform/GenericPlatformDriver.h"
