Seperating iOS and Mac backends (addition of AudioMixerPlatformCoreAudio)

diff --git a/Engine/Config/Mac/MacEngine.ini b/Engine/Config/Mac/MacEngine.ini
--- a/Engine/Config/Mac/MacEngine.ini
+++ b/Engine/Config/Mac/MacEngine.ini
@@ -2,7 +2,7 @@
 AudioDeviceModuleName=CoreAudio
 
 ; Defining below allows switching to audio mixer using -audiomixer commandline
-AudioMixerModuleName=AudioMixerAudioUnit
+AudioMixerModuleName=AudioMixerCoreAudio
 
 ; Defines a platform-specific volume headroom (in dB) for audio to provide better platform consistency with respect to volume levels.
 PlatformHeadroomDB=-6
@@ -39,4 +39,4 @@ GameHUDBind=(Key=Tilde,bShift=True,bCtrl=False,bAlt=False,bCmd=False)
 PlatformRequiresDataCrypto=True
 
 [DevOptions.Shaders]
-MaxShaderJobBatchSize=25
\ No newline at end of file
+MaxShaderJobBatchSize=25
diff --git a/Engine/Source/Runtime/Apple/AudioMixerAudioUnit/Private/AudioMixerPlatformAudioUnit.cpp b/Engine/Source/Runtime/Apple/AudioMixerAudioUnit/Private/AudioMixerPlatformAudioUnit.cpp
--- a/Engine/Source/Runtime/Apple/AudioMixerAudioUnit/Private/AudioMixerPlatformAudioUnit.cpp
+++ b/Engine/Source/Runtime/Apple/AudioMixerAudioUnit/Private/AudioMixerPlatformAudioUnit.cpp
@@ -7,59 +7,44 @@
 #include "CoreGlobals.h"
 #include "CoreMinimal.h"
 #include "Misc/ConfigCacheIni.h"
-#if PLATFORM_IOS || PLATFORM_TVOS
 #include "ADPCMAudioInfo.h"
 
-#else
-
-#include "VorbisAudioInfo.h"
-#include "OpusAudioInfo.h"
-#include "ADPCMAudioInfo.h"
-#endif // #if PLATFORM_IOS || PLATFORM_TVOS
 /*
-	This implementation only depends on the audio units API which allows it to run on MacOS, iOS and tvOS.
-
-	For now just assume an iOS configuration (only 2 left and right channels on a single device)
-*/
+ This implementation only depends on the audio units API which allows it to run on MacOS, iOS and tvOS.
+ 
+ For now just assume an iOS configuration (only 2 left and right channels on a single device)
+ */
 
 /**
-* CoreAudio System Headers
-*/
+ * CoreAudio System Headers
+ */
 #include <AudioToolbox/AudioToolbox.h>
 #include <AudioUnit/AudioUnit.h>
-#if PLATFORM_IOS || PLATFORM_TVOS
 #include <AVFoundation/AVAudioSession.h>
-#elif  PLATFORM_MAC
-#include <CoreAudio/AudioHardware.h>
-#else
-#error Invalid Platform!
-#endif // #if PLATFORM_IOS || PLATFORM_TVOS
 
 DECLARE_LOG_CATEGORY_EXTERN(LogAudioMixerAudioUnit, Log, All);
 DEFINE_LOG_CATEGORY(LogAudioMixerAudioUnit);
 
 namespace Audio
 {
-#if PLATFORM_IOS || PLATFORM_TVOS
-	static const int32 DefaultBufferSize = 4096;
-#else
-	static const int32 DefaultBufferSize = 1024;
-	static const int32 AUBufferSize = 256;
-#endif //#if PLATFORM_IOS || PLATFORM_TVOS
-	static const double DefaultSampleRate = 48000.0;
+	static const int32 DefaultBufferSize = 512;
 
+	static const double DefaultSampleRate = 48000.0;
+	
 	static int32 SuspendCounter = 0;
-
+	
 	FMixerPlatformAudioUnit::FMixerPlatformAudioUnit()
-		: bInitialized(false)
-		, bInCallback(false)
-		, SubmittedBufferPtr(nullptr)
-		, RemainingBytesInCurrentSubmittedBuffer(0)
-		, BytesPerSubmittedBuffer(0)
-		, GraphSampleRate(DefaultSampleRate)
+	: bInitialized(false)
+	, bInCallback(false)
+	, SubmittedBufferPtr(nullptr)
+	, RemainingBytesInCurrentSubmittedBuffer(0)
+	, BytesPerSubmittedBuffer(0)
+	, GraphSampleRate(DefaultSampleRate)
+	, NumSamplesPerRenderCallback(0)
+	, NumSamplesPerDeviceCallback(0)
 	{
 	}
-
+	
 	FMixerPlatformAudioUnit::~FMixerPlatformAudioUnit()
 	{
 		if (bInitialized)
@@ -67,134 +52,96 @@ namespace Audio
 			TeardownHardware();
 		}
 	}
-
+	
 	int32 FMixerPlatformAudioUnit::GetNumFrames(const int32 InNumReqestedFrames)
 	{
-#if PLATFORM_IOS || PLATFORM_TVOS
-		AVAudioSession* AudioSession = [AVAudioSession sharedInstance];
-		double BufferSizeInSec = [AudioSession preferredIOBufferDuration];
-		double SampleRate = [AudioSession preferredSampleRate];
-
-		if (BufferSizeInSec == 0.0)
-		{
-			return DefaultBufferSize;
-		}
-
-		int32 NumFrames = (int32)(SampleRate * BufferSizeInSec);
-
-		return NumFrames;
-#else
-	   //On MacOS, we hardcode buffer sizes.
-		return DefaultBufferSize;
-#endif
+		return AlignArbitrary(InNumReqestedFrames, 4);
 	}
-
+	
 	bool FMixerPlatformAudioUnit::InitializeHardware()
 	{
 		if (bInitialized)
 		{
 			return false;
 		}
-
+		
 		OSStatus Status;
-		GraphSampleRate = (double) AudioStreamInfo.DeviceInfo.SampleRate;
-		UInt32 BufferSize = (UInt32) GetNumFrames(OpenStreamParams.NumFrames);
-		const int32 NumChannels = 2;
+		GraphSampleRate = (double) InternalPlatformSettings.SampleRate;
+		UInt32 BufferSize = (UInt32) GetNumFrames(InternalPlatformSettings.CallbackBufferFrameSize);
 
+		const int32 NumChannels = 2;
+		
 		if (GraphSampleRate == 0)
 		{
 			GraphSampleRate = DefaultSampleRate;
 		}
-
+		
 		if (BufferSize == 0)
 		{
 			BufferSize = DefaultBufferSize;
 		}
-
+		
 		BytesPerSubmittedBuffer = BufferSize * NumChannels * sizeof(float);
 		check(BytesPerSubmittedBuffer != 0);
-		UE_LOG(LogAudioMixerAudioUnit, Display, TEXT("Bytes per submitted buffer: %d"), BytesPerSubmittedBuffer);
-
-#if PLATFORM_IOS || PLATFORM_TVOS
+		
 		NSError* error;
-
+		
 		AVAudioSession* AudioSession = [AVAudioSession sharedInstance];
-
-		GraphSampleRate = [AudioSession preferredSampleRate];
+		
+		// this sample rate is currently gotten from AudioSession in GetPlatformSettings, so there should be no issue
 		bool Success = [AudioSession setPreferredSampleRate:GraphSampleRate error:&error];
-
+		
 		if (!Success)
 		{
 			UE_LOG(LogAudioMixerAudioUnit, Display, TEXT("Error setting sample rate."));
 		}
-
+		
 		// By calling setPreferredIOBufferDuration, we indicate that we would prefer that the buffer size not change if possible.
-		Success = [AudioSession setPreferredIOBufferDuration:[AudioSession preferredIOBufferDuration] error: &error];
-
+		float AudioMixerBufferSizeInSec = InternalPlatformSettings.CallbackBufferFrameSize / GraphSampleRate;
+		Success = [AudioSession setPreferredIOBufferDuration:AudioMixerBufferSizeInSec error: &error];
+		
+		int32 FinalBufferSize = [AudioSession IOBufferDuration] * GraphSampleRate;
+		int32 FinalPreferredBufferSize = [AudioSession preferredIOBufferDuration] * GraphSampleRate;
+		
+		BytesPerSubmittedBuffer = FinalBufferSize * NumChannels * sizeof(float);
+		check(BytesPerSubmittedBuffer != 0);
+		
 		UE_LOG(LogAudioMixerAudioUnit, Display, TEXT("Device Sample Rate: %f"), GraphSampleRate);
 		check(GraphSampleRate != 0);
-
+		
 		Success = [AudioSession setActive:true error:&error];
-
+		
 		if (!Success)
 		{
 			UE_LOG(LogAudioMixerAudioUnit, Display, TEXT("Error starting audio session."));
 		}
-#else
-		AudioObjectID DeviceAudioObjectID;
-		AudioObjectPropertyAddress DevicePropertyAddress;
-		UInt32 AudioDeviceQuerySize;
-
-		//Get Audio Device ID- this will be used throughout initialization to query the audio hardware.
-		DevicePropertyAddress.mSelector = kAudioHardwarePropertyDefaultOutputDevice;
-		DevicePropertyAddress.mScope = kAudioObjectPropertyScopeGlobal;
-		DevicePropertyAddress.mElement = 0;
-		AudioDeviceQuerySize = sizeof(AudioDeviceID);
-		Status = AudioObjectGetPropertyData(kAudioObjectSystemObject, &DevicePropertyAddress, 0, nullptr, &AudioDeviceQuerySize, &DeviceAudioObjectID);
-
-		if(Status != 0)
-		{
-			UE_LOG(LogAudioMixerAudioUnit, Display, TEXT("ERROR setting sample rate to %f"), GraphSampleRate);
-		}
-
-		Status = AudioObjectGetPropertyData(DeviceAudioObjectID, &DevicePropertyAddress, 0, nullptr, &AudioDeviceQuerySize, &GraphSampleRate);
-
-		if(Status == 0)
-		{
-			UE_LOG(LogAudioMixerAudioUnit, Display, TEXT("Sample Rate: %f"), GraphSampleRate);
-		}
-
-#endif // #if PLATFORM_IOS || PLATFORM_TVOS
-
-
+		
+		UE_LOG(LogAudioMixerAudioUnit, Display, TEXT("Bytes per submitted buffer: %d"), BytesPerSubmittedBuffer);
+		
 		// Linear PCM stream format
 		OutputFormat.mFormatID         = kAudioFormatLinearPCM;
-		OutputFormat.mFormatFlags	   = kAudioFormatFlagIsFloat | kAudioFormatFlagsNativeEndian | kAudioFormatFlagIsPacked;
+		OutputFormat.mFormatFlags       = kAudioFormatFlagIsFloat | kAudioFormatFlagsNativeEndian | kAudioFormatFlagIsPacked;
 		OutputFormat.mChannelsPerFrame = 2;
 		OutputFormat.mBytesPerFrame    = sizeof(float) * OutputFormat.mChannelsPerFrame;
 		OutputFormat.mFramesPerPacket  = 1;
 		OutputFormat.mBytesPerPacket   = OutputFormat.mBytesPerFrame * OutputFormat.mFramesPerPacket;
 		OutputFormat.mBitsPerChannel   = 8 * sizeof(float);
 		OutputFormat.mSampleRate       = GraphSampleRate;
-
+		
 		Status = NewAUGraph(&AudioUnitGraph);
 		if (Status != noErr)
 		{
 			HandleError(TEXT("Failed to create audio unit graph!"));
 			return false;
 		}
-
+		
 		AudioComponentDescription UnitDescription;
-
+		
 		// Setup audio output unit
 		UnitDescription.componentType         = kAudioUnitType_Output;
-#if PLATFORM_IOS || PLATFORM_TVOS
+		
 		//On iOS, we'll use the RemoteIO AudioUnit.
 		UnitDescription.componentSubType      = kAudioUnitSubType_RemoteIO;
-#else //PLATFORM_MAC
-		//On MacOS, we'll use the DefaultOutput AudioUnit.
-		UnitDescription.componentSubType      = kAudioUnitSubType_DefaultOutput;
-#endif // #if PLATFORM_IOS || PLATFORM_TVOS
 		UnitDescription.componentManufacturer = kAudioUnitManufacturer_Apple;
 		UnitDescription.componentFlags        = 0;
 		UnitDescription.componentFlagsMask    = 0;
@@ -204,21 +151,21 @@ namespace Audio
 			HandleError(TEXT("Failed to initialize audio output node!"), true);
 			return false;
 		}
-
+		
 		Status = AUGraphOpen(AudioUnitGraph);
 		if (Status != noErr)
 		{
 			HandleError(TEXT("Failed to open audio unit graph"), true);
 			return false;
 		}
-
+		
 		Status = AUGraphNodeInfo(AudioUnitGraph, OutputNode, nullptr, &OutputUnit);
 		if (Status != noErr)
 		{
 			HandleError(TEXT("Failed to retrieve output unit reference!"), true);
 			return false;
 		}
-
+		
 		Status = AudioUnitSetProperty(OutputUnit,
 									  kAudioUnitProperty_StreamFormat,
 									  kAudioUnitScope_Input,
@@ -230,25 +177,10 @@ namespace Audio
 			HandleError(TEXT("Failed to set output format!"), true);
 			return false;
 		}
-
-#if PLATFORM_MAC
-		DevicePropertyAddress.mSelector = kAudioDevicePropertyBufferFrameSize;
-		DevicePropertyAddress.mScope = kAudioObjectPropertyScopeGlobal;
-		DevicePropertyAddress.mElement = 0;
-		AudioDeviceQuerySize = sizeof(AUBufferSize);
-		Status = AudioObjectSetPropertyData(DeviceAudioObjectID, &DevicePropertyAddress, 0, nullptr, AudioDeviceQuerySize, &AUBufferSize);
-		if(Status != 0)
-		{
-			HandleError(TEXT("Failed to set output format!"), true);
-			return false;
-		}
-
-#endif //#if PLATFORM_MAC
-
-		AudioStreamInfo.NumOutputFrames = BufferSize;
+		
 
 		AudioStreamInfo.DeviceInfo = GetPlatformDeviceInfo();
-
+		
 		AURenderCallbackStruct InputCallback;
 		InputCallback.inputProc = &AudioRenderCallback;
 		InputCallback.inputProcRefCon = this;
@@ -257,79 +189,90 @@ namespace Audio
 											 0,
 											 &InputCallback);
 		UE_CLOG(Status != noErr, LogAudioMixerAudioUnit, Error, TEXT("Failed to set input callback for audio output node"));
-
-		OpenStreamParams.NumFrames = BufferSize;
+		
 		AudioStreamInfo.StreamState = EAudioOutputStreamState::Closed;
-
+		
 		bInitialized = true;
-
+		
 		return true;
 	}
-
+	
 	bool FMixerPlatformAudioUnit::CheckAudioDeviceChange()
 	{
 		//TODO
 		return false;
 	}
-
+	
 	bool FMixerPlatformAudioUnit::TeardownHardware()
 	{
 		if(!bInitialized)
 		{
 			return true;
 		}
-
+		
 		StopAudioStream();
 		CloseAudioStream();
-
+		
 		DisposeAUGraph(AudioUnitGraph);
-
+		
 		AudioUnitGraph = nullptr;
 		OutputNode = -1;
 		OutputUnit = nullptr;
-
+		
 		bInitialized = false;
-
+		
 		return true;
 	}
-
+	
 	bool FMixerPlatformAudioUnit::IsInitialized() const
 	{
 		return bInitialized;
 	}
-
+	
 	bool FMixerPlatformAudioUnit::GetNumOutputDevices(uint32& OutNumOutputDevices)
 	{
 		OutNumOutputDevices = 1;
-
+		
 		return true;
 	}
-
+	
 	bool FMixerPlatformAudioUnit::GetOutputDeviceInfo(const uint32 InDeviceIndex, FAudioPlatformDeviceInfo& OutInfo)
 	{
 		OutInfo = AudioStreamInfo.DeviceInfo;
 		return true;
 	}
-
+	
 	bool FMixerPlatformAudioUnit::GetDefaultOutputDeviceIndex(uint32& OutDefaultDeviceIndex) const
 	{
 		OutDefaultDeviceIndex = 0;
-
+		
 		return true;
 	}
-
+	
 	bool FMixerPlatformAudioUnit::OpenAudioStream(const FAudioMixerOpenStreamParams& Params)
 	{
 		if (!bInitialized || AudioStreamInfo.StreamState != EAudioOutputStreamState::Closed)
 		{
 			return false;
 		}
-
-		AudioStreamInfo.OutputDeviceIndex = Params.OutputDeviceIndex;
-		AudioStreamInfo.AudioMixer = Params.AudioMixer;
-
+		
 		OpenStreamParams = Params;
-
+		//todo: AudioStreamInfo.SampleRate = OpenStreamParams.SampleRate;
+		AudioStreamInfo.Reset();
+		AudioStreamInfo.OutputDeviceIndex = OpenStreamParams.OutputDeviceIndex;
+		AudioStreamInfo.NumOutputFrames = OpenStreamParams.NumFrames;
+		AudioStreamInfo.NumBuffers = OpenStreamParams.NumBuffers;
+		AudioStreamInfo.AudioMixer = OpenStreamParams.AudioMixer;
+		AudioStreamInfo.DeviceInfo = GetPlatformDeviceInfo();
+		
+		// Set up circular buffer between our rendering buffer size and the device's buffer size.
+		// Since we are only using this circular buffer on a single thread, we do not need to add extra slack.
+		NumSamplesPerRenderCallback = AudioStreamInfo.NumOutputFrames * AudioStreamInfo.DeviceInfo.NumChannels;
+		NumSamplesPerDeviceCallback = InternalPlatformSettings.CallbackBufferFrameSize * AudioStreamInfo.DeviceInfo.NumChannels;
+		
+		// initial circular buffer capacity is zero, so this initializes it.
+		GrowCircularBufferIfNeeded(NumSamplesPerRenderCallback, NumSamplesPerDeviceCallback);
+		
 		// Initialize the audio unit graph
 		OSStatus Status = AUGraphInitialize(AudioUnitGraph);
 		if (Status != noErr)
@@ -337,33 +280,32 @@ namespace Audio
 			HandleError(TEXT("Failed to initialize audio graph!"), true);
 			return false;
 		}
-
 		AudioStreamInfo.StreamState = EAudioOutputStreamState::Open;
-
+		
 		return true;
 	}
-
+	
 	bool FMixerPlatformAudioUnit::CloseAudioStream()
 	{
 		if (!bInitialized || (AudioStreamInfo.StreamState != EAudioOutputStreamState::Open && AudioStreamInfo.StreamState != EAudioOutputStreamState::Stopped))
 		{
 			return false;
 		}
-
+		
 		AudioStreamInfo.StreamState = EAudioOutputStreamState::Closed;
-
+		
 		return true;
 	}
-
+	
 	bool FMixerPlatformAudioUnit::StartAudioStream()
 	{
 		if (!bInitialized || (AudioStreamInfo.StreamState != EAudioOutputStreamState::Open && AudioStreamInfo.StreamState != EAudioOutputStreamState::Stopped))
 		{
 			return false;
 		}
-
+		
 		BeginGeneratingAudio();
-
+		
 		// This will start the render audio callback
 		OSStatus Status = AUGraphStart(AudioUnitGraph);
 		if (Status != noErr)
@@ -371,151 +313,111 @@ namespace Audio
 			HandleError(TEXT("Failed to start audio graph!"), true);
 			return false;
 		}
-
+		
 		return true;
 	}
-
+	
 	bool FMixerPlatformAudioUnit::StopAudioStream()
 	{
-		if(!bInitialized)
+		if(!bInitialized || AudioStreamInfo.StreamState != EAudioOutputStreamState::Running)
 		{
 			return false;
 		}
-
-		if (AudioStreamInfo.StreamState != EAudioOutputStreamState::Stopped && AudioStreamInfo.StreamState != EAudioOutputStreamState::Closed)
-		{
-			AUGraphStop(AudioUnitGraph);
-
-			if (AudioStreamInfo.StreamState == EAudioOutputStreamState::Running)
-			{
-				StopGeneratingAudio();
-			}
-
-			check(AudioStreamInfo.StreamState == EAudioOutputStreamState::Stopped);
-		}
-
-
+		
+		AudioStreamInfo.StreamState = EAudioOutputStreamState::Stopping;
+		
+		AUGraphStop(AudioUnitGraph);
+		
+		AudioStreamInfo.StreamState = EAudioOutputStreamState::Stopped;
+		
 		return true;
 	}
-
+	
 	bool FMixerPlatformAudioUnit::MoveAudioStreamToNewAudioDevice(const FString& InNewDeviceId)
 	{
 		//TODO
+		
 		return false;
 	}
-
+	
 	FAudioPlatformDeviceInfo FMixerPlatformAudioUnit::GetPlatformDeviceInfo() const
 	{
 		FAudioPlatformDeviceInfo DeviceInfo;
-
-	#if PLATFORM_IOS || PLATFORM_TVOS
+		
 		AVAudioSession* AudioSession = [AVAudioSession sharedInstance];
 		double SampleRate = [AudioSession preferredSampleRate];
 		DeviceInfo.SampleRate = (int32)SampleRate;
-#else
-		DeviceInfo.SampleRate = GraphSampleRate;
-#endif
 		DeviceInfo.NumChannels = 2;
 		DeviceInfo.Format = EAudioMixerStreamDataFormat::Float;
 		DeviceInfo.OutputChannelArray.SetNum(2);
 		DeviceInfo.OutputChannelArray[0] = EAudioMixerChannel::FrontLeft;
 		DeviceInfo.OutputChannelArray[1] = EAudioMixerChannel::FrontRight;
 		DeviceInfo.bIsSystemDefault = true;
-
+		
 		return DeviceInfo;
 	}
-
+	
 	void FMixerPlatformAudioUnit::SubmitBuffer(const uint8* Buffer)
 	{
-		SubmittedBufferPtr = (uint8*) Buffer;
-		SubmittedBytes = 0;
-		RemainingBytesInCurrentSubmittedBuffer = BytesPerSubmittedBuffer;
-	}
+		
+		if(!Buffer)
+		{
+			return;
+		}
 
+		const int32 BytesToSubmitToAudioMixer = NumSamplesPerRenderCallback * sizeof(float);
+		
+		int32 PushResult = CircularOutputBuffer.Push((const int8*)Buffer, BytesToSubmitToAudioMixer);
+		check(PushResult == BytesToSubmitToAudioMixer);
+	}
+	
 	FName FMixerPlatformAudioUnit::GetRuntimeFormat(USoundWave* InSoundWave)
 	{
-		static const FName NAME_ADPCM(TEXT("ADPCM"));
-		static const FName NAME_OGG(TEXT("OGG"));
-		static const FName NAME_OPUS(TEXT("OPUS"));
-
-#if PLATFORM_IOS || PLATFORM_TVOS
+		static FName NAME_ADPCM(TEXT("ADPCM"));
 		return NAME_ADPCM;
-#else
-
-		if (InSoundWave->IsSeekableStreaming())
-		{
-			return NAME_ADPCM;
-		}
-
-		if (InSoundWave->IsStreaming())
-		{
-			return NAME_OPUS;
-		}
-
-		return NAME_OGG;
-#endif
 	}
-
+	
 	bool FMixerPlatformAudioUnit::HasCompressedAudioInfoClass(USoundWave* InSoundWave)
 	{
 		return true;
 	}
-
+	
 	ICompressedAudioInfo* FMixerPlatformAudioUnit::CreateCompressedAudioInfo(USoundWave* InSoundWave)
 	{
-#if PLATFORM_IOS || PLATFORM_TVOS
 		return new FADPCMAudioInfo();
-#else
-		check(InSoundWave);
-
-		if (InSoundWave->IsSeekableStreaming())
-		{
-			return new FADPCMAudioInfo();
-		}
-
-		if (InSoundWave->IsStreaming())
-		{
-			return new FOpusAudioInfo();
-		}
-
-#if WITH_OGGVORBIS
-		static const FName NAME_OGG(TEXT("OGG"));
-		if (FPlatformProperties::RequiresCookedData() ? InSoundWave->HasCompressedData(NAME_OGG) : (InSoundWave->GetCompressedData(NAME_OGG) != nullptr))
-		{
-			ICompressedAudioInfo* CompressedInfo = new FVorbisAudioInfo();
-			if (!CompressedInfo)
-			{
-				UE_LOG(LogAudio, Error, TEXT("Failed to create new FVorbisAudioInfo for SoundWave %s: out of memory."), *InSoundWave->GetName());
-				return nullptr;
-			}
-			return CompressedInfo;
-		}
-		else
-		{
-			return nullptr;
-		}
-#else
-		return nullptr;
-#endif // WITH_OGGVORBIS
-#endif // PLATFORM_IOS || PLATFORM_TVOS
 	}
-
+	
 	FString FMixerPlatformAudioUnit::GetDefaultDeviceName()
 	{
 		return FString();
 	}
-
+	
 	FAudioPlatformSettings FMixerPlatformAudioUnit::GetPlatformSettings() const
 	{
-		FAudioPlatformSettings Settings;
-		Settings.NumBuffers = 2;
-#if PLATFORM_IOS || PLATFORM_TVOS
+		InternalPlatformSettings = FAudioPlatformSettings::GetPlatformSettings(TEXT("/Script/IOSRuntimeSettings.IOSRuntimeSettings"));
+		
+		// Check for command line overrides
+		FString TempString;
+		
+		// Buffer Size
+		if(FParse::Value(FCommandLine::Get(), TEXT("-ForceIOSAudioMixerBufferSize="), TempString))
+		{
+			InternalPlatformSettings.CallbackBufferFrameSize = FCString::Atoi(*TempString);
+		}
+		
+		// NumBuffers
+		if(FParse::Value(FCommandLine::Get(), TEXT("-ForceIOSAudioMixerNumBuffers="), TempString))
+		{
+			InternalPlatformSettings.NumBuffers = FCString::Atoi(*TempString);
+		}
+		
 		AVAudioSession* AudioSession = [AVAudioSession sharedInstance];
-		double BufferSizeInSec = [AudioSession preferredIOBufferDuration];
+		double PreferredBufferSizeInSec = [AudioSession preferredIOBufferDuration];
+		double BufferSizeInSec = [AudioSession IOBufferDuration];
 		double SampleRate = [AudioSession preferredSampleRate];
-
+		
 		int32 NumFrames;
-
+		
 		if (BufferSizeInSec == 0.0)
 		{
 			NumFrames = DefaultBufferSize;
@@ -524,19 +426,11 @@ namespace Audio
 		{
 			NumFrames = (int32)(SampleRate * BufferSizeInSec);
 		}
-		Settings.CallbackBufferFrameSize = NumFrames;
-		Settings.SampleRate = SampleRate;
-		Settings.MaxChannels = 32;
-
-#else
-		Settings.SampleRate = GraphSampleRate;
-		Settings.CallbackBufferFrameSize = DefaultBufferSize;
+		InternalPlatformSettings.SampleRate = SampleRate;
 
-#endif //#if PLATFORM_IOS || PLATFORM_TVOS
-
-		return Settings;
+		return InternalPlatformSettings;
 	}
-
+	
 	void FMixerPlatformAudioUnit::ResumeContext()
 	{
 		if (SuspendCounter > 0)
@@ -547,7 +441,7 @@ namespace Audio
 			bSuspended = false;
 		}
 	}
-
+	
 	void FMixerPlatformAudioUnit::SuspendContext()
 	{
 		if (SuspendCounter == 0)
@@ -558,7 +452,7 @@ namespace Audio
 			bSuspended = true;
 		}
 	}
-
+	
 	void FMixerPlatformAudioUnit::HandleError(const TCHAR* InLogOutput, bool bTeardown)
 	{
 		UE_LOG(LogAudioMixerAudioUnit, Log, TEXT("%s"), InLogOutput);
@@ -567,74 +461,67 @@ namespace Audio
 			TeardownHardware();
 		}
 	}
-
+	
+	void FMixerPlatformAudioUnit::GrowCircularBufferIfNeeded(const int32 InNumSamplesPerRenderCallback, const int32 InNumSamplesPerDeviceCallback)
+	{
+		const int32 MaxCircularBufferCapacity = 2 * sizeof(float) * FMath::Max<int32>(NumSamplesPerRenderCallback, NumSamplesPerDeviceCallback);
+		
+		if (CircularOutputBuffer.GetCapacity() < MaxCircularBufferCapacity)
+		{
+			// SetCapacity also zeros-out data
+			CircularOutputBuffer.SetCapacity(MaxCircularBufferCapacity);
+			UE_LOG(LogAudioMixerAudioUnit, Display, TEXT("Growing iOS circular buffer to %i bytes."), MaxCircularBufferCapacity);
+		}
+	}
+	
 	bool FMixerPlatformAudioUnit::PerformCallback(AudioBufferList* OutputBufferData)
 	{
 		bInCallback = true;
-
+		
 		if (AudioStreamInfo.StreamState == EAudioOutputStreamState::Running)
 		{
-			if (!SubmittedBufferPtr)
-			{
-				ReadNextBuffer();
-			}
-
 			// How many bytes we have left over from previous callback
-			int32 SubmittedBufferBytesLeft = BytesPerSubmittedBuffer - SubmittedBytes;
-			int32 OutputBufferBytesLeft = OutputBufferData->mBuffers[0].mDataByteSize;
-			uint8* OutputBufferPtr = (uint8*) OutputBufferData->mBuffers[0].mData;
-			while (OutputBufferBytesLeft > 0)
+			BytesPerSubmittedBuffer = OutputBufferData->mBuffers[0].mDataByteSize;
+			uint8* OutputBufferPtr = (uint8*)OutputBufferData->mBuffers[0].mData;
+			
+			NumSamplesPerDeviceCallback = BytesPerSubmittedBuffer / static_cast<float>(sizeof(float));
+			GrowCircularBufferIfNeeded(NumSamplesPerRenderCallback, NumSamplesPerDeviceCallback);
+			
+			// Check to see if the system has requested a larger callback size
+			AVAudioSession* AudioSession = [AVAudioSession sharedInstance];
+			double BufferSizeInSec = [AudioSession IOBufferDuration];
+			
+			
+			while (CircularOutputBuffer.Num() < BytesPerSubmittedBuffer)
 			{
-				const int32 BytesToCopy = FMath::Min(SubmittedBufferBytesLeft, OutputBufferBytesLeft);
-
-				FMemory::Memcpy((void*) OutputBufferPtr, SubmittedBufferPtr + SubmittedBytes, BytesToCopy);
-
-				OutputBufferBytesLeft -= BytesToCopy;
-				SubmittedBufferBytesLeft -= BytesToCopy;
-
-				if (SubmittedBufferBytesLeft <= 0)
-				{
-					ReadNextBuffer();
-					SubmittedBytes = 0;
-					SubmittedBufferBytesLeft = BytesPerSubmittedBuffer;
-				}
-				else
-				{
-					SubmittedBytes += BytesToCopy;
-				}
-
-				if (OutputBufferBytesLeft <= 0)
-				{
-					break;
-				}
-				else
-				{
-					OutputBufferPtr += BytesToCopy;
-				}
+				ReadNextBuffer();
 			}
+			
+			int32 PopResult = CircularOutputBuffer.Pop((int8*)OutputBufferPtr, BytesPerSubmittedBuffer);
+			check(PopResult == BytesPerSubmittedBuffer);
 		}
-		else
+		else // AudioStreamInfo.StreamState != EAudioOutputStreamState::Running
 		{
 			for (uint32 bufferItr = 0; bufferItr < OutputBufferData->mNumberBuffers; ++bufferItr)
 			{
 				memset(OutputBufferData->mBuffers[bufferItr].mData, 0, OutputBufferData->mBuffers[bufferItr].mDataByteSize);
 			}
 		}
-
+		
 		bInCallback = false;
-
+		
 		return true;
 	}
-
+	
 	OSStatus FMixerPlatformAudioUnit::AudioRenderCallback(void* RefCon, AudioUnitRenderActionFlags* ActionFlags,
 														  const AudioTimeStamp* TimeStamp, UInt32 BusNumber,
 														  UInt32 NumFrames, AudioBufferList* IOData)
 	{
 		// Get the user data and cast to our FMixerPlatformCoreAudio object
 		FMixerPlatformAudioUnit* me = (FMixerPlatformAudioUnit*) RefCon;
-
+		
 		me->PerformCallback(IOData);
-
+		
 		return noErr;
 	}
 }
diff --git a/Engine/Source/Runtime/Apple/AudioMixerAudioUnit/Private/AudioMixerPlatformAudioUnit.h b/Engine/Source/Runtime/Apple/AudioMixerAudioUnit/Private/AudioMixerPlatformAudioUnit.h
--- a/Engine/Source/Runtime/Apple/AudioMixerAudioUnit/Private/AudioMixerPlatformAudioUnit.h
+++ b/Engine/Source/Runtime/Apple/AudioMixerAudioUnit/Private/AudioMixerPlatformAudioUnit.h
@@ -3,22 +3,22 @@
 #pragma once
 
 #include "AudioMixer.h"
-
+#include "DSP/Dsp.h"
 #include <AudioToolbox/AudioToolbox.h>
 #include <AudioUnit/AudioUnit.h>
 
 // Any platform defines
 namespace Audio
 {
-
+	
 	class FMixerPlatformAudioUnit : public IAudioMixerPlatformInterface
 	{
-
+		
 	public:
-
+		
 		FMixerPlatformAudioUnit();
 		~FMixerPlatformAudioUnit();
-
+		
 		//~ Begin IAudioMixerPlatformInterface
 		virtual EAudioMixerPlatformApi::Type GetPlatformApi() const override { return EAudioMixerPlatformApi::AudioUnit; }
 		virtual bool InitializeHardware() override;
@@ -40,11 +40,18 @@ namespace Audio
 		virtual ICompressedAudioInfo* CreateCompressedAudioInfo(USoundWave* InSoundWave) override;
 		virtual FString GetDefaultDeviceName() override;
 		virtual FAudioPlatformSettings GetPlatformSettings() const override;
-        virtual int32 GetNumFrames(const int32 InNumReqestedFrames) override;
-        virtual void ResumeContext() override;
-        virtual void SuspendContext() override;
-        //~ End IAudioMixerPlatformInterface
-        
+		virtual int32 GetNumFrames(const int32 InNumReqestedFrames) override;
+		virtual void ResumeContext() override;
+		virtual void SuspendContext() override;
+		
+		/** Whether or not the platform supports realtime decompression. */
+		virtual bool SupportsRealtimeDecompression() const override { return true; }
+		
+		/** Whether or not the platform disables caching of decompressed PCM data (i.e. to save memory on fixed memory platforms) */
+		virtual bool DisablePCMAudioCaching() const override { return true; }
+		
+		//~ End IAudioMixerPlatformInterface
+		
 	private:
 		AudioStreamBasicDescription OutputFormat;
 
@@ -58,15 +65,28 @@ namespace Audio
 
 		AUGraph     AudioUnitGraph;
 		AUNode      OutputNode;
-		AudioUnit	OutputUnit;
-        uint8*      SubmittedBufferPtr;
-        int32 SubmittedBytes = 0;
-
-        int32       RemainingBytesInCurrentSubmittedBuffer;
-        int32       BytesPerSubmittedBuffer;
-        
-        double GraphSampleRate;
-        
+		AudioUnit    OutputUnit;
+		uint8*      SubmittedBufferPtr;
+		int32 SubmittedBytes = 0;
+		
+		int32       RemainingBytesInCurrentSubmittedBuffer;
+		int32       BytesPerSubmittedBuffer;
+		
+		double GraphSampleRate;
+		
+		// We may have to grow the circular buffer capacity since Audio Unit callback size is not guaranteed to be constant
+		// Currently, this just zero's-out and reallocates, so it will pop. (We always keep largest capacity)
+		void GrowCircularBufferIfNeeded(const int32 InNumSamplesPerRenderCallback, const int32 InNumSamplesPerDeviceCallback);
+		
+		// This buffer is pushed to and popped from in the SubmitBuffer callback.
+		// This is required for devices that require frame counts per callback that are not powers of two.
+		Audio::TCircularAudioBuffer<int8> CircularOutputBuffer;
+		
+		int32 NumSamplesPerRenderCallback;
+		int32 NumSamplesPerDeviceCallback;
+		mutable bool bInternalPlatformSettingsInitialized{ false };
+		mutable FAudioPlatformSettings InternalPlatformSettings;
+		
 		bool PerformCallback(AudioBufferList* OutputBufferData);
 		void HandleError(const TCHAR* InLogOutput, bool bTeardown = true);
 		static OSStatus AudioRenderCallback(void* RefCon, AudioUnitRenderActionFlags* ActionFlags,
diff --git a/Engine/Source/Runtime/Apple/AudioMixerCoreAudio/AudioMixerCoreAudio.Build.cs b/Engine/Source/Runtime/Apple/AudioMixerCoreAudio/AudioMixerCoreAudio.Build.cs
new file mode 100644
--- /dev/null
+++ b/Engine/Source/Runtime/Apple/AudioMixerCoreAudio/AudioMixerCoreAudio.Build.cs
@@ -0,0 +1,59 @@
+// Copyright 1998-2019 Epic Games, Inc. All Rights Reserved.
+
+
+
+using UnrealBuildTool;
+
+
+
+public class AudioMixerCoreAudio : ModuleRules
+{
+	public AudioMixerCoreAudio(ReadOnlyTargetRules Target) : base(Target)
+	{
+		PrivateIncludePathModuleNames.Add("TargetPlatform");
+		PublicIncludePaths.Add("Runtime/AudioMixer/Public");
+		PrivateIncludePaths.Add("Runtime/AudioMixer/Private");
+
+
+		PrivateDependencyModuleNames.AddRange(
+			new string[] {
+				"Core",
+				"CoreUObject",
+				"Engine",
+			}
+			);
+
+		PrecompileForTargets = PrecompileTargetsType.None;
+		PrivateDependencyModuleNames.Add("AudioMixer");
+
+
+
+		AddEngineThirdPartyPrivateStaticDependencies(Target, 
+
+			"UEOgg",
+			"Vorbis",
+			"VorbisFile"
+			);
+
+
+
+		PublicFrameworks.AddRange(new string[]
+		{
+			"AudioToolbox",
+			"CoreAudio"
+		});
+
+		if (Target.Platform == UnrealTargetPlatform.Mac)
+        {
+            PublicFrameworks.Add("AudioUnit");
+        }
+
+
+		PublicDefinitions.Add("WITH_OGGVORBIS=1");
+
+		if (Target.Platform == UnrealTargetPlatform.Mac)
+		{
+			PrecompileForTargets = PrecompileTargetsType.Any;
+		}
+	}
+}
diff --git a/Engine/Source/Runtime/Apple/AudioMixerCoreAudio/Private/AudioMixerCoreAudio.cpp b/Engine/Source/Runtime/Apple/AudioMixerCoreAudio/Private/AudioMixerCoreAudio.cpp
new file mode 100644
--- /dev/null
+++ b/Engine/Source/Runtime/Apple/AudioMixerCoreAudio/Private/AudioMixerCoreAudio.cpp
@@ -0,0 +1,20 @@
+// Copyright 1998-2019 Epic Games, Inc. All Rights Reserved.
+/**
+	Concrete implementation of FAudioDevice for Apple's CoreAudio
+*/
+#include "AudioMixer.h"
+#include "AudioMixerDevice.h"
+#include "AudioMixerPlatformCoreAudio.h"
+
+
+class FAudioMixerModuleCoreAudio : public IAudioDeviceModule
+{
+public:
+	virtual bool IsAudioMixerModule() const override { return true; }
+
+	virtual FAudioDevice* CreateAudioDevice() override
+	{
+		return new Audio::FMixerDevice(new Audio::FMixerPlatformCoreAudio());
+	}
+};
+IMPLEMENT_MODULE(FAudioMixerModuleCoreAudio, AudioMixerCoreAudio);
diff --git a/Engine/Source/Runtime/Apple/AudioMixerCoreAudio/Private/AudioMixerPlatformCoreAudio.cpp b/Engine/Source/Runtime/Apple/AudioMixerCoreAudio/Private/AudioMixerPlatformCoreAudio.cpp
new file mode 100644
--- /dev/null
+++ b/Engine/Source/Runtime/Apple/AudioMixerCoreAudio/Private/AudioMixerPlatformCoreAudio.cpp
@@ -0,0 +1,1050 @@
+// Copyright 1998-2019 Epic Games, Inc. All Rights Reserved.
+
+
+
+#include "AudioMixerPlatformCoreAudio.h"
+
+#include "Modules/ModuleManager.h"
+
+#include "AudioMixer.h"
+
+#include "AudioMixerDevice.h"
+
+#include "CoreGlobals.h"
+
+#include "CoreMinimal.h"
+
+#include "Misc/ConfigCacheIni.h"
+
+#include "VorbisAudioInfo.h"
+
+#include "OpusAudioInfo.h"
+
+#include "ADPCMAudioInfo.h"
+
+
+
+/**
+
+* CoreAudio System Headers
+
+*/
+
+#include <AudioToolbox/AudioToolbox.h>
+
+#include <AudioUnit/AudioUnit.h>
+
+#include <CoreAudio/AudioHardware.h>
+
+
+
+DECLARE_LOG_CATEGORY_EXTERN(LogAudioMixerAudioUnit, Log, All);
+
+DEFINE_LOG_CATEGORY(LogAudioMixerAudioUnit);
+
+
+
+namespace Audio
+
+{
+
+	static const int32 DefaultBufferSize = 1024;
+
+	static const int32 AUBufferSize = 256;
+
+	static const double DefaultSampleRate = 48000.0;
+
+
+
+	static int32 SuspendCounter = 0;
+
+
+
+	FMixerPlatformCoreAudio::FMixerPlatformCoreAudio()
+
+		: bInitialized(false)
+
+		, bInCallback(false)
+
+		, SubmittedBufferPtr(nullptr)
+
+		, RemainingBytesInCurrentSubmittedBuffer(0)
+
+		, BytesPerSubmittedBuffer(0)
+
+		, GraphSampleRate(DefaultSampleRate)
+
+	{
+
+	}
+
+
+
+	FMixerPlatformCoreAudio::~FMixerPlatformCoreAudio()
+
+	{
+
+		if (bInitialized)
+
+		{
+
+			TeardownHardware();
+
+		}
+
+	}
+
+
+
+	int32 FMixerPlatformCoreAudio::GetNumFrames(const int32 InNumReqestedFrames)
+
+	{
+
+	   //On MacOS, we hardcode buffer sizes.
+
+		return DefaultBufferSize;
+
+	}
+
+
+
+	bool FMixerPlatformCoreAudio::InitializeHardware()
+
+	{
+
+		if (bInitialized)
+
+		{
+
+			return false;
+
+		}
+
+
+
+		OSStatus Status;
+
+		GraphSampleRate = (double) AudioStreamInfo.DeviceInfo.SampleRate;
+
+		UInt32 BufferSize = (UInt32) GetNumFrames(OpenStreamParams.NumFrames);
+
+		const int32 NumChannels = 2;
+
+
+
+		if (GraphSampleRate == 0)
+
+		{
+
+			GraphSampleRate = DefaultSampleRate;
+
+		}
+
+
+
+		if (BufferSize == 0)
+
+		{
+
+			BufferSize = DefaultBufferSize;
+
+		}
+
+
+
+		BytesPerSubmittedBuffer = BufferSize * NumChannels * sizeof(float);
+
+		check(BytesPerSubmittedBuffer != 0);
+
+		UE_LOG(LogAudioMixerAudioUnit, Display, TEXT("Bytes per submitted buffer: %d"), BytesPerSubmittedBuffer);
+
+
+
+		AudioObjectID DeviceAudioObjectID;
+
+		AudioObjectPropertyAddress DevicePropertyAddress;
+
+		UInt32 AudioDeviceQuerySize;
+
+
+
+		//Get Audio Device ID- this will be used throughout initialization to query the audio hardware.
+
+		DevicePropertyAddress.mSelector = kAudioHardwarePropertyDefaultOutputDevice;
+
+		DevicePropertyAddress.mScope = kAudioObjectPropertyScopeGlobal;
+
+		DevicePropertyAddress.mElement = 0;
+
+		AudioDeviceQuerySize = sizeof(AudioDeviceID);
+
+		Status = AudioObjectGetPropertyData(kAudioObjectSystemObject, &DevicePropertyAddress, 0, nullptr, &AudioDeviceQuerySize, &DeviceAudioObjectID);
+
+
+
+		if(Status != 0)
+
+		{
+
+			UE_LOG(LogAudioMixerAudioUnit, Display, TEXT("ERROR setting sample rate to %f"), GraphSampleRate);
+
+		}
+
+
+
+		Status = AudioObjectGetPropertyData(DeviceAudioObjectID, &DevicePropertyAddress, 0, nullptr, &AudioDeviceQuerySize, &GraphSampleRate);
+
+
+
+		if(Status == 0)
+
+		{
+
+			UE_LOG(LogAudioMixerAudioUnit, Display, TEXT("Sample Rate: %f"), GraphSampleRate);
+
+		}
+
+
+
+
+
+		// Linear PCM stream format
+
+		OutputFormat.mFormatID         = kAudioFormatLinearPCM;
+
+		OutputFormat.mFormatFlags	   = kAudioFormatFlagIsFloat | kAudioFormatFlagsNativeEndian | kAudioFormatFlagIsPacked;
+
+		OutputFormat.mChannelsPerFrame = 2;
+
+		OutputFormat.mBytesPerFrame    = sizeof(float) * OutputFormat.mChannelsPerFrame;
+
+		OutputFormat.mFramesPerPacket  = 1;
+
+		OutputFormat.mBytesPerPacket   = OutputFormat.mBytesPerFrame * OutputFormat.mFramesPerPacket;
+
+		OutputFormat.mBitsPerChannel   = 8 * sizeof(float);
+
+		OutputFormat.mSampleRate       = GraphSampleRate;
+
+
+
+		Status = NewAUGraph(&AudioUnitGraph);
+
+		if (Status != noErr)
+
+		{
+
+			HandleError(TEXT("Failed to create audio unit graph!"));
+
+			return false;
+
+		}
+
+
+
+		AudioComponentDescription UnitDescription;
+
+
+
+		// Setup audio output unit
+
+		UnitDescription.componentType         = kAudioUnitType_Output;
+
+		//On MacOS, we'll use the DefaultOutput AudioUnit.
+
+		UnitDescription.componentSubType      = kAudioUnitSubType_DefaultOutput;
+
+		UnitDescription.componentManufacturer = kAudioUnitManufacturer_Apple;
+
+		UnitDescription.componentFlags        = 0;
+
+		UnitDescription.componentFlagsMask    = 0;
+
+		Status = AUGraphAddNode(AudioUnitGraph, &UnitDescription, &OutputNode);
+
+		if (Status != noErr)
+
+		{
+
+			HandleError(TEXT("Failed to initialize audio output node!"), true);
+
+			return false;
+
+		}
+
+
+
+		Status = AUGraphOpen(AudioUnitGraph);
+
+		if (Status != noErr)
+
+		{
+
+			HandleError(TEXT("Failed to open audio unit graph"), true);
+
+			return false;
+
+		}
+
+
+
+		Status = AUGraphNodeInfo(AudioUnitGraph, OutputNode, nullptr, &OutputUnit);
+
+		if (Status != noErr)
+
+		{
+
+			HandleError(TEXT("Failed to retrieve output unit reference!"), true);
+
+			return false;
+
+		}
+
+
+
+		Status = AudioUnitSetProperty(OutputUnit,
+
+									  kAudioUnitProperty_StreamFormat,
+
+									  kAudioUnitScope_Input,
+
+									  0,
+
+									  &OutputFormat,
+
+									  sizeof(AudioStreamBasicDescription));
+
+		if (Status != noErr)
+
+		{
+
+			HandleError(TEXT("Failed to set output format!"), true);
+
+			return false;
+
+		}
+
+
+
+		DevicePropertyAddress.mSelector = kAudioDevicePropertyBufferFrameSize;
+
+		DevicePropertyAddress.mScope = kAudioObjectPropertyScopeGlobal;
+
+		DevicePropertyAddress.mElement = 0;
+
+		AudioDeviceQuerySize = sizeof(AUBufferSize);
+
+		Status = AudioObjectSetPropertyData(DeviceAudioObjectID, &DevicePropertyAddress, 0, nullptr, AudioDeviceQuerySize, &AUBufferSize);
+
+		if(Status != 0)
+
+		{
+
+			HandleError(TEXT("Failed to set output format!"), true);
+
+			return false;
+
+		}
+
+
+
+		AudioStreamInfo.NumOutputFrames = BufferSize;
+
+
+
+		AudioStreamInfo.DeviceInfo = GetPlatformDeviceInfo();
+
+
+
+		AURenderCallbackStruct InputCallback;
+
+		InputCallback.inputProc = &AudioRenderCallback;
+
+		InputCallback.inputProcRefCon = this;
+
+		Status = AUGraphSetNodeInputCallback(AudioUnitGraph,
+
+											 OutputNode,
+
+											 0,
+
+											 &InputCallback);
+
+		UE_CLOG(Status != noErr, LogAudioMixerAudioUnit, Error, TEXT("Failed to set input callback for audio output node"));
+
+
+
+		OpenStreamParams.NumFrames = BufferSize;
+
+		AudioStreamInfo.StreamState = EAudioOutputStreamState::Closed;
+
+
+
+		bInitialized = true;
+
+
+
+		return true;
+
+	}
+
+
+
+	bool FMixerPlatformCoreAudio::CheckAudioDeviceChange()
+
+	{
+
+		//TODO
+
+		return false;
+
+	}
+
+
+
+	bool FMixerPlatformCoreAudio::TeardownHardware()
+
+	{
+
+		if(!bInitialized)
+
+		{
+
+			return true;
+
+		}
+
+
+
+		StopAudioStream();
+
+		CloseAudioStream();
+
+
+
+		DisposeAUGraph(AudioUnitGraph);
+
+
+
+		AudioUnitGraph = nullptr;
+
+		OutputNode = -1;
+
+		OutputUnit = nullptr;
+
+
+
+		bInitialized = false;
+
+
+
+		return true;
+
+	}
+
+
+
+	bool FMixerPlatformCoreAudio::IsInitialized() const
+
+	{
+
+		return bInitialized;
+
+	}
+
+
+
+	bool FMixerPlatformCoreAudio::GetNumOutputDevices(uint32& OutNumOutputDevices)
+
+	{
+
+		OutNumOutputDevices = 1;
+
+
+
+		return true;
+
+	}
+
+
+
+	bool FMixerPlatformCoreAudio::GetOutputDeviceInfo(const uint32 InDeviceIndex, FAudioPlatformDeviceInfo& OutInfo)
+
+	{
+
+		OutInfo = AudioStreamInfo.DeviceInfo;
+
+		return true;
+
+	}
+
+
+
+	bool FMixerPlatformCoreAudio::GetDefaultOutputDeviceIndex(uint32& OutDefaultDeviceIndex) const
+
+	{
+
+		OutDefaultDeviceIndex = 0;
+
+
+
+		return true;
+
+	}
+
+
+
+	bool FMixerPlatformCoreAudio::OpenAudioStream(const FAudioMixerOpenStreamParams& Params)
+
+	{
+
+		if (!bInitialized || AudioStreamInfo.StreamState != EAudioOutputStreamState::Closed)
+
+		{
+
+			return false;
+
+		}
+
+
+
+		AudioStreamInfo.OutputDeviceIndex = Params.OutputDeviceIndex;
+
+		AudioStreamInfo.AudioMixer = Params.AudioMixer;
+
+
+
+		OpenStreamParams = Params;
+
+
+
+		// Initialize the audio unit graph
+
+		OSStatus Status = AUGraphInitialize(AudioUnitGraph);
+
+		if (Status != noErr)
+
+		{
+
+			HandleError(TEXT("Failed to initialize audio graph!"), true);
+
+			return false;
+
+		}
+
+
+
+		AudioStreamInfo.StreamState = EAudioOutputStreamState::Open;
+
+
+
+		return true;
+
+	}
+
+
+
+	bool FMixerPlatformCoreAudio::CloseAudioStream()
+
+	{
+
+		if (!bInitialized || (AudioStreamInfo.StreamState != EAudioOutputStreamState::Open && AudioStreamInfo.StreamState != EAudioOutputStreamState::Stopped))
+
+		{
+
+			return false;
+
+		}
+
+
+
+		AudioStreamInfo.StreamState = EAudioOutputStreamState::Closed;
+
+
+
+		return true;
+
+	}
+
+
+
+	bool FMixerPlatformCoreAudio::StartAudioStream()
+
+	{
+
+		if (!bInitialized || (AudioStreamInfo.StreamState != EAudioOutputStreamState::Open && AudioStreamInfo.StreamState != EAudioOutputStreamState::Stopped))
+
+		{
+
+			return false;
+
+		}
+
+
+
+		BeginGeneratingAudio();
+
+
+
+		// This will start the render audio callback
+
+		OSStatus Status = AUGraphStart(AudioUnitGraph);
+
+		if (Status != noErr)
+
+		{
+
+			HandleError(TEXT("Failed to start audio graph!"), true);
+
+			return false;
+
+		}
+
+
+
+		return true;
+
+	}
+
+
+
+	bool FMixerPlatformCoreAudio::StopAudioStream()
+
+	{
+
+		if(!bInitialized)
+
+		{
+
+			return false;
+
+		}
+
+
+
+		if (AudioStreamInfo.StreamState != EAudioOutputStreamState::Stopped && AudioStreamInfo.StreamState != EAudioOutputStreamState::Closed)
+
+		{
+
+			AUGraphStop(AudioUnitGraph);
+
+
+
+			if (AudioStreamInfo.StreamState == EAudioOutputStreamState::Running)
+
+			{
+
+				StopGeneratingAudio();
+
+			}
+
+
+
+			check(AudioStreamInfo.StreamState == EAudioOutputStreamState::Stopped);
+
+		}
+
+
+
+
+
+		return true;
+
+	}
+
+
+
+	bool FMixerPlatformCoreAudio::MoveAudioStreamToNewAudioDevice(const FString& InNewDeviceId)
+
+	{
+
+		//TODO
+
+		return false;
+
+	}
+
+
+
+	FAudioPlatformDeviceInfo FMixerPlatformCoreAudio::GetPlatformDeviceInfo() const
+
+	{
+
+		FAudioPlatformDeviceInfo DeviceInfo;
+
+
+
+		DeviceInfo.SampleRate = GraphSampleRate;
+
+		DeviceInfo.NumChannels = 2;
+
+		DeviceInfo.Format = EAudioMixerStreamDataFormat::Float;
+
+		DeviceInfo.OutputChannelArray.SetNum(2);
+
+		DeviceInfo.OutputChannelArray[0] = EAudioMixerChannel::FrontLeft;
+
+		DeviceInfo.OutputChannelArray[1] = EAudioMixerChannel::FrontRight;
+
+		DeviceInfo.bIsSystemDefault = true;
+
+
+
+		return DeviceInfo;
+
+	}
+
+
+
+	void FMixerPlatformCoreAudio::SubmitBuffer(const uint8* Buffer)
+
+	{
+
+		SubmittedBufferPtr = (uint8*) Buffer;
+
+		SubmittedBytes = 0;
+
+		RemainingBytesInCurrentSubmittedBuffer = BytesPerSubmittedBuffer;
+
+	}
+
+
+
+	FName FMixerPlatformCoreAudio::GetRuntimeFormat(USoundWave* InSoundWave)
+
+	{
+
+		static const FName NAME_ADPCM(TEXT("ADPCM"));
+
+		static const FName NAME_OGG(TEXT("OGG"));
+
+		static const FName NAME_OPUS(TEXT("OPUS"));
+
+
+
+		if (InSoundWave->IsSeekableStreaming())
+
+		{
+
+			return NAME_ADPCM;
+
+		}
+
+
+
+		if (InSoundWave->IsStreaming())
+
+		{
+
+			return NAME_OPUS;
+
+		}
+
+
+
+		return NAME_OGG;
+
+	}
+
+
+
+	bool FMixerPlatformCoreAudio::HasCompressedAudioInfoClass(USoundWave* InSoundWave)
+
+	{
+
+		return true;
+
+	}
+
+
+
+	ICompressedAudioInfo* FMixerPlatformCoreAudio::CreateCompressedAudioInfo(USoundWave* InSoundWave)
+
+	{
+
+		check(InSoundWave);
+
+
+
+		if (InSoundWave->IsSeekableStreaming())
+
+		{
+
+			return new FADPCMAudioInfo();
+
+		}
+
+
+
+		if (InSoundWave->IsStreaming())
+
+		{
+
+			return new FOpusAudioInfo();
+
+		}
+
+
+
+		static const FName NAME_OGG(TEXT("OGG"));
+
+		if (FPlatformProperties::RequiresCookedData() ? InSoundWave->HasCompressedData(NAME_OGG) : (InSoundWave->GetCompressedData(NAME_OGG) != nullptr))
+
+		{
+
+			ICompressedAudioInfo* CompressedInfo = new FVorbisAudioInfo();
+
+			if (!CompressedInfo)
+
+			{
+
+				UE_LOG(LogAudio, Error, TEXT("Failed to create new FVorbisAudioInfo for SoundWave %s: out of memory."), *InSoundWave->GetName());
+
+				return nullptr;
+
+			}
+
+			return CompressedInfo;
+
+		}
+
+
+
+		return nullptr;
+
+	}
+
+
+
+	FString FMixerPlatformCoreAudio::GetDefaultDeviceName()
+
+	{
+
+		return FString();
+
+	}
+
+
+
+	FAudioPlatformSettings FMixerPlatformCoreAudio::GetPlatformSettings() const
+
+	{
+
+		FAudioPlatformSettings Settings;
+
+		Settings.NumBuffers = 2;
+
+		Settings.SampleRate = GraphSampleRate;
+
+		Settings.CallbackBufferFrameSize = DefaultBufferSize;
+
+
+
+		return Settings;
+
+	}
+
+
+
+	void FMixerPlatformCoreAudio::ResumeContext()
+
+	{
+
+		if (SuspendCounter > 0)
+
+		{
+
+			FPlatformAtomics::InterlockedDecrement(&SuspendCounter);
+
+			AUGraphStart(AudioUnitGraph);
+
+			UE_LOG(LogAudioMixerAudioUnit, Display, TEXT("Resuming Audio"));
+
+			bSuspended = false;
+
+		}
+
+	}
+
+
+
+	void FMixerPlatformCoreAudio::SuspendContext()
+
+	{
+
+		if (SuspendCounter == 0)
+
+		{
+
+			FPlatformAtomics::InterlockedIncrement(&SuspendCounter);
+
+			AUGraphStop(AudioUnitGraph);
+
+			UE_LOG(LogAudioMixerAudioUnit, Display, TEXT("Suspending Audio"));
+
+			bSuspended = true;
+
+		}
+
+	}
+
+
+
+	void FMixerPlatformCoreAudio::HandleError(const TCHAR* InLogOutput, bool bTeardown)
+
+	{
+
+		UE_LOG(LogAudioMixerAudioUnit, Log, TEXT("%s"), InLogOutput);
+
+		if (bTeardown)
+
+		{
+
+			TeardownHardware();
+
+		}
+
+	}
+
+
+
+	bool FMixerPlatformCoreAudio::PerformCallback(AudioBufferList* OutputBufferData)
+
+	{
+
+		bInCallback = true;
+
+
+
+		if (AudioStreamInfo.StreamState == EAudioOutputStreamState::Running)
+
+		{
+
+			if (!SubmittedBufferPtr)
+
+			{
+
+				ReadNextBuffer();
+
+			}
+
+
+
+			// How many bytes we have left over from previous callback
+
+			int32 SubmittedBufferBytesLeft = BytesPerSubmittedBuffer - SubmittedBytes;
+
+			int32 OutputBufferBytesLeft = OutputBufferData->mBuffers[0].mDataByteSize;
+
+			uint8* OutputBufferPtr = (uint8*) OutputBufferData->mBuffers[0].mData;
+
+			while (OutputBufferBytesLeft > 0)
+
+			{
+
+				const int32 BytesToCopy = FMath::Min(SubmittedBufferBytesLeft, OutputBufferBytesLeft);
+
+
+
+				FMemory::Memcpy((void*) OutputBufferPtr, SubmittedBufferPtr + SubmittedBytes, BytesToCopy);
+
+
+
+				OutputBufferBytesLeft -= BytesToCopy;
+
+				SubmittedBufferBytesLeft -= BytesToCopy;
+
+
+
+				if (SubmittedBufferBytesLeft <= 0)
+
+				{
+
+					ReadNextBuffer();
+
+					SubmittedBytes = 0;
+
+					SubmittedBufferBytesLeft = BytesPerSubmittedBuffer;
+
+				}
+
+				else
+
+				{
+
+					SubmittedBytes += BytesToCopy;
+
+				}
+
+
+
+				if (OutputBufferBytesLeft <= 0)
+
+				{
+
+					break;
+
+				}
+
+				else
+
+				{
+
+					OutputBufferPtr += BytesToCopy;
+
+				}
+
+			}
+
+		}
+
+		else
+
+		{
+
+			for (uint32 bufferItr = 0; bufferItr < OutputBufferData->mNumberBuffers; ++bufferItr)
+
+			{
+
+				memset(OutputBufferData->mBuffers[bufferItr].mData, 0, OutputBufferData->mBuffers[bufferItr].mDataByteSize);
+
+			}
+
+		}
+
+
+
+		bInCallback = false;
+
+
+
+		return true;
+
+	}
+
+
+
+	OSStatus FMixerPlatformCoreAudio::AudioRenderCallback(void* RefCon, AudioUnitRenderActionFlags* ActionFlags,
+
+														  const AudioTimeStamp* TimeStamp, UInt32 BusNumber,
+
+														  UInt32 NumFrames, AudioBufferList* IOData)
+
+	{
+
+		// Get the user data and cast to our FMixerPlatformCoreAudio object
+
+		FMixerPlatformCoreAudio* me = (FMixerPlatformCoreAudio*) RefCon;
+
+
+
+		me->PerformCallback(IOData);
+
+
+
+		return noErr;
+
+	}
+
+}
+
diff --git a/Engine/Source/Runtime/Apple/AudioMixerCoreAudio/Private/AudioMixerPlatformCoreAudio.h b/Engine/Source/Runtime/Apple/AudioMixerCoreAudio/Private/AudioMixerPlatformCoreAudio.h
new file mode 100644
--- /dev/null
+++ b/Engine/Source/Runtime/Apple/AudioMixerCoreAudio/Private/AudioMixerPlatformCoreAudio.h
@@ -0,0 +1,156 @@
+// Copyright 1998-2019 Epic Games, Inc. All Rights Reserved.
+
+
+
+#pragma once
+
+
+
+#include "AudioMixer.h"
+
+
+
+#include <AudioToolbox/AudioToolbox.h>
+
+#include <AudioUnit/AudioUnit.h>
+
+
+
+// Any platform defines
+
+namespace Audio
+
+{
+
+
+
+	class FMixerPlatformCoreAudio : public IAudioMixerPlatformInterface
+
+	{
+
+
+
+	public:
+
+
+
+		FMixerPlatformCoreAudio();
+
+		~FMixerPlatformCoreAudio();
+
+
+
+		//~ Begin IAudioMixerPlatformInterface
+
+		virtual EAudioMixerPlatformApi::Type GetPlatformApi() const override { return EAudioMixerPlatformApi::CoreAudio; }
+
+		virtual bool InitializeHardware() override;
+
+		virtual bool CheckAudioDeviceChange() override;
+
+		virtual bool TeardownHardware() override;
+
+		virtual bool IsInitialized() const override;
+
+		virtual bool GetNumOutputDevices(uint32& OutNumOutputDevices) override;
+
+		virtual bool GetOutputDeviceInfo(const uint32 InDeviceIndex, FAudioPlatformDeviceInfo& OutInfo) override;
+
+		virtual bool GetDefaultOutputDeviceIndex(uint32& OutDefaultDeviceIndex) const override;
+
+		virtual bool OpenAudioStream(const FAudioMixerOpenStreamParams& Params) override;
+
+		virtual bool CloseAudioStream() override;
+
+		virtual bool StartAudioStream() override;
+
+		virtual bool StopAudioStream() override;
+
+		virtual bool MoveAudioStreamToNewAudioDevice(const FString& InNewDeviceId) override;
+
+		virtual FAudioPlatformDeviceInfo GetPlatformDeviceInfo() const override;
+
+		virtual void SubmitBuffer(const uint8* Buffer) override;
+
+		virtual FName GetRuntimeFormat(USoundWave* InSoundWave) override;
+
+		virtual bool HasCompressedAudioInfoClass(USoundWave* InSoundWave) override;
+
+		virtual ICompressedAudioInfo* CreateCompressedAudioInfo(USoundWave* InSoundWave) override;
+
+		virtual FString GetDefaultDeviceName() override;
+
+		virtual FAudioPlatformSettings GetPlatformSettings() const override;
+
+        virtual int32 GetNumFrames(const int32 InNumReqestedFrames) override;
+
+        virtual void ResumeContext() override;
+
+        virtual void SuspendContext() override;
+
+        //~ End IAudioMixerPlatformInterface
+
+        
+
+	private:
+
+		AudioStreamBasicDescription OutputFormat;
+
+
+
+		bool	bSuspended;
+
+
+
+		/** True if the connection to the device has been initialized */
+
+		bool	bInitialized;
+
+		
+
+		/** True if execution is in the callback */
+
+		bool	bInCallback;
+
+
+
+		AUGraph     AudioUnitGraph;
+
+		AUNode      OutputNode;
+
+		AudioUnit	OutputUnit;
+
+        uint8*      SubmittedBufferPtr;
+
+        int32 SubmittedBytes = 0;
+
+
+
+        int32       RemainingBytesInCurrentSubmittedBuffer;
+
+        int32       BytesPerSubmittedBuffer;
+
+        
+
+        double GraphSampleRate;
+
+        
+
+		bool PerformCallback(AudioBufferList* OutputBufferData);
+
+		void HandleError(const TCHAR* InLogOutput, bool bTeardown = true);
+
+		static OSStatus AudioRenderCallback(void* RefCon, AudioUnitRenderActionFlags* ActionFlags,
+
+														  const AudioTimeStamp* TimeStamp, UInt32 BusNumber,
+
+														  UInt32 NumFrames, AudioBufferList* IOData);
+
+
+
+	};
+
+
+
+}
+
diff --git a/Engine/Source/Runtime/AudioMixer/Public/DSP/Dsp.h b/Engine/Source/Runtime/AudioMixer/Public/DSP/Dsp.h
--- a/Engine/Source/Runtime/AudioMixer/Public/DSP/Dsp.h
+++ b/Engine/Source/Runtime/AudioMixer/Public/DSP/Dsp.h
@@ -586,6 +586,12 @@ namespace Audio
 			}
 		}
 
+        // Get the current capacity of the buffer
+        uint32 GetCapacity()
+        {
+            return Capacity;
+        }
+        
 		// Get number of samples that can be pushed onto the buffer before it is full.
 		uint32 Remainder()
 		{
diff --git a/Engine/Source/Runtime/Launch/Launch.Build.cs b/Engine/Source/Runtime/Launch/Launch.Build.cs
--- a/Engine/Source/Runtime/Launch/Launch.Build.cs
+++ b/Engine/Source/Runtime/Launch/Launch.Build.cs
@@ -92,7 +92,7 @@ public class Launch : ModuleRules
 			else if (Target.Platform == UnrealTargetPlatform.Mac)
 			{
 				DynamicallyLoadedModuleNames.AddRange(new string[] {
-					"AudioMixerAudioUnit",
+					"AudioMixerCoreAudio",
 					"CoreAudio",
 				});
 			}
