Fix for build error with Mac Editor.  Unwanted compilation of iOS AudioMixer platform was being triggered, resulting in missing include

diff --git a/Engine/Source/Runtime/Apple/AudioMixerAudioUnit/AudioMixerAudioUnit.Build.cs b/Engine/Source/Runtime/Apple/AudioMixerAudioUnit/AudioMixerAudioUnit.Build.cs
--- a/Engine/Source/Runtime/Apple/AudioMixerAudioUnit/AudioMixerAudioUnit.Build.cs
+++ b/Engine/Source/Runtime/Apple/AudioMixerAudioUnit/AudioMixerAudioUnit.Build.cs
@@ -38,14 +38,11 @@ public class AudioMixerAudioUnit : ModuleRules
 		{
 			PublicFrameworks.Add("AVFoundation");
 		}
-        else if (Target.Platform == UnrealTargetPlatform.Mac)
-        {
-            PublicFrameworks.Add("AudioUnit");
-        }
+
 
 		PublicDefinitions.Add("WITH_OGGVORBIS=1");
 
-		if (Target.Platform == UnrealTargetPlatform.Mac || Target.Platform == UnrealTargetPlatform.IOS || Target.Platform == UnrealTargetPlatform.TVOS)
+		if (Target.Platform == UnrealTargetPlatform.IOS || Target.Platform == UnrealTargetPlatform.TVOS)
 		{
 			PrecompileForTargets = PrecompileTargetsType.Any;
 		}
diff --git a/Engine/Source/Runtime/ApplicationCore/Private/Mac/MacPlatformApplicationMisc.cpp b/Engine/Source/Runtime/ApplicationCore/Private/Mac/MacPlatformApplicationMisc.cpp
--- a/Engine/Source/Runtime/ApplicationCore/Private/Mac/MacPlatformApplicationMisc.cpp
+++ b/Engine/Source/Runtime/ApplicationCore/Private/Mac/MacPlatformApplicationMisc.cpp
@@ -329,7 +329,7 @@ void FMacPlatformApplicationMisc::TearDown()
 void FMacPlatformApplicationMisc::LoadPreInitModules()
 {
 	FModuleManager::Get().LoadModule(TEXT("CoreAudio"));
-	FModuleManager::Get().LoadModule(TEXT("AudioMixerAudioUnit"));
+	FModuleManager::Get().LoadModule(TEXT("AudioMixerCoreAudio"));
 }
 
 class FOutputDeviceConsole* FMacPlatformApplicationMisc::CreateConsoleOutputDevice()
