PF_Unknown is assumed to mean no conversion in other rhis. Added that behavior to Metal.

diff --git a/Engine/Source/Runtime/Apple/MetalRHI/Private/MetalUAV.cpp b/Engine/Source/Runtime/Apple/MetalRHI/Private/MetalUAV.cpp
--- a/Engine/Source/Runtime/Apple/MetalRHI/Private/MetalUAV.cpp
+++ b/Engine/Source/Runtime/Apple/MetalRHI/Private/MetalUAV.cpp
@@ -242,7 +242,16 @@ FShaderResourceViewRHIRef FMetalDynamicRHI::RHICreateShaderResourceView(FRHIText
 		SRV->SourceTexture = (FRHITexture*)Texture2DRHI;
 		
 		FMetalSurface* Surface = GetMetalSurfaceFromRHITexture(Texture2DRHI);
-		SRV->TextureView = Surface ? new FMetalSurface(*Surface, NSMakeRange(CreateInfo.MipLevel, CreateInfo.NumMipLevels), (EPixelFormat)CreateInfo.Format) : nullptr;
+		
+		// Asking to make a SRV with PF_Unknown means to use the same format.
+		// This matches the behavior of the DX11 RHI.
+		EPixelFormat Format = (EPixelFormat) CreateInfo.Format;
+		if(Surface && Format == PF_Unknown)
+		{
+			Format = Surface->PixelFormat;
+		}
+		
+		SRV->TextureView = Surface ? new FMetalSurface(*Surface, NSMakeRange(CreateInfo.MipLevel, CreateInfo.NumMipLevels), Format) : nullptr;
 		
 		SRV->SourceVertexBuffer = nullptr;
 		SRV->SourceIndexBuffer = nullptr;
