[NO NON-PII COMMIT MESSAGE]

diff --git a/Engine/Source/Runtime/Apple/AudioMixerCoreAudio/AudioMixerCoreAudio.Build.cs b/Engine/Source/Runtime/Apple/AudioMixerCoreAudio/AudioMixerCoreAudio.Build.cs
--- a/Engine/Source/Runtime/Apple/AudioMixerCoreAudio/AudioMixerCoreAudio.Build.cs
+++ b/Engine/Source/Runtime/Apple/AudioMixerCoreAudio/AudioMixerCoreAudio.Build.cs
@@ -19,21 +19,23 @@ public class AudioMixerCoreAudio : ModuleRules
 			new string[] {
 				"Core",
 				"CoreUObject",
-				"Engine",
+				"AudioMixerCore",
 			}
 			);
 
 		PrecompileForTargets = PrecompileTargetsType.None;
-		PrivateDependencyModuleNames.Add("AudioMixer");
 
-
-
-		AddEngineThirdPartyPrivateStaticDependencies(Target, 
-
-			"UEOgg",
-			"Vorbis",
-			"VorbisFile"
-			);
+        if (Target.bCompileAgainstEngine)
+        {
+            // Engine module is required for CompressedAudioInfo implementations.
+            PrivateDependencyModuleNames.Add("Engine");
+
+            AddEngineThirdPartyPrivateStaticDependencies(Target,
+                "UEOgg",
+                "Vorbis",
+                "VorbisFile"
+            );
+        }
 
 
 
