[NO NON-PII COMMIT MESSAGE]

diff --git a/Engine/Source/Runtime/Apple/AudioMixerCoreAudio/Private/AudioMixerCoreAudio.cpp b/Engine/Source/Runtime/Apple/AudioMixerCoreAudio/Private/AudioMixerCoreAudio.cpp
--- a/Engine/Source/Runtime/Apple/AudioMixerCoreAudio/Private/AudioMixerCoreAudio.cpp
+++ b/Engine/Source/Runtime/Apple/AudioMixerCoreAudio/Private/AudioMixerCoreAudio.cpp
@@ -12,9 +12,9 @@ class FAudioMixerModuleCoreAudio : public IAudioDeviceModule
 public:
 	virtual bool IsAudioMixerModule() const override { return true; }
 
-	virtual FAudioDevice* CreateAudioDevice() override
+	virtual Audio::IAudioMixerPlatformInterface* CreateAudioMixerPlatformInterface() override
 	{
-		return new Audio::FMixerDevice(new Audio::FMixerPlatformCoreAudio());
+		return new new Audio::FMixerPlatformCoreAudio();
 	}
 };
 IMPLEMENT_MODULE(FAudioMixerModuleCoreAudio, AudioMixerCoreAudio);
