[NO NON-PII COMMIT MESSAGE]

diff --git a/Engine/Source/Runtime/Apple/AudioMixerCoreAudio/Private/AudioMixerCoreAudio.cpp b/Engine/Source/Runtime/Apple/AudioMixerCoreAudio/Private/AudioMixerCoreAudio.cpp
--- a/Engine/Source/Runtime/Apple/AudioMixerCoreAudio/Private/AudioMixerCoreAudio.cpp
+++ b/Engine/Source/Runtime/Apple/AudioMixerCoreAudio/Private/AudioMixerCoreAudio.cpp
@@ -14,7 +14,7 @@ public:
 
 	virtual Audio::IAudioMixerPlatformInterface* CreateAudioMixerPlatformInterface() override
 	{
-		return new new Audio::FMixerPlatformCoreAudio();
+		return new Audio::FMixerPlatformCoreAudio();
 	}
 };
 IMPLEMENT_MODULE(FAudioMixerModuleCoreAudio, AudioMixerCoreAudio);
