fix for audio crackling on app suspension

diff --git a/Engine/Source/Runtime/Apple/AudioMixerAudioUnit/Private/AudioMixerPlatformAudioUnit.cpp b/Engine/Source/Runtime/Apple/AudioMixerAudioUnit/Private/AudioMixerPlatformAudioUnit.cpp
--- a/Engine/Source/Runtime/Apple/AudioMixerAudioUnit/Private/AudioMixerPlatformAudioUnit.cpp
+++ b/Engine/Source/Runtime/Apple/AudioMixerAudioUnit/Private/AudioMixerPlatformAudioUnit.cpp
@@ -281,6 +281,9 @@ namespace Audio
 			return false;
 		}
 		AudioStreamInfo.StreamState = EAudioOutputStreamState::Open;
+        
+        FCoreDelegates::ApplicationWillDeactivateDelegate.AddRaw(this, &FMixerPlatformAudioUnit::SuspendContext);
+        FCoreDelegates::ApplicationHasReactivatedDelegate.AddRaw(this, &FMixerPlatformAudioUnit::ResumeContext);
 		
 		return true;
 	}
@@ -293,6 +296,9 @@ namespace Audio
 		}
 		
 		AudioStreamInfo.StreamState = EAudioOutputStreamState::Closed;
+        
+        FCoreDelegates::ApplicationWillDeactivateDelegate.RemoveAll(this);
+        FCoreDelegates::ApplicationHasReactivatedDelegate.RemoveAll(this);
 		
 		return true;
 	}
