Fix for iOS crash when toggling audio mixer

diff --git a/Engine/Source/Runtime/Apple/AudioMixerAudioUnit/Private/AudioMixerPlatformAudioUnit.cpp b/Engine/Source/Runtime/Apple/AudioMixerAudioUnit/Private/AudioMixerPlatformAudioUnit.cpp
--- a/Engine/Source/Runtime/Apple/AudioMixerAudioUnit/Private/AudioMixerPlatformAudioUnit.cpp
+++ b/Engine/Source/Runtime/Apple/AudioMixerAudioUnit/Private/AudioMixerPlatformAudioUnit.cpp
@@ -332,12 +332,9 @@ namespace Audio
 			return false;
 		}
 		
-		AudioStreamInfo.StreamState = EAudioOutputStreamState::Stopping;
-		
+        StopGeneratingAudio();
 		AUGraphStop(AudioUnitGraph);
 		
-		AudioStreamInfo.StreamState = EAudioOutputStreamState::Stopped;
-		
 		return true;
 	}
 	
