Fix for iOS audio mixer thread contention with OS events

diff --git a/Engine/Source/Runtime/Apple/AudioMixerAudioUnit/Private/AudioMixerPlatformAudioUnit.cpp b/Engine/Source/Runtime/Apple/AudioMixerAudioUnit/Private/AudioMixerPlatformAudioUnit.cpp
--- a/Engine/Source/Runtime/Apple/AudioMixerAudioUnit/Private/AudioMixerPlatformAudioUnit.cpp
+++ b/Engine/Source/Runtime/Apple/AudioMixerAudioUnit/Private/AudioMixerPlatformAudioUnit.cpp
@@ -489,14 +489,10 @@ namespace Audio
 			BytesPerSubmittedBuffer = OutputBufferData->mBuffers[0].mDataByteSize;
 			uint8* OutputBufferPtr = (uint8*)OutputBufferData->mBuffers[0].mData;
 			
+			// Check to see if the system has requested a larger callback size
 			NumSamplesPerDeviceCallback = BytesPerSubmittedBuffer / static_cast<float>(sizeof(float));
 			GrowCircularBufferIfNeeded(NumSamplesPerRenderCallback, NumSamplesPerDeviceCallback);
 			
-			// Check to see if the system has requested a larger callback size
-			AVAudioSession* AudioSession = [AVAudioSession sharedInstance];
-			double BufferSizeInSec = [AudioSession IOBufferDuration];
-			
-			
 			while (CircularOutputBuffer.Num() < BytesPerSubmittedBuffer)
 			{
 				ReadNextBuffer();
