Fix for CIS error in Dev-Anim (missing include)

diff --git a/Engine/Source/Runtime/Apple/AudioMixerAudioUnit/Private/AudioMixerAudioUnit.cpp b/Engine/Source/Runtime/Apple/AudioMixerAudioUnit/Private/AudioMixerAudioUnit.cpp
--- a/Engine/Source/Runtime/Apple/AudioMixerAudioUnit/Private/AudioMixerAudioUnit.cpp
+++ b/Engine/Source/Runtime/Apple/AudioMixerAudioUnit/Private/AudioMixerAudioUnit.cpp
@@ -7,7 +7,7 @@
 
 #include "AudioMixer.h"
 #include "AudioMixerPlatformAudioUnit.h"
-
+#include "Modules/ModuleManager.h"
 
 class FAudioMixerModuleAudioUnit : public IAudioDeviceModule
 {
