fix for development compile failure for IOS

diff --git a/Engine/Source/Runtime/Apple/MetalRHI/Private/MetalStateCache.cpp b/Engine/Source/Runtime/Apple/MetalRHI/Private/MetalStateCache.cpp
--- a/Engine/Source/Runtime/Apple/MetalRHI/Private/MetalStateCache.cpp
+++ b/Engine/Source/Runtime/Apple/MetalRHI/Private/MetalStateCache.cpp
@@ -2301,8 +2301,12 @@ void FMetalStateCache::SetRenderPipelineState(FMetalCommandEncoder& CommandEncod
 			uint32 Index = __builtin_ctz(VertexMask.BufferMask);
 			VertexMask.BufferMask &= ~(1 << Index);
 			
-			if ((VertexStage == EMetalShaderStages::Vertex) || (Pipeline->TessellationPipelineDesc.TessellationInputPatchConstBufferIndex != Index
-				&& Pipeline->TessellationPipelineDesc.TessellationInputControlPointBufferIndex != Index))
+			if (VertexStage == EMetalShaderStages::Vertex
+#if PLATFORM_SUPPORTS_TESSELLATION_SHADERS
+				|| (Pipeline->TessellationPipelineDesc.TessellationInputPatchConstBufferIndex != Index
+				&& Pipeline->TessellationPipelineDesc.TessellationInputControlPointBufferIndex != Index)
+#endif // PLATFORM_SUPPORTS_TESSELLATION_SHADERS
+				)
 			{
 				FMetalBufferBinding const& Binding = ShaderBuffers[VertexStage].Buffers[Index];
 				ensure(Binding.Buffer || Binding.Bytes);
