Fixed Metal debug compile errors

diff --git a/Engine/Source/Runtime/Apple/MetalRHI/Private/MetalPipeline.cpp b/Engine/Source/Runtime/Apple/MetalRHI/Private/MetalPipeline.cpp
--- a/Engine/Source/Runtime/Apple/MetalRHI/Private/MetalPipeline.cpp
+++ b/Engine/Source/Runtime/Apple/MetalRHI/Private/MetalPipeline.cpp
@@ -1551,7 +1551,9 @@ static FMetalShaderPipeline* CreateMTLRenderPipeline(bool const bSync, FMetalGra
 					mtlpp::AutoReleasedComputePipelineReflection Reflection;
 					ComputeOption = mtlpp::PipelineOption::ArgumentInfo|mtlpp::PipelineOption::BufferTypeInfo METAL_STATISTICS_ONLY(|NSUInteger(EMTLPipelineStats));
 					Pipeline->ComputePipelineState = Device.NewComputePipelineState(ComputePipelineDesc, (mtlpp::PipelineOption)ComputeOption, &Reflection, &AutoError);
+#if METAL_DEBUG_OPTIONS
 					Pipeline->ComputePipelineReflection = Reflection;
+#endif
 				}
 				else
 				{
diff --git a/Engine/Source/Runtime/Apple/MetalRHI/Private/MetalRenderPass.cpp b/Engine/Source/Runtime/Apple/MetalRHI/Private/MetalRenderPass.cpp
--- a/Engine/Source/Runtime/Apple/MetalRHI/Private/MetalRenderPass.cpp
+++ b/Engine/Source/Runtime/Apple/MetalRHI/Private/MetalRenderPass.cpp
@@ -1126,6 +1126,7 @@ void FMetalRenderPass::DrawPatches(uint32 PrimitiveType,FMetalBuffer const& Inde
 			METAL_DEBUG_LAYER(EMetalDebugLevelFastValidation, CurrentEncoder.GetRenderCommandEncoderDebugging().DrawPatches(boundShaderState->VertexShader->TessellationOutputControlPoints, 0, NumPrimitives * NumInstances, nil, 0, 1, 0));
 		}
 		
+#if METAL_DEBUG_OPTIONS
 		if (GMetalCommandBufferDebuggingEnabled)
 		{
 			FMetalCommandData Data;
@@ -1137,6 +1138,7 @@ void FMetalRenderPass::DrawPatches(uint32 PrimitiveType,FMetalBuffer const& Inde
 			
 			InsertDebugDraw(Data);
 		}
+#endif
 		
 		if(hullShaderOutputBufferSize)
 		{
diff --git a/Engine/Source/Runtime/Apple/MetalRHI/Private/MetalStateCache.cpp b/Engine/Source/Runtime/Apple/MetalRHI/Private/MetalStateCache.cpp
--- a/Engine/Source/Runtime/Apple/MetalRHI/Private/MetalStateCache.cpp
+++ b/Engine/Source/Runtime/Apple/MetalRHI/Private/MetalStateCache.cpp
@@ -2220,6 +2220,7 @@ void FMetalStateCache::SetRenderPipelineState(FMetalCommandEncoder& CommandEncod
         PipelineBits &= EMetalPipelineFlagComputeMask;
     }
 	
+#if METAL_DEBUG_OPTIONS
 	if (SafeGetRuntimeDebuggingLevel() >= EMetalDebugLevelFastValidation)
 	{
 		FMetalShaderPipeline* Pipeline = GetPipelineState();
