#UE4 CIS quick fix

diff --git a/Engine/Source/Runtime/Apple/MetalRHI/Private/MetalStateCache.cpp b/Engine/Source/Runtime/Apple/MetalRHI/Private/MetalStateCache.cpp
--- a/Engine/Source/Runtime/Apple/MetalRHI/Private/MetalStateCache.cpp
+++ b/Engine/Source/Runtime/Apple/MetalRHI/Private/MetalStateCache.cpp
@@ -2402,6 +2402,7 @@ void FMetalStateCache::SetRenderPipelineState(FMetalCommandEncoder& CommandEncod
 			ensure(ShaderSamplers[EMetalShaderStages::Pixel].Samplers[Index]);
 		}
 	}
+#endif // METAL_DEBUG_OPTIONS
 }
 
 #if PLATFORM_SUPPORTS_TESSELLATION_SHADERS
