Force all textures to be shared on ios. This avoids the blits while updating textures. Those blits caused the driver's internal storage to grow (and it never gets compacted).

This saves ~35-40MB on the first run and potentially more as you play more matches.

Protected by a CVar rhi.Metal.ForceIOSTexturesShared which defaults to on.

diff --git a/Engine/Source/Runtime/Apple/MetalRHI/Private/MetalTexture.cpp b/Engine/Source/Runtime/Apple/MetalRHI/Private/MetalTexture.cpp
--- a/Engine/Source/Runtime/Apple/MetalRHI/Private/MetalTexture.cpp
+++ b/Engine/Source/Runtime/Apple/MetalRHI/Private/MetalTexture.cpp
@@ -22,6 +22,13 @@ FAutoConsoleVariableRef CVarMetalMaxOutstandingAsyncTexUploads(
 															   ECVF_ReadOnly|ECVF_RenderThreadSafe
 															   );
 
+int32 GMetalForceIOSTexturesShared = 1;
+FAutoConsoleVariableRef CVarMetalForceIOSTexturesShared(
+														TEXT("rhi.Metal.ForceIOSTexturesShared"),
+														GMetalForceIOSTexturesShared,
+														TEXT("If true, forces all textures to be Shared on iOS"),
+														ECVF_RenderThreadSafe);
+
 
 /** Texture reference class. */
 class FMetalTextureReference : public FRHITextureReference
@@ -799,8 +806,21 @@ FMetalSurface::FMetalSurface(ERHIResourceType ResourceType, EPixelFormat Format,
 		{
 			check(!(Flags & TexCreate_CPUReadback));
 			Desc.SetCpuCacheMode(mtlpp::CpuCacheMode::DefaultCache);
+#if PLATFORM_MAC
 			Desc.SetStorageMode(mtlpp::StorageMode::Private);
 			Desc.SetResourceOptions((mtlpp::ResourceOptions)(mtlpp::ResourceOptions::CpuCacheModeDefaultCache|mtlpp::ResourceOptions::StorageModePrivate));
+#else
+			if(GMetalForceIOSTexturesShared)
+			{
+				Desc.SetStorageMode(mtlpp::StorageMode::Shared);
+				Desc.SetResourceOptions((mtlpp::ResourceOptions)(mtlpp::ResourceOptions::CpuCacheModeDefaultCache|mtlpp::ResourceOptions::StorageModeShared));
+			}
+			else
+			{
+				Desc.SetStorageMode(mtlpp::StorageMode::Private);
+				Desc.SetResourceOptions((mtlpp::ResourceOptions)(mtlpp::ResourceOptions::CpuCacheModeDefaultCache|mtlpp::ResourceOptions::StorageModePrivate));
+			}
+#endif
 		}
 		else
 		{
@@ -811,17 +831,25 @@ FMetalSurface::FMetalSurface(ERHIResourceType ResourceType, EPixelFormat Format,
 			Desc.SetResourceOptions((mtlpp::ResourceOptions)(mtlpp::ResourceOptions::CpuCacheModeWriteCombined|mtlpp::ResourceOptions::StorageModePrivate));
 #else
 			Desc.SetCpuCacheMode(mtlpp::CpuCacheMode::DefaultCache);
+			if(GMetalForceIOSTexturesShared)
+			{
+				Desc.SetStorageMode(mtlpp::StorageMode::Shared);
+				Desc.SetResourceOptions((mtlpp::ResourceOptions)(mtlpp::ResourceOptions::CpuCacheModeDefaultCache|mtlpp::ResourceOptions::StorageModeShared));
+			}
 			// No private storage for PVRTC as it messes up the blit-encoder usage.
 			// note: this is set to always be on and will be re-addressed in a future release
-            if (IsPixelFormatPVRTCCompressed(Format))
-            {
-                Desc.SetStorageMode(mtlpp::StorageMode::Shared);
-                Desc.SetResourceOptions((mtlpp::ResourceOptions)(mtlpp::ResourceOptions::CpuCacheModeDefaultCache|mtlpp::ResourceOptions::StorageModeShared));
-            }
-            else
+			else
 			{
-				Desc.SetStorageMode(mtlpp::StorageMode::Private);
-				Desc.SetResourceOptions((mtlpp::ResourceOptions)(mtlpp::ResourceOptions::CpuCacheModeDefaultCache|mtlpp::ResourceOptions::StorageModePrivate));
+				if (IsPixelFormatPVRTCCompressed(Format))
+				{
+					Desc.SetStorageMode(mtlpp::StorageMode::Shared);
+					Desc.SetResourceOptions((mtlpp::ResourceOptions)(mtlpp::ResourceOptions::CpuCacheModeDefaultCache|mtlpp::ResourceOptions::StorageModeShared));
+				}
+				else
+				{
+					Desc.SetStorageMode(mtlpp::StorageMode::Private);
+					Desc.SetResourceOptions((mtlpp::ResourceOptions)(mtlpp::ResourceOptions::CpuCacheModeDefaultCache|mtlpp::ResourceOptions::StorageModePrivate));
+				}
 			}
 #endif
 		}
