Force accumulation of potentially redundant pixel shader texture binding dirty state.

diff --git a/Engine/Source/Runtime/Apple/MetalRHI/Private/MetalStateCache.cpp b/Engine/Source/Runtime/Apple/MetalRHI/Private/MetalStateCache.cpp
--- a/Engine/Source/Runtime/Apple/MetalRHI/Private/MetalStateCache.cpp
+++ b/Engine/Source/Runtime/Apple/MetalRHI/Private/MetalStateCache.cpp
@@ -6,6 +6,16 @@
 #include "MetalProfiler.h"
 #include "MetalCommandBuffer.h"
 
+#if PLATFORM_MAC
+int32 GMetalBypassPixelShaderTextureBindOptimization = 1;
+static FAutoConsoleVariableRef CVarMetalBypassPixelShaderTextureBindOptimization(
+									TEXT("rhi.Metal.BypassPixelShaderTextureBindOptimization"),
+									GMetalBypassPixelShaderTextureBindOptimization,
+									TEXT("When enabled (default), this allows redundant calls to set the shader texture bindings to accumulate the dirty bits for deferred validation and resource bindings, only on the pixel shader stage."),
+									ECVF_RenderThreadSafe
+									);
+#endif // PLATFORM_MAC
+
 static mtlpp::TriangleFillMode TranslateFillMode(ERasterizerFillMode FillMode)
 {
 	switch (FillMode)
@@ -273,6 +283,7 @@ void FMetalStateCache::Reset(void)
 		for (uint32 i = 0; i < ML_MaxTextures; i++)
 		{
 			ShaderTextures[Frequency].Textures[i] = nil;
+			ShaderTextures[Frequency].Usage[i] = mtlpp::ResourceUsage(0);
 		}
 		ShaderTextures[Frequency].Bound = 0;
 	}
@@ -1414,8 +1425,12 @@ void FMetalStateCache::SetShaderTexture(EMetalShaderStages const Frequency, FMet
 	check(Frequency < EMetalShaderStages::Num);
 	check(Index < ML_MaxTextures);
 	
-	if (ShaderTextures[Frequency].Textures[Index] != Texture
-		|| ShaderTextures[Frequency].Usage[Index] != Usage)
+	if (
+#if PLATFORM_MAC
+		 (GMetalBypassPixelShaderTextureBindOptimization && (Frequency == EMetalShaderStages::Pixel)) ||
+#endif
+		 (ShaderTextures[Frequency].Textures[Index] != Texture || ShaderTextures[Frequency].Usage[Index] != Usage)
+	   )
 	{
 		bool bMemoryLess = false;
 #if PLATFORM_IOS
@@ -2716,10 +2731,11 @@ void FMetalStateCache::CommitResourceTable(EMetalShaderStages const Frequency, m
 	{
 		uint32 Index = __builtin_ctzll(HiTextures);
 		HiTextures &= ~(uint64(1) << uint64(Index));
+		Index += 64;
 		
 		if (Index < ML_MaxTextures && TextureBindings.Textures[Index])
 		{
-			CommandEncoder.SetShaderTexture(Type, TextureBindings.Textures[Index], Index + 64, TextureBindings.Usage[Index]);
+			CommandEncoder.SetShaderTexture(Type, TextureBindings.Textures[Index], Index, TextureBindings.Usage[Index]);
 		}
 	}
 	
