Remove forced accumulation of potentially redundant pixel shader texture binding dirty state.

#rnx

diff --git a/Engine/Source/Runtime/Apple/MetalRHI/Private/MetalStateCache.cpp b/Engine/Source/Runtime/Apple/MetalRHI/Private/MetalStateCache.cpp
--- a/Engine/Source/Runtime/Apple/MetalRHI/Private/MetalStateCache.cpp
+++ b/Engine/Source/Runtime/Apple/MetalRHI/Private/MetalStateCache.cpp
@@ -6,16 +6,6 @@
 #include "MetalProfiler.h"
 #include "MetalCommandBuffer.h"
 
-#if PLATFORM_MAC
-int32 GMetalBypassPixelShaderTextureBindOptimization = 1;
-static FAutoConsoleVariableRef CVarMetalBypassPixelShaderTextureBindOptimization(
-									TEXT("rhi.Metal.BypassPixelShaderTextureBindOptimization"),
-									GMetalBypassPixelShaderTextureBindOptimization,
-									TEXT("When enabled (default), this allows redundant calls to set the shader texture bindings to accumulate the dirty bits for deferred validation and resource bindings, only on the pixel shader stage."),
-									ECVF_RenderThreadSafe
-									);
-#endif // PLATFORM_MAC
-
 static mtlpp::TriangleFillMode TranslateFillMode(ERasterizerFillMode FillMode)
 {
 	switch (FillMode)
@@ -1425,12 +1415,8 @@ void FMetalStateCache::SetShaderTexture(EMetalShaderStages const Frequency, FMet
 	check(Frequency < EMetalShaderStages::Num);
 	check(Index < ML_MaxTextures);
 	
-	if (
-#if PLATFORM_MAC
-		 (GMetalBypassPixelShaderTextureBindOptimization && (Frequency == EMetalShaderStages::Pixel)) ||
-#endif
-		 (ShaderTextures[Frequency].Textures[Index] != Texture || ShaderTextures[Frequency].Usage[Index] != Usage)
-	   )
+	if (ShaderTextures[Frequency].Textures[Index] != Texture
+		|| ShaderTextures[Frequency].Usage[Index] != Usage)
 	{
 		bool bMemoryLess = false;
 #if PLATFORM_IOS
@@ -1442,11 +1428,11 @@ void FMetalStateCache::SetShaderTexture(EMetalShaderStages const Frequency, FMet
 		
 		if (Texture && !bMemoryLess)
 		{
-			ShaderTextures[Frequency].Bound |= (1 << Index);
+			ShaderTextures[Frequency].Bound |= (FMetalTextureMask(1) << FMetalTextureMask(Index));
 		}
 		else
 		{
-			ShaderTextures[Frequency].Bound &= ~(1 << Index);
+			ShaderTextures[Frequency].Bound &= ~(FMetalTextureMask(1) << FMetalTextureMask(Index));
 		}
 	}
 }
