Whoops - fix mismerge that duplicated the UsageTypes enum.

diff --git a/Engine/Source/Runtime/Apple/MetalRHI/Private/MetalBuffer.h b/Engine/Source/Runtime/Apple/MetalRHI/Private/MetalBuffer.h
--- a/Engine/Source/Runtime/Apple/MetalRHI/Private/MetalBuffer.h
+++ b/Engine/Source/Runtime/Apple/MetalRHI/Private/MetalBuffer.h
@@ -411,13 +411,6 @@ class FMetalResourceHeap
 		/** Number of texture usage types */
 		EMetalHeapTextureUsageNum = 2
 	};
-
-    enum UsageTypes
-    {
-        UsageStatic,
-        UsageDynamic,
-        NumUsageTypes = 2
-    };
     
 public:
 	FMetalResourceHeap(void);
