Fix another P4 merge error in MetalRHI.

diff --git a/Engine/Source/Runtime/Apple/MetalRHI/Private/MetalVertexBuffer.cpp b/Engine/Source/Runtime/Apple/MetalRHI/Private/MetalVertexBuffer.cpp
--- a/Engine/Source/Runtime/Apple/MetalRHI/Private/MetalVertexBuffer.cpp
+++ b/Engine/Source/Runtime/Apple/MetalRHI/Private/MetalVertexBuffer.cpp
@@ -268,7 +268,6 @@ void FMetalRHIBuffer::Alloc(uint32 InSize, EResourceLockMode LockMode)
 		CPUBuffer = GetMetalDeviceContext().CreatePooledBuffer(ArgsCPU);
 		CPUBuffer.SetOwner(this);
 		check(CPUBuffer && CPUBuffer.GetPtr());
-        CPUBuffer.SetOwner(this);
         METAL_INC_DWORD_STAT_BY(Type, MemAlloc, InSize);
 		check(CPUBuffer.GetLength() >= Buffer.GetLength());
 	}
