Metal command-buffer failure debugger in the vein of Aftermath toggled with "r.GPUCrashDebugging" prior to running the game or by using the -metalgpudebug command line flag. Requires TileShaders on iOS (iOS 11, A11) or macOS 10.14 (memoryBarrier) - it will automatically disable itself on other configurations.

By default every draw call is instrumented on macOS and every 10th on iOS as TBDR GPUs suffer a much larger performance penalty.

On failure it should report the last Vertex + Fragment pair to succeed and then a more detailed list of the following commands where we believe we failed.

To get the best use of the output you need the Stable.csv files and the SF_METAL.tgz file for the build so you can track the PSOs back to their HLSL and Metal shaders.

diff --git a/Engine/Source/Runtime/Apple/MetalRHI/Private/MetalCommandEncoder.cpp b/Engine/Source/Runtime/Apple/MetalRHI/Private/MetalCommandEncoder.cpp
--- a/Engine/Source/Runtime/Apple/MetalRHI/Private/MetalCommandEncoder.cpp
+++ b/Engine/Source/Runtime/Apple/MetalRHI/Private/MetalCommandEncoder.cpp
@@ -19,6 +19,153 @@ const uint32 EncoderRingBufferSize = 1024 * 1024;
 extern int32 GMetalBufferScribble;
 #endif
 
+static TCHAR const* const GMetalCommandDataTypeName[] = {
+	TEXT("DrawPrimitive"),
+	TEXT("DrawPrimitiveIndexed"),
+	TEXT("DrawPrimitivePatch"),
+	TEXT("DrawPrimitiveIndirect"),
+	TEXT("DrawPrimitiveIndexedIndirect"),
+};
+
+
+FString FMetalCommandData::ToString() const
+{
+	FString Result;
+	if ((uint32)CommandType < (uint32)FMetalCommandData::Type::Num)
+	{
+		Result = GMetalCommandDataTypeName[(uint32)CommandType];
+		switch(CommandType)
+		{
+			case FMetalCommandData::Type::DrawPrimitive:
+				Result += FString::Printf(TEXT(" BaseInstance: %u InstanceCount: %u VertexCount: %u VertexStart: %u"), Draw.BaseInstance, Draw.InstanceCount, Draw.VertexCount, Draw.VertexStart);
+				break;
+			case FMetalCommandData::Type::DrawPrimitiveIndexed:
+				Result += FString::Printf(TEXT(" BaseInstance: %u BaseVertex: %u IndexCount: %u IndexStart: %u InstanceCount: %u"), DrawIndexed.BaseInstance, DrawIndexed.BaseVertex, DrawIndexed.IndexCount, DrawIndexed.IndexStart, DrawIndexed.InstanceCount);
+				break;
+			case FMetalCommandData::Type::DrawPrimitivePatch:
+				Result += FString::Printf(TEXT(" BaseInstance: %u InstanceCount: %u PatchCount: %u PatchStart: %u"), DrawPatch.BaseInstance, DrawPatch.InstanceCount, DrawPatch.PatchCount, DrawPatch.PatchStart);
+				break;
+			case FMetalCommandData::Type::DrawPrimitiveIndirect:
+			case FMetalCommandData::Type::DrawPrimitiveIndexedIndirect:
+			case FMetalCommandData::Type::Num:
+			default:
+				break;
+		}
+	}
+	return Result;
+};
+
+struct FMetalCommandContextDebug
+{
+	TArray<FMetalCommandDebug> Commands;
+	TSet<TRefCountPtr<FMetalGraphicsPipelineState>> PSOs;
+	FMetalBuffer DebugBuffer;
+};
+
+@interface FMetalCommandBufferDebug : FApplePlatformObject
+{
+	@public
+	TArray<FMetalCommandContextDebug> Contexts;
+}
+@end
+@implementation FMetalCommandBufferDebug
+APPLE_PLATFORM_OBJECT_ALLOC_OVERRIDES(FMetalCommandBufferDebug)
+- (void)dealloc
+{
+	Contexts.Empty();
+	[super dealloc];
+}
+@end
+
+char const* FMetalCommandBufferMarkers::kTableAssociationKey = "FMetalCommandBufferMarkers::kTableAssociationKey";
+
+FMetalCommandBufferMarkers::FMetalCommandBufferMarkers(void)
+: ns::Object<FMetalCommandBufferDebug*, ns::CallingConvention::ObjectiveC>(nil)
+{
+	
+}
+
+FMetalCommandBufferMarkers::FMetalCommandBufferMarkers(mtlpp::CommandBuffer& CmdBuf)
+: ns::Object<FMetalCommandBufferDebug*, ns::CallingConvention::ObjectiveC>([FMetalCommandBufferDebug new], ns::Ownership::Assign)
+{
+	CmdBuf.SetAssociatedObject<FMetalCommandBufferMarkers>(FMetalCommandBufferMarkers::kTableAssociationKey, *this);
+	m_ptr->Contexts.SetNum(1);
+}
+
+
+FMetalCommandBufferMarkers::FMetalCommandBufferMarkers(FMetalCommandBufferDebug* CmdBuf)
+: ns::Object<FMetalCommandBufferDebug*, ns::CallingConvention::ObjectiveC>(CmdBuf)
+{
+	
+}
+
+void FMetalCommandBufferMarkers::AllocateContexts(uint32 NumContexts)
+{
+	if (m_ptr && m_ptr->Contexts.Num() < NumContexts)
+	{
+		m_ptr->Contexts.SetNum(NumContexts);
+	}
+}
+
+uint32 FMetalCommandBufferMarkers::AddCommand(uint32 Encoder, uint32 ContextIndex, FMetalBuffer& DebugBuffer, FMetalGraphicsPipelineState* PSO, FMetalCommandData& Data)
+{
+	uint32 Num = 0;
+	if (m_ptr)
+	{
+		FMetalCommandContextDebug& Context = m_ptr->Contexts[ContextIndex];
+		if (Context.DebugBuffer != DebugBuffer)
+		{
+			Context.DebugBuffer = DebugBuffer;
+		}
+		Context.PSOs.Add(PSO);
+		Num = Context.Commands.Num();
+		FMetalCommandDebug Command;
+		Command.Encoder = Encoder;
+		Command.Index = Num;
+		Command.PSO = PSO;
+		Command.Data = Data;
+		Context.Commands.Add(Command);
+	}
+	return Num;
+}
+
+TArray<FMetalCommandDebug>* FMetalCommandBufferMarkers::GetCommands(uint32 ContextIndex)
+{
+	TArray<FMetalCommandDebug>* Result = nullptr;
+	if (m_ptr)
+	{
+		FMetalCommandContextDebug& Context = m_ptr->Contexts[ContextIndex];
+		Result = &Context.Commands;
+	}
+	return Result;
+}
+
+ns::AutoReleased<FMetalBuffer> FMetalCommandBufferMarkers::GetDebugBuffer(uint32 ContextIndex)
+{
+	ns::AutoReleased<FMetalBuffer> Buffer;
+	if (m_ptr)
+	{
+		FMetalCommandContextDebug& Context = m_ptr->Contexts[ContextIndex];
+		Buffer = Context.DebugBuffer;
+	}
+	return Buffer;
+}
+
+uint32 FMetalCommandBufferMarkers::NumContexts() const
+{
+	uint32 Num = 0;
+	if (m_ptr)
+	{
+		Num = m_ptr->Contexts.Num();
+	}
+	return Num;
+}
+
+FMetalCommandBufferMarkers FMetalCommandBufferMarkers::Get(mtlpp::CommandBuffer const& CmdBuf)
+{
+	return CmdBuf.GetAssociatedObject<FMetalCommandBufferMarkers>(FMetalCommandBufferMarkers::kTableAssociationKey);
+}
+
 #pragma mark - Public C++ Boilerplate -
 
 FMetalCommandEncoder::FMetalCommandEncoder(FMetalCommandList& CmdList)
@@ -37,6 +184,7 @@ FMetalCommandEncoder::FMetalCommandEncoder(FMetalCommandList& CmdList)
 , DebugGroups([NSMutableArray new])
 , FenceStage(mtlpp::RenderStages::Fragment)
 , EncoderNum(0)
+, CmdBufIndex(0)
 {
 	for (uint32 Frequency = 0; Frequency < uint32(mtlpp::FunctionType::Kernel)+1; Frequency++)
 	{
@@ -185,9 +333,15 @@ void FMetalCommandEncoder::StartCommandBuffer(void)
 
 	if (!CommandBuffer)
 	{
+		CmdBufIndex++;
 		CommandBuffer = CommandList.GetCommandQueue().CreateCommandBuffer();
 		METAL_DEBUG_LAYER(EMetalDebugLevelFastValidation, CommandBufferDebug = FMetalCommandBufferDebugging::Get(CommandBuffer));
 		
+		if (GMetalCommandBufferDebuggingEnabled)
+		{
+			CommandBufferMarkers = FMetalCommandBufferMarkers(CommandBuffer);
+		}
+		
 		if ([DebugGroups count] > 0)
 		{
 			CommandBuffer.SetLabel([DebugGroups lastObject]);
@@ -1144,6 +1298,11 @@ void FMetalCommandEncoder::PopDebugGroup(void)
 	}
 }
 
+FMetalCommandBufferMarkers& FMetalCommandEncoder::GetMarkers(void)
+{
+	return CommandBufferMarkers;
+}
+
 #if ENABLE_METAL_GPUPROFILE
 FMetalCommandBufferStats* FMetalCommandEncoder::GetCommandBufferStats(void)
 {
diff --git a/Engine/Source/Runtime/Apple/MetalRHI/Private/MetalCommandEncoder.h b/Engine/Source/Runtime/Apple/MetalRHI/Private/MetalCommandEncoder.h
--- a/Engine/Source/Runtime/Apple/MetalRHI/Private/MetalCommandEncoder.h
+++ b/Engine/Source/Runtime/Apple/MetalRHI/Private/MetalCommandEncoder.h
@@ -40,6 +40,55 @@ enum EMetalSubmitFlags
 	EMetalSubmitFlagsLastCommandBuffer = 1 << 6,
 };
 
+@class FMetalCommandBufferDebug;
+
+struct FMetalCommandData
+{
+	enum class Type
+	{
+		DrawPrimitive,
+		DrawPrimitiveIndexed,
+		DrawPrimitivePatch,
+		DrawPrimitiveIndirect,
+		DrawPrimitiveIndexedIndirect,
+		Num,
+	};
+	Type CommandType;
+	union
+	{
+		mtlpp::DrawPrimitivesIndirectArguments Draw;
+		mtlpp::DrawIndexedPrimitivesIndirectArguments DrawIndexed;
+		mtlpp::DrawPatchIndirectArguments DrawPatch;
+	};
+	FString ToString() const;
+};
+
+struct FMetalCommandDebug
+{
+	uint32 Encoder;
+	uint32 Index;
+	FMetalGraphicsPipelineState* PSO;
+	FMetalCommandData Data;
+};
+
+class FMetalCommandBufferMarkers : public ns::Object<FMetalCommandBufferDebug*, ns::CallingConvention::ObjectiveC>
+{
+public:
+	FMetalCommandBufferMarkers(void);
+	FMetalCommandBufferMarkers(mtlpp::CommandBuffer& CmdBuf);
+	FMetalCommandBufferMarkers(FMetalCommandBufferDebug* CmdBuf);
+	
+	void AllocateContexts(uint32 NumContexts);
+	uint32 AddCommand(uint32 Encoder, uint32 Context, FMetalBuffer& DebugBuffer, FMetalGraphicsPipelineState* PSO, FMetalCommandData& Data);
+	TArray<FMetalCommandDebug>* GetCommands(uint32 Context);
+	ns::AutoReleased<FMetalBuffer> GetDebugBuffer(uint32 ContextIndex);
+	uint32 NumContexts() const;
+	
+	static FMetalCommandBufferMarkers Get(mtlpp::CommandBuffer const& CmdBuf);
+	
+	static char const* kTableAssociationKey;
+};
+
 /**
  * FMetalCommandEncoder:
  *	Wraps the details of switching between different command encoders on the command-buffer, allowing for restoration of the render encoder if needed.
@@ -85,6 +134,9 @@ public:
 
 	/** @returns the current command buffer */
 	mtlpp::CommandBuffer const& GetCommandBuffer() const { return CommandBuffer; }
+	
+	/** @returns the monotonically incremented command buffer index */
+	uint32 GetCommandBufferIndex() const { return CmdBufIndex; }
 
 #pragma mark - Public Command Encoder Accessors -
 	
@@ -196,6 +248,9 @@ public:
 	/* Pop the latest named string off of the stack. */
 	void PopDebugGroup(void);
 	
+	/** Returns the object that records debug command markers into the command-buffer. */
+	FMetalCommandBufferMarkers& GetMarkers(void);
+	
 #if ENABLE_METAL_GPUPROFILE
 	/* Get the command-buffer stats object. */
 	FMetalCommandBufferStats* GetCommandBufferStats(void);
@@ -471,6 +526,7 @@ private:
 	mtlpp::ComputeCommandEncoder ComputeCommandEncoder;
 	mtlpp::BlitCommandEncoder BlitCommandEncoder;
 	TArray<mtlpp::RenderCommandEncoder> ChildRenderCommandEncoders;
+	FMetalCommandBufferMarkers CommandBufferMarkers;
 	
 	METAL_DEBUG_ONLY(FMetalCommandBufferDebugging CommandBufferDebug);
 	METAL_DEBUG_ONLY(FMetalRenderCommandEncoderDebugging RenderEncoderDebug);
@@ -504,4 +560,5 @@ private:
 	
 	mtlpp::RenderStages FenceStage;
 	uint32 EncoderNum;
+	uint32 CmdBufIndex;
 };
diff --git a/Engine/Source/Runtime/Apple/MetalRHI/Private/MetalCommandList.cpp b/Engine/Source/Runtime/Apple/MetalRHI/Private/MetalCommandList.cpp
--- a/Engine/Source/Runtime/Apple/MetalRHI/Private/MetalCommandList.cpp
+++ b/Engine/Source/Runtime/Apple/MetalRHI/Private/MetalCommandList.cpp
@@ -10,6 +10,7 @@
 #include "MetalCommandQueue.h"
 #include "MetalProfiler.h"
 #include "MetalCommandBuffer.h"
+#include "MetalCommandEncoder.h"
 #include "ns.hpp"
 
 #pragma mark - Public C++ Boilerplate -
@@ -17,6 +18,7 @@
 #if PLATFORM_IOS
 extern bool GIsSuspended;
 #endif
+extern int32 GMetalDebugOpsCount;
 
 FMetalCommandList::FMetalCommandList(FMetalCommandQueue& InCommandQueue, bool const bInImmediate)
 : CommandQueue(InCommandQueue)
@@ -69,7 +71,130 @@ static void ReportMetalCommandBufferFailure(mtlpp::CommandBuffer const& Complete
     if (bDoCheck && !GIsSuspended && !GIsRenderingThreadSuspended)
 #endif
     {
-		UE_LOG(LogMetal, Fatal, TEXT("Command Buffer %s Failed with %s Error! Error Domain: %s Code: %d Description %s %s %s"), *LabelString, ErrorType, *DomainString, Code, *ErrorString, *FailureString, *RecoveryString);
+		FMetalCommandBufferMarkers Markers = FMetalCommandBufferMarkers::Get(CompletedBuffer);
+		if (Markers)
+		{
+			FMetalDebugInfo BrokenDraw;
+			BrokenDraw.EncoderIndex = UINT32_MAX;
+			BrokenDraw.CommandIndex = UINT32_MAX;
+			BrokenDraw.CmdBuffIndex = UINT32_MAX;
+			BrokenDraw.ContextIndex = UINT32_MAX;
+			uint32 BrokenContext = 0;
+			bool bFoundBrokenDraw = false;
+			for (uint32 i = 0; !bFoundBrokenDraw && i < Markers.NumContexts(); i++)
+			{
+				ns::AutoReleased<FMetalBuffer> DebugBuffer = Markers.GetDebugBuffer(i);
+				TArray<FMetalCommandDebug>* Commands = Markers.GetCommands(i);
+				if (DebugBuffer && Commands && Commands->Num())
+				{
+					uint32 DebugLength = DebugBuffer.GetLength();
+					uint32 DebugCount = DebugLength / sizeof(FMetalDebugInfo);
+					check(DebugCount >= 1);
+					
+					FMetalDebugInfo* DebugArray = (FMetalDebugInfo*)DebugBuffer.GetContents();
+					FMetalDebugInfo CurrentDraw = DebugArray[0];
+					
+					// On TBDR we find the disjoint where one tile has progressed further than another
+					for (uint32 j = 0; j < DebugCount; j++)
+					{
+						if (DebugArray[j].CommandIndex < CurrentDraw.CommandIndex)
+						{
+							CurrentDraw = DebugArray[j];
+							break;
+						}
+					}
+
+					// Find the first command to fail - which depends on the order of encoders - parallel contexts make this more complicated
+					if (CurrentDraw.CmdBuffIndex < BrokenDraw.CmdBuffIndex)
+					{
+						BrokenDraw = CurrentDraw;
+						BrokenContext = CurrentDraw.ContextIndex;
+						bFoundBrokenDraw = true;
+					}
+					else if (CurrentDraw.CommandIndex < Commands->Num() && CurrentDraw.EncoderIndex < BrokenDraw.EncoderIndex)
+					{
+						BrokenDraw = CurrentDraw;
+						BrokenContext = i;
+                        bFoundBrokenDraw = true;
+					}
+                    else if(Commands->Num() > 1)
+                    {
+                        BrokenDraw = DebugArray[1];
+                        bFoundBrokenDraw = true;
+                    }
+				}
+			}
+			if (bFoundBrokenDraw)
+			{
+				id<MTLCommandBuffer> Ptr = reinterpret_cast<id<MTLCommandBuffer>>(BrokenDraw.CommandBuffer);
+				UE_LOG(LogMetal, Error, TEXT("GPU last wrote Command Buffer: %u (%p) Encoder Index: %d Context Index: %d Draw Index: %d PSO: VS: %u_%u, PS: %u_%u."), BrokenDraw.CmdBuffIndex, Ptr, BrokenDraw.EncoderIndex, BrokenDraw.ContextIndex, BrokenDraw.CommandIndex, BrokenDraw.PSOSignature[0], BrokenDraw.PSOSignature[1], BrokenDraw.PSOSignature[2], BrokenDraw.PSOSignature[3]);
+
+				if (Ptr != CompletedBuffer.GetPtr())
+				{
+					// Have we failed between commands-buffers? Or have our command-buffers been subject to coalescing?
+					ns::AutoReleased<mtlpp::CommandBuffer> PrevCommandBuffer(Ptr);
+					FMetalCommandBufferMarkers OtherMarkers = FMetalCommandBufferMarkers::Get(PrevCommandBuffer);
+					if (OtherMarkers)
+					{
+						ns::AutoReleased<FMetalBuffer> DebugBuffer = OtherMarkers.GetDebugBuffer(BrokenContext);
+						TArray<FMetalCommandDebug>* Commands = OtherMarkers.GetCommands(BrokenContext);
+						
+						UE_LOG(LogMetal, Error, TEXT("Failed executing following commands:"));
+						uint32 StartIdx = BrokenDraw.CommandIndex;
+						for (uint32 CmdIdx = StartIdx; CmdIdx < Commands->Num(); CmdIdx++)
+						{
+							FMetalCommandDebug& Command = (*Commands)[CmdIdx];
+							FString VSHash = Command.PSO->VertexShader->GetHash().ToString();
+							uint32 VSSig[2] = { Command.PSO->VertexShader->SourceLen, Command.PSO->VertexShader->SourceCRC };
+							
+							FString PSHash;
+							uint32 PSSig[2] = { 0, 0 };
+							if (Command.PSO->PixelShader.IsValid())
+							{
+								PSHash = Command.PSO->PixelShader->GetHash().ToString();
+								PSSig[0] = Command.PSO->PixelShader->SourceLen;
+								PSSig[1] = Command.PSO->PixelShader->SourceCRC;
+							}
+							
+							UE_LOG(LogMetal, Error, TEXT("Command Buffer: %p Encoder: %d Command: %d: %s PSO: VS: %s (%u_%u), PS: %s (%u_%u)"), Ptr, Command.Encoder, Command.Index, *Command.Data.ToString(), *VSHash, VSSig[0], VSSig[1], *PSHash, PSSig[0], PSSig[1]);
+						}
+					}
+					
+					BrokenDraw.CommandIndex = 0;
+				}
+				
+				ns::AutoReleased<FMetalBuffer> DebugBuffer = Markers.GetDebugBuffer(BrokenContext);
+				TArray<FMetalCommandDebug>* Commands = Markers.GetCommands(BrokenContext);
+				if (Commands->Num())
+				{
+					UE_LOG(LogMetal, Error, TEXT("Failed executing following commands:"));
+					uint32 StartIdx = BrokenDraw.CommandIndex;
+					uint32 EndIdx = BrokenDraw.CommandIndex + GMetalDebugOpsCount + 1;
+					for (uint32 CmdIdx = StartIdx; CmdIdx < Commands->Num() && CmdIdx < EndIdx; CmdIdx++)
+					{
+						FMetalCommandDebug& Command = (*Commands)[CmdIdx];
+						FString VSHash = Command.PSO->VertexShader->GetHash().ToString();
+						uint32 VSSig[2] = { Command.PSO->VertexShader->SourceLen, Command.PSO->VertexShader->SourceCRC };
+						
+						FString PSHash;
+						uint32 PSSig[2] = { 0, 0 };
+						if (Command.PSO->PixelShader.IsValid())
+						{
+							PSHash = Command.PSO->PixelShader->GetHash().ToString();
+							PSSig[0] = Command.PSO->PixelShader->SourceLen;
+							PSSig[1] = Command.PSO->PixelShader->SourceCRC;
+						}
+						
+						UE_LOG(LogMetal, Error, TEXT("Command Buffer: %p Encoder: %d Command: %d: %s PSO: VS: %s (%u_%u), PS: %s (%u_%u)"), CompletedBuffer.GetPtr(), Command.Encoder, Command.Index, *Command.Data.ToString(), *VSHash, VSSig[0], VSSig[1], *PSHash, PSSig[0], PSSig[1]);
+					}
+				}
+			}
+			UE_LOG(LogMetal, Fatal, TEXT("Command Buffer %s Failed with %s Error! Error Domain: %s Code: %d Description %s %s %s"), *LabelString, ErrorType, *DomainString, Code, *ErrorString, *FailureString, *RecoveryString);
+		}
+		else
+		{
+			UE_LOG(LogMetal, Fatal, TEXT("Command Buffer %s Failed with %s Error! Error Domain: %s Code: %d Description %s %s %s"), *LabelString, ErrorType, *DomainString, Code, *ErrorString, *FailureString, *RecoveryString);
+		}
     }
 }
 
diff --git a/Engine/Source/Runtime/Apple/MetalRHI/Private/MetalCommandQueue.cpp b/Engine/Source/Runtime/Apple/MetalRHI/Private/MetalCommandQueue.cpp
--- a/Engine/Source/Runtime/Apple/MetalRHI/Private/MetalCommandQueue.cpp
+++ b/Engine/Source/Runtime/Apple/MetalRHI/Private/MetalCommandQueue.cpp
@@ -19,6 +19,7 @@
 #pragma mark - Private C++ Statics -
 uint64 FMetalCommandQueue::Features = 0;
 extern mtlpp::VertexFormat GMetalFColorVertexFormat;
+bool GMetalCommandBufferDebuggingEnabled = 0;
 
 #pragma mark - Public C++ Boilerplate -
 
@@ -137,6 +138,17 @@ FMetalCommandQueue::FMetalCommandQueue(mtlpp::Device InDevice, uint32 const MaxN
 				}
 				
 				Features |= EMetalFeaturesPresentMinDuration | EMetalFeaturesGPUCaptureManager | EMetalFeaturesBufferSubAllocation | EMetalFeaturesParallelRenderEncoders | EMetalFeaturesPipelineBufferMutability;
+                
+                if(Device.SupportsFeatureSet(mtlpp::FeatureSet::iOS_GPUFamily4_v1))
+                {
+                    Features |= EMetalFeaturesTileShaders;
+                    
+                    GMetalSupportsTileShaders = true;
+                    
+                    // The below implies tile shaders which are necessary to order the draw calls and generate a buffer that shows what PSOs/draws ran on each tile.
+                    IConsoleVariable* GPUCrashDebuggingCVar = IConsoleManager::Get().FindConsoleVariable(TEXT("r.GPUCrashDebugging"));
+                    GMetalCommandBufferDebuggingEnabled = (GPUCrashDebuggingCVar && GPUCrashDebuggingCVar->GetInt() != 0) || FParse::Param(FCommandLine::Get(),TEXT("metalgpudebug"));
+                }
 				
 				// Turn on Texture Buffers! These are faster on the GPU as we don't need to do out-of-bounds tests but require Metal 2.1 and macOS 10.14
 				if (Vers.majorVersion >= 12)
@@ -223,7 +235,10 @@ FMetalCommandQueue::FMetalCommandQueue(mtlpp::Device InDevice, uint32 const MaxN
             {
                 Features |= EMetalFeaturesIABs;
             }
-			
+            
+            IConsoleVariable* GPUCrashDebuggingCVar = IConsoleManager::Get().FindConsoleVariable(TEXT("r.GPUCrashDebugging"));
+            GMetalCommandBufferDebuggingEnabled = (GPUCrashDebuggingCVar && GPUCrashDebuggingCVar->GetInt() != 0) || FParse::Param(FCommandLine::Get(),TEXT("metalgpudebug"));
+            
             // The editor spawns so many viewports and preview icons that we can run out of hardware fences!
 			// Need to figure out a way to safely flush the rendering and reuse the fences when that happens.
 #if WITH_EDITORONLY_DATA
diff --git a/Engine/Source/Runtime/Apple/MetalRHI/Private/MetalCommandQueue.h b/Engine/Source/Runtime/Apple/MetalRHI/Private/MetalCommandQueue.h
--- a/Engine/Source/Runtime/Apple/MetalRHI/Private/MetalCommandQueue.h
+++ b/Engine/Source/Runtime/Apple/MetalRHI/Private/MetalCommandQueue.h
@@ -68,7 +68,9 @@ typedef NS_OPTIONS(uint64, EMetalFeatures)
 	/** Supports indirect argument buffers */
 	EMetalFeaturesIABs = 1llu << 24llu,
 	/** Supports specifying the mutability of buffers bound to PSOs */
-	EMetalFeaturesPipelineBufferMutability = 1llu << 25llu,
+    EMetalFeaturesPipelineBufferMutability = 1llu << 25llu,
+    /** Supports tile shaders */
+    EMetalFeaturesTileShaders = 1llu << 26llu,
 };
 
 /**
diff --git a/Engine/Source/Runtime/Apple/MetalRHI/Private/MetalPipeline.cpp b/Engine/Source/Runtime/Apple/MetalRHI/Private/MetalPipeline.cpp
--- a/Engine/Source/Runtime/Apple/MetalRHI/Private/MetalPipeline.cpp
+++ b/Engine/Source/Runtime/Apple/MetalRHI/Private/MetalPipeline.cpp
@@ -56,6 +56,68 @@ static float RoundTessLevel(float TessFactor, mtlpp::TessellationPartitionMode P
 	}
 }
 
+// A tile-based or vertex-based debug shader for trying to emulate Aftermath style failure reporting
+static NSString* GMetalDebugShader = @"#include <metal_stdlib>\n"
+"#include <metal_compute>\n"
+"\n"
+"using namespace metal;\n"
+"\n"
+"struct DebugInfo\n"
+"{\n"
+"   uint CmdBuffIndex;\n"
+"	uint EncoderIndex;\n"
+"   uint ContextIndex;\n"
+"   uint CommandIndex;\n"
+"   uint2 CommandBuffer;\n"
+"	uint4 PSOSignature;\n"
+"};\n"
+"\n"
+#if !PLATFORM_MAC
+"// Executes once per-tile\n"
+"fragment void Main_Debug(constant DebugInfo *debugTable [[ buffer(0) ]], device DebugInfo* debugBuffer [[ buffer(1) ]], uint tile_index [[ tile_index ]])\n"
+"{\n"
+"	// Write Pass, Draw indices\n"
+"	// Write Vertex+Fragment PSO sig (in form VertexLen, VertexCRC, FragLen, FragCRC)\n"
+"	debugBuffer[tile_index] = debugTable[0];\n"
+"}";
+#else
+"// Executes once as a point draw call\n"
+"vertex void Main_Debug(constant DebugInfo *debugTable [[ buffer(0) ]], device DebugInfo* debugBuffer [[ buffer(1) ]])\n"
+"{\n"
+"	// Write Pass, Draw indices\n"
+"	// Write Vertex+Fragment PSO sig (in form VertexLen, VertexCRC, FragLen, FragCRC)\n"
+"	debugBuffer[0] = debugTable[0];\n"
+"}";
+#endif
+
+struct FMetalHelperFunctions
+{
+    mtlpp::Library DebugShadersLib;
+    mtlpp::Function DebugFunc;
+    
+    FMetalHelperFunctions()
+    {
+        if (GMetalCommandBufferDebuggingEnabled)
+        {
+            mtlpp::CompileOptions CompileOptions;
+            ns::AutoReleasedError Error;
+            DebugShadersLib = GetMetalDeviceContext().GetDevice().NewLibrary(GMetalDebugShader, CompileOptions, &Error);
+            DebugFunc = DebugShadersLib.NewFunction(@"Main_Debug");
+        }
+    }
+    
+    static FMetalHelperFunctions& Get()
+    {
+        static FMetalHelperFunctions sSelf;
+        return sSelf;
+    }
+    
+    mtlpp::Function GetDebugFunction()
+    {
+        return DebugFunc;
+    }
+};
+
 struct FMetalGraphicsPipelineKey
 {
 	FMetalRenderPipelineHash RenderPipelineHash;
@@ -478,7 +540,12 @@ static FMetalShaderPipeline* CreateMTLRenderPipeline(bool const bSync, FMetalGra
 		METAL_DEBUG_OPTION(FMemory::Memzero(Pipeline->ResourceMask, sizeof(Pipeline->ResourceMask)));
 
 		mtlpp::RenderPipelineDescriptor RenderPipelineDesc;
-		mtlpp::ComputePipelineDescriptor ComputePipelineDesc(nil);
+        mtlpp::ComputePipelineDescriptor ComputePipelineDesc(nil);
+#if PLATFORM_MAC
+        mtlpp::RenderPipelineDescriptor DebugPipelineDesc;
+#else
+        mtlpp::TileRenderPipelineDescriptor DebugPipelineDesc;
+#endif
 		
 		if (FMetalCommandQueue::SupportsFeature(EMetalFeaturesPipelineBufferMutability))
 		{
@@ -536,8 +603,9 @@ static FMetalShaderPipeline* CreateMTLRenderPipeline(bool const bSync, FMetalGra
         FMetalBlendState* BlendState = (FMetalBlendState*)Init.BlendState;
 		
 		ns::Array<mtlpp::RenderPipelineColorAttachmentDescriptor> ColorAttachments = RenderPipelineDesc.GetColorAttachments();
-		
-        for (uint32 i = 0; i < NumActiveTargets; i++)
+		auto DebugColorAttachements = DebugPipelineDesc.GetColorAttachments();
+
+		for (uint32 i = 0; i < NumActiveTargets; i++)
         {
             EPixelFormat TargetFormat = (EPixelFormat)Init.RenderTargetFormats[i];
             if (TargetFormat == PF_Unknown && PixelShader && (((PixelShader->Bindings.InOutMask & 0x7fff) & (1 << i))))
@@ -561,7 +629,10 @@ static FMetalShaderPipeline* CreateMTLRenderPipeline(bool const bSync, FMetalGra
             
             mtlpp::RenderPipelineColorAttachmentDescriptor Attachment = ColorAttachments[i];
             Attachment.SetPixelFormat(MetalFormat);
-            
+			
+			auto DebugAttachment = DebugColorAttachements[i];;
+			DebugAttachment.SetPixelFormat(MetalFormat);
+			
             mtlpp::RenderPipelineColorAttachmentDescriptor Blend = BlendState->RenderTargetStates[i].BlendState;
             if(TargetFormat != PF_Unknown)
             {
@@ -574,11 +645,26 @@ static FMetalShaderPipeline* CreateMTLRenderPipeline(bool const bSync, FMetalGra
                 Attachment.SetDestinationAlphaBlendFactor(Blend.GetDestinationAlphaBlendFactor());
                 Attachment.SetAlphaBlendOperation(Blend.GetAlphaBlendOperation());
                 Attachment.SetWriteMask(Blend.GetWriteMask());
+				
+#if PLATFORM_MAC
+				DebugAttachment.SetBlendingEnabled(Blend.IsBlendingEnabled());
+				DebugAttachment.SetSourceRgbBlendFactor(Blend.GetSourceRgbBlendFactor());
+				DebugAttachment.SetDestinationRgbBlendFactor(Blend.GetDestinationRgbBlendFactor());
+				DebugAttachment.SetRgbBlendOperation(Blend.GetRgbBlendOperation());
+				DebugAttachment.SetSourceAlphaBlendFactor(Blend.GetSourceAlphaBlendFactor());
+				DebugAttachment.SetDestinationAlphaBlendFactor(Blend.GetDestinationAlphaBlendFactor());
+				DebugAttachment.SetAlphaBlendOperation(Blend.GetAlphaBlendOperation());
+				DebugAttachment.SetWriteMask(Blend.GetWriteMask());
+#endif
             }
             else
             {
                 Attachment.SetBlendingEnabled(NO);
 				Attachment.SetWriteMask(mtlpp::ColorWriteMask::None);
+#if PLATFORM_MAC
+				DebugAttachment.SetBlendingEnabled(NO);
+				DebugAttachment.SetWriteMask(mtlpp::ColorWriteMask::None);
+#endif
             }
         }
         
@@ -592,22 +678,35 @@ static FMetalShaderPipeline* CreateMTLRenderPipeline(bool const bSync, FMetalGra
                     if (Init.DepthTargetLoadAction != ERenderTargetLoadAction::ENoAction || Init.DepthTargetStoreAction != ERenderTargetStoreAction::ENoAction)
                     {
                         RenderPipelineDesc.SetDepthAttachmentPixelFormat(MetalFormat);
+#if PLATFORM_MAC
+						DebugPipelineDesc.SetDepthAttachmentPixelFormat(MetalFormat);
+#endif
                     }
                     if (Init.StencilTargetLoadAction != ERenderTargetLoadAction::ENoAction || Init.StencilTargetStoreAction != ERenderTargetStoreAction::ENoAction)
                     {
                         RenderPipelineDesc.SetStencilAttachmentPixelFormat(mtlpp::PixelFormat::Stencil8);
+#if PLATFORM_MAC
+						DebugPipelineDesc.SetStencilAttachmentPixelFormat(mtlpp::PixelFormat::Stencil8);
+#endif
                     }
                 }
                 else
                 {
                     RenderPipelineDesc.SetDepthAttachmentPixelFormat(MetalFormat);
                     RenderPipelineDesc.SetStencilAttachmentPixelFormat(MetalFormat);
+#if PLATFORM_MAC
+					DebugPipelineDesc.SetDepthAttachmentPixelFormat(MetalFormat);
+					DebugPipelineDesc.SetStencilAttachmentPixelFormat(MetalFormat);
+#endif
                 }
                 break;
             }
             case PF_ShadowDepth:
             {
                 RenderPipelineDesc.SetDepthAttachmentPixelFormat((mtlpp::PixelFormat)GPixelFormats[PF_ShadowDepth].PlatformFormat);
+#if PLATFORM_MAC
+				DebugPipelineDesc.SetDepthAttachmentPixelFormat((mtlpp::PixelFormat)GPixelFormats[PF_ShadowDepth].PlatformFormat);
+#endif
                 break;
             }
             default:
@@ -625,11 +724,18 @@ static FMetalShaderPipeline* CreateMTLRenderPipeline(bool const bSync, FMetalGra
         {
             RenderPipelineDesc.SetDepthAttachmentPixelFormat((mtlpp::PixelFormat)GPixelFormats[PF_DepthStencil].PlatformFormat);
             RenderPipelineDesc.SetStencilAttachmentPixelFormat((mtlpp::PixelFormat)GPixelFormats[PF_DepthStencil].PlatformFormat);
+			
+#if PLATFORM_MAC
+			DebugPipelineDesc.SetDepthAttachmentPixelFormat((mtlpp::PixelFormat)GPixelFormats[PF_DepthStencil].PlatformFormat);
+			DebugPipelineDesc.SetStencilAttachmentPixelFormat((mtlpp::PixelFormat)GPixelFormats[PF_DepthStencil].PlatformFormat);
+#endif
         }
         
         RenderPipelineDesc.SetSampleCount(FMath::Max(Init.NumSamples, (uint16)1u));
     #if PLATFORM_MAC
         RenderPipelineDesc.SetInputPrimitiveTopology(TranslatePrimitiveTopology(Init.PrimitiveType));
+		DebugPipelineDesc.SetSampleCount(FMath::Max(Init.NumSamples, (uint16)1u));
+		DebugPipelineDesc.SetInputPrimitiveTopology(TranslatePrimitiveTopology(Init.PrimitiveType));
     #endif
         
         FMetalVertexDeclaration* VertexDecl = (FMetalVertexDeclaration*)Init.BoundShaderState.VertexDeclarationRHI;
@@ -1020,6 +1126,25 @@ static FMetalShaderPipeline* CreateMTLRenderPipeline(bool const bSync, FMetalGra
 	#endif
         Pipeline->FragmentSource = PixelShader ? PixelShader->GetSourceCode() : nil;
     #endif
+		
+		if (GMetalCommandBufferDebuggingEnabled)
+		{
+#if PLATFORM_MAC
+			DebugPipelineDesc.SetVertexFunction(FMetalHelperFunctions::Get().GetDebugFunction());
+			DebugPipelineDesc.SetRasterizationEnabled(false);
+#else
+			DebugPipelineDesc.SetTileFunction(FMetalHelperFunctions::Get().GetDebugFunction());
+			DebugPipelineDesc.SetRasterSampleCount(RenderPipelineDesc.GetSampleCount());
+			DebugPipelineDesc.SetThreadgroupSizeMatchesTileSize(false);
+#endif
+#if ENABLE_METAL_GPUPROFILE
+			DebugPipelineDesc.SetLabel(@"Main_Debug");
+#endif
+
+			ns::AutoReleasedError RenderError;
+			METAL_GPUPROFILE(FScopedMetalCPUStats CPUStat(FString::Printf(TEXT("NewDebugPipeline: %s"), TEXT("")/**FString([RenderPipelineDesc.GetPtr() description])*/)));
+			Pipeline->DebugPipelineState = Device.NewRenderPipelineState(DebugPipelineDesc, (mtlpp::PipelineOption)RenderOption, Reflection, &RenderError);
+		}
         
 #if METAL_DEBUG_OPTIONS
         if (GFrameCounter > 3)
diff --git a/Engine/Source/Runtime/Apple/MetalRHI/Private/MetalPipeline.h b/Engine/Source/Runtime/Apple/MetalRHI/Private/MetalPipeline.h
--- a/Engine/Source/Runtime/Apple/MetalRHI/Private/MetalPipeline.h
+++ b/Engine/Source/Runtime/Apple/MetalRHI/Private/MetalPipeline.h
@@ -73,6 +73,7 @@ struct FMetalTessellationPipelineDesc
 @public
 	mtlpp::RenderPipelineState RenderPipelineState;
 	mtlpp::ComputePipelineState ComputePipelineState;
+	mtlpp::RenderPipelineState DebugPipelineState;
 	FMetalTessellationPipelineDesc TessellationPipelineDesc;
 #if METAL_DEBUG_OPTIONS
 	FMetalDebugShaderResourceMask ResourceMask[EMetalShaderStagesNum];
diff --git a/Engine/Source/Runtime/Apple/MetalRHI/Private/MetalRHIPrivate.h b/Engine/Source/Runtime/Apple/MetalRHI/Private/MetalRHIPrivate.h
--- a/Engine/Source/Runtime/Apple/MetalRHI/Private/MetalRHIPrivate.h
+++ b/Engine/Source/Runtime/Apple/MetalRHI/Private/MetalRHIPrivate.h
@@ -103,6 +103,8 @@ extern int32 GMetalResetOnPSOChange;
 #define METAL_DEBUG_LAYER(Level, Code)
 #endif
 
+extern bool GMetalCommandBufferDebuggingEnabled;
+
 /** Set to 1 to enable GPU events in Xcode frame debugger */
 #ifndef ENABLE_METAL_GPUEVENTS_IN_TEST
 	#define ENABLE_METAL_GPUEVENTS_IN_TEST 0
@@ -127,6 +129,16 @@ extern int32 GMetalResetOnPSOChange;
 
 #define METAL_NEW_NONNULL_DECL (__clang_major__ >= 9)
 
+struct FMetalDebugInfo
+{
+	uint32 CmdBuffIndex;
+	uint32 EncoderIndex;
+	uint32 ContextIndex;
+	uint32 CommandIndex;
+	uint64 CommandBuffer;
+	uint32 PSOSignature[4];
+};
+
 // Access the internal context for the device-owning DynamicRHI object
 FMetalDeviceContext& GetMetalDeviceContext();
 
diff --git a/Engine/Source/Runtime/Apple/MetalRHI/Private/MetalRenderPass.cpp b/Engine/Source/Runtime/Apple/MetalRHI/Private/MetalRenderPass.cpp
--- a/Engine/Source/Runtime/Apple/MetalRHI/Private/MetalRenderPass.cpp
+++ b/Engine/Source/Runtime/Apple/MetalRHI/Private/MetalRenderPass.cpp
@@ -37,6 +37,13 @@ static FAutoConsoleVariableRef CVarMetalDeferRenderPasses(
 	GMetalDeferRenderPasses,
 	TEXT("Whether to defer creating render command encoders. (Default: 1)"));
 
+// Deliberately not static!
+int32 GMetalDebugOpsCount = PLATFORM_MAC ? 1 : 10;
+static FAutoConsoleVariableRef CVarMetalDebugOpsCount(
+	TEXT("rhi.Metal.DebugOpsCount"),
+	GMetalDebugOpsCount,
+	TEXT("The number of operations to allow between GPU debug markers for the r.GPUCrashDebugging reports. (Default: Mac = 1 : iOS/tvOS = 10)"));
+
 #pragma mark - Public C++ Boilerplate -
 
 FMetalRenderPass::FMetalRenderPass(FMetalCommandList& InCmdList, FMetalStateCache& Cache)
@@ -406,6 +413,18 @@ void FMetalRenderPass::DrawPrimitive(uint32 PrimitiveType, uint32 BaseVertexInde
 		METAL_GPUPROFILE(FMetalProfiler::GetProfiler()->EncodeDraw(CurrentEncoder.GetCommandBufferStats(), __FUNCTION__, NumPrimitives, NumVertices, NumInstances));
 		CurrentEncoder.GetRenderCommandEncoder().Draw(TranslatePrimitiveType(PrimitiveType), BaseVertexIndex, NumVertices, NumInstances);
 		METAL_DEBUG_LAYER(EMetalDebugLevelFastValidation, CurrentEncoder.GetRenderCommandEncoderDebugging().Draw(TranslatePrimitiveType(PrimitiveType), BaseVertexIndex, NumVertices, NumInstances));
+
+		if (GMetalCommandBufferDebuggingEnabled)
+		{
+			FMetalCommandData Data;
+			Data.CommandType = FMetalCommandData::Type::DrawPrimitive;
+			Data.Draw.BaseInstance = 0;
+			Data.Draw.InstanceCount = NumInstances;
+			Data.Draw.VertexCount = NumVertices;
+			Data.Draw.VertexStart = BaseVertexIndex;
+			
+			InsertDebugDraw(Data);
+		}
 	}
 #if PLATFORM_SUPPORTS_TESSELLATION_SHADERS
 	else
@@ -431,6 +450,14 @@ void FMetalRenderPass::DrawPrimitiveIndirect(uint32 PrimitiveType, FMetalVertexB
 		CurrentEncoder.GetRenderCommandEncoder().Draw(TranslatePrimitiveType(PrimitiveType), VertexBuffer->Buffer, ArgumentOffset);
 		METAL_DEBUG_LAYER(EMetalDebugLevelFastValidation, CurrentEncoder.GetRenderCommandEncoderDebugging().Draw(TranslatePrimitiveType(PrimitiveType), VertexBuffer->Buffer, ArgumentOffset));
 
+		if (GMetalCommandBufferDebuggingEnabled)
+		{
+			FMetalCommandData Data;
+			Data.CommandType = FMetalCommandData::Type::DrawPrimitiveIndirect;
+
+			InsertDebugDraw(Data);
+		}
+
 		ConditionalSubmit();
 	}
 	else
@@ -512,6 +539,19 @@ void FMetalRenderPass::DrawIndexedPrimitive(FMetalBuffer const& IndexBuffer, uin
 			CurrentEncoder.GetRenderCommandEncoder().DrawIndexed(TranslatePrimitiveType(PrimitiveType), NumIndices, ((IndexStride == 2) ? mtlpp::IndexType::UInt16 : mtlpp::IndexType::UInt32), IndexBuffer, StartIndex * IndexStride, NumInstances);
 			METAL_DEBUG_LAYER(EMetalDebugLevelFastValidation, CurrentEncoder.GetRenderCommandEncoderDebugging().DrawIndexed(TranslatePrimitiveType(PrimitiveType), NumIndices, ((IndexStride == 2) ? mtlpp::IndexType::UInt16 : mtlpp::IndexType::UInt32), IndexBuffer, StartIndex * IndexStride, NumInstances));
 		}
+		
+		if (GMetalCommandBufferDebuggingEnabled)
+		{
+			FMetalCommandData Data;
+			Data.CommandType = FMetalCommandData::Type::DrawPrimitiveIndexed;
+			Data.DrawIndexed.BaseInstance = FirstInstance;
+			Data.DrawIndexed.BaseVertex = BaseVertexIndex;
+			Data.DrawIndexed.IndexCount = NumIndices;
+			Data.DrawIndexed.IndexStart = StartIndex;
+			Data.DrawIndexed.InstanceCount = NumInstances;
+			
+			InsertDebugDraw(Data);
+		}
 	}
 #if PLATFORM_SUPPORTS_TESSELLATION_SHADERS
 	else
@@ -540,6 +580,13 @@ void FMetalRenderPass::DrawIndexedIndirect(FMetalIndexBuffer* IndexBuffer, uint3
 		CurrentEncoder.GetRenderCommandEncoder().DrawIndexed(TranslatePrimitiveType(PrimitiveType), (mtlpp::IndexType)IndexBuffer->IndexType, IndexBuffer->Buffer, 0, VertexBuffer->Buffer, (DrawArgumentsIndex * 5 * sizeof(uint32)));
 		METAL_DEBUG_LAYER(EMetalDebugLevelFastValidation, CurrentEncoder.GetRenderCommandEncoderDebugging().DrawIndexed(TranslatePrimitiveType(PrimitiveType), (mtlpp::IndexType)IndexBuffer->IndexType, IndexBuffer->Buffer, 0, VertexBuffer->Buffer, (DrawArgumentsIndex * 5 * sizeof(uint32))));
 
+		if (GMetalCommandBufferDebuggingEnabled)
+		{
+			FMetalCommandData Data;
+			Data.CommandType = FMetalCommandData::Type::DrawPrimitiveIndexedIndirect;
+			
+			InsertDebugDraw(Data);
+		}
 		ConditionalSubmit();
 	}
 	else
@@ -562,6 +609,14 @@ void FMetalRenderPass::DrawIndexedPrimitiveIndirect(uint32 PrimitiveType,FMetalI
 		CurrentEncoder.GetRenderCommandEncoder().DrawIndexed(TranslatePrimitiveType(PrimitiveType), (mtlpp::IndexType)IndexBuffer->IndexType, IndexBuffer->Buffer, 0, VertexBuffer->Buffer, ArgumentOffset);
 		METAL_DEBUG_LAYER(EMetalDebugLevelFastValidation, CurrentEncoder.GetRenderCommandEncoderDebugging().DrawIndexed(TranslatePrimitiveType(PrimitiveType), (mtlpp::IndexType)IndexBuffer->IndexType, IndexBuffer->Buffer, 0, VertexBuffer->Buffer, ArgumentOffset));
 
+		if (GMetalCommandBufferDebuggingEnabled)
+		{
+			FMetalCommandData Data;
+			Data.CommandType = FMetalCommandData::Type::DrawPrimitiveIndirect;
+			
+			InsertDebugDraw(Data);
+		}
+		
 		ConditionalSubmit();
 	}
 	else
@@ -720,6 +775,18 @@ void FMetalRenderPass::DrawPatches(uint32 PrimitiveType,FMetalBuffer const& Inde
 			METAL_DEBUG_LAYER(EMetalDebugLevelFastValidation, CurrentEncoder.GetRenderCommandEncoderDebugging().DrawPatches(boundShaderState->VertexShader->TessellationOutputControlPoints, 0, NumPrimitives * NumInstances, nil, 0, 1, 0));
 		}
 		
+		if (GMetalCommandBufferDebuggingEnabled)
+		{
+			FMetalCommandData Data;
+			Data.CommandType = FMetalCommandData::Type::DrawPrimitivePatch;
+			Data.DrawPatch.BaseInstance = FirstInstance;
+			Data.DrawPatch.InstanceCount = NumInstances;
+			Data.DrawPatch.PatchCount = NumPrimitives * NumInstances;
+			Data.DrawPatch.PatchStart = 0;
+			
+			InsertDebugDraw(Data);
+		}
+		
 		if(hullShaderOutputBufferSize)
 		{
 			deviceContext.ReleaseBuffer(hullShaderOutputBuffer);
@@ -1788,3 +1855,85 @@ void FMetalRenderPass::ConditionalSubmit()
 #endif
 	}
 }
+
+uint32 FMetalRenderPass::GetEncoderIndex(void) const
+{
+	if (!CmdList.IsParallel())
+	{
+		return PrologueEncoder.NumEncodedPasses() + CurrentEncoder.NumEncodedPasses();
+	}
+	else
+	{
+		return GetMetalDeviceContext().GetCurrentRenderPass().GetEncoderIndex();
+	}
+}
+
+uint32 FMetalRenderPass::GetCommandBufferIndex(void) const
+{
+	if (!CmdList.IsParallel())
+	{
+		return CurrentEncoder.GetCommandBufferIndex();
+	}
+	else
+	{
+		return GetMetalDeviceContext().GetCurrentRenderPass().GetCommandBufferIndex();
+	}
+}
+
+void FMetalRenderPass::InsertDebugDraw(FMetalCommandData& Data)
+{
+	if (GMetalCommandBufferDebuggingEnabled)
+	{
+		FMetalGraphicsPipelineState* BoundShaderState = State.GetGraphicsPSO();
+		
+		uint32 NumCommands = CurrentEncoder.GetMarkers().AddCommand(GetEncoderIndex(), CmdList.GetParallelIndex(), State.GetDebugBuffer(), BoundShaderState, Data);
+		
+		if ((NumCommands % GMetalDebugOpsCount) == 0)
+		{
+			FMetalDebugInfo DebugInfo;
+			DebugInfo.EncoderIndex = GetEncoderIndex();
+			DebugInfo.ContextIndex = CmdList.GetParallelIndex();
+			DebugInfo.CommandIndex = NumCommands;
+			DebugInfo.CmdBuffIndex = GetCommandBufferIndex();
+			DebugInfo.CommandBuffer = reinterpret_cast<uintptr_t>(CurrentEncoder.GetCommandBuffer().GetPtr());
+			DebugInfo.PSOSignature[0] = BoundShaderState->VertexShader->SourceLen;
+			DebugInfo.PSOSignature[1] = BoundShaderState->VertexShader->SourceCRC;
+			if (IsValidRef(BoundShaderState->PixelShader))
+			{
+				DebugInfo.PSOSignature[2] = BoundShaderState->PixelShader->SourceLen;
+				DebugInfo.PSOSignature[3] = BoundShaderState->PixelShader->SourceCRC;
+			}
+			else
+			{
+				DebugInfo.PSOSignature[2] = 0;
+				DebugInfo.PSOSignature[3] = 0;
+			}
+			
+			FMetalShaderPipeline* PSO = State.GetPipelineState();
+			
+			CurrentEncoder.GetRenderCommandEncoder().SetRenderPipelineState(PSO->DebugPipelineState);
+
+		#if PLATFORM_MAC
+			id<MTLBuffer> DebugBufferPtr = State.GetDebugBuffer().GetPtr();
+			[CurrentEncoder.GetRenderCommandEncoder().GetPtr() memoryBarrierWithResources:&DebugBufferPtr count:1 afterStages:MTLRenderStageFragment beforeStages:MTLRenderStageVertex];
+			
+			CurrentEncoder.SetShaderBytes(mtlpp::FunctionType::Vertex, (uint8 const*)&DebugInfo, sizeof(DebugInfo), 0);
+			State.SetShaderBufferDirty(EMetalShaderStages::Vertex, 0);
+
+			CurrentEncoder.SetShaderBuffer(mtlpp::FunctionType::Vertex, State.GetDebugBuffer(), 0, State.GetDebugBuffer().GetLength(), 1, mtlpp::ResourceUsage::Write);
+			State.SetShaderBufferDirty(EMetalShaderStages::Vertex, 1);
+
+			CurrentEncoder.GetRenderCommandEncoder().Draw(mtlpp::PrimitiveType::Point, 0, 1);
+			
+			[CurrentEncoder.GetRenderCommandEncoder().GetPtr() memoryBarrierWithResources:&DebugBufferPtr count:1 afterStages:MTLRenderStageVertex beforeStages:MTLRenderStageVertex];
+		#else
+			CurrentEncoder.GetRenderCommandEncoder().SetTileData((uint8 const*)&DebugInfo, sizeof(DebugInfo), 0);
+			CurrentEncoder.GetRenderCommandEncoder().SetTileBuffer(State.GetDebugBuffer(), 0, 1);
+			mtlpp::Size ThreadsPerTile(1, 1, 1);
+			CurrentEncoder.GetRenderCommandEncoder().DispatchThreadsPerTile(ThreadsPerTile);
+		#endif
+
+			CurrentEncoder.GetRenderCommandEncoder().SetRenderPipelineState(PSO->RenderPipelineState);
+		}
+	}
+}
diff --git a/Engine/Source/Runtime/Apple/MetalRHI/Private/MetalRenderPass.h b/Engine/Source/Runtime/Apple/MetalRHI/Private/MetalRenderPass.h
--- a/Engine/Source/Runtime/Apple/MetalRHI/Private/MetalRenderPass.h
+++ b/Engine/Source/Runtime/Apple/MetalRHI/Private/MetalRenderPass.h
@@ -173,6 +173,11 @@ private:
 	void CommitAsyncDispatchResourceTables(void);
     
     void ConditionalSubmit();
+	
+	uint32 GetEncoderIndex(void) const;
+	uint32 GetCommandBufferIndex(void) const;
+	
+	void InsertDebugDraw(FMetalCommandData& Data);
 private:
 #pragma mark -
 	FMetalCommandList& CmdList;
diff --git a/Engine/Source/Runtime/Apple/MetalRHI/Private/MetalStateCache.cpp b/Engine/Source/Runtime/Apple/MetalRHI/Private/MetalStateCache.cpp
--- a/Engine/Source/Runtime/Apple/MetalRHI/Private/MetalStateCache.cpp
+++ b/Engine/Source/Runtime/Apple/MetalRHI/Private/MetalStateCache.cpp
@@ -1975,6 +1975,11 @@ void FMetalStateCache::SetStateDirty(void)
 	}
 }
 
+void FMetalStateCache::SetShaderBufferDirty(EMetalShaderStages const Frequency, NSUInteger const Index)
+{
+	ShaderBuffers[Frequency].Bound |= (1 << Index);
+}
+
 void FMetalStateCache::SetRenderStoreActions(FMetalCommandEncoder& CommandEncoder, bool const bConditionalSwitch)
 {
 	check(CommandEncoder.IsRenderCommandEncoderActive())
@@ -2261,6 +2266,17 @@ void FMetalStateCache::CommitResourceTable(EMetalShaderStages const Frequency, m
 	}
 }
 
+FMetalBuffer& FMetalStateCache::GetDebugBuffer()
+{
+    if (!DebugBuffer)
+    {
+        // Assume worst case tiling (16x16) and render-target size (4096x4096) on iOS for now
+        uint32 NumTiles = PLATFORM_MAC ? 1 : 65536;
+        DebugBuffer = GetMetalDeviceContext().CreatePooledBuffer(FMetalPooledBufferArgs(GetMetalDeviceContext().GetDevice(), NumTiles * sizeof(FMetalDebugInfo), BUF_Dynamic, mtlpp::StorageMode::Shared));
+    }
+    return DebugBuffer;
+}
+
 FTexture2DRHIRef FMetalStateCache::CreateFallbackDepthStencilSurface(uint32 Width, uint32 Height)
 {
 #if PLATFORM_MAC
diff --git a/Engine/Source/Runtime/Apple/MetalRHI/Private/MetalStateCache.h b/Engine/Source/Runtime/Apple/MetalRHI/Private/MetalStateCache.h
--- a/Engine/Source/Runtime/Apple/MetalRHI/Private/MetalStateCache.h
+++ b/Engine/Source/Runtime/Apple/MetalRHI/Private/MetalStateCache.h
@@ -98,6 +98,8 @@ public:
 
 	void SetStateDirty(void);
 	
+	void SetShaderBufferDirty(EMetalShaderStages const Frequency, NSUInteger const Index);
+	
 	void SetRenderStoreActions(FMetalCommandEncoder& CommandEncoder, bool const bConditionalSwitch);
 	
 	void SetRenderState(FMetalCommandEncoder& CommandEncoder, FMetalCommandEncoder* PrologueEncoder);
@@ -114,6 +116,7 @@ public:
 	
 	bool PrepareToRestart(bool const bCurrentApplied);
 	
+	FMetalBuffer& GetDebugBuffer();
 	FMetalShaderParameterCache& GetShaderParameters(EMetalShaderStages const Stage) { return ShaderParameters[Stage]; }
 	FLinearColor const& GetBlendFactor() const { return BlendFactor; }
 	uint32 GetStencilRef() const { return StencilRef; }
@@ -281,6 +284,7 @@ private:
 	FTextureRHIRef DepthStencilResolve;
 	/** A fallback depth-stencil surface for draw calls that write to depth without a depth-stencil surface bound. */
 	FTexture2DRHIRef FallbackDepthStencilSurface;
+	FMetalBuffer DebugBuffer;
 	mtlpp::RenderPassDescriptor RenderPassDesc;
 	uint32 RasterBits;
     uint8 PipelineBits;
