Remove reference to obsolete variable that slipped back in with a merge.

diff --git a/Engine/Source/Runtime/Apple/MetalRHI/Private/MetalCommandQueue.cpp b/Engine/Source/Runtime/Apple/MetalRHI/Private/MetalCommandQueue.cpp
--- a/Engine/Source/Runtime/Apple/MetalRHI/Private/MetalCommandQueue.cpp
+++ b/Engine/Source/Runtime/Apple/MetalRHI/Private/MetalCommandQueue.cpp
@@ -143,8 +143,6 @@ FMetalCommandQueue::FMetalCommandQueue(mtlpp::Device InDevice, uint32 const MaxN
                 {
                     Features |= EMetalFeaturesTileShaders;
                     
-                    GMetalSupportsTileShaders = true;
-                    
                     // The below implies tile shaders which are necessary to order the draw calls and generate a buffer that shows what PSOs/draws ran on each tile.
                     IConsoleVariable* GPUCrashDebuggingCVar = IConsoleManager::Get().FindConsoleVariable(TEXT("r.GPUCrashDebugging"));
                     GMetalCommandBufferDebuggingEnabled = (GPUCrashDebuggingCVar && GPUCrashDebuggingCVar->GetInt() != 0) || FParse::Param(FCommandLine::Get(),TEXT("metalgpudebug"));
