Change the way we restore the OS binary shader cache as this may help on Intel/Nvidia.

diff --git a/Engine/Source/Runtime/Apple/MetalRHI/Private/MetalRHI.cpp b/Engine/Source/Runtime/Apple/MetalRHI/Private/MetalRHI.cpp
--- a/Engine/Source/Runtime/Apple/MetalRHI/Private/MetalRHI.cpp
+++ b/Engine/Source/Runtime/Apple/MetalRHI/Private/MetalRHI.cpp
@@ -807,12 +807,20 @@ FMetalDynamicRHI::FMetalDynamicRHI(ERHIFeatureLevel::Type RequestedFeatureLevel)
 			NSString* TempDir = [NSString stringWithFormat:@"%@/../C/%@/com.apple.metal", NSTemporaryDirectory(), [NSBundle mainBundle].bundleIdentifier];
 
 			NSError* Err = nil;
-			BOOL bOK = [[NSFileManager defaultManager] removeItemAtPath:TempDir
-						error:&Err];
+			NSDirectoryEnumerator<NSString *> * Enum = [[NSFileManager defaultManager] enumeratorAtPath:DstPath];
+			for (NSString* Path in Enum)
+			{
+				[Enum skipDescendents];
 
-			bOK = [[NSFileManager defaultManager] copyItemAtPath:DstPath
-						toPath:TempDir
-						error:&Err];
+				NSString* Dest = [NSString stringWithFormat:@"%@/%@", TempDir, Path];
+				if(![[NSFileManager defaultManager] fileExistsAtPath:Dest])
+				{
+					NSString* Src = [NSString stringWithFormat:@"%@/%@", DstPath, Path];
+					[[NSFileManager defaultManager] copyItemAtPath:Src
+																   toPath:Dest
+																	error:&Err];
+				}
+			}
 		}
 	}
 #endif
