Need to create the temporary directory if it doesn't exist too.

diff --git a/Engine/Source/Runtime/Apple/MetalRHI/Private/MetalRHI.cpp b/Engine/Source/Runtime/Apple/MetalRHI/Private/MetalRHI.cpp
--- a/Engine/Source/Runtime/Apple/MetalRHI/Private/MetalRHI.cpp
+++ b/Engine/Source/Runtime/Apple/MetalRHI/Private/MetalRHI.cpp
@@ -807,6 +807,14 @@ FMetalDynamicRHI::FMetalDynamicRHI(ERHIFeatureLevel::Type RequestedFeatureLevel)
 			NSString* TempDir = [NSString stringWithFormat:@"%@/../C/%@/com.apple.metal", NSTemporaryDirectory(), [NSBundle mainBundle].bundleIdentifier];
 
 			NSError* Err = nil;
+			if(![[NSFileManager defaultManager] fileExistsAtPath:TempDir])
+			{
+				[[NSFileManager defaultManager] createDirectoryAtPath:TempDir
+											withIntermediateDirectories:YES
+															 attributes:nil
+																  error:&Err];
+			}
+			
 			NSDirectoryEnumerator<NSString *> * Enum = [[NSFileManager defaultManager] enumeratorAtPath:DstPath];
 			for (NSString* Path in Enum)
 			{
