Changes to Metal GPU debugging code from working on iOS.

diff --git a/Engine/Source/Runtime/Apple/MetalRHI/Private/MetalCommandEncoder.cpp b/Engine/Source/Runtime/Apple/MetalRHI/Private/MetalCommandEncoder.cpp
--- a/Engine/Source/Runtime/Apple/MetalRHI/Private/MetalCommandEncoder.cpp
+++ b/Engine/Source/Runtime/Apple/MetalRHI/Private/MetalCommandEncoder.cpp
@@ -107,7 +107,7 @@ void FMetalCommandBufferMarkers::AllocateContexts(uint32 NumContexts)
 	}
 }
 
-uint32 FMetalCommandBufferMarkers::AddCommand(uint32 Encoder, uint32 ContextIndex, FMetalBuffer& DebugBuffer, FMetalGraphicsPipelineState* PSO, FMetalCommandData& Data)
+uint32 FMetalCommandBufferMarkers::AddCommand(uint32 CmdBufIndex, uint32 Encoder, uint32 ContextIndex, FMetalBuffer& DebugBuffer, FMetalGraphicsPipelineState* PSO, FMetalCommandData& Data)
 {
 	uint32 Num = 0;
 	if (m_ptr)
@@ -120,6 +120,7 @@ uint32 FMetalCommandBufferMarkers::AddCommand(uint32 Encoder, uint32 ContextInde
 		Context.PSOs.Add(PSO);
 		Num = Context.Commands.Num();
 		FMetalCommandDebug Command;
+        Command.CmdBufIndex = CmdBufIndex;
 		Command.Encoder = Encoder;
 		Command.Index = Num;
 		Command.PSO = PSO;
diff --git a/Engine/Source/Runtime/Apple/MetalRHI/Private/MetalCommandEncoder.h b/Engine/Source/Runtime/Apple/MetalRHI/Private/MetalCommandEncoder.h
--- a/Engine/Source/Runtime/Apple/MetalRHI/Private/MetalCommandEncoder.h
+++ b/Engine/Source/Runtime/Apple/MetalRHI/Private/MetalCommandEncoder.h
@@ -65,6 +65,7 @@ struct FMetalCommandData
 
 struct FMetalCommandDebug
 {
+    uint32 CmdBufIndex;
 	uint32 Encoder;
 	uint32 Index;
 	FMetalGraphicsPipelineState* PSO;
@@ -79,7 +80,7 @@ public:
 	FMetalCommandBufferMarkers(FMetalCommandBufferDebug* CmdBuf);
 	
 	void AllocateContexts(uint32 NumContexts);
-	uint32 AddCommand(uint32 Encoder, uint32 Context, FMetalBuffer& DebugBuffer, FMetalGraphicsPipelineState* PSO, FMetalCommandData& Data);
+	uint32 AddCommand(uint32 CmdBufIndex, uint32 Encoder, uint32 Context, FMetalBuffer& DebugBuffer, FMetalGraphicsPipelineState* PSO, FMetalCommandData& Data);
 	TArray<FMetalCommandDebug>* GetCommands(uint32 Context);
 	ns::AutoReleased<FMetalBuffer> GetDebugBuffer(uint32 ContextIndex);
 	uint32 NumContexts() const;
diff --git a/Engine/Source/Runtime/Apple/MetalRHI/Private/MetalCommandList.cpp b/Engine/Source/Runtime/Apple/MetalRHI/Private/MetalCommandList.cpp
--- a/Engine/Source/Runtime/Apple/MetalRHI/Private/MetalCommandList.cpp
+++ b/Engine/Source/Runtime/Apple/MetalRHI/Private/MetalCommandList.cpp
@@ -79,6 +79,7 @@ static void ReportMetalCommandBufferFailure(mtlpp::CommandBuffer const& Complete
 			BrokenDraw.CommandIndex = UINT32_MAX;
 			BrokenDraw.CmdBuffIndex = UINT32_MAX;
 			BrokenDraw.ContextIndex = UINT32_MAX;
+            BrokenDraw.CommandBuffer = 0;
 			uint32 BrokenContext = 0;
 			bool bFoundBrokenDraw = false;
 			for (uint32 i = 0; !bFoundBrokenDraw && i < Markers.NumContexts(); i++)
@@ -127,9 +128,9 @@ static void ReportMetalCommandBufferFailure(mtlpp::CommandBuffer const& Complete
 			if (bFoundBrokenDraw)
 			{
 				id<MTLCommandBuffer> Ptr = reinterpret_cast<id<MTLCommandBuffer>>(BrokenDraw.CommandBuffer);
-				UE_LOG(LogMetal, Error, TEXT("GPU last wrote Command Buffer: %u (%p) Encoder Index: %d Context Index: %d Draw Index: %d PSO: VS: %u_%u, PS: %u_%u."), BrokenDraw.CmdBuffIndex, Ptr, BrokenDraw.EncoderIndex, BrokenDraw.ContextIndex, BrokenDraw.CommandIndex, BrokenDraw.PSOSignature[0], BrokenDraw.PSOSignature[1], BrokenDraw.PSOSignature[2], BrokenDraw.PSOSignature[3]);
+				UE_LOG(LogMetal, Error, TEXT("GPU last wrote Command Buffer: %u (%llx) Encoder Index: %d Context Index: %d Draw Index: %d PSO: VS: %u_%u, PS: %u_%u."), BrokenDraw.CmdBuffIndex, BrokenDraw.CommandBuffer, BrokenDraw.EncoderIndex, BrokenDraw.ContextIndex, BrokenDraw.CommandIndex, BrokenDraw.PSOSignature[0], BrokenDraw.PSOSignature[1], BrokenDraw.PSOSignature[2], BrokenDraw.PSOSignature[3]);
 
-				if (Ptr != CompletedBuffer.GetPtr())
+                if (Ptr && Ptr != CompletedBuffer.GetPtr())
 				{
 					// Have we failed between commands-buffers? Or have our command-buffers been subject to coalescing?
 					ns::AutoReleased<mtlpp::CommandBuffer> PrevCommandBuffer(Ptr);
@@ -156,7 +157,7 @@ static void ReportMetalCommandBufferFailure(mtlpp::CommandBuffer const& Complete
 								PSSig[1] = Command.PSO->PixelShader->SourceCRC;
 							}
 							
-							UE_LOG(LogMetal, Error, TEXT("Command Buffer: %p Encoder: %d Command: %d: %s PSO: VS: %s (%u_%u), PS: %s (%u_%u)"), Ptr, Command.Encoder, Command.Index, *Command.Data.ToString(), *VSHash, VSSig[0], VSSig[1], *PSHash, PSSig[0], PSSig[1]);
+                            UE_LOG(LogMetal, Error, TEXT("Command Buffer: %d (%p) Encoder: %d Command: %d: %s PSO: VS: %s (%u_%u), PS: %s (%u_%u)"), Command.CmdBufIndex, Ptr, Command.Encoder, Command.Index, *Command.Data.ToString(), *VSHash, VSSig[0], VSSig[1], *PSHash, PSSig[0], PSSig[1]);
 						}
 					}
 					
@@ -169,8 +170,7 @@ static void ReportMetalCommandBufferFailure(mtlpp::CommandBuffer const& Complete
 				{
 					UE_LOG(LogMetal, Error, TEXT("Failed executing following commands:"));
 					uint32 StartIdx = BrokenDraw.CommandIndex;
-					uint32 EndIdx = BrokenDraw.CommandIndex + GMetalDebugOpsCount + 1;
-					for (uint32 CmdIdx = StartIdx; CmdIdx < Commands->Num() && CmdIdx < EndIdx; CmdIdx++)
+					for (uint32 CmdIdx = StartIdx; CmdIdx < Commands->Num(); CmdIdx++)
 					{
 						FMetalCommandDebug& Command = (*Commands)[CmdIdx];
 						FString VSHash = Command.PSO->VertexShader->GetHash().ToString();
@@ -185,7 +185,20 @@ static void ReportMetalCommandBufferFailure(mtlpp::CommandBuffer const& Complete
 							PSSig[1] = Command.PSO->PixelShader->SourceCRC;
 						}
 						
-						UE_LOG(LogMetal, Error, TEXT("Command Buffer: %p Encoder: %d Command: %d: %s PSO: VS: %s (%u_%u), PS: %s (%u_%u)"), CompletedBuffer.GetPtr(), Command.Encoder, Command.Index, *Command.Data.ToString(), *VSHash, VSSig[0], VSSig[1], *PSHash, PSSig[0], PSSig[1]);
+						UE_LOG(LogMetal, Error, TEXT("Command Buffer: %d (%p) Encoder: %d Command: %d: %s PSO: VS: %s (%u_%u), PS: %s (%u_%u)"), Command.CmdBufIndex, CompletedBuffer.GetPtr(), Command.Encoder, Command.Index, *Command.Data.ToString(), *VSHash, VSSig[0], VSSig[1], *PSHash, PSSig[0], PSSig[1]);
+					}
+				}
+                if (DebugBuffer)
+                {
+                    uint32 DebugLength = DebugBuffer.GetLength();
+                    uint32 DebugCount = DebugLength / sizeof(FMetalDebugInfo);
+                    check(DebugCount >= 1);
+                    
+                    FMetalDebugInfo* DebugArray = (FMetalDebugInfo*)DebugBuffer.GetContents();
+                    for (uint32 i = 0; i < DebugCount; i++)
+                    {
+                        FMetalDebugInfo& Command = DebugArray[i];
+                        UE_LOG(LogMetal, Error, TEXT("Command Buffer: %d (%p) Debug Buffer: %p Tile: %u Context: %d Encoder: %d Command: %d PSO: VS: %u_%u, PS: %u_%u"), Command.CmdBuffIndex, Command.CommandBuffer, DebugBuffer.GetPtr(), i, Command.ContextIndex, Command.EncoderIndex, Command.CommandIndex, Command.PSOSignature[0], Command.PSOSignature[1], Command.PSOSignature[2], Command.PSOSignature[3]);
 					}
 				}
 			}
diff --git a/Engine/Source/Runtime/Apple/MetalRHI/Private/MetalPipeline.cpp b/Engine/Source/Runtime/Apple/MetalRHI/Private/MetalPipeline.cpp
--- a/Engine/Source/Runtime/Apple/MetalRHI/Private/MetalPipeline.cpp
+++ b/Engine/Source/Runtime/Apple/MetalRHI/Private/MetalPipeline.cpp
@@ -68,16 +68,17 @@ static NSString* GMetalDebugShader = @"#include <metal_stdlib>\n"
 "	uint EncoderIndex;\n"
 "   uint ContextIndex;\n"
 "   uint CommandIndex;\n"
-"   uint2 CommandBuffer;\n"
-"	uint4 PSOSignature;\n"
+"   uint CommandBuffer[2];\n"
+"	uint PSOSignature[4];\n"
 "};\n"
 "\n"
 #if !PLATFORM_MAC
 "// Executes once per-tile\n"
-"fragment void Main_Debug(constant DebugInfo *debugTable [[ buffer(0) ]], device DebugInfo* debugBuffer [[ buffer(1) ]], uint tile_index [[ tile_index ]])\n"
+"kernel void Main_Debug(constant DebugInfo *debugTable [[ buffer(0) ]], device DebugInfo* debugBuffer [[ buffer(1) ]], uint2 threadgroup_position_in_grid [[ threadgroup_position_in_grid ]], uint2 threadgroups_per_grid [[ threadgroups_per_grid ]])\n"
 "{\n"
 "	// Write Pass, Draw indices\n"
 "	// Write Vertex+Fragment PSO sig (in form VertexLen, VertexCRC, FragLen, FragCRC)\n"
+"   uint tile_index = threadgroup_position_in_grid.x + (threadgroup_position_in_grid.y * threadgroups_per_grid.x);"
 "	debugBuffer[tile_index] = debugTable[0];\n"
 "}";
 #else
diff --git a/Engine/Source/Runtime/Apple/MetalRHI/Private/MetalRenderPass.cpp b/Engine/Source/Runtime/Apple/MetalRHI/Private/MetalRenderPass.cpp
--- a/Engine/Source/Runtime/Apple/MetalRHI/Private/MetalRenderPass.cpp
+++ b/Engine/Source/Runtime/Apple/MetalRHI/Private/MetalRenderPass.cpp
@@ -1886,7 +1886,7 @@ void FMetalRenderPass::InsertDebugDraw(FMetalCommandData& Data)
 	{
 		FMetalGraphicsPipelineState* BoundShaderState = State.GetGraphicsPSO();
 		
-		uint32 NumCommands = CurrentEncoder.GetMarkers().AddCommand(GetEncoderIndex(), CmdList.GetParallelIndex(), State.GetDebugBuffer(), BoundShaderState, Data);
+		uint32 NumCommands = CurrentEncoder.GetMarkers().AddCommand(GetCommandBufferIndex(), GetEncoderIndex(), CmdList.GetParallelIndex(), State.GetDebugBuffer(), BoundShaderState, Data);
 		
 		if ((NumCommands % GMetalDebugOpsCount) == 0)
 		{
