Duplicate 5676864 & 5672343 from //UE4/Dev-Main to Dev-Rendering (//UE4/Dev-Rendering)

The current tvOS hardware & API doesn't support tile-shaders so we must disable the GPU debugging code on PLATFORM_TVOS.
Missed one of the tile-shader references that aren't available on tvOS.

diff --git a/Engine/Source/Runtime/Apple/MetalRHI/Private/MetalPipeline.cpp b/Engine/Source/Runtime/Apple/MetalRHI/Private/MetalPipeline.cpp
--- a/Engine/Source/Runtime/Apple/MetalRHI/Private/MetalPipeline.cpp
+++ b/Engine/Source/Runtime/Apple/MetalRHI/Private/MetalPipeline.cpp
@@ -98,6 +98,7 @@ struct FMetalHelperFunctions
     
     FMetalHelperFunctions()
     {
+#if !PLATFORM_TVOS
         if (GMetalCommandBufferDebuggingEnabled)
         {
             mtlpp::CompileOptions CompileOptions;
@@ -105,6 +106,7 @@ struct FMetalHelperFunctions
             DebugShadersLib = GetMetalDeviceContext().GetDevice().NewLibrary(GMetalDebugShader, CompileOptions, &Error);
             DebugFunc = DebugShadersLib.NewFunction(@"Main_Debug");
         }
+#endif
     }
     
     static FMetalHelperFunctions& Get()
@@ -544,7 +546,7 @@ static FMetalShaderPipeline* CreateMTLRenderPipeline(bool const bSync, FMetalGra
         mtlpp::ComputePipelineDescriptor ComputePipelineDesc(nil);
 #if PLATFORM_MAC
         mtlpp::RenderPipelineDescriptor DebugPipelineDesc;
-#else
+#elif !PLATFORM_TVOS
         mtlpp::TileRenderPipelineDescriptor DebugPipelineDesc;
 #endif
 		
@@ -604,7 +606,9 @@ static FMetalShaderPipeline* CreateMTLRenderPipeline(bool const bSync, FMetalGra
         FMetalBlendState* BlendState = (FMetalBlendState*)Init.BlendState;
 		
 		ns::Array<mtlpp::RenderPipelineColorAttachmentDescriptor> ColorAttachments = RenderPipelineDesc.GetColorAttachments();
+#if !PLATFORM_TVOS
 		auto DebugColorAttachements = DebugPipelineDesc.GetColorAttachments();
+#endif
 
 		uint32 TargetWidth = 0;
 		for (uint32 i = 0; i < NumActiveTargets; i++)
@@ -633,8 +637,10 @@ static FMetalShaderPipeline* CreateMTLRenderPipeline(bool const bSync, FMetalGra
             mtlpp::RenderPipelineColorAttachmentDescriptor Attachment = ColorAttachments[i];
             Attachment.SetPixelFormat(MetalFormat);
 			
+#if !PLATFORM_TVOS
 			auto DebugAttachment = DebugColorAttachements[i];;
 			DebugAttachment.SetPixelFormat(MetalFormat);
+#endif
 			
             mtlpp::RenderPipelineColorAttachmentDescriptor Blend = BlendState->RenderTargetStates[i].BlendState;
             if(TargetFormat != PF_Unknown)
@@ -1136,6 +1142,7 @@ static FMetalShaderPipeline* CreateMTLRenderPipeline(bool const bSync, FMetalGra
         Pipeline->FragmentSource = PixelShader ? PixelShader->GetSourceCode() : nil;
     #endif
 		
+#if !PLATFORM_TVOS
 		if (GMetalCommandBufferDebuggingEnabled)
 		{
 #if PLATFORM_MAC
@@ -1154,6 +1161,7 @@ static FMetalShaderPipeline* CreateMTLRenderPipeline(bool const bSync, FMetalGra
 			METAL_GPUPROFILE(FScopedMetalCPUStats CPUStat(FString::Printf(TEXT("NewDebugPipeline: %s"), TEXT("")/**FString([RenderPipelineDesc.GetPtr() description])*/)));
 			Pipeline->DebugPipelineState = Device.NewRenderPipelineState(DebugPipelineDesc, (mtlpp::PipelineOption)RenderOption, Reflection, &RenderError);
 		}
+#endif
         
 #if METAL_DEBUG_OPTIONS
         if (GFrameCounter > 3)
diff --git a/Engine/Source/Runtime/Apple/MetalRHI/Private/MetalRenderPass.cpp b/Engine/Source/Runtime/Apple/MetalRHI/Private/MetalRenderPass.cpp
--- a/Engine/Source/Runtime/Apple/MetalRHI/Private/MetalRenderPass.cpp
+++ b/Engine/Source/Runtime/Apple/MetalRHI/Private/MetalRenderPass.cpp
@@ -1882,6 +1882,7 @@ uint32 FMetalRenderPass::GetCommandBufferIndex(void) const
 
 void FMetalRenderPass::InsertDebugDraw(FMetalCommandData& Data)
 {
+#if !PLATFORM_TVOS
 	if (GMetalCommandBufferDebuggingEnabled)
 	{
 		FMetalGraphicsPipelineState* BoundShaderState = State.GetGraphicsPSO();
@@ -1936,4 +1937,5 @@ void FMetalRenderPass::InsertDebugDraw(FMetalCommandData& Data)
 			CurrentEncoder.GetRenderCommandEncoder().SetRenderPipelineState(PSO->RenderPipelineState);
 		}
 	}
+#endif
 }
