DR - Remove unused EResourceAliasability API
#rnx

diff --git a/Engine/Source/Runtime/Apple/MetalRHI/Private/MetalTexture.cpp b/Engine/Source/Runtime/Apple/MetalRHI/Private/MetalTexture.cpp
--- a/Engine/Source/Runtime/Apple/MetalRHI/Private/MetalTexture.cpp
+++ b/Engine/Source/Runtime/Apple/MetalRHI/Private/MetalTexture.cpp
@@ -2546,44 +2546,6 @@ void FMetalDynamicRHI::RHIDiscardTransientResource_RenderThread(FStructuredBuffe
 	}
 }
 
-void FMetalDynamicRHI::RHISetResourceAliasability_RenderThread(class FRHICommandListImmediate& RHICmdList, EResourceAliasability AliasMode, FTextureRHIParamRef* InTextures, int32 NumTextures)
-{
-	check(InTextures && NumTextures);
-	@autoreleasepool
-	{
-		switch(AliasMode)
-		{
-			case EResourceAliasability::EAliasable:
-			{
-				for (int32 i = 0; i < NumTextures; ++i)
-				{
-					FMetalSurface* Source = GetMetalSurfaceFromRHITexture(InTextures[i]);
-					Source->MakeAliasable();
-				}
-				break;
-			}
-			case EResourceAliasability::EUnaliasable:
-			{
-				if (RHICmdList.Bypass() || !IsRunningRHIInSeparateThread())
-				{
-					for (int32 i = 0; i < NumTextures; ++i)
-					{
-						FMetalSurface* Source = GetMetalSurfaceFromRHITexture(InTextures[i]);
-						Source->MakeUnAliasable();
-					}
-				}
-				else
-				{
-					new (RHICmdList.AllocCommand<FMetalRHICommandUnaliasTextures>()) FMetalRHICommandUnaliasTextures(InTextures, NumTextures);
-				}
-				break;
-			}
-			default:
-				break;
-		}
-	}
-}
-
 struct FRHICopySubTextureRegion final : public FRHICommand<FRHICopySubTextureRegion>
 {
 	FTexture2DRHIParamRef SourceTexture;
diff --git a/Engine/Source/Runtime/Apple/MetalRHI/Public/MetalRHI.h b/Engine/Source/Runtime/Apple/MetalRHI/Public/MetalRHI.h
--- a/Engine/Source/Runtime/Apple/MetalRHI/Public/MetalRHI.h
+++ b/Engine/Source/Runtime/Apple/MetalRHI/Public/MetalRHI.h
@@ -199,7 +199,6 @@ public:
 	virtual void RHIDiscardTransientResource_RenderThread(FVertexBufferRHIParamRef Buffer) final override;
 	virtual void RHIAcquireTransientResource_RenderThread(FStructuredBufferRHIParamRef Buffer) final override;
 	virtual void RHIDiscardTransientResource_RenderThread(FStructuredBufferRHIParamRef Buffer) final override;
-	virtual void RHISetResourceAliasability_RenderThread(class FRHICommandListImmediate& RHICmdList, EResourceAliasability AliasMode, FTextureRHIParamRef* InTextures, int32 NumTextures) final override;
 	
 	/* Copy the source box pixels in the destination box texture, return true if implemented for the current platform*/
 	virtual void RHICopySubTextureRegion_RenderThread(class FRHICommandListImmediate& RHICmdList, FTexture2DRHIParamRef SourceTexture, FTexture2DRHIParamRef DestinationTexture, FBox2D SourceBox, FBox2D DestinationBox) final override;
diff --git a/Engine/Source/Runtime/RHI/Public/DynamicRHI.h b/Engine/Source/Runtime/RHI/Public/DynamicRHI.h
--- a/Engine/Source/Runtime/RHI/Public/DynamicRHI.h
+++ b/Engine/Source/Runtime/RHI/Public/DynamicRHI.h
@@ -1103,8 +1103,6 @@ public:
 	//Utilities
 	virtual void EnableIdealGPUCaptureOptions(bool bEnable);
 	
-	virtual void RHISetResourceAliasability_RenderThread(class FRHICommandListImmediate& RHICmdList, EResourceAliasability AliasMode, FTextureRHIParamRef* InTextures, int32 NumTextures) {}
-
 	//checks if the GPU is still alive.
 	virtual bool CheckGpuHeartbeat() const { return true; }
 
diff --git a/Engine/Source/Runtime/RHI/Public/RHI.h b/Engine/Source/Runtime/RHI/Public/RHI.h
--- a/Engine/Source/Runtime/RHI/Public/RHI.h
+++ b/Engine/Source/Runtime/RHI/Public/RHI.h
@@ -1322,12 +1322,6 @@ enum class EResourceTransitionAccess
 	EMaxAccess,
 };
 
-enum class EResourceAliasability
-{
-	EAliasable, // Make the resource aliasable with other resources
-	EUnaliasable, // Make the resource unaliasable with any other resources
-};
-
 class RHI_API FResourceTransitionUtility
 {
 public:
diff --git a/Engine/Source/Runtime/RHI/Public/RHICommandList.h b/Engine/Source/Runtime/RHI/Public/RHICommandList.h
--- a/Engine/Source/Runtime/RHI/Public/RHICommandList.h
+++ b/Engine/Source/Runtime/RHI/Public/RHICommandList.h
@@ -4166,11 +4166,6 @@ public:
 		GDynamicRHI->RHIExecuteCommandList(CmdList);
 	}
 	
-	FORCEINLINE void SetResourceAliasability(EResourceAliasability AliasMode, FTextureRHIParamRef* InTextures, int32 NumTextures)
-	{
-		GDynamicRHI->RHISetResourceAliasability_RenderThread(*this, AliasMode, InTextures, NumTextures);
-	}
-	
 	FORCEINLINE void* GetNativeDevice()
 	{
 		QUICK_SCOPE_CYCLE_COUNTER(STAT_RHIMETHOD_GetNativeDevice_Flush);
diff --git a/Engine/Source/Runtime/RHI/Public/ValidationRHI.h b/Engine/Source/Runtime/RHI/Public/ValidationRHI.h
--- a/Engine/Source/Runtime/RHI/Public/ValidationRHI.h
+++ b/Engine/Source/Runtime/RHI/Public/ValidationRHI.h
@@ -582,6 +582,7 @@ public:
 	// FlushType: Thread safe
 	virtual FTexture2DRHIRef RHIAsyncCreateTexture2D(uint32 SizeX, uint32 SizeY, uint8 Format, uint32 NumMips, uint32 Flags, void** InitialMipData, uint32 NumInitialMips) override final
 	{
+		check(GRHISupportsAsyncTextureCreation);
 		return RHI->RHIAsyncCreateTexture2D(SizeX, SizeY, Format, NumMips, Flags, InitialMipData, NumInitialMips);
 	}
 
@@ -1556,11 +1557,6 @@ public:
 		RHI->EnableIdealGPUCaptureOptions(bEnable);
 	}
 
-	virtual void RHISetResourceAliasability_RenderThread(class FRHICommandListImmediate& RHICmdList, EResourceAliasability AliasMode, FRHITexture** InTextures, int32 NumTextures) override final
-	{
-		RHI->RHISetResourceAliasability_RenderThread(RHICmdList, AliasMode, InTextures, NumTextures);
-	}
-
 	//checks if the GPU is still alive.
 	virtual bool CheckGpuHeartbeat() const override final
 	{
