Remove extraneous extern declaration.

diff --git a/Engine/Source/Runtime/Apple/MetalRHI/Private/MetalRHIPrivate.h b/Engine/Source/Runtime/Apple/MetalRHI/Private/MetalRHIPrivate.h
--- a/Engine/Source/Runtime/Apple/MetalRHI/Private/MetalRHIPrivate.h
+++ b/Engine/Source/Runtime/Apple/MetalRHI/Private/MetalRHIPrivate.h
@@ -93,7 +93,6 @@ extern FMetalBufferFormat GMetalBufferFormats[PF_MAX];
 #else
 #define METAL_DEBUG_OPTION(Code)
 #endif
-extern int32 GMetalResetOnPSOChange;
 
 #if MTLPP_CONFIG_VALIDATE && METAL_DEBUG_OPTIONS
 #define METAL_DEBUG_ONLY(Code) Code
