Fix merge issues.

diff --git a/Engine/Source/Runtime/Apple/MetalRHI/Private/MetalRenderPass.cpp b/Engine/Source/Runtime/Apple/MetalRHI/Private/MetalRenderPass.cpp
--- a/Engine/Source/Runtime/Apple/MetalRHI/Private/MetalRenderPass.cpp
+++ b/Engine/Source/Runtime/Apple/MetalRHI/Private/MetalRenderPass.cpp
@@ -1936,7 +1936,6 @@ void FMetalRenderPass::ConditionalSubmit()
 				RestartRenderPass(State.GetRenderPassDescriptor());
 			}
 		}
-#endif
 	}
 }
 
@@ -2030,6 +2029,7 @@ void FMetalRenderPass::InsertDebugDraw(FMetalCommandData& Data)
 			CurrentEncoder.GetRenderCommandEncoder().SetRenderPipelineState(PSO->RenderPipelineState);
 		}
 	}
+#endif
 }
 
 void FMetalRenderPass::InsertDebugDispatch(FMetalCommandData& Data)
@@ -2076,87 +2076,3 @@ void FMetalRenderPass::InsertDebugDispatch(FMetalCommandData& Data)
 	}
 #endif
 }
-
-uint32 FMetalRenderPass::GetEncoderIndex(void) const
-{
-	if (!CmdList.IsParallel())
-	{
-		return PrologueEncoder.NumEncodedPasses() + CurrentEncoder.NumEncodedPasses();
-	}
-	else
-	{
-		return GetMetalDeviceContext().GetCurrentRenderPass().GetEncoderIndex();
-	}
-}
-
-uint32 FMetalRenderPass::GetCommandBufferIndex(void) const
-{
-	if (!CmdList.IsParallel())
-	{
-		return CurrentEncoder.GetCommandBufferIndex();
-	}
-	else
-	{
-		return GetMetalDeviceContext().GetCurrentRenderPass().GetCommandBufferIndex();
-	}
-}
-
-void FMetalRenderPass::InsertDebugDraw(FMetalCommandData& Data)
-{
-#if !PLATFORM_TVOS
-	if (GMetalCommandBufferDebuggingEnabled)
-	{
-		FMetalGraphicsPipelineState* BoundShaderState = State.GetGraphicsPSO();
-		
-		uint32 NumCommands = CurrentEncoder.GetMarkers().AddCommand(GetCommandBufferIndex(), GetEncoderIndex(), CmdList.GetParallelIndex(), State.GetDebugBuffer(), BoundShaderState, Data);
-		
-		if ((NumCommands % GMetalDebugOpsCount) == 0)
-		{
-			FMetalDebugInfo DebugInfo;
-			DebugInfo.EncoderIndex = GetEncoderIndex();
-			DebugInfo.ContextIndex = CmdList.GetParallelIndex();
-			DebugInfo.CommandIndex = NumCommands;
-			DebugInfo.CmdBuffIndex = GetCommandBufferIndex();
-			DebugInfo.CommandBuffer = reinterpret_cast<uintptr_t>(CurrentEncoder.GetCommandBuffer().GetPtr());
-			DebugInfo.PSOSignature[0] = BoundShaderState->VertexShader->SourceLen;
-			DebugInfo.PSOSignature[1] = BoundShaderState->VertexShader->SourceCRC;
-			if (IsValidRef(BoundShaderState->PixelShader))
-			{
-				DebugInfo.PSOSignature[2] = BoundShaderState->PixelShader->SourceLen;
-				DebugInfo.PSOSignature[3] = BoundShaderState->PixelShader->SourceCRC;
-			}
-			else
-			{
-				DebugInfo.PSOSignature[2] = 0;
-				DebugInfo.PSOSignature[3] = 0;
-			}
-			
-			FMetalShaderPipeline* PSO = State.GetPipelineState();
-			
-			CurrentEncoder.GetRenderCommandEncoder().SetRenderPipelineState(PSO->DebugPipelineState);
-
-		#if PLATFORM_MAC
-			id<MTLBuffer> DebugBufferPtr = State.GetDebugBuffer().GetPtr();
-			[CurrentEncoder.GetRenderCommandEncoder().GetPtr() memoryBarrierWithResources:&DebugBufferPtr count:1 afterStages:MTLRenderStageFragment beforeStages:MTLRenderStageVertex];
-			
-			CurrentEncoder.SetShaderBytes(mtlpp::FunctionType::Vertex, (uint8 const*)&DebugInfo, sizeof(DebugInfo), 0);
-			State.SetShaderBufferDirty(EMetalShaderStages::Vertex, 0);
-
-			CurrentEncoder.SetShaderBuffer(mtlpp::FunctionType::Vertex, State.GetDebugBuffer(), 0, State.GetDebugBuffer().GetLength(), 1, mtlpp::ResourceUsage::Write);
-			State.SetShaderBufferDirty(EMetalShaderStages::Vertex, 1);
-
-			CurrentEncoder.GetRenderCommandEncoder().Draw(mtlpp::PrimitiveType::Point, 0, 1);
-			
-			[CurrentEncoder.GetRenderCommandEncoder().GetPtr() memoryBarrierWithResources:&DebugBufferPtr count:1 afterStages:MTLRenderStageVertex beforeStages:MTLRenderStageVertex];
-		#else
-			CurrentEncoder.GetRenderCommandEncoder().SetTileData((uint8 const*)&DebugInfo, sizeof(DebugInfo), 0);
-			CurrentEncoder.GetRenderCommandEncoder().SetTileBuffer(State.GetDebugBuffer(), 0, 1);
-			mtlpp::Size ThreadsPerTile(1, 1, 1);
-			CurrentEncoder.GetRenderCommandEncoder().DispatchThreadsPerTile(ThreadsPerTile);
-		#endif
-
-			CurrentEncoder.GetRenderCommandEncoder().SetRenderPipelineState(PSO->RenderPipelineState);
-		}
-	}
-#endif
-}
