Expose the depth-bounds from Metal PSOs.

diff --git a/Engine/Source/Runtime/Apple/MetalRHI/Public/MetalResources.h b/Engine/Source/Runtime/Apple/MetalRHI/Public/MetalResources.h
--- a/Engine/Source/Runtime/Apple/MetalRHI/Public/MetalResources.h
+++ b/Engine/Source/Runtime/Apple/MetalRHI/Public/MetalResources.h
@@ -296,6 +296,7 @@ public:
 	TRefCountPtr<FMetalRasterizerState> RasterizerState;
 	
 	inline EPrimitiveType GetPrimitiveType() { return Initializer.PrimitiveType; }
+	inline bool GetDepthBounds() const { return Initializer.bDepthBounds; }
 	
 	friend class FMetalDynamicRHI;
 	
