Let's have a macro for the Metal RHI functions we purposefully ignore.

diff --git a/Engine/Source/Runtime/Apple/MetalRHI/Private/MetalCommands.cpp b/Engine/Source/Runtime/Apple/MetalRHI/Private/MetalCommands.cpp
--- a/Engine/Source/Runtime/Apple/MetalRHI/Private/MetalCommands.cpp
+++ b/Engine/Source/Runtime/Apple/MetalRHI/Private/MetalCommands.cpp
@@ -686,7 +686,7 @@ void FMetalRHICommandContext::RHISetRenderTargetsAndClear(const FRHISetRenderTar
 			PassInfo.ColorRenderTargets[i].ArraySlice = RenderTargetsInfo.ColorRenderTarget[i].ArraySliceIndex;
 			PassInfo.ColorRenderTargets[i].MipIndex = RenderTargetsInfo.ColorRenderTarget[i].MipIndex;
 			PassInfo.ColorRenderTargets[i].Action = MakeRenderTargetActions(RenderTargetsInfo.ColorRenderTarget[i].LoadAction, RenderTargetsInfo.ColorRenderTarget[i].StoreAction);
-			bHasTarget = (RenderTargetsInfo.ColorRenderTarget[i].Texture != nullptr);
+		bHasTarget = (RenderTargetsInfo.ColorRenderTarget[i].Texture != nullptr);
 			PassInfo.bIsMSAA |= PassInfo.ColorRenderTargets[i].RenderTarget->GetNumSamples() > 1;
 		}
 	}
@@ -927,12 +927,12 @@ uint32 FMetalDynamicRHI::RHIGetGPUFrameCycles()
 
 void FMetalRHICommandContext::RHIAutomaticCacheFlushAfterComputeShader(bool bEnable)
 {
-	// Nothing required here
+	METAL_IGNORED(FMetalRHICommandContextRHIAutomaticCacheFlushAfterComputeShader);
 }
 
 void FMetalRHICommandContext::RHIFlushComputeShaderCache()
 {
-	// Nothing required here
+	METAL_IGNORED(FMetalRHICommandContextRHIFlushComputeShaderCache);
 }
 
 void FMetalDynamicRHI::RHIExecuteCommandList(FRHICommandList* RHICmdList)
@@ -942,7 +942,7 @@ void FMetalDynamicRHI::RHIExecuteCommandList(FRHICommandList* RHICmdList)
 
 void FMetalRHICommandContext::RHISetDepthBounds(float MinDepth, float MaxDepth)
 {
-	// Nothing required here
+	METAL_IGNORED(FMetalRHICommandContextSetDepthBounds);
 }
 
 void FMetalRHICommandContext::RHISubmitCommandsHint()
diff --git a/Engine/Source/Runtime/Apple/MetalRHI/Private/MetalRHIPrivate.h b/Engine/Source/Runtime/Apple/MetalRHI/Private/MetalRHIPrivate.h
--- a/Engine/Source/Runtime/Apple/MetalRHI/Private/MetalRHIPrivate.h
+++ b/Engine/Source/Runtime/Apple/MetalRHI/Private/MetalRHIPrivate.h
@@ -137,6 +137,10 @@ extern bool GMetalCommandBufferDebuggingEnabled;
 #endif
 #define METAL_FATAL_ASSERT(Condition, Format, ...) if (!(Condition)) { METAL_FATAL_ERROR(Format, __VA_ARGS__); }
 
+#if !defined(METAL_IGNORED)
+	#define METAL_IGNORED(Func)
+#endif
+
 struct FMetalDebugInfo
 {
 	uint32 CmdBuffIndex;
