Texture copies on Metal require exact format matches - it will not perform any texture conversion for you.

diff --git a/Engine/Source/Runtime/Apple/MetalRHI/Private/MetalTexture.cpp b/Engine/Source/Runtime/Apple/MetalRHI/Private/MetalTexture.cpp
--- a/Engine/Source/Runtime/Apple/MetalRHI/Private/MetalTexture.cpp
+++ b/Engine/Source/Runtime/Apple/MetalRHI/Private/MetalTexture.cpp
@@ -2588,11 +2588,11 @@ void FMetalDynamicRHI::RHICopySubTextureRegion(FTexture2DRHIParamRef SourceTextu
 		check(SourceTexture);
 		check(DestinationTexture);
 		
+		FMetalTexture2D* MetalSrcTexture = ResourceCast(SourceTexture);
+		FMetalTexture2D* MetalDestTexture = ResourceCast(DestinationTexture);
+		
 		if(SourceTexture->GetFormat() == DestinationTexture->GetFormat())
 		{
-			FMetalTexture2D* MetalSrcTexture = ResourceCast(SourceTexture);
-			FMetalTexture2D* MetalDestTexture = ResourceCast(DestinationTexture);
-			
 			FVector2D SourceSizeVector = SourceBox.GetSize();
 			FVector2D DestinatioSizeVector = DestinationBox.GetSize();
 			
@@ -2659,7 +2659,7 @@ void FMetalDynamicRHI::RHICopySubTextureRegion(FTexture2DRHIParamRef SourceTextu
 		}
 		else
 		{
-			UE_LOG(LogMetal, Warning, TEXT("RHICopySubTextureRegion Source <-> Destination texture format mismatch"));
+			UE_LOG(LogMetal, Error, TEXT("RHICopySubTextureRegion Source (UE4 %d: MTL %d) <-> Destination (UE4 %d: MTL %d) texture format mismatch"), SourceTexture->GetFormat(), MetalSrcTexture->Texture.GetPixelFormat(), DestinationTexture->GetFormat(), MetalDestTexture->Texture.GetPixelFormat());
 		}
 	}
 }
@@ -2680,11 +2680,11 @@ void FMetalRHICommandContext::RHICopyTexture(FTextureRHIParamRef SourceTextureRH
 		check(SourceTextureRHI);
 		check(DestTextureRHI);
 		
+		FMetalSurface* MetalSrcTexture = GetMetalSurfaceFromRHITexture(SourceTextureRHI);
+		FMetalSurface* MetalDestTexture = GetMetalSurfaceFromRHITexture(DestTextureRHI);
+		
 		if(SourceTextureRHI->GetFormat() == DestTextureRHI->GetFormat())
 		{
-			FMetalSurface* MetalSrcTexture = GetMetalSurfaceFromRHITexture(SourceTextureRHI);
-			FMetalSurface* MetalDestTexture = GetMetalSurfaceFromRHITexture(DestTextureRHI);
-			
 			FIntVector Size = (CopyInfo.Size != FIntVector::ZeroValue) ? CopyInfo.Size : FIntVector(MetalSrcTexture->SizeX, MetalSrcTexture->SizeY, MetalSrcTexture->SizeZ);
 			
 			mtlpp::Origin SourceOrigin(CopyInfo.SourcePosition.X, CopyInfo.SourcePosition.Y, CopyInfo.SourcePosition.Z);
@@ -2758,7 +2758,7 @@ void FMetalRHICommandContext::RHICopyTexture(FTextureRHIParamRef SourceTextureRH
 		}
 		else
 		{
-			UE_LOG(LogMetal, Warning, TEXT("RHICopyTexture Source <-> Destination texture format mismatch"));
+			UE_LOG(LogMetal, Error, TEXT("RHICopyTexture Source (UE4 %d: MTL %d) <-> Destination (UE4 %d: MTL %d) texture format mismatch"), SourceTextureRHI->GetFormat(), MetalSrcTexture->Texture.GetPixelFormat(), DestTextureRHI->GetFormat(), MetalDestTexture->Texture.GetPixelFormat());
 		}
 	}
 }
