Missed some formatting casts so that the error message prints out correctly.

diff --git a/Engine/Source/Runtime/Apple/MetalRHI/Private/MetalTexture.cpp b/Engine/Source/Runtime/Apple/MetalRHI/Private/MetalTexture.cpp
--- a/Engine/Source/Runtime/Apple/MetalRHI/Private/MetalTexture.cpp
+++ b/Engine/Source/Runtime/Apple/MetalRHI/Private/MetalTexture.cpp
@@ -2659,7 +2659,7 @@ void FMetalDynamicRHI::RHICopySubTextureRegion(FTexture2DRHIParamRef SourceTextu
 		}
 		else
 		{
-			UE_LOG(LogMetal, Error, TEXT("RHICopySubTextureRegion Source (UE4 %d: MTL %d) <-> Destination (UE4 %d: MTL %d) texture format mismatch"), SourceTexture->GetFormat(), MetalSrcTexture->Texture.GetPixelFormat(), DestinationTexture->GetFormat(), MetalDestTexture->Texture.GetPixelFormat());
+			UE_LOG(LogMetal, Error, TEXT("RHICopySubTextureRegion Source (UE4 %d: MTL %d) <-> Destination (UE4 %d: MTL %d) texture format mismatch"), (uint32)SourceTexture->GetFormat(), (uint32)MetalSrcTexture->Texture.GetPixelFormat(), (uint32)DestinationTexture->GetFormat(), (uint32)MetalDestTexture->Texture.GetPixelFormat());
 		}
 	}
 }
@@ -2758,7 +2758,7 @@ void FMetalRHICommandContext::RHICopyTexture(FTextureRHIParamRef SourceTextureRH
 		}
 		else
 		{
-			UE_LOG(LogMetal, Error, TEXT("RHICopyTexture Source (UE4 %d: MTL %d) <-> Destination (UE4 %d: MTL %d) texture format mismatch"), SourceTextureRHI->GetFormat(), MetalSrcTexture->Texture.GetPixelFormat(), DestTextureRHI->GetFormat(), MetalDestTexture->Texture.GetPixelFormat());
+			UE_LOG(LogMetal, Error, TEXT("RHICopyTexture Source (UE4 %d: MTL %d) <-> Destination (UE4 %d: MTL %d) texture format mismatch"), (uint32)SourceTextureRHI->GetFormat(), (uint32)MetalSrcTexture->Texture.GetPixelFormat(), (uint32)DestTextureRHI->GetFormat(), (uint32)MetalDestTexture->Texture.GetPixelFormat());
 		}
 	}
 }
