Fix compile errors.

diff --git a/Engine/Source/Runtime/Apple/MetalRHI/Private/MetalTexture.cpp b/Engine/Source/Runtime/Apple/MetalRHI/Private/MetalTexture.cpp
--- a/Engine/Source/Runtime/Apple/MetalRHI/Private/MetalTexture.cpp
+++ b/Engine/Source/Runtime/Apple/MetalRHI/Private/MetalTexture.cpp
@@ -2659,7 +2659,7 @@ void FMetalDynamicRHI::RHICopySubTextureRegion(FTexture2DRHIParamRef SourceTextu
 		}
 		else
 		{
-			UE_LOG(LogMetal, Error, TEXT("RHICopySubTextureRegion Source (UE4 %d: MTL %d) <-> Destination (UE4 %d: MTL %d) texture format mismatch"), (uint32)SourceTexture->GetFormat(), (uint32)MetalSrcTexture->Texture.GetPixelFormat(), (uint32)DestinationTexture->GetFormat(), (uint32)MetalDestTexture->Texture.GetPixelFormat());
+			UE_LOG(LogMetal, Error, TEXT("RHICopySubTextureRegion Source (UE4 %d: MTL %d) <-> Destination (UE4 %d: MTL %d) texture format mismatch"), (uint32)SourceTexture->GetFormat(), (uint32)MetalSrcTexture->Surface.Texture.GetPixelFormat(), (uint32)DestinationTexture->GetFormat(), (uint32)MetalDestTexture->Surface.Texture.GetPixelFormat());
 		}
 	}
 }
