Rather than convert shader resolve to render pass - remove the shader implemention code path of MSAA depth resolve for devices that don't support EMetalFeaturesMSAADepthResolve.  NumSamples works it's self out via a call to FMetalDynamicRHI::RHIGetPlatformTextureMaxSampleCount() from FSceneRenderTargets::GetNumSceneColorMSAASamples().

diff --git a/Engine/Source/Runtime/Apple/MetalRHI/Private/MetalRenderTarget.cpp b/Engine/Source/Runtime/Apple/MetalRHI/Private/MetalRenderTarget.cpp
--- a/Engine/Source/Runtime/Apple/MetalRHI/Private/MetalRenderTarget.cpp
+++ b/Engine/Source/Runtime/Apple/MetalRHI/Private/MetalRenderTarget.cpp
@@ -99,116 +99,17 @@ void FMetalRHICommandContext::RHICopyToResolveTarget(FTextureRHIParamRef SourceT
 			Profiler->RegisterGPUWork();
 		}
 
-		//if we're trying to resolve an MSAA Target we must perform a manual resolve when the HW resolve can't do the proper resolve
-        //or when we are on Mac where we are supported the 'separate MSAA and Resolve target' path.
 		const bool bMSAASource = Source->MSAATexture;
         const bool bMSAADest = Destination->MSAATexture;
         const bool bDepthStencil = Source->PixelFormat == PF_DepthStencil;
 		if (bMSAASource && !bMSAADest)
 		{
-            //This path is a layering violation (using high level constructs to implement an RHI call).
-            //It needs to be re-implementing as raw Metal and a built-in shader before going back to main.
+			// Resolve required - Device must support this - Using Shader for resolve not supported amd NumSamples should be 1
 			const bool bSupportsMSAADepthResolve = GetMetalDeviceContext().SupportsFeature(EMetalFeaturesMSAADepthResolve);
-            if (bDepthStencil && !bSupportsMSAADepthResolve)
-            {
-                FRHICommandList_RecursiveHazardous RHICmdList(this);
-
-                FResolveRect ResolveRect = ResolveParams.Rect;
-
-				PRAGMA_DISABLE_DEPRECATION_WARNINGS
-                SetRenderTargets(RHICmdList, 0, nullptr, DestTextureRHI, ESimpleRenderTargetMode::EClearColorExistingDepth, FExclusiveDepthStencil::DepthWrite_StencilWrite, true);
-				PRAGMA_ENABLE_DEPRECATION_WARNINGS
-
-                FGraphicsPipelineStateInitializer GraphicsPSOInit;
-                RHICmdList.ApplyCachedRenderTargets(GraphicsPSOInit);
-
-                // No alpha blending, no depth tests or writes, no stencil tests or writes, no backface culling.
-                GraphicsPSOInit.BlendState = TStaticBlendState<>::GetRHI();
-                GraphicsPSOInit.RasterizerState = TStaticRasterizerState<FM_Solid, CM_None>::GetRHI();
-
-                GraphicsPSOInit.DepthStencilState = TStaticDepthStencilState<true, CF_Always>::GetRHI();
-
-                const uint32 SourceWidth = Source->SizeX;
-                const uint32 SourceHeight = Source->SizeY;
-
-                const uint32 TargetWidth = Destination->SizeX;
-                const uint32 TargetHeight = Destination->SizeY;
-
-                RHICmdList.SetViewport(0.0f, 0.0f, 0.0f, TargetWidth, TargetHeight, 1.0f);
-
-                FResolveRect SourceRect = GetDefaultRect(ResolveParams.Rect, SourceWidth, SourceHeight);
-                FResolveRect DestRect = GetDefaultRect(ResolveParams.Rect, TargetWidth, TargetHeight);
-
-                // Generate the vertices used to copy from the source surface to the destination surface.
-                const float MinU = SourceRect.X1;
-                const float MinV = SourceRect.Y1;
-                const float MaxU = SourceRect.X2;
-                const float MaxV = SourceRect.Y2;
-                const float MinX = -1.f + DestRect.X1 / ((float)TargetWidth * 0.5f);
-                const float MinY = +1.f - DestRect.Y1 / ((float)TargetHeight * 0.5f);
-                const float MaxX = -1.f + DestRect.X2 / ((float)TargetWidth * 0.5f);
-                const float MaxY = +1.f - DestRect.Y2 / ((float)TargetHeight * 0.5f);
-
-                // Set the vertex and pixel shader
-                auto ShaderMap = GetGlobalShaderMap(GMaxRHIFeatureLevel);
-                TShaderMapRef<FResolveVS> ResolveVertexShader(ShaderMap);
-                TShaderMapRef<FResolveDepthPS> ResolvePixelShader(ShaderMap);
-
-                GraphicsPSOInit.BoundShaderState.VertexDeclarationRHI = GScreenVertexDeclaration.VertexDeclarationRHI;
-                GraphicsPSOInit.BoundShaderState.VertexShaderRHI = GETSAFERHISHADER_VERTEX(*ResolveVertexShader);
-                GraphicsPSOInit.BoundShaderState.PixelShaderRHI = GETSAFERHISHADER_PIXEL(*ResolvePixelShader);
-                GraphicsPSOInit.PrimitiveType = PT_TriangleStrip;
-
-                SetGraphicsPipelineState(RHICmdList, GraphicsPSOInit);
-                RHICmdList.SetBlendFactor(FLinearColor::White);
-
-                ResolvePixelShader->SetParameters(RHICmdList, FDummyResolveParameter());
-
-                // Set the source texture.
-                const uint32 TextureIndex = ResolvePixelShader->UnresolvedSurface.GetBaseIndex();
-                if (SourceTextureRHI)
-                {
-                    RHICmdList.SetShaderTexture(ResolvePixelShader->GetPixelShader(), TextureIndex, SourceTextureRHI);
-                }
-
-				FRHIResourceCreateInfo CreateInfo;
-				FVertexBufferRHIRef VertexBufferRHI = RHICreateVertexBuffer(sizeof(FScreenVertex) * 4, BUF_Volatile, CreateInfo);
-				void* VoidPtr = RHILockVertexBuffer(VertexBufferRHI, 0, sizeof(FScreenVertex) * 4, RLM_WriteOnly);
-
-                // Generate the vertices used
-                FScreenVertex* Vertices = (FScreenVertex*)VoidPtr;
-
-                Vertices[0].Position.X = MaxX;
-                Vertices[0].Position.Y = MinY;
-                Vertices[0].UV.X = MaxU;
-                Vertices[0].UV.Y = MinV;
-
-                Vertices[1].Position.X = MaxX;
-                Vertices[1].Position.Y = MaxY;
-                Vertices[1].UV.X = MaxU;
-                Vertices[1].UV.Y = MaxV;
-
-                Vertices[2].Position.X = MinX;
-                Vertices[2].Position.Y = MinY;
-                Vertices[2].UV.X = MinU;
-                Vertices[2].UV.Y = MinV;
-
-                Vertices[3].Position.X = MinX;
-                Vertices[3].Position.Y = MaxY;
-                Vertices[3].UV.X = MinU;
-                Vertices[3].UV.Y = MaxV;
-
-				RHIUnlockVertexBuffer(VertexBufferRHI);
-				RHICmdList.SetStreamSource(0, VertexBufferRHI, 0);
-				RHICmdList.DrawPrimitive(0, 2, 1);
-
-                RHICmdList.Flush();
-            }
-            else
-            {
-                //MacOS MSAA will do extra work here until we have renderpasses.  As-is it will be copying from our 'auto' resolved target to the high level's destination target.  We would prefer to resolve directly to the high level's destination target.
-                Context->CopyFromTextureToTexture(Source->MSAAResolveTexture, SrcIndex, ResolveParams.MipIndex, Origin, Size, Destination->Texture, DestIndex, ResolveParams.MipIndex, Origin);
-            }
+			const bool bSupportsMSAAStoreAndResolve = GetMetalDeviceContext().SupportsFeature(EMetalFeaturesMSAAStoreAndResolve);
+			check( (!bDepthStencil && bSupportsMSAAStoreAndResolve) || (bDepthStencil && bSupportsMSAADepthResolve) );
+			
+			Context->CopyFromTextureToTexture(Source->MSAAResolveTexture, SrcIndex, ResolveParams.MipIndex, Origin, Size, Destination->Texture, DestIndex, ResolveParams.MipIndex, Origin);
 		}
 		else
 		{
