Fix merge mistakes in MetalRHI.

diff --git a/Engine/Source/Runtime/Apple/MetalRHI/Private/MetalBuffer.cpp b/Engine/Source/Runtime/Apple/MetalRHI/Private/MetalBuffer.cpp
--- a/Engine/Source/Runtime/Apple/MetalRHI/Private/MetalBuffer.cpp
+++ b/Engine/Source/Runtime/Apple/MetalRHI/Private/MetalBuffer.cpp
@@ -428,7 +428,6 @@ NSUInteger FMetalSubBufferHeap::MaxAvailableSize() const
 	}
 }
 
-
 bool FMetalSubBufferHeap::CanAllocateSize(NSUInteger Size) const
 {
 	if (ParentHeap)
@@ -461,7 +460,6 @@ FMetalBuffer FMetalSubBufferHeap::NewBuffer(NSUInteger length)
 #if STATS || ENABLE_LOW_LEVEL_MEM_TRACKER
 		MetalLLM::LogAllocBuffer(GetMetalDeviceContext().GetDevice(), Result);
 #endif
-		
 		DEC_MEMORY_STAT_BY(STAT_MetalBufferUnusedMemory, Result.GetLength());
 		DEC_MEMORY_STAT_BY(STAT_MetalHeapBufferUnusedMemory, Result.GetLength());
 		INC_MEMORY_STAT_BY(STAT_MetalHeapBufferMemory, Result.GetLength());
@@ -1742,8 +1740,30 @@ void FMetalResourceHeap::Compact(FMetalRenderPass* Pass, bool const bForce)
                 }
             }
             
+            uint32 BytesCompacted = 0;
+            uint32 const BytesToCompact = GMetalHeapBufferBytesToCompact;
+            
+            for (uint32 i = 0; i < NumHeapSizes; i++)
+            {
 				// When not forcing the disposal of all resources we can compact them using the render-pass
                 for (uint32 j = 1; !bForce && j < BufferHeaps[u][t][i].Num(); j++)
+                {
+                    FMetalSubBufferHeap* Data = BufferHeaps[u][t][i][j];
+                    if (Data->AllocRanges.Num() > 0 && BytesCompacted < BytesToCompact)
+                    {
+                        for (FMetalSubBufferHeap::Allocation const& Alloc : Data->AllocRanges)
+                        {
+                            if (Alloc.Owner)
+                            {
+                                for (uint32 AllocIndex = 0; AllocIndex < j && BytesCompacted < BytesToCompact; AllocIndex++)
+                                {
+                                    FMetalSubBufferHeap* Prev = BufferHeaps[u][t][i][AllocIndex];
+                                    if (Prev->MaxAvailableSize() >= Alloc.Range.Length)
+                                    {
+                                        FMetalBuffer NewBuffer = Prev->NewBuffer(Alloc.Range.Length);
+                                        check(NewBuffer && NewBuffer.GetPtr());
+                                        
+                                        FMetalBuffer SourceResource = FMetalBuffer((id<MTLBuffer>)Alloc.Resource);
 										Pass->AsyncCopyFromBufferToBuffer(SourceResource, Alloc.Range.Location, NewBuffer, 0, Alloc.Range.Length);
 										
                                         FMetalBuffer PrevBuffer;
@@ -1782,6 +1802,8 @@ void FMetalResourceHeap::Compact(FMetalRenderPass* Pass, bool const bForce)
                     }
                 }
             }
+        }
+    }
 
 	Buffers[AllocShared].DrainPool(bForce);
 	Buffers[AllocPrivate].DrainPool(bForce);
diff --git a/Engine/Source/Runtime/Apple/MetalRHI/Private/MetalCommandQueue.cpp b/Engine/Source/Runtime/Apple/MetalRHI/Private/MetalCommandQueue.cpp
--- a/Engine/Source/Runtime/Apple/MetalRHI/Private/MetalCommandQueue.cpp
+++ b/Engine/Source/Runtime/Apple/MetalRHI/Private/MetalCommandQueue.cpp
@@ -148,15 +148,6 @@ FMetalCommandQueue::FMetalCommandQueue(mtlpp::Device InDevice, uint32 const MaxN
 				
 				Features |= EMetalFeaturesPresentMinDuration | EMetalFeaturesGPUCaptureManager | EMetalFeaturesBufferSubAllocation | EMetalFeaturesParallelRenderEncoders | EMetalFeaturesPipelineBufferMutability;
                 
-                if(Device.SupportsFeatureSet(mtlpp::FeatureSet::iOS_GPUFamily4_v1))
-                {
-                    Features |= EMetalFeaturesTileShaders;
-                    
-                    // The below implies tile shaders which are necessary to order the draw calls and generate a buffer that shows what PSOs/draws ran on each tile.
-                    IConsoleVariable* GPUCrashDebuggingCVar = IConsoleManager::Get().FindConsoleVariable(TEXT("r.GPUCrashDebugging"));
-                    GMetalCommandBufferDebuggingEnabled = (GPUCrashDebuggingCVar && GPUCrashDebuggingCVar->GetInt() != 0) || FParse::Param(FCommandLine::Get(),TEXT("metalgpudebug"));
-                }
-				
 				// Turn on Texture Buffers! These are faster on the GPU as we don't need to do out-of-bounds tests but require Metal 2.1 and macOS 10.14
 				if (Vers.majorVersion >= 12)
 				{
@@ -273,9 +264,6 @@ FMetalCommandQueue::FMetalCommandQueue(mtlpp::Device InDevice, uint32 const MaxN
 					Features |= EMetalFeaturesHeaps;
 				}
 			}
-            
-			IConsoleVariable* GPUCrashDebuggingCVar = IConsoleManager::Get().FindConsoleVariable(TEXT("r.GPUCrashDebugging"));
-			GMetalCommandBufferDebuggingEnabled = (GPUCrashDebuggingCVar && GPUCrashDebuggingCVar->GetInt() != 0) || FParse::Param(FCommandLine::Get(),TEXT("metalgpudebug"));
 		}
     }
     else if ([Device.GetName().GetPtr() rangeOfString:@"Nvidia" options:NSCaseInsensitiveSearch].location != NSNotFound)
@@ -458,7 +446,12 @@ FMetalFence* FMetalCommandQueue::CreateFence(ns::String const& Label) const
 				FMetalDebugFence* Fence = (FMetalDebugFence*)InnerFence.GetPtr();
 				Fence.label = String;
 			}
+			else
 	#endif
+			if(InnerFence && String)
+			{
+				InnerFence.SetLabel(String);
+			}
 		}
 		return InternalFence;
 	}
diff --git a/Engine/Source/Runtime/Apple/MetalRHI/Private/MetalTexture.cpp b/Engine/Source/Runtime/Apple/MetalRHI/Private/MetalTexture.cpp
--- a/Engine/Source/Runtime/Apple/MetalRHI/Private/MetalTexture.cpp
+++ b/Engine/Source/Runtime/Apple/MetalRHI/Private/MetalTexture.cpp
@@ -1148,7 +1148,7 @@ FMetalBuffer FMetalSurface::AllocSurface(uint32 MipIndex, uint32 ArrayIndex, ERe
 	uint32 MipBytes = GetMipSize(MipIndex, &DestStride, false);
 	
 	// allocate some temporary memory
-	mtlpp::ResourceOptions ResMode = FMetalCommandQueue::GetCompatibleResourceOptions((mtlpp::ResourceOptions)(mtlpp::ResourceOptions::StorageModeShared | (bSupportsResourceOptions && !(PLATFORM_MAC && PixelFormat == PF_G8 && (Flags & TexCreate_SRGB)) ? mtlpp::ResourceOptions::CpuCacheModeWriteCombined : 0)));
+	mtlpp::ResourceOptions ResMode = FMetalCommandQueue::GetCompatibleResourceOptions((mtlpp::ResourceOptions)(mtlpp::ResourceOptions::StorageModeShared | (!(PLATFORM_MAC && PixelFormat == PF_G8 && (Flags & TexCreate_SRGB)) ? mtlpp::ResourceOptions::CpuCacheModeWriteCombined : 0)));
 	
 	FMetalBuffer Buffer = GetMetalDeviceContext().GetResourceHeap().CreateBuffer(MipBytes, BufferOffsetAlignment, BUF_Dynamic, ResMode);
 	
diff --git a/Engine/Source/Runtime/Apple/MetalRHI/Public/ue4_stdlib.h b/Engine/Source/Runtime/Apple/MetalRHI/Public/ue4_stdlib.h
--- a/Engine/Source/Runtime/Apple/MetalRHI/Public/ue4_stdlib.h
+++ b/Engine/Source/Runtime/Apple/MetalRHI/Public/ue4_stdlib.h
@@ -9919,6 +9919,6 @@ unsigned char ue4_stdlib_metal[] = {
   0x09, 0x7d, 0x3b, 0x0a, 0x7d, 0x0a, 0x09, 0x0a, 0x23, 0x70, 0x72, 0x61,
   0x67, 0x6d, 0x61, 0x20, 0x63, 0x6c, 0x61, 0x6e, 0x67, 0x20, 0x64, 0x69,
   0x61, 0x67, 0x6e, 0x6f, 0x73, 0x74, 0x69, 0x63, 0x20, 0x70, 0x6f, 0x70,
-  0x0a, 
+  0x0a
 };
-unsigned int ue4_stdlib_metal_len =  119041;
+unsigned int ue4_stdlib_metal_len = 119041;
