Fixing build failures on Metal and Switch.

diff --git a/Engine/Source/Runtime/Apple/MetalRHI/Private/MetalUAV.cpp b/Engine/Source/Runtime/Apple/MetalRHI/Private/MetalUAV.cpp
--- a/Engine/Source/Runtime/Apple/MetalRHI/Private/MetalUAV.cpp
+++ b/Engine/Source/Runtime/Apple/MetalRHI/Private/MetalUAV.cpp
@@ -187,81 +187,6 @@ FUnorderedAccessViewRHIRef FMetalDynamicRHI::RHICreateUnorderedAccessView(FIndex
 	}
 }
 
-FShaderResourceViewRHIRef FMetalDynamicRHI::RHICreateShaderResourceView_RenderThread(class FRHICommandListImmediate& RHICmdList, FTexture2DRHIParamRef Texture2DRHI, uint8 MipLevel)
-{
-	FMetalTexture2D* Texture = ResourceCast(Texture2DRHI);
-	FMetalTexture Tex = Texture->Surface.Texture;
-	if (!(Tex.GetUsage() & mtlpp::TextureUsage::PixelFormatView))
-	{
-		FScopedRHIThreadStaller StallRHIThread(RHICmdList);
-		return GDynamicRHI->RHICreateShaderResourceView(Texture2DRHI, MipLevel);
-	}
-	else
-	{
-		return GDynamicRHI->RHICreateShaderResourceView(Texture2DRHI, MipLevel);
-	}
-}
-
-FShaderResourceViewRHIRef FMetalDynamicRHI::RHICreateShaderResourceView_RenderThread(class FRHICommandListImmediate& RHICmdList, FTexture2DRHIParamRef Texture2DRHI, uint8 MipLevel, uint8 NumMipLevels, uint8 Format)
-{
-	FMetalTexture2D* Texture = ResourceCast(Texture2DRHI);
-	FMetalTexture Tex = Texture->Surface.Texture;
-	if (!(Tex.GetUsage() & mtlpp::TextureUsage::PixelFormatView))
-	{
-		FScopedRHIThreadStaller StallRHIThread(RHICmdList);
-		return GDynamicRHI->RHICreateShaderResourceView(Texture2DRHI, MipLevel, NumMipLevels, Format);
-	}
-	else
-	{
-		return GDynamicRHI->RHICreateShaderResourceView(Texture2DRHI, MipLevel, NumMipLevels, Format);
-	}
-}
-
-FShaderResourceViewRHIRef FMetalDynamicRHI::RHICreateShaderResourceView_RenderThread(class FRHICommandListImmediate& RHICmdList, FTexture3DRHIParamRef Texture3DRHI, uint8 MipLevel)
-{
-	FMetalTexture3D* Texture = ResourceCast(Texture3DRHI);
-	FMetalTexture Tex = Texture->Surface.Texture;
-	if (!(Tex.GetUsage() & mtlpp::TextureUsage::PixelFormatView))
-	{
-		FScopedRHIThreadStaller StallRHIThread(RHICmdList);
-		return GDynamicRHI->RHICreateShaderResourceView(Texture3DRHI, MipLevel);
-	}
-	else
-	{
-		return GDynamicRHI->RHICreateShaderResourceView(Texture3DRHI, MipLevel);
-	}
-}
-
-FShaderResourceViewRHIRef FMetalDynamicRHI::RHICreateShaderResourceView_RenderThread(class FRHICommandListImmediate& RHICmdList, FTexture2DArrayRHIParamRef Texture2DArrayRHI, uint8 MipLevel)
-{
-	FMetalTexture2DArray* Texture = ResourceCast(Texture2DArrayRHI);
-	FMetalTexture Tex = Texture->Surface.Texture;
-	if (!(Tex.GetUsage() & mtlpp::TextureUsage::PixelFormatView))
-	{
-		FScopedRHIThreadStaller StallRHIThread(RHICmdList);
-		return GDynamicRHI->RHICreateShaderResourceView(Texture2DArrayRHI, MipLevel);
-	}
-	else
-	{
-		return GDynamicRHI->RHICreateShaderResourceView(Texture2DArrayRHI, MipLevel);
-	}
-}
-
-FShaderResourceViewRHIRef FMetalDynamicRHI::RHICreateShaderResourceView_RenderThread(class FRHICommandListImmediate& RHICmdList, FTextureCubeRHIParamRef TextureCubeRHI, uint8 MipLevel)
-{
-	FMetalTextureCube* Texture = ResourceCast(TextureCubeRHI);
-	FMetalTexture Tex = Texture->Surface.Texture;
-	if (!(Tex.GetUsage() & mtlpp::TextureUsage::PixelFormatView))
-	{
-		FScopedRHIThreadStaller StallRHIThread(RHICmdList);
-		return GDynamicRHI->RHICreateShaderResourceView(TextureCubeRHI, MipLevel);
-	}
-	else
-	{
-		return GDynamicRHI->RHICreateShaderResourceView(TextureCubeRHI, MipLevel);
-	}
-}
-
 FShaderResourceViewRHIRef FMetalDynamicRHI::CreateShaderResourceView_RenderThread(class FRHICommandListImmediate& RHICmdList, FVertexBufferRHIParamRef VertexBuffer, uint32 Stride, uint8 Format)
 {
 	FShaderResourceViewRHIRef Result = GDynamicRHI->RHICreateShaderResourceView(VertexBuffer, Stride, Format);
@@ -374,136 +299,6 @@ FShaderResourceViewRHIRef FMetalDynamicRHI::RHICreateShaderResourceView(FIndexBu
 	}
 }
 
-FShaderResourceViewRHIRef FMetalDynamicRHI::RHICreateShaderResourceView(FTexture2DRHIParamRef Texture2DRHI, uint8 MipLevel)
-{
-	@autoreleasepool {
-    FMetalShaderResourceView* SRV = new FMetalShaderResourceView;
-	SRV->SourceTexture = (FRHITexture*)Texture2DRHI;
-	
-	FMetalSurface* Surface = GetMetalSurfaceFromRHITexture(Texture2DRHI);
-	SRV->TextureView = Surface ? new FMetalSurface(*Surface, NSMakeRange(MipLevel, 1)) : nullptr;
-	
-	SRV->SourceVertexBuffer = nullptr;
-	SRV->SourceIndexBuffer = nullptr;
-	SRV->SourceStructuredBuffer = nullptr;
-	
-	SRV->MipLevel = MipLevel;
-	SRV->NumMips = 1;
-	SRV->Format = PF_Unknown;
-		
-	if (Surface)
-	{
-		Surface->SRVs.Add(SRV);
-	}
-	
-	return SRV;
-	}
-}
-
-FShaderResourceViewRHIRef FMetalDynamicRHI::RHICreateShaderResourceView(FTexture2DRHIParamRef Texture2DRHI, uint8 MipLevel, uint8 NumMipLevels, uint8 Format)
-{
-	@autoreleasepool {
-	FMetalShaderResourceView* SRV = new FMetalShaderResourceView;
-	SRV->SourceTexture = (FRHITexture*)Texture2DRHI;
-	
-	FMetalSurface* Surface = GetMetalSurfaceFromRHITexture(Texture2DRHI);
-	SRV->TextureView = Surface ? new FMetalSurface(*Surface, NSMakeRange(MipLevel, NumMipLevels), (EPixelFormat)Format) : nullptr;
-		
-	SRV->SourceVertexBuffer = nullptr;
-	SRV->SourceIndexBuffer = nullptr;
-	SRV->SourceStructuredBuffer = nullptr;
-	
-	SRV->MipLevel = MipLevel;
-	SRV->NumMips = NumMipLevels;
-	SRV->Format = Format;
-		
-	if (Surface)
-	{
-		Surface->SRVs.Add(SRV);
-	}
-	
-	return SRV;
-	}
-}
-
-FShaderResourceViewRHIRef FMetalDynamicRHI::RHICreateShaderResourceView(FTexture3DRHIParamRef Texture3DRHI, uint8 MipLevel)
-{
-	@autoreleasepool {
-	FMetalShaderResourceView* SRV = new FMetalShaderResourceView;
-	SRV->SourceTexture = (FRHITexture*)Texture3DRHI;
-	
-	FMetalSurface* Surface = GetMetalSurfaceFromRHITexture(Texture3DRHI);
-	SRV->TextureView = Surface ? new FMetalSurface(*Surface, NSMakeRange(MipLevel, 1)) : nullptr;
-	
-	SRV->SourceVertexBuffer = nullptr;
-	SRV->SourceIndexBuffer = nullptr;
-	SRV->SourceStructuredBuffer = nullptr;
-	
-	SRV->MipLevel = MipLevel;
-	SRV->NumMips = 1;
-	SRV->Format = PF_Unknown;
-		
-	if (Surface)
-	{
-		Surface->SRVs.Add(SRV);
-	}
-	
-	return SRV;
-	}
-}
-
-FShaderResourceViewRHIRef FMetalDynamicRHI::RHICreateShaderResourceView(FTexture2DArrayRHIParamRef Texture2DArrayRHI, uint8 MipLevel)
-{
-	@autoreleasepool {
-	FMetalShaderResourceView* SRV = new FMetalShaderResourceView;
-	SRV->SourceTexture = (FRHITexture*)Texture2DArrayRHI;
-	
-	FMetalSurface* Surface = GetMetalSurfaceFromRHITexture(Texture2DArrayRHI);
-	SRV->TextureView = Surface ? new FMetalSurface(*Surface, NSMakeRange(MipLevel, 1)) : nullptr;
-	
-	SRV->SourceVertexBuffer = nullptr;
-	SRV->SourceIndexBuffer = nullptr;
-	SRV->SourceStructuredBuffer = nullptr;
-	
-	SRV->MipLevel = MipLevel;
-	SRV->NumMips = 1;
-	SRV->Format = PF_Unknown;
-		
-	if (Surface)
-	{
-		Surface->SRVs.Add(SRV);
-	}
-	
-	return SRV;
-	}
-}
-
-FShaderResourceViewRHIRef FMetalDynamicRHI::RHICreateShaderResourceView(FTextureCubeRHIParamRef TextureCubeRHI, uint8 MipLevel)
-{
-	@autoreleasepool {
-	FMetalShaderResourceView* SRV = new FMetalShaderResourceView;
-	SRV->SourceTexture = (FRHITexture*)TextureCubeRHI;
-	
-	FMetalSurface* Surface = GetMetalSurfaceFromRHITexture(TextureCubeRHI);
-	SRV->TextureView = Surface ? new FMetalSurface(*Surface, NSMakeRange(MipLevel, 1)) : nullptr;
-	
-	SRV->SourceVertexBuffer = nullptr;
-	SRV->SourceIndexBuffer = nullptr;
-	SRV->SourceStructuredBuffer = nullptr;
-	
-	SRV->MipLevel = MipLevel;
-	SRV->NumMips = 1;
-	SRV->Format = PF_Unknown;
-		
-	if (Surface)
-	{
-		Surface->SRVs.Add(SRV);
-	}
-	
-	return SRV;
-	}
-}
-
 void FMetalDynamicRHI::RHIUpdateShaderResourceView(FShaderResourceViewRHIParamRef SRVRHI, FVertexBufferRHIParamRef VertexBufferRHI, uint32 Stride, uint8 Format)
 {
 	check(SRVRHI);
diff --git a/Engine/Source/Runtime/Apple/MetalRHI/Public/MetalRHI.h b/Engine/Source/Runtime/Apple/MetalRHI/Public/MetalRHI.h
--- a/Engine/Source/Runtime/Apple/MetalRHI/Public/MetalRHI.h
+++ b/Engine/Source/Runtime/Apple/MetalRHI/Public/MetalRHI.h
@@ -97,11 +97,6 @@ public:
 	virtual FTexture2DArrayRHIRef RHICreateTexture2DArray(uint32 SizeX, uint32 SizeY, uint32 SizeZ, uint8 Format, uint32 NumMips, uint32 Flags, FRHIResourceCreateInfo& CreateInfo) final override;
 	virtual FTexture3DRHIRef RHICreateTexture3D(uint32 SizeX, uint32 SizeY, uint32 SizeZ, uint8 Format, uint32 NumMips, uint32 Flags, FRHIResourceCreateInfo& CreateInfo) final override;
 	virtual void RHIGetResourceInfo(FTextureRHIParamRef Ref, FRHIResourceInfo& OutInfo) final override;
-	virtual FShaderResourceViewRHIRef RHICreateShaderResourceView(FTexture2DRHIParamRef Texture2DRHI, uint8 MipLevel) final override;
-	virtual FShaderResourceViewRHIRef RHICreateShaderResourceView(FTexture2DRHIParamRef Texture2DRHI, uint8 MipLevel, uint8 NumMipLevels, uint8 Format) final override;
-	virtual FShaderResourceViewRHIRef RHICreateShaderResourceView(FTexture3DRHIParamRef Texture3DRHI, uint8 MipLevel) final override;
-	virtual FShaderResourceViewRHIRef RHICreateShaderResourceView(FTexture2DArrayRHIParamRef Texture2DArrayRHI, uint8 MipLevel) final override;
-	virtual FShaderResourceViewRHIRef RHICreateShaderResourceView(FTextureCubeRHIParamRef TextureCubeRHI, uint8 MipLevel) final override;
 	virtual void RHIGenerateMips(FTextureRHIParamRef Texture) final override;
 	virtual uint32 RHIComputeMemorySize(FTextureRHIParamRef TextureRHI) final override;
 	virtual FTexture2DRHIRef RHIAsyncReallocateTexture2D(FTexture2DRHIParamRef Texture2D, int32 NewMipCount, int32 NewSizeX, int32 NewSizeY, FThreadSafeCounter* RequestStatus) final override;
@@ -177,11 +172,6 @@ public:
 	virtual FUnorderedAccessViewRHIRef RHICreateUnorderedAccessView_RenderThread(class FRHICommandListImmediate& RHICmdList, FVertexBufferRHIParamRef VertexBuffer, uint8 Format) final override;
 	virtual FUnorderedAccessViewRHIRef RHICreateUnorderedAccessView_RenderThread(class FRHICommandListImmediate& RHICmdList, FIndexBufferRHIParamRef IndexBuffer, uint8 Format) final override;
 	
-	virtual FShaderResourceViewRHIRef RHICreateShaderResourceView_RenderThread(class FRHICommandListImmediate& RHICmdList, FTexture2DRHIParamRef Texture2DRHI, uint8 MipLevel) final override;
-	virtual FShaderResourceViewRHIRef RHICreateShaderResourceView_RenderThread(class FRHICommandListImmediate& RHICmdList, FTexture2DRHIParamRef Texture2DRHI, uint8 MipLevel, uint8 NumMipLevels, uint8 Format) final override;
-	virtual FShaderResourceViewRHIRef RHICreateShaderResourceView_RenderThread(class FRHICommandListImmediate& RHICmdList, FTexture3DRHIParamRef Texture3DRHI, uint8 MipLevel) final override;
-	virtual FShaderResourceViewRHIRef RHICreateShaderResourceView_RenderThread(class FRHICommandListImmediate& RHICmdList, FTexture2DArrayRHIParamRef Texture2DArrayRHI, uint8 MipLevel) final override;
-	virtual FShaderResourceViewRHIRef RHICreateShaderResourceView_RenderThread(class FRHICommandListImmediate& RHICmdList, FTextureCubeRHIParamRef TextureCubeRHI, uint8 MipLevel) final override;
 	virtual FShaderResourceViewRHIRef RHICreateShaderResourceView_RenderThread(class FRHICommandListImmediate& RHICmdList, FVertexBufferRHIParamRef VertexBuffer, uint32 Stride, uint8 Format) final override;
 	virtual FShaderResourceViewRHIRef CreateShaderResourceView_RenderThread(class FRHICommandListImmediate& RHICmdList, FIndexBufferRHIParamRef Buffer) final override;
 	virtual FShaderResourceViewRHIRef CreateShaderResourceView_RenderThread(class FRHICommandListImmediate& RHICmdList, FVertexBufferRHIParamRef VertexBuffer, uint32 Stride, uint8 Format) final override;
