Fixing Metal failures. Converted SRV creation to CreateInfo.

diff --git a/Engine/Source/Runtime/Apple/MetalRHI/Private/MetalUAV.cpp b/Engine/Source/Runtime/Apple/MetalRHI/Private/MetalUAV.cpp
--- a/Engine/Source/Runtime/Apple/MetalRHI/Private/MetalUAV.cpp
+++ b/Engine/Source/Runtime/Apple/MetalRHI/Private/MetalUAV.cpp
@@ -220,6 +220,46 @@ FShaderResourceViewRHIRef FMetalDynamicRHI::RHICreateShaderResourceView_RenderTh
 	return GDynamicRHI->RHICreateShaderResourceView(StructuredBuffer);
 }
 
+FShaderResourceViewRHIRef FMetalDynamicRHI::RHICreateShaderResourceView_RenderThread(class FRHICommandListImmediate& RHICmdList, FTextureRHIParamRef Texture2DRHI, const FRHITextureSRVCreateInfo& CreateInfo)
+{
+	FMetalSurface* Surface = (FMetalSurface*)Texture2DRHI->GetTextureBaseRHI();
+	FMetalTexture Tex = Surface->Texture;
+	if (!(Tex.GetUsage() & mtlpp::TextureUsage::PixelFormatView))
+	{
+		FScopedRHIThreadStaller StallRHIThread(RHICmdList);
+		return GDynamicRHI->RHICreateShaderResourceView(Texture2DRHI, CreateInfo);
+	}
+	else
+	{
+		return GDynamicRHI->RHICreateShaderResourceView(Texture2DRHI, CreateInfo);
+	}
+}
+
+FShaderResourceViewRHIRef FMetalDynamicRHI::RHICreateShaderResourceView(FTextureRHIParamRef Texture2DRHI, const FRHITextureSRVCreateInfo& CreateInfo)
+{
+	@autoreleasepool {
+		FMetalShaderResourceView* SRV = new FMetalShaderResourceView;
+		SRV->SourceTexture = (FRHITexture*)Texture2DRHI;
+		
+		FMetalSurface* Surface = GetMetalSurfaceFromRHITexture(Texture2DRHI);
+		SRV->TextureView = Surface ? new FMetalSurface(*Surface, NSMakeRange(CreateInfo.MipLevel, CreateInfo.NumMipLevels), (EPixelFormat)CreateInfo.Format) : nullptr;
+		
+		SRV->SourceVertexBuffer = nullptr;
+		SRV->SourceIndexBuffer = nullptr;
+		SRV->SourceStructuredBuffer = nullptr;
+		
+		SRV->MipLevel = CreateInfo.MipLevel;
+		SRV->NumMips = CreateInfo.NumMipLevels;
+		SRV->Format = CreateInfo.Format;
+		
+		if (Surface)
+		{
+			Surface->SRVs.Add(SRV);
+		}
+		
+		return SRV;
+	}
+}
 
 FShaderResourceViewRHIRef FMetalDynamicRHI::RHICreateShaderResourceView(FStructuredBufferRHIParamRef StructuredBufferRHI)
 {
diff --git a/Engine/Source/Runtime/Apple/MetalRHI/Public/MetalRHI.h b/Engine/Source/Runtime/Apple/MetalRHI/Public/MetalRHI.h
--- a/Engine/Source/Runtime/Apple/MetalRHI/Public/MetalRHI.h
+++ b/Engine/Source/Runtime/Apple/MetalRHI/Public/MetalRHI.h
@@ -80,6 +80,7 @@ public:
 	virtual FUnorderedAccessViewRHIRef RHICreateUnorderedAccessView(FVertexBufferRHIParamRef VertexBuffer, uint8 Format) final override;
 	virtual FUnorderedAccessViewRHIRef RHICreateUnorderedAccessView(FIndexBufferRHIParamRef IndexBuffer, uint8 Format) final override;
 	virtual FShaderResourceViewRHIRef RHICreateShaderResourceView(FStructuredBufferRHIParamRef StructuredBuffer) final override;
+	virtual FShaderResourceViewRHIRef RHICreateShaderResourceView(FTextureRHIParamRef Texture2DRHI, const FRHITextureSRVCreateInfo& CreateInfo) final override;
 	virtual FShaderResourceViewRHIRef RHICreateShaderResourceView(FVertexBufferRHIParamRef VertexBuffer, uint32 Stride, uint8 Format) final override;
 	virtual FShaderResourceViewRHIRef RHICreateShaderResourceView(FIndexBufferRHIParamRef Buffer) final override;
 	virtual void RHIUpdateShaderResourceView(FShaderResourceViewRHIParamRef SRV, FVertexBufferRHIParamRef VertexBuffer, uint32 Stride, uint8 Format) final override;
@@ -177,6 +178,7 @@ public:
 	virtual FShaderResourceViewRHIRef CreateShaderResourceView_RenderThread(class FRHICommandListImmediate& RHICmdList, FVertexBufferRHIParamRef VertexBuffer, uint32 Stride, uint8 Format) final override;
 	virtual FShaderResourceViewRHIRef RHICreateShaderResourceView_RenderThread(class FRHICommandListImmediate& RHICmdList, FIndexBufferRHIParamRef Buffer) final override;
 	virtual FShaderResourceViewRHIRef RHICreateShaderResourceView_RenderThread(class FRHICommandListImmediate& RHICmdList, FStructuredBufferRHIParamRef StructuredBuffer) final override;
+	virtual FShaderResourceViewRHIRef RHICreateShaderResourceView_RenderThread(class FRHICommandListImmediate& RHICmdList, FTextureRHIParamRef Texture2DRHI, const FRHITextureSRVCreateInfo& CreateInfo) final override;
 
 	virtual FComputeFenceRHIRef RHICreateComputeFence(const FName& Name) final override;
 	virtual FGPUFenceRHIRef RHICreateGPUFence(const FName &Name) final override;
