Update the names for the new SPIRV-Cross output & settings.

diff --git a/Engine/Source/Developer/Apple/MetalShaderFormat/Private/MetalDerivedData.cpp b/Engine/Source/Developer/Apple/MetalShaderFormat/Private/MetalDerivedData.cpp
--- a/Engine/Source/Developer/Apple/MetalShaderFormat/Private/MetalDerivedData.cpp
+++ b/Engine/Source/Developer/Apple/MetalShaderFormat/Private/MetalDerivedData.cpp
@@ -1132,7 +1132,7 @@ bool FMetalShaderOutputCooker::Build(TArray<uint8>& OutData)
 			char BufferIdx[3];
 			FCStringAnsi::Snprintf(BufferIdx, 3, "%d", SideTableIndex);
 
-			ShaderConductor::MacroDefine Defines[7] = {{"texel_buffer_texture_width", "0"}, {"enforce_storge_buffer_bounds", "1"}, {"metadata_buffer_index", BufferIdx}, {nullptr, nullptr}, {nullptr, nullptr}, {nullptr, nullptr}, {nullptr, nullptr}};
+			ShaderConductor::MacroDefine Defines[7] = {{"texel_buffer_texture_width", "0"}, {"enforce_storge_buffer_bounds", "1"}, {"buffer_size_buffer_index", BufferIdx}, {nullptr, nullptr}, {nullptr, nullptr}, {nullptr, nullptr}, {nullptr, nullptr}};
 			TargetDesc.numOptions = 3;
 			TargetDesc.options = &Defines[0];
 			switch(Semantics)
@@ -1201,9 +1201,9 @@ bool FMetalShaderOutputCooker::Build(TArray<uint8>& OutData)
 		Result = (Results.hasError) ? 0 : 1;
 		if (!Results.hasError && Results.target)
 		{
-			if(strstr((const char*)Results.target->Data(), "spvBufferSizes"))
+			if(strstr((const char*)Results.target->Data(), "spvBufferSizeConstants"))
 			{
-				MetaData += FString::Printf(TEXT("// @SideTable: spvBufferSizes(%d)\n"), SideTableIndex);
+				MetaData += FString::Printf(TEXT("// @SideTable: spvBufferSizeConstants(%d)\n"), SideTableIndex);
 			}
 			MetaData += FString::Printf(TEXT("\n\n"));
 			
diff --git a/Engine/Source/Runtime/Apple/MetalRHI/Private/MetalPipeline.cpp b/Engine/Source/Runtime/Apple/MetalRHI/Private/MetalPipeline.cpp
--- a/Engine/Source/Runtime/Apple/MetalRHI/Private/MetalPipeline.cpp
+++ b/Engine/Source/Runtime/Apple/MetalRHI/Private/MetalPipeline.cpp
@@ -543,7 +543,7 @@ APPLE_PLATFORM_OBJECT_ALLOC_OVERRIDES(FMetalShaderPipeline)
 			case MTLArgumentTypeBuffer:
 			{
 				checkf(Arg.index < ML_MaxBuffers, TEXT("Metal buffer index exceeded!"));
-				if (FString(Arg.name) != TEXT("BufferSizes") && FString(Arg.name) != TEXT("spvBufferSizes"))
+				if (FString(Arg.name) != TEXT("BufferSizes") && FString(Arg.name) != TEXT("spvBufferSizeConstants"))
 				{
 					ResourceMask[Frequency].BufferMask |= (1 << Arg.index);
 				
@@ -1296,14 +1296,14 @@ static void ReleaseMTLRenderPipeline(FMetalShaderPipeline* Pipeline)
 bool FMetalGraphicsPipelineState::Compile()
 {
 	FMemory::Memzero(PipelineStates);
-	for (uint32 i = 0; i < EMetalIndexType_Num; i++)
-	{
-		PipelineStates[i] = [GetMTLRenderPipeline(true, this, Initializer, (EMetalIndexType)i) retain];
-		if(!PipelineStates[i])
+		for (uint32 i = 0; i < EMetalIndexType_Num; i++)
 		{
-			return false;
+			PipelineStates[i] = [GetMTLRenderPipeline(true, this, Initializer, (EMetalIndexType)i) retain];
+			if(!PipelineStates[i])
+			{
+				return false;
+			}
 		}
-	}
 	
 	return true;
 }
@@ -1321,13 +1321,13 @@ FMetalShaderPipeline* FMetalGraphicsPipelineState::GetPipeline(EMetalIndexType I
 {
 	check(IndexType < EMetalIndexType_Num);
 
-	if(!PipelineStates[IndexType])
-	{
-		PipelineStates[IndexType] = [GetMTLRenderPipeline(true, this, Initializer, IndexType) retain];
-	}
+		if(!PipelineStates[IndexType])
+		{
+			PipelineStates[IndexType] = [GetMTLRenderPipeline(true, this, Initializer, IndexType) retain];
+		}
 	FMetalShaderPipeline* Pipe = PipelineStates[IndexType];
 
-	check(Pipe);
+		check(Pipe);
     return Pipe;
 }
 
