Added support for MIP-level selection in RHIReadSurfaceData function for Metal.

#rnx

diff --git a/Engine/Source/Runtime/Apple/MetalRHI/Private/MetalRenderTarget.cpp b/Engine/Source/Runtime/Apple/MetalRHI/Private/MetalRenderTarget.cpp
--- a/Engine/Source/Runtime/Apple/MetalRHI/Private/MetalRenderTarget.cpp
+++ b/Engine/Source/Runtime/Apple/MetalRHI/Private/MetalRenderTarget.cpp
@@ -481,7 +481,7 @@ void FMetalDynamicRHI::RHIReadSurfaceData(FTextureRHIParamRef TextureRHI, FIntRe
 			Desc.SetWidth(SizeX);
 			Desc.SetHeight(SizeY);
 			Desc.SetDepth(1);
-			Desc.SetMipmapLevelCount(Texture.GetMipmapLevelCount());
+			Desc.SetMipmapLevelCount(1); // Only consider a single subresource and not the whole texture (like in the other RHIs)
 			Desc.SetSampleCount(Texture.GetSampleCount());
 			Desc.SetArrayLength(Texture.GetArrayLength());
 			
@@ -494,7 +494,7 @@ void FMetalDynamicRHI::RHIReadSurfaceData(FTextureRHIParamRef TextureRHI, FIntRe
 			
 			TempTexture = GetMetalDeviceContext().GetDevice().NewTexture(Desc);
 			
-			ImmediateContext.Context->CopyFromTextureToTexture(Texture, 0, 0, mtlpp::Origin(Region.origin), mtlpp::Size(Region.size), TempTexture, 0, 0, mtlpp::Origin(0, 0, 0));
+			ImmediateContext.Context->CopyFromTextureToTexture(Texture, 0, InFlags.GetMip(), mtlpp::Origin(Region.origin), mtlpp::Size(Region.size), TempTexture, 0, 0, mtlpp::Origin(0, 0, 0));
 			
 			Texture = TempTexture;
 			Region = mtlpp::Region(0, 0, SizeX, SizeY);
@@ -503,7 +503,7 @@ void FMetalDynamicRHI::RHIReadSurfaceData(FTextureRHIParamRef TextureRHI, FIntRe
 		if(Texture.GetStorageMode() == mtlpp::StorageMode::Managed)
 		{
 			// Synchronise the texture with the CPU
-			ImmediateContext.Context->SynchronizeTexture(Texture, 0, 0);
+			ImmediateContext.Context->SynchronizeTexture(Texture, 0, InFlags.GetMip());
 		}
 #endif
 
@@ -539,17 +539,17 @@ void FMetalDynamicRHI::RHIReadSurfaceData(FTextureRHIParamRef TextureRHI, FIntRe
 			
 			if (Surface->PixelFormat != PF_DepthStencil)
 			{
-				ImmediateContext.Context->CopyFromTextureToBuffer(Texture, 0, 0, Region.origin, Region.size, Buffer, 0, AlignedStride, BytesPerImage, mtlpp::BlitOption::None);
+				ImmediateContext.Context->CopyFromTextureToBuffer(Texture, 0, InFlags.GetMip(), Region.origin, Region.size, Buffer, 0, AlignedStride, BytesPerImage, mtlpp::BlitOption::None);
 			}
 			else
 			{
 				if (!InFlags.GetOutputStencil())
 				{
-					ImmediateContext.Context->CopyFromTextureToBuffer(Texture, 0, 0, Region.origin, Region.size, Buffer, 0, AlignedStride, BytesPerImage, mtlpp::BlitOption::DepthFromDepthStencil);
+					ImmediateContext.Context->CopyFromTextureToBuffer(Texture, 0, InFlags.GetMip(), Region.origin, Region.size, Buffer, 0, AlignedStride, BytesPerImage, mtlpp::BlitOption::DepthFromDepthStencil);
 				}
 				else
 				{
-					ImmediateContext.Context->CopyFromTextureToBuffer(Texture, 0, 0, Region.origin, Region.size, Buffer, 0, AlignedStride, BytesPerImage, mtlpp::BlitOption::StencilFromDepthStencil);
+					ImmediateContext.Context->CopyFromTextureToBuffer(Texture, 0, InFlags.GetMip(), Region.origin, Region.size, Buffer, 0, AlignedStride, BytesPerImage, mtlpp::BlitOption::StencilFromDepthStencil);
 				}
 			}
 			
