Implement support for SPIRV-Cross's separate Metal tessellation shaders - this means streaming out the vertex data using a render-pass in the prologue command-buffer, followed by generating a patch index buffer & running the Hull shader as a compute pass and then finally running the draw call. This is slower and uses more memory than the unified approach taken in hlslcc but it is conceptually simpler and much easier to deal with one the shader compiler side.

diff --git a/Engine/Source/Developer/Apple/MetalShaderFormat/Private/MetalDerivedData.cpp b/Engine/Source/Developer/Apple/MetalShaderFormat/Private/MetalDerivedData.cpp
--- a/Engine/Source/Developer/Apple/MetalShaderFormat/Private/MetalDerivedData.cpp
+++ b/Engine/Source/Developer/Apple/MetalShaderFormat/Private/MetalDerivedData.cpp
@@ -609,6 +609,7 @@ bool FMetalShaderOutputCooker::Build(TArray<uint8>& OutData)
 		FString MetaData = FString::Printf(TEXT("// ! %s/%s.usf:%s\n"), *Input.DebugGroupName, *Input.GetSourceFilename(), *Input.EntryPointName);
 		FString EntryPoint = Input.EntryPointName;
 		EHlslShaderFrequency Freq = Frequency;
+		FString ALNString;
 
 		TargetDesc.language = ShaderConductor::ShadingLanguage::SpirV;
 		ShaderConductor::Compiler::ResultDesc Results = ShaderConductor::Compiler::Compile(SourceDesc, Options, TargetDesc);
@@ -702,12 +703,14 @@ bool FMetalShaderOutputCooker::Build(TArray<uint8>& OutData)
 							if (Frequency == HSF_HullShader || Frequency == HSF_DomainShader)
 							{
 								TESStrings[(uint8)EMetalTessellationMetadataTags::TessellationDomain] = TEXT("// @TessellationDomain: tri\n");
+								Attribs.HSTFOutSize = 8; // sizeof(MTLTriangleTessellationFactorsHalf);
 							}
 							break;
 						case SpvExecutionModeQuads:
 							if (Frequency == HSF_HullShader || Frequency == HSF_DomainShader)
 							{
 								TESStrings[(uint8)EMetalTessellationMetadataTags::TessellationDomain] = TEXT("// @TessellationDomain: quad\n");
+								Attribs.HSTFOutSize = 12; // sizeof(MTLQuadTessellationFactorsHalf);
 							}
 							break;
 						case SpvExecutionModeIsolines:
@@ -1217,6 +1220,311 @@ bool FMetalShaderOutputCooker::Build(TArray<uint8>& OutData)
 				}
 			}
 			
+			if (Frequency == HSF_VertexShader && bUsingTessellation)
+			{
+				Count = 0;
+				SPVRResult = Reflection.EnumerateOutputVariables(&Count, nullptr);
+				check(SPVRResult == SPV_REFLECT_RESULT_SUCCESS);
+				OutputVars.SetNum(Count);
+				SPVRResult = Reflection.EnumerateOutputVariables(&Count, OutputVars.GetData());
+				check(SPVRResult == SPV_REFLECT_RESULT_SUCCESS);
+				if (Count > 0)
+				{
+					
+					OutputVars.Sort([](SpvReflectInterfaceVariable& LHS, SpvReflectInterfaceVariable& RHS)  { return LHS.location < RHS.location; });
+					uint32 Offset = 0;
+					uint32 InitialAlign = 0;
+					for (auto const& Var : OutputVars)
+					{
+						if (Var->storage_class == SpvStorageClassOutput && Var->built_in == -1)
+						{
+							unsigned Location = Var->location;
+							
+							uint32 ArrayCount = 1;
+							for (uint32 Dim = 0; Dim < Var->array.dims_count; Dim++)
+							{
+								ArrayCount *= Var->array.dims[Dim];
+							}
+							
+							auto const type = *Var->type_description;
+							uint32_t masked_type = type.type_flags & 0xF;
+							
+							for (uint32 j = 0; j < ArrayCount; j++)
+							{
+								FMetalAttribute& Attrib = Attribs.HSOut.AddDefaulted_GetRef();
+								Attrib.Index = Location + j;
+								uint32 ElementSize = 0;
+								switch (masked_type) {
+									default: checkf(false, TEXT("unsupported component type %d"), masked_type); break;
+									case SPV_REFLECT_TYPE_FLAG_BOOL  : Attrib.Type = EMetalComponentType::Bool; ElementSize = 1; break;
+									case SPV_REFLECT_TYPE_FLAG_INT   : Attrib.Type = (type.traits.numeric.scalar.signedness ? EMetalComponentType::Int : EMetalComponentType::Uint); ElementSize = 4; break;
+									case SPV_REFLECT_TYPE_FLAG_FLOAT : Attrib.Type = (type.traits.numeric.scalar.width == 32 ? EMetalComponentType::Float : EMetalComponentType::Half); if(type.traits.numeric.scalar.width == 32) { ElementSize = 4; } else { ElementSize = 2; } break;
+								}
+								if (type.type_flags & SPV_REFLECT_TYPE_FLAG_MATRIX)
+								{
+									Attrib.Components = (type.traits.numeric.matrix.row_count * type.traits.numeric.matrix.column_count);
+								}
+								else if (type.type_flags & SPV_REFLECT_TYPE_FLAG_VECTOR)
+								{
+									Attrib.Components = (type.traits.numeric.vector.component_count);
+								}
+								else
+								{
+									Attrib.Components = 1;
+								}
+								
+								char const* NumberedSemantic = strstr(Var->name, Var->semantic);
+								FString Semantic = FString::Printf(TEXT("%s_%u"), UTF8_TO_TCHAR(NumberedSemantic), j);
+								Attrib.Semantic = GetTypeHash(Semantic);
+								
+								uint32 AlignedComponents = Attrib.Components == 3 ? 4 : Attrib.Components;
+								InitialAlign = InitialAlign == 0 ? ElementSize * AlignedComponents : InitialAlign;
+								Offset = Align(Offset, ElementSize * AlignedComponents);
+								Attrib.Offset = Offset;
+								Offset += ElementSize * AlignedComponents;
+
+								if (Input.Environment.CompilerFlags.Contains(CFLAG_KeepDebugInfo))
+									ALNString += FString::Printf(TEXT("%sHSOut.%s(%u)"), ALNString.Len() ? TEXT(",") : TEXT(""), *Semantic, Attrib.Offset);
+							}
+						}
+					}
+					Attribs.HSOutSize = Align(Offset, 16);
+					
+					if (Input.Environment.CompilerFlags.Contains(CFLAG_KeepDebugInfo))
+						ALNString += FString::Printf(TEXT(" = HSOutSize(%u)"), Attribs.HSOutSize);
+				}
+			}
+			
+			if (Frequency == HSF_HullShader)
+			{
+				Count = 0;
+				SPVRResult = Reflection.EnumerateInputVariables(&Count, nullptr);
+				check(SPVRResult == SPV_REFLECT_RESULT_SUCCESS);
+				InputVars.SetNum(Count);
+				SPVRResult = Reflection.EnumerateInputVariables(&Count, InputVars.GetData());
+				check(SPVRResult == SPV_REFLECT_RESULT_SUCCESS);
+				if (Count > 0)
+				{
+					
+					InputVars.Sort([](SpvReflectInterfaceVariable& LHS, SpvReflectInterfaceVariable& RHS)  { return LHS.location < RHS.location; });
+					uint32 Offset = 0;
+					for (auto const& Var : InputVars)
+					{
+						if (Var->storage_class == SpvStorageClassInput && Var->built_in == -1)
+						{
+							unsigned Location = Var->location;
+							
+							uint32 ArrayCount = 1;
+							for (uint32 Dim = 1; Dim < Var->array.dims_count; Dim++)
+							{
+								ArrayCount *= Var->array.dims[Dim];
+							}
+							
+							auto const type = *Var->type_description;
+							uint32_t masked_type = type.type_flags & 0xF;
+							
+							for (uint32 j = 0; j < ArrayCount; j++)
+							{
+								FMetalAttribute& Attrib = Attribs.HSIn.AddDefaulted_GetRef();
+								Attrib.Index = Location + j;
+								uint32 ElementSize = 0;
+								switch (masked_type) {
+									default: checkf(false, TEXT("unsupported component type %d"), masked_type); break;
+									case SPV_REFLECT_TYPE_FLAG_BOOL  : Attrib.Type = EMetalComponentType::Bool; ElementSize = 1; break;
+									case SPV_REFLECT_TYPE_FLAG_INT   : Attrib.Type = (type.traits.numeric.scalar.signedness ? EMetalComponentType::Int : EMetalComponentType::Uint); ElementSize = 4; break;
+									case SPV_REFLECT_TYPE_FLAG_FLOAT : Attrib.Type = (type.traits.numeric.scalar.width == 32 ? EMetalComponentType::Float : EMetalComponentType::Half); if(type.traits.numeric.scalar.width == 32) { ElementSize = 4; } else { ElementSize = 2; } break;
+								}
+								if (type.type_flags & SPV_REFLECT_TYPE_FLAG_MATRIX)
+								{
+									Attrib.Components = (type.traits.numeric.matrix.row_count * type.traits.numeric.matrix.column_count);
+								}
+								else if (type.type_flags & SPV_REFLECT_TYPE_FLAG_VECTOR)
+								{
+									Attrib.Components = (type.traits.numeric.vector.component_count);
+								}
+								else
+								{
+									Attrib.Components = 1;
+								}
+								
+								char const* NumberedSemantic = strstr(Var->name, Var->semantic);
+								FString Semantic = FString::Printf(TEXT("%s_%u"), UTF8_TO_TCHAR(NumberedSemantic), j);
+								Attrib.Semantic = GetTypeHash(Semantic);
+								
+								uint32 AlignedComponents = Attrib.Components == 3 ? 4 : Attrib.Components;
+								Offset = Align(Offset, ElementSize * AlignedComponents);
+								Attrib.Offset = Offset;
+								Offset += ElementSize * AlignedComponents;
+								
+								if (Input.Environment.CompilerFlags.Contains(CFLAG_KeepDebugInfo))
+									ALNString += FString::Printf(TEXT("%sHSIn.%s(%u)"), ALNString.Len() ? TEXT(",") : TEXT(""), *Semantic, Attrib.Offset);
+							}
+						}
+					}
+				}
+				
+				Count = 0;
+				SPVRResult = Reflection.EnumerateOutputVariables(&Count, nullptr);
+				check(SPVRResult == SPV_REFLECT_RESULT_SUCCESS);
+				OutputVars.SetNum(Count);
+				SPVRResult = Reflection.EnumerateOutputVariables(&Count, OutputVars.GetData());
+				check(SPVRResult == SPV_REFLECT_RESULT_SUCCESS);
+				if (Count > 0)
+				{
+					uint32 HSOutOffset = 0;
+					uint32 PatchControlPointOffset = 0;
+					uint32 HSOutAlign = 0;
+					uint32 PatchControlPointAlign = 0;
+					
+					OutputVars.Sort([](SpvReflectInterfaceVariable& LHS, SpvReflectInterfaceVariable& RHS)  { return LHS.location < RHS.location; });
+					
+					for (auto const& Var : OutputVars)
+					{
+						if (Var->storage_class == SpvStorageClassOutput && Var->built_in == -1)
+						{
+							unsigned Location = Var->location;
+							uint32 ArrayCount = 1;
+							for (uint32 Dim = 1; Dim < Var->array.dims_count; Dim++)
+							{
+								ArrayCount *= Var->array.dims[Dim];
+							}
+							
+							auto const type = *Var->type_description;
+							uint32_t masked_type = type.type_flags & 0xF;
+							
+							for (uint32 j = 0; j < ArrayCount; j++)
+							{
+								FMetalAttribute& Attrib = Var->array.dims_count ? Attribs.PatchControlPointOut.AddDefaulted_GetRef() : Attribs.HSOut.AddDefaulted_GetRef();
+								Attrib.Index = Location + j;
+								uint32 ElementSize = 0;
+								switch (masked_type) {
+									default: checkf(false, TEXT("unsupported component type %d"), masked_type); break;
+									case SPV_REFLECT_TYPE_FLAG_BOOL  : Attrib.Type = EMetalComponentType::Bool; ElementSize = 1; break;
+									case SPV_REFLECT_TYPE_FLAG_INT   : Attrib.Type = (type.traits.numeric.scalar.signedness ? EMetalComponentType::Int : EMetalComponentType::Uint); ElementSize = 4; break;
+									case SPV_REFLECT_TYPE_FLAG_FLOAT : Attrib.Type = (type.traits.numeric.scalar.width == 32 ? EMetalComponentType::Float : EMetalComponentType::Half); if(type.traits.numeric.scalar.width == 32) { ElementSize = 4; } else { ElementSize = 2; } break;
+								}
+								if (type.type_flags & SPV_REFLECT_TYPE_FLAG_MATRIX)
+								{
+									Attrib.Components = (type.traits.numeric.matrix.row_count * type.traits.numeric.matrix.column_count);
+								}
+								else if (type.type_flags & SPV_REFLECT_TYPE_FLAG_VECTOR)
+								{
+									Attrib.Components = (type.traits.numeric.vector.component_count);
+								}
+								else
+								{
+									Attrib.Components = 1;
+								}
+								
+								uint32& Offset = Var->array.dims_count ? PatchControlPointOffset : HSOutOffset;
+								uint32& InitialAlign = Var->array.dims_count ? PatchControlPointAlign : HSOutAlign;
+								TCHAR const* Name = Var->array.dims_count ? TEXT("PatchOut") : TEXT("HSOut");
+								
+								char const* NumberedSemantic = strstr(Var->name, Var->semantic);
+								FString Semantic = FString::Printf(TEXT("%s_%u"), UTF8_TO_TCHAR(NumberedSemantic), j);
+								Attrib.Semantic = GetTypeHash(Semantic);
+								
+								uint32 AlignedComponents = Attrib.Components == 3 ? 4 : Attrib.Components;
+								InitialAlign = InitialAlign == 0 ? ElementSize * AlignedComponents : InitialAlign;
+								Offset = Align(Offset, ElementSize * AlignedComponents);
+								Attrib.Offset = Offset;
+								Offset += ElementSize * AlignedComponents;
+								
+								if (Input.Environment.CompilerFlags.Contains(CFLAG_KeepDebugInfo))
+									ALNString += FString::Printf(TEXT("%s%s.%s(%u)"), ALNString.Len() ? TEXT(",") : TEXT(""), Name, *Semantic, Attrib.Offset);
+							}
+						}
+					}
+					Attribs.HSOutSize = Align(HSOutOffset, 16);
+					Attribs.PatchControlPointOutSize = Align(PatchControlPointOffset, 16);
+					
+					if (Input.Environment.CompilerFlags.Contains(CFLAG_KeepDebugInfo))
+						ALNString += FString::Printf(TEXT(" = HSOutSize(%u), PatchControlPointOutSize(%u)"), Attribs.HSOutSize, Attribs.PatchControlPointOutSize);
+				}
+			}
+			
+			if (Frequency == HSF_DomainShader)
+			{
+				Count = 0;
+				SPVRResult = Reflection.EnumerateInputVariables(&Count, nullptr);
+				check(SPVRResult == SPV_REFLECT_RESULT_SUCCESS);
+				InputVars.SetNum(Count);
+				SPVRResult = Reflection.EnumerateInputVariables(&Count, InputVars.GetData());
+				check(SPVRResult == SPV_REFLECT_RESULT_SUCCESS);
+				if (Count > 0)
+				{
+					uint32 HSOutOffset = 0;
+					uint32 PatchControlPointOffset = 0;
+					uint32 HSOutAlign = 0;
+					uint32 PatchControlPointAlign = 0;
+					
+					InputVars.Sort([](SpvReflectInterfaceVariable& LHS, SpvReflectInterfaceVariable& RHS)  { return LHS.location < RHS.location; });
+					for (auto const& Var : InputVars)
+					{
+						if (Var->storage_class == SpvStorageClassInput && Var->built_in == -1)
+						{
+							unsigned Location = Var->location;
+							
+							uint32 ArrayCount = 1;
+							for (uint32 Dim = 1; Dim < Var->array.dims_count; Dim++)
+							{
+								ArrayCount *= Var->array.dims[Dim];
+							}
+							
+							auto const type = *Var->type_description;
+							uint32_t masked_type = type.type_flags & 0xF;
+							
+							for (uint32 j = 0; j < ArrayCount; j++)
+							{
+								FMetalAttribute& Attrib = Var->array.dims_count ? Attribs.PatchControlPointOut.AddDefaulted_GetRef() : Attribs.HSOut.AddDefaulted_GetRef();
+								Attrib.Index = Location + j;
+								uint32 ElementSize = 0;
+								switch (masked_type) {
+									default: checkf(false, TEXT("unsupported component type %d"), masked_type); break;
+									case SPV_REFLECT_TYPE_FLAG_BOOL  : Attrib.Type = EMetalComponentType::Bool; ElementSize = 1; break;
+									case SPV_REFLECT_TYPE_FLAG_INT   : Attrib.Type = (type.traits.numeric.scalar.signedness ? EMetalComponentType::Int : EMetalComponentType::Uint); ElementSize = 4; break;
+									case SPV_REFLECT_TYPE_FLAG_FLOAT : Attrib.Type = (type.traits.numeric.scalar.width == 32 ? EMetalComponentType::Float : EMetalComponentType::Half); if(type.traits.numeric.scalar.width == 32) { ElementSize = 4; } else { ElementSize = 2; } break;
+								}
+								if (type.type_flags & SPV_REFLECT_TYPE_FLAG_MATRIX)
+								{
+									Attrib.Components = (type.traits.numeric.matrix.row_count * type.traits.numeric.matrix.column_count);
+								}
+								else if (type.type_flags & SPV_REFLECT_TYPE_FLAG_VECTOR)
+								{
+									Attrib.Components = (type.traits.numeric.vector.component_count);
+								}
+								else
+								{
+									Attrib.Components = 1;
+								}
+								
+								uint32& Offset = Var->array.dims_count ? PatchControlPointOffset : HSOutOffset;
+								uint32& InitialAlign = Var->array.dims_count ? PatchControlPointAlign : HSOutAlign;
+								TCHAR const* Name = Var->array.dims_count ? TEXT("PatchOut") : TEXT("HSOut");
+								
+								char const* NumberedSemantic = strstr(Var->name, Var->semantic);
+								FString Semantic = FString::Printf(TEXT("%s_%u"), UTF8_TO_TCHAR(NumberedSemantic), j);
+								Attrib.Semantic = GetTypeHash(Semantic);
+								
+								uint32 AlignedComponents = Attrib.Components == 3 ? 4 : Attrib.Components;
+								InitialAlign = InitialAlign == 0 ? ElementSize * AlignedComponents : InitialAlign;
+								Offset = Align(Offset, ElementSize * AlignedComponents);
+								Attrib.Offset = Offset;
+								Offset += ElementSize * AlignedComponents;
+								
+								if (Input.Environment.CompilerFlags.Contains(CFLAG_KeepDebugInfo))
+									ALNString += FString::Printf(TEXT("%s%s.%s(%u)"), ALNString.Len() ? TEXT(",") : TEXT(""), Name, *Semantic, Attrib.Offset);
+							}
+						}
+					}
+					Attribs.HSOutSize = Align(HSOutOffset, 16);
+					Attribs.PatchControlPointOutSize = Align(PatchControlPointOffset, 16);
+					
+					if (Input.Environment.CompilerFlags.Contains(CFLAG_KeepDebugInfo))
+						ALNString += FString::Printf(TEXT(" = HSOutSize(%u), PatchControlPointOutSize(%u)"), Attribs.HSOutSize, Attribs.PatchControlPointOutSize);
+				}
+			}
+			
 			MetaData += TEXT("// Compiled by ShaderConductor\n");
 			if (INPString.Len())
 			{
@@ -1274,7 +1582,7 @@ bool FMetalShaderOutputCooker::Build(TArray<uint8>& OutData)
 		uint32 IndirectParamsIndex = 0;
 		uint32 PatchBufferIndex = 0;
 		uint32 TessFactorBufferIndex = 0;
-		uint32 ShaderInputWGIndex = 0;
+		uint32 HullIndexBuffer = 0;
 		if (Result)
 		{
 			ShaderConductor::DestroyBlob(Results.errorWarningMsg);
@@ -1315,7 +1623,7 @@ bool FMetalShaderOutputCooker::Build(TArray<uint8>& OutData)
 			
 			char OutputBufferIdx[3];
 			char IndirectParamsIdx[3];
-			if ((Frequency == HSF_VertexShader && bUsingTessellation) || Frequency == HSF_HullShader)
+			if (Frequency == HSF_VertexShader && bUsingTessellation)
 			{
 				OutputBufferIndex = FPlatformMath::CountTrailingZeros(BufferIndices);
 				BufferIndices &= ~(1 << OutputBufferIndex);
@@ -1331,11 +1639,32 @@ bool FMetalShaderOutputCooker::Build(TArray<uint8>& OutData)
 				TESStrings[(uint8)EMetalTessellationMetadataTags::TessellationPatchCountBuffer] = FString::Printf(TEXT("// @TessellationPatchCountBuffer: %u\n"), IndirectParamsIndex);
 			}
 			
+			if (Frequency == HSF_DomainShader)
+			{
+				OutputBufferIndex = FPlatformMath::CountTrailingZeros(BufferIndices);
+				BufferIndices &= ~(1 << OutputBufferIndex);
+				
+				IndirectParamsIndex = FPlatformMath::CountTrailingZeros(BufferIndices);
+				BufferIndices &= ~(1 << IndirectParamsIndex);
+				
+				TESStrings[(uint8)EMetalTessellationMetadataTags::TessellationHSOutBuffer] = FString::Printf(TEXT("// @TessellationHSOutBuffer: %u\n"), OutputBufferIndex);
+				TESStrings[(uint8)EMetalTessellationMetadataTags::TessellationControlPointOutBuffer] = FString::Printf(TEXT("// @TessellationControlPointOutBuffer: %u\n"), IndirectParamsIndex);
+			}
+			
 			char PatchBufferIdx[3];
 			char TessFactorBufferIdx[3];
-			char ShaderInputWGIdx[3];
 			if (Frequency == HSF_HullShader)
 			{
+				OutputBufferIndex = FPlatformMath::CountTrailingZeros(BufferIndices);
+				BufferIndices &= ~(1 << OutputBufferIndex);
+				FCStringAnsi::Snprintf(OutputBufferIdx, 3, "%u", OutputBufferIndex);
+				Defines[TargetDesc.numOptions++] = { "shader_output_buffer_index", OutputBufferIdx };
+				
+				IndirectParamsIndex = FPlatformMath::CountTrailingZeros(BufferIndices);
+				BufferIndices &= ~(1 << IndirectParamsIndex);
+				FCStringAnsi::Snprintf(IndirectParamsIdx, 3, "%u", IndirectParamsIndex);
+				Defines[TargetDesc.numOptions++] = { "indirect_params_buffer_index", IndirectParamsIdx };
+				
 				PatchBufferIndex = FPlatformMath::CountTrailingZeros(BufferIndices);
 				BufferIndices &= ~(1 << PatchBufferIndex);
 				FCStringAnsi::Snprintf(PatchBufferIdx, 3, "%u", PatchBufferIndex);
@@ -1345,13 +1674,20 @@ bool FMetalShaderOutputCooker::Build(TArray<uint8>& OutData)
 				BufferIndices &= ~(1 << TessFactorBufferIndex);
 				FCStringAnsi::Snprintf(TessFactorBufferIdx, 3, "%u", TessFactorBufferIndex);
 				Defines[TargetDesc.numOptions++] = { "shader_tess_factor_buffer_index", TessFactorBufferIdx };
+				Defines[TargetDesc.numOptions++] = { "shader_input_wg_index", "0" };
+
+				HullIndexBuffer = FPlatformMath::CountTrailingZeros(BufferIndices);
+				BufferIndices &= ~(1 << HullIndexBuffer);
 				
-				FCStringAnsi::Snprintf(ShaderInputWGIdx, 3, "%u", ShaderInputWGIndex);
-				Defines[TargetDesc.numOptions++] = { "shader_input_wg_index", ShaderInputWGIdx };
+				uint32 HullVertexBuffer = FPlatformMath::CountTrailingZeros(BufferIndices);
+				BufferIndices &= ~(1 << HullVertexBuffer);
 				
+				TESStrings[(uint8)EMetalTessellationMetadataTags::TessellationHSOutBuffer] = FString::Printf(TEXT("// @TessellationHSOutBuffer: %u\n"), PatchBufferIndex);
+				TESStrings[(uint8)EMetalTessellationMetadataTags::TessellationPatchCountBuffer] = FString::Printf(TEXT("// @TessellationPatchCountBuffer: %u\n"), IndirectParamsIndex);
+				TESStrings[(uint8)EMetalTessellationMetadataTags::TessellationControlPointIndexBuffer] = FString::Printf(TEXT("// @TessellationControlPointIndexBuffer: %u\n"), HullVertexBuffer);
 				TESStrings[(uint8)EMetalTessellationMetadataTags::TessellationHSTFOutBuffer] = FString::Printf(TEXT("// @TessellationHSTFOutBuffer: %u\n"), TessFactorBufferIndex);
-				TESStrings[(uint8)EMetalTessellationMetadataTags::TessellationControlPointOutBuffer] = FString::Printf(TEXT("// @TessellationControlPointOutBuffer: %u\n"), PatchBufferIndex);
-				TESStrings[(uint8)EMetalTessellationMetadataTags::TessellationIndexBuffer] = FString::Printf(TEXT("// @TessellationIndexBuffer: %u\n"), ShaderInputWGIndex);
+				TESStrings[(uint8)EMetalTessellationMetadataTags::TessellationControlPointOutBuffer] = FString::Printf(TEXT("// @TessellationControlPointOutBuffer: %u\n"), OutputBufferIndex);
+				TESStrings[(uint8)EMetalTessellationMetadataTags::TessellationIndexBuffer] = FString::Printf(TEXT("// @TessellationIndexBuffer: %u\n"), HullIndexBuffer);
 			}
 			
 			switch (VersionEnum)
@@ -1410,11 +1746,56 @@ bool FMetalShaderOutputCooker::Build(TArray<uint8>& OutData)
 			{
 				MetaData += TESString;
 			}
-			MetaData += FString::Printf(TEXT("\n\n"));
-			
+			MetaData += TEXT("\n\n");
+			if (ALNString.Len())
+			{
+				MetaData += TEXT("// Attributes: ");
+				MetaData += ALNString;
+				MetaData += TEXT("\n\n");
+			}
 			MetalSource = TCHAR_TO_UTF8(*MetaData);
 			MetalSource += std::string((const char*)Results.target->Data(), Results.target->Size());
 			
+			if (Frequency == HSF_DomainShader)
+			{
+				uint32 PatchSize = 0;
+				size_t ThreadSizeIdx = SourceData.find("OutputPatch<");
+				check(ThreadSizeIdx != std::string::npos);
+				char const* String = SourceData.c_str() + ThreadSizeIdx;
+				while(!isdigit(*String))
+				{
+					String++;
+				}
+#if !PLATFORM_WINDOWS
+				size_t Found = sscanf(String, "%u", &PatchSize);
+#else
+				size_t Found = sscanf_s(String, "%u", &PatchSize);
+#endif
+				check(Found == 1);
+				
+				std::string PatchAttr;
+				std::string PatchSearch = "[[ patch(triangle, 0) ]]";
+				size_t PatchIdx = MetalSource.find(PatchSearch);
+				if (PatchIdx != std::string::npos)
+				{
+					PatchAttr = "[[ patch(triangle, ";
+				}
+				else
+				{
+					PatchSearch = "[[ patch(quad, 0) ]]";
+					PatchIdx = MetalSource.find(PatchSearch);
+					PatchAttr = "[[ patch(quad, ";
+				}
+				if (PatchIdx != std::string::npos)
+				{
+					char BufferIdx[3];
+					FCStringAnsi::Snprintf(BufferIdx, 3, "%d", PatchSize);
+					PatchAttr += BufferIdx;
+					PatchAttr += ") ]]";
+					MetalSource.replace(PatchIdx, PatchSearch.length(), PatchAttr);
+				}
+			}
+			
 			CRCLen = MetalSource.length();
 			CRC = FCrc::MemCrc_DEPRECATED(MetalSource.c_str(), CRCLen);
 			
diff --git a/Engine/Source/Developer/Apple/MetalShaderFormat/Private/MetalShaderFormat.cpp b/Engine/Source/Developer/Apple/MetalShaderFormat/Private/MetalShaderFormat.cpp
--- a/Engine/Source/Developer/Apple/MetalShaderFormat/Private/MetalShaderFormat.cpp
+++ b/Engine/Source/Developer/Apple/MetalShaderFormat/Private/MetalShaderFormat.cpp
@@ -213,7 +213,7 @@ class FMetalShaderFormat : public IShaderFormat
 public:
 	enum
 	{
-		HEADER_VERSION = 66,
+		HEADER_VERSION = 68,
 	};
 	
 	struct FVersion
diff --git a/Engine/Source/Runtime/Apple/MetalRHI/Private/MetalCommandQueue.cpp b/Engine/Source/Runtime/Apple/MetalRHI/Private/MetalCommandQueue.cpp
--- a/Engine/Source/Runtime/Apple/MetalRHI/Private/MetalCommandQueue.cpp
+++ b/Engine/Source/Runtime/Apple/MetalRHI/Private/MetalCommandQueue.cpp
@@ -276,6 +276,12 @@ FMetalCommandQueue::FMetalCommandQueue(mtlpp::Device InDevice, uint32 const MaxN
     {
         Features |= EMetalFeaturesMultipleViewports | EMetalFeaturesPipelineBufferMutability | EMetalFeaturesGPUCaptureManager;
 		
+		static const auto CVar = IConsoleManager::Get().FindConsoleVariable(TEXT("r.Metal.ForceDXC"));
+		if (CVar && CVar->GetInt() != 0)
+		{
+			Features |= EMetalFeaturesSeparateTessellation;
+		}
+		
 		if (FParse::Param(FCommandLine::Get(),TEXT("metalfence")))
 		{
 			Features |= EMetalFeaturesFences;
diff --git a/Engine/Source/Runtime/Apple/MetalRHI/Private/MetalCommandQueue.h b/Engine/Source/Runtime/Apple/MetalRHI/Private/MetalCommandQueue.h
--- a/Engine/Source/Runtime/Apple/MetalRHI/Private/MetalCommandQueue.h
+++ b/Engine/Source/Runtime/Apple/MetalRHI/Private/MetalCommandQueue.h
@@ -71,6 +71,8 @@ typedef NS_OPTIONS(uint64, EMetalFeatures)
     EMetalFeaturesPipelineBufferMutability = 1llu << 25llu,
     /** Supports tile shaders */
     EMetalFeaturesTileShaders = 1llu << 26llu,
+	/** Supports separate tessellation shader execution */
+	EMetalFeaturesSeparateTessellation = 1llu << 27llu
 };
 
 /**
diff --git a/Engine/Source/Runtime/Apple/MetalRHI/Private/MetalCommands.cpp b/Engine/Source/Runtime/Apple/MetalRHI/Private/MetalCommands.cpp
--- a/Engine/Source/Runtime/Apple/MetalRHI/Private/MetalCommands.cpp
+++ b/Engine/Source/Runtime/Apple/MetalRHI/Private/MetalCommands.cpp
@@ -80,12 +80,6 @@ mtlpp::PrimitiveType TranslatePrimitiveType(uint32 PrimitiveType)
 		case PT_31_ControlPointPatchList:
 		case PT_32_ControlPointPatchList:
 		{
-			static uint32 Logged = 0;
-			if (!Logged)
-			{
-				Logged = 1;
-				UE_LOG(LogMetal, Warning, TEXT("Untested primitive type %d"), (int32)PrimitiveType);
-			}
 			return mtlpp::PrimitiveType::Triangle;
 		}
 		default:
@@ -504,6 +498,12 @@ void FMetalRHICommandContext::RHISetShaderParameter(FRHIHullShader* HullShaderRH
 {
 	// Just ignore Hull shader parameter sets - none of our Hull shaders have any loose parameters to bind.
 	// @todo Whenever we do put a shader parameter into a hull shader we'll need to map it into the vertex-shader parameter buffer so that it can be set on the device.
+	if (FMetalCommandQueue::SupportsFeature(EMetalFeaturesSeparateTessellation))
+	{
+		@autoreleasepool {
+			Context->GetCurrentState().GetShaderParameters(EMetalShaderStages::Hull).Set(BufferIndex, BaseIndex, NumBytes, NewValue);
+		}
+	}
 }
 
 void FMetalRHICommandContext::RHISetShaderParameter(FRHIPixelShader* PixelShaderRHI, uint32 BufferIndex, uint32 BaseIndex, uint32 NumBytes, const void* NewValue)
diff --git a/Engine/Source/Runtime/Apple/MetalRHI/Private/MetalDebugCommandEncoder.h b/Engine/Source/Runtime/Apple/MetalRHI/Private/MetalDebugCommandEncoder.h
--- a/Engine/Source/Runtime/Apple/MetalRHI/Private/MetalDebugCommandEncoder.h
+++ b/Engine/Source/Runtime/Apple/MetalRHI/Private/MetalDebugCommandEncoder.h
@@ -43,8 +43,9 @@ enum EMetalShaderFrequency
     EMetalShaderVertex = 0,
     EMetalShaderFragment = 1,
     EMetalShaderCompute = 2,
+	EMetalShaderStream = 3,
     EMetalShaderRenderNum = 2,
-	EMetalShaderStagesNum = 3
+	EMetalShaderStagesNum = 4
 };
 
 /** A structure for quick mask-testing of shader-stage resource bindings */
diff --git a/Engine/Source/Runtime/Apple/MetalRHI/Private/MetalPipeline.cpp b/Engine/Source/Runtime/Apple/MetalRHI/Private/MetalPipeline.cpp
--- a/Engine/Source/Runtime/Apple/MetalRHI/Private/MetalPipeline.cpp
+++ b/Engine/Source/Runtime/Apple/MetalRHI/Private/MetalPipeline.cpp
@@ -116,6 +116,49 @@ static NSString* GMetalDebugMarkerComputeShader = @"#include <metal_stdlib>\n"
 "	debugBuffer[0] = debugTable[0];\n"
 "}";
 
+// A compute shader for copying indices for separate tessellation
+static NSString* GMetalCopyIndexComputeShader = @"#include <metal_stdlib>\n"
+"#include <metal_compute>\n"
+"#include <metal_stdlib>\n"
+"using namespace metal;\n"
+"\n"
+"// Executes once\n"
+"kernel void Main_CopyIndex32(const device uint* source [[ buffer(0) ]], device uint* dest [[ buffer(1) ]], constant uint2& controlPointCount [[ buffer(2) ]], constant MTLDrawIndexedPrimitivesIndirectArguments& Params [[ buffer(3) ]], uint2 threadgroup_position_in_grid [[ threadgroup_position_in_grid ]], uint2 thread_position_in_threadgroup [[ thread_position_in_threadgroup ]])\n"
+"{\n"
+"	uint i = thread_position_in_threadgroup.y;\n"
+"	uint j = threadgroup_position_in_grid.x;\n"
+"	uint k = thread_position_in_threadgroup.x;\n"
+"	if (k < controlPointCount.x) {\n"
+"		dest[i * Params.indexCount + j * controlPointCount.y + k] = source[Params.indexStart + j * controlPointCount.x + k] + i * Params.indexCount;\n"
+"	} else {\n"
+"		dest[i * Params.indexCount + j * controlPointCount.y + k] = 0;\n"
+"	}\n"
+"}\n"
+"\n"
+"// Executes once\n"
+"kernel void Main_CopyIndex16(const device ushort* source [[ buffer(0) ]], device uint* dest [[ buffer(1) ]], constant uint2& controlPointCount [[ buffer(2) ]], constant MTLDrawIndexedPrimitivesIndirectArguments& Params [[ buffer(3) ]], uint2 threadgroup_position_in_grid [[ threadgroup_position_in_grid ]], uint2 thread_position_in_threadgroup [[ thread_position_in_threadgroup ]])\n"
+"{\n"
+"	uint i = thread_position_in_threadgroup.y;\n"
+"	uint j = threadgroup_position_in_grid.x;\n"
+"	uint k = thread_position_in_threadgroup.x;\n"
+"	if (k < controlPointCount.x) {\n"
+"		dest[i * Params.indexCount + j * controlPointCount.y + k] = source[Params.indexStart + j * controlPointCount.x + k] + i * Params.indexCount;\n"
+"	} else {\n"
+"		dest[i * Params.indexCount + j * controlPointCount.y + k] = 0;\n"
+"	}\n"
+"}\n"
+"\n"
+"// Executes once\n"
+"kernel void Main_FlattenTess(device MTLTriangleTessellationFactorsHalf* dest [[ buffer(0) ]], constant MTLDrawIndexedPrimitivesIndirectArguments& Params [[ buffer(1) ]])\n"
+"{\n"
+"	for(uint i = 0; i < Params.indexCount; i++) {\n"
+"		dest[i].edgeTessellationFactor[0] = half(1.0);\n"
+"		dest[i].edgeTessellationFactor[1] = half(1.0);\n"
+"		dest[i].edgeTessellationFactor[2] = half(1.0);\n"
+"		dest[i].insideTessellationFactor = half(1.0);\n"
+"	}\n"
+"}";
+
 struct FMetalHelperFunctions
 {
     mtlpp::Library DebugShadersLib;
@@ -124,7 +167,16 @@ struct FMetalHelperFunctions
 	mtlpp::Library DebugComputeShadersLib;
 	mtlpp::Function DebugComputeFunc;
 	mtlpp::ComputePipelineState DebugComputeState;
-    
+	
+	mtlpp::Library CopyIndexLib;
+	mtlpp::Function CopyIndex32Func;
+	mtlpp::Function CopyIndex16Func;
+	mtlpp::ComputePipelineState CopyIndex32State;
+	mtlpp::ComputePipelineState CopyIndex16State;
+	
+	mtlpp::Function FlattenTessFunc;
+	mtlpp::ComputePipelineState FlattenTessState;
+	
     FMetalHelperFunctions()
     {
 #if !PLATFORM_TVOS
@@ -142,6 +194,19 @@ struct FMetalHelperFunctions
 			DebugComputeState = GetMetalDeviceContext().GetDevice().NewComputePipelineState(DebugComputeFunc, &Error);
         }
 #endif
+		{
+			mtlpp::CompileOptions CompileOptions;
+			ns::AutoReleasedError Error;
+			
+			CopyIndexLib = GetMetalDeviceContext().GetDevice().NewLibrary(GMetalCopyIndexComputeShader, CompileOptions, &Error);
+			CopyIndex32Func = CopyIndexLib.NewFunction(@"Main_CopyIndex16");
+			CopyIndex16Func = CopyIndexLib.NewFunction(@"Main_CopyIndex16");
+			CopyIndex32State = GetMetalDeviceContext().GetDevice().NewComputePipelineState(CopyIndex32Func, &Error);
+			CopyIndex16State = GetMetalDeviceContext().GetDevice().NewComputePipelineState(CopyIndex16Func, &Error);
+			
+			FlattenTessFunc = CopyIndexLib.NewFunction(@"Main_FlattenTess");
+			FlattenTessState = GetMetalDeviceContext().GetDevice().NewComputePipelineState(FlattenTessFunc, &Error);
+		}
     }
     
     static FMetalHelperFunctions& Get()
@@ -159,6 +224,21 @@ struct FMetalHelperFunctions
 	{
 		return DebugComputeState;
 	}
+	
+	mtlpp::ComputePipelineState GetCopyIndex32Function()
+	{
+		return CopyIndex32State;
+	}
+	
+	mtlpp::ComputePipelineState GetCopyIndex16Function()
+	{
+		return CopyIndex16State;
+	}
+	
+	mtlpp::ComputePipelineState GetFlattenTessState()
+	{
+		return FlattenTessState;
+	}
 };
 
 mtlpp::ComputePipelineState GetMetalDebugComputeState()
@@ -166,6 +246,21 @@ mtlpp::ComputePipelineState GetMetalDebugComputeState()
 	return FMetalHelperFunctions::Get().GetDebugComputeState();
 }
 
+mtlpp::ComputePipelineState GetMetalCopyIndex32Function()
+{
+	return FMetalHelperFunctions::Get().GetCopyIndex32Function();
+}
+
+mtlpp::ComputePipelineState GetMetalCopyIndex16Function()
+{
+	return FMetalHelperFunctions::Get().GetCopyIndex16Function();
+}
+
+mtlpp::ComputePipelineState GetMetalFlattenTessState()
+{
+	return FMetalHelperFunctions::Get().GetFlattenTessState();
+}
+
 struct FMetalGraphicsPipelineKey
 {
 	FMetalRenderPipelineHash RenderPipelineHash;
@@ -470,8 +565,10 @@ APPLE_PLATFORM_OBJECT_ALLOC_OVERRIDES(FMetalShaderPipeline)
 	{
 		RenderPipelineReflection = mtlpp::RenderPipelineReflection(nil);
 		ComputePipelineReflection = mtlpp::ComputePipelineReflection(nil);
+		StreamPipelineReflection = mtlpp::RenderPipelineReflection(nil);
 #if METAL_DEBUG_OPTIONS
 		RenderDesc = mtlpp::RenderPipelineDescriptor(nil);
+		StreamDesc = mtlpp::RenderPipelineDescriptor(nil);
 		ComputeDesc = mtlpp::ComputePipelineDescriptor(nil);
 #endif
 	}
@@ -499,6 +596,15 @@ APPLE_PLATFORM_OBJECT_ALLOC_OVERRIDES(FMetalShaderPipeline)
 			ComputePipelineReflection = mtlpp::ComputePipelineReflection(nil);
 		}
 	}
+	if (StreamPipelineReflection)
+	{
+		[self initResourceMask:EMetalShaderStream];
+		
+		if (SafeGetRuntimeDebuggingLevel() < EMetalDebugLevelValidation METAL_STATISTICS_ONLY(&& !GetMetalDeviceContext().GetCommandQueue().GetStatistics()))
+		{
+			StreamPipelineReflection = mtlpp::RenderPipelineReflection(nil);
+		}
+	}
 }
 - (void)initResourceMask:(EMetalShaderFrequency)Frequency
 {
@@ -529,6 +635,14 @@ APPLE_PLATFORM_OBJECT_ALLOC_OVERRIDES(FMetalShaderPipeline)
 			Arguments = Reflection.arguments;
 			break;
 		}
+		case EMetalShaderStream:
+		{
+			MTLRenderPipelineReflection* Reflection = StreamPipelineReflection;
+			check(Reflection);
+			
+			Arguments = Reflection.vertexArguments;
+			break;
+		}
 		default:
 			check(false);
 			break;
@@ -612,254 +726,621 @@ static MTLVertexDescriptor* GetMaskedVertexDescriptor(MTLVertexDescriptor* Input
 	return InputDesc;
 }
 
-static FMetalShaderPipeline* CreateMTLRenderPipeline(bool const bSync, FMetalGraphicsPipelineKey const& Key, const FGraphicsPipelineStateInitializer& Init, EMetalIndexType const IndexType)
+static bool ConfigureRenderPipelineDescriptor(mtlpp::RenderPipelineDescriptor& RenderPipelineDesc,
+#if PLATFORM_MAC
+											  mtlpp::RenderPipelineDescriptor& DebugPipelineDesc,
+#elif !PLATFORM_TVOS
+											  mtlpp::TileRenderPipelineDescriptor& DebugPipelineDesc,
+#endif
+											  FMetalGraphicsPipelineKey const& Key,
+											  const FGraphicsPipelineStateInitializer& Init,
+											  EMetalIndexType const IndexType)
 {
-    FMetalVertexShader* VertexShader = (FMetalVertexShader*)Init.BoundShaderState.VertexShaderRHI;
-    FMetalPixelShader* PixelShader = (FMetalPixelShader*)Init.BoundShaderState.PixelShaderRHI;
-    
-    mtlpp::Function vertexFunction = VertexShader->GetFunction();
-    mtlpp::Function fragmentFunction = PixelShader ? PixelShader->GetFunction() : nil;
-
-#if PLATFORM_SUPPORTS_TESSELLATION_SHADERS
-	FMetalDomainShader* DomainShader = (FMetalDomainShader*)Init.BoundShaderState.DomainShaderRHI;
-    mtlpp::Function domainFunction = DomainShader ? DomainShader->GetFunction() : nil;
+	FMetalPixelShader* PixelShader = (FMetalPixelShader*)Init.BoundShaderState.PixelShaderRHI;
+	uint32 const NumActiveTargets = Init.ComputeNumValidRenderTargets();
+	check(NumActiveTargets <= MaxSimultaneousRenderTargets);
+	if (PixelShader)
+	{
+		if ((PixelShader->Bindings.InOutMask & 0x8000) == 0 && (PixelShader->Bindings.InOutMask & 0x7fff) == 0 && PixelShader->Bindings.NumUAVs == 0 && PixelShader->Bindings.bDiscards == false)
+		{
+			UE_LOG(LogMetal, Error, TEXT("Pixel shader has no outputs which is not permitted. No Discards, In-Out Mask: %x\nNumber UAVs: %d\nSource Code:\n%s"), PixelShader->Bindings.InOutMask, PixelShader->Bindings.NumUAVs, *FString(PixelShader->GetSourceCode()));
+			return false;
+		}
+		
+		UE_CLOG((NumActiveTargets < __builtin_popcount(PixelShader->Bindings.InOutMask & 0x7fff)), LogMetal, Verbose, TEXT("NumActiveTargets doesn't match pipeline's pixel shader output mask: %u, %hx"), NumActiveTargets, PixelShader->Bindings.InOutMask);
+	}
+	
+	FMetalBlendState* BlendState = (FMetalBlendState*)Init.BlendState;
+	
+	ns::Array<mtlpp::RenderPipelineColorAttachmentDescriptor> ColorAttachments = RenderPipelineDesc.GetColorAttachments();
+#if !PLATFORM_TVOS
+	auto DebugColorAttachements = DebugPipelineDesc.GetColorAttachments();
 #endif
-    
-    FMetalShaderPipeline* Pipeline = nil;
-    if (vertexFunction && ((PixelShader != nullptr) == (fragmentFunction != nil))
-#if PLATFORM_SUPPORTS_TESSELLATION_SHADERS
-		&& ((DomainShader != nullptr) == (domainFunction != nil))
+	
+	uint32 TargetWidth = 0;
+	for (uint32 i = 0; i < NumActiveTargets; i++)
+	{
+		EPixelFormat TargetFormat = (EPixelFormat)Init.RenderTargetFormats[i];
+		
+		METAL_FATAL_ASSERT(!(TargetFormat == PF_Unknown && PixelShader && (((PixelShader->Bindings.InOutMask & 0x7fff) & (1 << i)))), TEXT("Pipeline pixel shader expects target %u to be bound but it isn't: %s."), i, *FString(PixelShader->GetSourceCode()));
+		
+		TargetWidth += GPixelFormats[TargetFormat].BlockBytes;
+		
+		mtlpp::PixelFormat MetalFormat = (mtlpp::PixelFormat)GPixelFormats[TargetFormat].PlatformFormat;
+		uint32 Flags = Init.RenderTargetFlags[i];
+		if (Flags & TexCreate_SRGB)
+		{
+#if PLATFORM_MAC // Expand as R8_sRGB is iOS only.
+			if (MetalFormat == mtlpp::PixelFormat::R8Unorm)
+			{
+				MetalFormat = mtlpp::PixelFormat::RGBA8Unorm;
+			}
 #endif
-	)
-    {
-		ns::Error Error;
-		mtlpp::Device Device = GetMetalDeviceContext().GetDevice();
-
-		uint32 const NumActiveTargets = Init.ComputeNumValidRenderTargets();
-        check(NumActiveTargets <= MaxSimultaneousRenderTargets);
-        if (PixelShader)
-        {
-			if ((PixelShader->Bindings.InOutMask & 0x8000) == 0 && (PixelShader->Bindings.InOutMask & 0x7fff) == 0 && PixelShader->Bindings.NumUAVs == 0 && PixelShader->Bindings.bDiscards == false)
+			MetalFormat = ToSRGBFormat(MetalFormat);
+		}
+		
+		mtlpp::RenderPipelineColorAttachmentDescriptor Attachment = ColorAttachments[i];
+		Attachment.SetPixelFormat(MetalFormat);
+		
+#if !PLATFORM_TVOS
+		auto DebugAttachment = DebugColorAttachements[i];;
+		DebugAttachment.SetPixelFormat(MetalFormat);
+#endif
+		
+		mtlpp::RenderPipelineColorAttachmentDescriptor Blend = BlendState->RenderTargetStates[i].BlendState;
+		if(TargetFormat != PF_Unknown)
+		{
+			// assign each property manually, would be nice if this was faster
+			Attachment.SetBlendingEnabled(Blend.IsBlendingEnabled());
+			Attachment.SetSourceRgbBlendFactor(Blend.GetSourceRgbBlendFactor());
+			Attachment.SetDestinationRgbBlendFactor(Blend.GetDestinationRgbBlendFactor());
+			Attachment.SetRgbBlendOperation(Blend.GetRgbBlendOperation());
+			Attachment.SetSourceAlphaBlendFactor(Blend.GetSourceAlphaBlendFactor());
+			Attachment.SetDestinationAlphaBlendFactor(Blend.GetDestinationAlphaBlendFactor());
+			Attachment.SetAlphaBlendOperation(Blend.GetAlphaBlendOperation());
+			Attachment.SetWriteMask(Blend.GetWriteMask());
+			
+#if PLATFORM_MAC
+			DebugAttachment.SetBlendingEnabled(Blend.IsBlendingEnabled());
+			DebugAttachment.SetSourceRgbBlendFactor(Blend.GetSourceRgbBlendFactor());
+			DebugAttachment.SetDestinationRgbBlendFactor(Blend.GetDestinationRgbBlendFactor());
+			DebugAttachment.SetRgbBlendOperation(Blend.GetRgbBlendOperation());
+			DebugAttachment.SetSourceAlphaBlendFactor(Blend.GetSourceAlphaBlendFactor());
+			DebugAttachment.SetDestinationAlphaBlendFactor(Blend.GetDestinationAlphaBlendFactor());
+			DebugAttachment.SetAlphaBlendOperation(Blend.GetAlphaBlendOperation());
+			DebugAttachment.SetWriteMask(Blend.GetWriteMask());
+#endif
+		}
+		else
+		{
+			Attachment.SetBlendingEnabled(NO);
+			Attachment.SetWriteMask(mtlpp::ColorWriteMask::None);
+#if PLATFORM_MAC
+			DebugAttachment.SetBlendingEnabled(NO);
+			DebugAttachment.SetWriteMask(mtlpp::ColorWriteMask::None);
+#endif
+		}
+	}
+	
+	// don't allow a PSO that is too wide
+	if (!GSupportsWideMRT && TargetWidth > 16)
+	{
+		return false;
+	}
+	
+	switch(Init.DepthStencilTargetFormat)
+	{
+		case PF_DepthStencil:
+		{
+			mtlpp::PixelFormat MetalFormat = (mtlpp::PixelFormat)GPixelFormats[PF_DepthStencil].PlatformFormat;
+			if(MetalFormat == mtlpp::PixelFormat::Depth32Float)
 			{
-				UE_LOG(LogMetal, Error, TEXT("Pixel shader has no outputs which is not permitted. No Discards, In-Out Mask: %x\nNumber UAVs: %d\nSource Code:\n%s"), PixelShader->Bindings.InOutMask, PixelShader->Bindings.NumUAVs, *FString(PixelShader->GetSourceCode()));
-				return nil;
+				if (Init.DepthTargetLoadAction != ERenderTargetLoadAction::ENoAction || Init.DepthTargetStoreAction != ERenderTargetStoreAction::ENoAction)
+				{
+					RenderPipelineDesc.SetDepthAttachmentPixelFormat(MetalFormat);
+#if PLATFORM_MAC
+					DebugPipelineDesc.SetDepthAttachmentPixelFormat(MetalFormat);
+#endif
+				}
+				if (Init.StencilTargetLoadAction != ERenderTargetLoadAction::ENoAction || Init.StencilTargetStoreAction != ERenderTargetStoreAction::ENoAction)
+				{
+					RenderPipelineDesc.SetStencilAttachmentPixelFormat(mtlpp::PixelFormat::Stencil8);
+#if PLATFORM_MAC
+					DebugPipelineDesc.SetStencilAttachmentPixelFormat(mtlpp::PixelFormat::Stencil8);
+#endif
+				}
 			}
-            
-            UE_CLOG((NumActiveTargets < __builtin_popcount(PixelShader->Bindings.InOutMask & 0x7fff)), LogMetal, Verbose, TEXT("NumActiveTargets doesn't match pipeline's pixel shader output mask: %u, %hx"), NumActiveTargets, PixelShader->Bindings.InOutMask);
-        }
-        
-		Pipeline = [FMetalShaderPipeline new];
-		METAL_DEBUG_OPTION(FMemory::Memzero(Pipeline->ResourceMask, sizeof(Pipeline->ResourceMask)));
-
-		mtlpp::RenderPipelineDescriptor RenderPipelineDesc;
-        mtlpp::ComputePipelineDescriptor ComputePipelineDesc(nil);
+			else
+			{
+				RenderPipelineDesc.SetDepthAttachmentPixelFormat(MetalFormat);
+				RenderPipelineDesc.SetStencilAttachmentPixelFormat(MetalFormat);
 #if PLATFORM_MAC
-        mtlpp::RenderPipelineDescriptor DebugPipelineDesc;
-#elif !PLATFORM_TVOS
-        mtlpp::TileRenderPipelineDescriptor DebugPipelineDesc;
+				DebugPipelineDesc.SetDepthAttachmentPixelFormat(MetalFormat);
+				DebugPipelineDesc.SetStencilAttachmentPixelFormat(MetalFormat);
 #endif
-		
-		if (FMetalCommandQueue::SupportsFeature(EMetalFeaturesPipelineBufferMutability))
+			}
+			break;
+		}
+		case PF_ShadowDepth:
+		{
+			RenderPipelineDesc.SetDepthAttachmentPixelFormat((mtlpp::PixelFormat)GPixelFormats[PF_ShadowDepth].PlatformFormat);
+#if PLATFORM_MAC
+			DebugPipelineDesc.SetDepthAttachmentPixelFormat((mtlpp::PixelFormat)GPixelFormats[PF_ShadowDepth].PlatformFormat);
+#endif
+			break;
+		}
+		default:
 		{
-			ns::AutoReleased<ns::Array<mtlpp::PipelineBufferDescriptor>> VertexPipelineBuffers = RenderPipelineDesc.GetVertexBuffers();
+			break;
+		}
+	}
+	
+	check(Init.BoundShaderState.VertexShaderRHI != nullptr);
+#if PLATFORM_SUPPORTS_GEOMETRY_SHADERS
+	check(Init.BoundShaderState.GeometryShaderRHI == nullptr);
+#endif
+	
+	if(RenderPipelineDesc.GetDepthAttachmentPixelFormat() == mtlpp::PixelFormat::Invalid && PixelShader && ((PixelShader->Bindings.InOutMask & 0x8000) || (NumActiveTargets == 0 && (PixelShader->Bindings.NumUAVs > 0))))
+	{
+		RenderPipelineDesc.SetDepthAttachmentPixelFormat((mtlpp::PixelFormat)GPixelFormats[PF_DepthStencil].PlatformFormat);
+		RenderPipelineDesc.SetStencilAttachmentPixelFormat((mtlpp::PixelFormat)GPixelFormats[PF_DepthStencil].PlatformFormat);
+		
+#if PLATFORM_MAC
+		DebugPipelineDesc.SetDepthAttachmentPixelFormat((mtlpp::PixelFormat)GPixelFormats[PF_DepthStencil].PlatformFormat);
+		DebugPipelineDesc.SetStencilAttachmentPixelFormat((mtlpp::PixelFormat)GPixelFormats[PF_DepthStencil].PlatformFormat);
+#endif
+	}
+	
+	static bool bNoMSAA = FParse::Param(FCommandLine::Get(), TEXT("nomsaa"));
+	RenderPipelineDesc.SetSampleCount(!bNoMSAA ? FMath::Max(Init.NumSamples, (uint16)1u) : (uint16)1u);
+#if PLATFORM_MAC
+	RenderPipelineDesc.SetInputPrimitiveTopology(TranslatePrimitiveTopology(Init.PrimitiveType));
+	DebugPipelineDesc.SetSampleCount(!bNoMSAA ? FMath::Max(Init.NumSamples, (uint16)1u) : (uint16)1u);
+	DebugPipelineDesc.SetInputPrimitiveTopology(mtlpp::PrimitiveTopologyClass::Point);
+#endif
+	
+	if (FMetalCommandQueue::SupportsFeature(EMetalFeaturesPipelineBufferMutability))
+	{
+		FMetalVertexShader* VertexShader = (FMetalVertexShader*)Init.BoundShaderState.VertexShaderRHI;
+		
+		ns::AutoReleased<ns::Array<mtlpp::PipelineBufferDescriptor>> VertexPipelineBuffers = RenderPipelineDesc.GetVertexBuffers();
 #if PLATFORM_SUPPORTS_TESSELLATION_SHADERS
-			FMetalShaderBindings& VertexBindings = DomainShader ? DomainShader->Bindings : VertexShader->Bindings;
-			int8 VertexSideTable = DomainShader ? DomainShader->SideTableBinding : VertexShader->SideTableBinding;
+		FMetalDomainShader* DomainShader = (FMetalDomainShader*)Init.BoundShaderState.DomainShaderRHI;
+		FMetalShaderBindings& VertexBindings = DomainShader ? DomainShader->Bindings : VertexShader->Bindings;
+		int8 VertexSideTable = DomainShader ? DomainShader->SideTableBinding : VertexShader->SideTableBinding;
 #else
-			FMetalShaderBindings& VertexBindings = VertexShader->Bindings;
-			int8 VertexSideTable = VertexShader->SideTableBinding;
+		FMetalShaderBindings& VertexBindings = VertexShader->Bindings;
+		int8 VertexSideTable = VertexShader->SideTableBinding;
 #endif
+		{
+			uint32 ImmutableBuffers = VertexBindings.ConstantBuffers | VertexBindings.ArgumentBuffers;
+			while(ImmutableBuffers)
 			{
-				uint32 ImmutableBuffers = VertexBindings.ConstantBuffers | VertexBindings.ArgumentBuffers;
-				while(ImmutableBuffers)
+				uint32 Index = __builtin_ctz(ImmutableBuffers);
+				ImmutableBuffers &= ~(1 << Index);
+				
+				if (Index < ML_MaxBuffers)
 				{
-					uint32 Index = __builtin_ctz(ImmutableBuffers);
-					ImmutableBuffers &= ~(1 << Index);
-					
-					if (Index < ML_MaxBuffers)
-					{
-						ns::AutoReleased<mtlpp::PipelineBufferDescriptor> PipelineBuffer = VertexPipelineBuffers[Index];
-						PipelineBuffer.SetMutability(mtlpp::Mutability::Immutable);
-					}
+					ns::AutoReleased<mtlpp::PipelineBufferDescriptor> PipelineBuffer = VertexPipelineBuffers[Index];
+					PipelineBuffer.SetMutability(mtlpp::Mutability::Immutable);
 				}
-				if (VertexSideTable > 0)
+			}
+			if (VertexSideTable > 0)
+			{
+				ns::AutoReleased<mtlpp::PipelineBufferDescriptor> PipelineBuffer = VertexPipelineBuffers[VertexSideTable];
+				PipelineBuffer.SetMutability(mtlpp::Mutability::Immutable);
+			}
+		}
+		
+		if (PixelShader)
+		{
+			ns::AutoReleased<ns::Array<mtlpp::PipelineBufferDescriptor>> FragmentPipelineBuffers = RenderPipelineDesc.GetFragmentBuffers();
+			uint32 ImmutableBuffers = PixelShader->Bindings.ConstantBuffers | PixelShader->Bindings.ArgumentBuffers;
+			while(ImmutableBuffers)
+			{
+				uint32 Index = __builtin_ctz(ImmutableBuffers);
+				ImmutableBuffers &= ~(1 << Index);
+				
+				if (Index < ML_MaxBuffers)
 				{
-					ns::AutoReleased<mtlpp::PipelineBufferDescriptor> PipelineBuffer = VertexPipelineBuffers[VertexSideTable];
+					ns::AutoReleased<mtlpp::PipelineBufferDescriptor> PipelineBuffer = FragmentPipelineBuffers[Index];
 					PipelineBuffer.SetMutability(mtlpp::Mutability::Immutable);
 				}
 			}
+			if (PixelShader->SideTableBinding > 0)
+			{
+				ns::AutoReleased<mtlpp::PipelineBufferDescriptor> PipelineBuffer = FragmentPipelineBuffers[PixelShader->SideTableBinding];
+				PipelineBuffer.SetMutability(mtlpp::Mutability::Immutable);
+			}
+		}
+	}
+	
+	return true;
+}
+
+#if PLATFORM_SUPPORTS_TESSELLATION_SHADERS
+
+static mtlpp::VertexFormat Formats[(uint8)EMetalComponentType::Max][4] = {
+	{mtlpp::VertexFormat::UInt, mtlpp::VertexFormat::UInt2, mtlpp::VertexFormat::UInt3, mtlpp::VertexFormat::UInt4},
+	{mtlpp::VertexFormat::Int, mtlpp::VertexFormat::Int2, mtlpp::VertexFormat::Int3, mtlpp::VertexFormat::Int4},
+	{mtlpp::VertexFormat::Invalid, mtlpp::VertexFormat::Half2, mtlpp::VertexFormat::Half3, mtlpp::VertexFormat::Half4},
+	{mtlpp::VertexFormat::Float, mtlpp::VertexFormat::Float2, mtlpp::VertexFormat::Float3, mtlpp::VertexFormat::Float4},
+	{mtlpp::VertexFormat::Invalid, mtlpp::VertexFormat::UChar2, mtlpp::VertexFormat::UChar3, mtlpp::VertexFormat::UChar4},
+};
+
+static FMetalShaderPipeline* CreateSeparateMetalTessellationPipeline(bool const bSync, FMetalGraphicsPipelineKey const& Key, const FGraphicsPipelineStateInitializer& Init, EMetalIndexType const IndexType)
+{
+	FMetalShaderPipeline* Pipeline = [FMetalShaderPipeline new];
+	METAL_DEBUG_OPTION(FMemory::Memzero(Pipeline->ResourceMask, sizeof(Pipeline->ResourceMask)));
+	
+	FMetalVertexShader* VertexShader = (FMetalVertexShader*)Init.BoundShaderState.VertexShaderRHI;
+	mtlpp::Function vertexFunction = VertexShader->GetFunction();
+	
+	mtlpp::RenderPipelineDescriptor VertexRenderPipelineDesc;
+	VertexRenderPipelineDesc.SetRasterizationEnabled(false);
+	
+#if PLATFORM_MAC
+	VertexRenderPipelineDesc.SetInputPrimitiveTopology(TranslatePrimitiveTopology(Init.PrimitiveType));
+#endif
+	
+	FMetalVertexDeclaration* VertexDecl = (FMetalVertexDeclaration*)Init.BoundShaderState.VertexDeclarationRHI;
+	VertexRenderPipelineDesc.SetVertexDescriptor(GetMaskedVertexDescriptor(VertexDecl->Layout.VertexDesc, VertexShader->Bindings.InOutMask));
+	VertexRenderPipelineDesc.SetVertexFunction(vertexFunction);
+	
+	VertexRenderPipelineDesc.SetDepthAttachmentPixelFormat((mtlpp::PixelFormat)GPixelFormats[PF_DepthStencil].PlatformFormat);
+	VertexRenderPipelineDesc.SetStencilAttachmentPixelFormat((mtlpp::PixelFormat)GPixelFormats[PF_DepthStencil].PlatformFormat);
+	
+#if ENABLE_METAL_GPUPROFILE
+	ns::String VertexName = vertexFunction.GetName();
+	VertexRenderPipelineDesc.SetLabel([NSString stringWithFormat:@"%@", VertexName.GetPtr()]);
+#endif
+	
+	NSUInteger RenderOption = mtlpp::PipelineOption::NoPipelineOption;
+	mtlpp::AutoReleasedRenderPipelineReflection* Reflection = nullptr;
+	mtlpp::AutoReleasedRenderPipelineReflection OutReflection;
+	Reflection = &OutReflection;
+	if (GetMetalDeviceContext().GetCommandQueue().GetRuntimeDebuggingLevel() >= EMetalDebugLevelFastValidation METAL_STATISTICS_ONLY(|| GetMetalDeviceContext().GetCommandQueue().GetStatistics()))
+	{
+		RenderOption = mtlpp::PipelineOption::ArgumentInfo|mtlpp::PipelineOption::BufferTypeInfo METAL_STATISTICS_ONLY(|NSUInteger(EMTLPipelineStats));
+	}
+	
+	ns::Error Error;
+	mtlpp::Device Device = GetMetalDeviceContext().GetDevice();
+	mtlpp::RenderPipelineState& VertexRenderPipelineState = Pipeline->StreamPipelineState;
+	{
+		ns::AutoReleasedError RenderError;
+		METAL_GPUPROFILE(FScopedMetalCPUStats CPUStat(FString::Printf(TEXT("NewRenderPipeline: %s"), TEXT("")/**FString([RenderPipelineDesc.GetPtr() description])*/)));
+		VertexRenderPipelineState = Device.NewRenderPipelineState(VertexRenderPipelineDesc, (mtlpp::PipelineOption)RenderOption, Reflection, &RenderError);
+		if (Reflection)
+		{
+			Pipeline->StreamPipelineReflection = *Reflection;
+#if METAL_DEBUG_OPTIONS
+			Pipeline->StreamDesc = VertexRenderPipelineDesc;
+#endif
+		}
+		Error = RenderError;
+	}
+	
+	UE_CLOG((VertexRenderPipelineState == nil), LogMetal, Error, TEXT("Failed to generate a pipeline state object: %s"), *FString(Error.GetPtr().description));
+	UE_CLOG((VertexRenderPipelineState == nil), LogMetal, Error, TEXT("Vertex shader: %s"), *FString(VertexShader->GetSourceCode()));
+	UE_CLOG((VertexRenderPipelineState == nil), LogMetal, Error, TEXT("Descriptor: %s"), *FString(VertexRenderPipelineDesc.GetPtr().description));
+	UE_CLOG((VertexRenderPipelineState == nil), LogMetal, Error, TEXT("Failed to generate a render pipeline state object:\n\n %s\n\n"), *FString(Error.GetLocalizedDescription()));
+
+	FMetalHullShader* HullShader = (FMetalHullShader*)Init.BoundShaderState.HullShaderRHI;
+	mtlpp::ComputePipelineDescriptor HullShaderPipelineDesc;
+	
+	mtlpp::StageInputOutputDescriptor HullStageInOut;
+	
+	ns::Array<mtlpp::BufferLayoutDescriptor> HullVertexLayouts = HullStageInOut.GetLayouts();
+	ns::Array<mtlpp::AttributeDescriptor> HullVertexAttribs = HullStageInOut.GetAttributes();
+	
+	bool bFoundAttrib = false;
+	for (FMetalAttribute const& HSAttrib : HullShader->TessellationOutputAttribs.HSIn)
+	{
+		bFoundAttrib = false;
+		for (FMetalAttribute const& VSAttrib : VertexShader->TessellationOutputAttribs.HSOut)
+		{
+			if (HSAttrib.Semantic == VSAttrib.Semantic)
+			{
+				uint32 attributeIndex = HSAttrib.Index;
+
+				mtlpp::VertexFormat format = Formats[(uint8)HSAttrib.Type][HSAttrib.Components-1];
+				check(format != mtlpp::VertexFormat::Invalid); // TODO support more cases
+				HullVertexAttribs[attributeIndex].SetFormat((mtlpp::AttributeFormat)format);
+				HullVertexAttribs[attributeIndex].SetOffset(VSAttrib.Offset);
+				HullVertexAttribs[attributeIndex].SetBufferIndex(HullShader->TessellationControlPointIndexBuffer);
+				bFoundAttrib = true;
+				break;
+			}
+		}
+		check(bFoundAttrib);
+	}
+	
+	// if (IndexType == EMetalIndexType_UInt16 || IndexType == EMetalIndexType_UInt32)
+	{
+		HullStageInOut.SetIndexType(GetMetalIndexType(EMetalIndexType_UInt32));
+		HullStageInOut.SetIndexBufferIndex(HullShader->TessellationIndexBuffer);
+		HullVertexLayouts[HullShader->TessellationControlPointIndexBuffer].SetStepFunction(mtlpp::StepFunction::ThreadPositionInGridXIndexed);
+	}
+//	else
+//	{
+//		HullVertexLayouts[HullShader->TessellationControlPointIndexBuffer].SetStepFunction(mtlpp::StepFunction::ThreadPositionInGridX);
+//	}
+	HullVertexLayouts[HullShader->TessellationControlPointIndexBuffer].SetStride(VertexShader->TessellationOutputAttribs.HSOutSize);
+	
+	mtlpp::Function hullFunction = HullShader->GetFunction();
+	HullShaderPipelineDesc.SetComputeFunction(hullFunction);
+	check(HullShaderPipelineDesc.GetComputeFunction());
+	HullShaderPipelineDesc.SetStageInputDescriptor(HullStageInOut);
+	
+	mtlpp::ComputePipelineState& HullPipelineState = Pipeline->ComputePipelineState;
+	{
+		METAL_GPUPROFILE(FScopedMetalCPUStats CPUStat(FString::Printf(TEXT("NewComputePipelineState: %s"), TEXT("")/**FString([ComputePipelineDesc.GetPtr() description])*/)));
+		ns::AutoReleasedError AutoError;
+		NSUInteger ComputeOption = mtlpp::PipelineOption::NoPipelineOption;
+#if ENABLE_METAL_GPUPROFILE
+		{
+			ns::String HullName = hullFunction.GetName();
+			HullShaderPipelineDesc.SetLabel([NSString stringWithFormat:@"%@", HullName.GetPtr()]);
+		}
+#endif
+		if (GetMetalDeviceContext().GetCommandQueue().GetRuntimeDebuggingLevel() >= EMetalDebugLevelFastValidation METAL_STATISTICS_ONLY(|| GetMetalDeviceContext().GetCommandQueue().GetStatistics()))
+		{
+			mtlpp::AutoReleasedComputePipelineReflection HullReflection;
+			ComputeOption = mtlpp::PipelineOption::ArgumentInfo|mtlpp::PipelineOption::BufferTypeInfo METAL_STATISTICS_ONLY(|NSUInteger(EMTLPipelineStats));
+			HullPipelineState = Device.NewComputePipelineState(HullShaderPipelineDesc, (mtlpp::PipelineOption)ComputeOption, &HullReflection, &AutoError);
+			Pipeline->ComputePipelineReflection = HullReflection;
+		}
+		else
+		{
+			HullPipelineState = Device.NewComputePipelineState(HullShaderPipelineDesc, (mtlpp::PipelineOption)ComputeOption, nullptr, &AutoError);
+		}
+		Error = AutoError;
+		
+#if METAL_DEBUG_OPTIONS
+		Pipeline->ComputeDesc = HullShaderPipelineDesc;
+#endif
+		
+		UE_CLOG((HullPipelineState == nil), LogMetal, Error, TEXT("Failed to generate a pipeline state object: %s"), *FString([Error.GetPtr() description]));
+		UE_CLOG((HullPipelineState == nil), LogMetal, Error, TEXT("Hull shader: %s"), *FString(HullShader->GetSourceCode()));
+		UE_CLOG((HullPipelineState == nil), LogMetal, Error, TEXT("Descriptor: %s"), *FString(HullShaderPipelineDesc.GetPtr().description));
+		UE_CLOG((HullPipelineState == nil), LogMetal, Error, TEXT("Failed to generate a hull pipeline state object:\n\n %s\n\n"), *FString(Error.GetLocalizedDescription()));
+	}
+	
+#if PLATFORM_MAC
+	mtlpp::RenderPipelineDescriptor DebugPipelineDesc;
+#elif !PLATFORM_TVOS
+	mtlpp::TileRenderPipelineDescriptor DebugPipelineDesc;
+#endif
+
+	FMetalDomainShader* DomainShader = (FMetalDomainShader*)Init.BoundShaderState.DomainShaderRHI;
+	mtlpp::RenderPipelineDescriptor DomainRenderPipelineDesc;
+	{
+		mtlpp::VertexDescriptor DomainVertexDesc;
+		ns::Array<mtlpp::VertexBufferLayoutDescriptor> DomainVertexLayouts = DomainVertexDesc.GetLayouts();
+		ns::Array<mtlpp::VertexAttributeDescriptor> Attribs = DomainVertexDesc.GetAttributes();
+		
+		if (DomainShader->TessellationHSOutBuffer != UINT_MAX && HullShader->TessellationOutputAttribs.HSOutSize)
+		{
+			check(DomainShader->TessellationHSOutBuffer < ML_MaxBuffers);
+			uint32 bufferIndex = DomainShader->TessellationHSOutBuffer;
+			uint32 bufferSize = HullShader->TessellationOutputAttribs.HSOutSize;
 			
-			if (PixelShader)
+			DomainVertexLayouts[bufferIndex].SetStride(bufferSize);
+			DomainVertexLayouts[bufferIndex].SetStepFunction(mtlpp::VertexStepFunction::PerPatch);
+			DomainVertexLayouts[bufferIndex].SetStepRate(1);
+			
+			bFoundAttrib = false;
+			for (FMetalAttribute const& VSAttrib : DomainShader->TessellationOutputAttribs.HSOut)
 			{
-				ns::AutoReleased<ns::Array<mtlpp::PipelineBufferDescriptor>> FragmentPipelineBuffers = RenderPipelineDesc.GetFragmentBuffers();
-				uint32 ImmutableBuffers = PixelShader->Bindings.ConstantBuffers | PixelShader->Bindings.ArgumentBuffers;
-				while(ImmutableBuffers)
+				bFoundAttrib = false;
+				for (FMetalAttribute const& HSAttrib : HullShader->TessellationOutputAttribs.HSOut)
 				{
-					uint32 Index = __builtin_ctz(ImmutableBuffers);
-					ImmutableBuffers &= ~(1 << Index);
-					
-					if (Index < ML_MaxBuffers)
+					if (HSAttrib.Semantic == VSAttrib.Semantic)
 					{
-						ns::AutoReleased<mtlpp::PipelineBufferDescriptor> PipelineBuffer = FragmentPipelineBuffers[Index];
-						PipelineBuffer.SetMutability(mtlpp::Mutability::Immutable);
+						uint32 attributeIndex = VSAttrib.Index;
+						
+						mtlpp::VertexFormat format = Formats[(uint8)VSAttrib.Type][VSAttrib.Components-1];
+						check(format != mtlpp::VertexFormat::Invalid); // TODO support more cases
+						Attribs[attributeIndex].SetFormat(format);
+						Attribs[attributeIndex].SetOffset(HSAttrib.Offset);
+						Attribs[attributeIndex].SetBufferIndex(DomainShader->TessellationHSOutBuffer);
+						bFoundAttrib = true;
+						break;
 					}
 				}
-				if (PixelShader->SideTableBinding > 0)
+				check(bFoundAttrib);
+			}
+		}
+		
+		if (DomainShader->TessellationControlPointOutBuffer != UINT_MAX && HullShader->TessellationOutputAttribs.PatchControlPointOutSize)
+		{
+			uint32 bufferIndex = DomainShader->TessellationControlPointOutBuffer;
+			uint32 bufferSize = HullShader->TessellationOutputAttribs.PatchControlPointOutSize;
+			
+			DomainVertexLayouts[bufferIndex].SetStride(bufferSize);
+			DomainVertexLayouts[bufferIndex].SetStepFunction(mtlpp::VertexStepFunction::PerPatchControlPoint);
+			DomainVertexLayouts[bufferIndex].SetStepRate(1);
+			
+			bFoundAttrib = false;
+			for (FMetalAttribute const& VSAttrib : DomainShader->TessellationOutputAttribs.PatchControlPointOut)
+			{
+				bFoundAttrib = false;
+				for (FMetalAttribute const& HSAttrib : HullShader->TessellationOutputAttribs.PatchControlPointOut)
 				{
-					ns::AutoReleased<mtlpp::PipelineBufferDescriptor> PipelineBuffer = FragmentPipelineBuffers[PixelShader->SideTableBinding];
-					PipelineBuffer.SetMutability(mtlpp::Mutability::Immutable);
+					if (HSAttrib.Semantic == VSAttrib.Semantic)
+					{
+						uint32 attributeIndex = VSAttrib.Index;
+						
+						mtlpp::VertexFormat format = Formats[(uint8)VSAttrib.Type][VSAttrib.Components-1];
+						check(format != mtlpp::VertexFormat::Invalid); // TODO support more cases
+						Attribs[attributeIndex].SetFormat(format);
+						Attribs[attributeIndex].SetOffset(HSAttrib.Offset);
+						Attribs[attributeIndex].SetBufferIndex(DomainShader->TessellationControlPointOutBuffer);
+						bFoundAttrib = true;
+						break;
+					}
 				}
+				check(bFoundAttrib);
 			}
 		}
 		
-        FMetalBlendState* BlendState = (FMetalBlendState*)Init.BlendState;
+		DomainRenderPipelineDesc.SetTessellationPartitionMode(GMetalTessellationForcePartitionMode == 0 ? HullShader->TessellationPartitioning : (mtlpp::TessellationPartitionMode)(GMetalTessellationForcePartitionMode - 1));
+		DomainRenderPipelineDesc.SetTessellationFactorStepFunction(mtlpp::TessellationFactorStepFunction::PerPatch);
+		DomainRenderPipelineDesc.SetTessellationOutputWindingOrder(HullShader->TessellationOutputWinding);
+		int FixedMaxTessFactor = (int)RoundTessLevel(HullShader->TessellationMaxTessFactor, DomainRenderPipelineDesc.GetTessellationPartitionMode());
+		DomainRenderPipelineDesc.SetMaxTessellationFactor(FixedMaxTessFactor);
+		DomainRenderPipelineDesc.SetTessellationFactorScaleEnabled(NO);
+		DomainRenderPipelineDesc.SetTessellationFactorFormat(mtlpp::TessellationFactorFormat::Half);
+		DomainRenderPipelineDesc.SetTessellationControlPointIndexType(mtlpp::TessellationControlPointIndexType::None);
+		DomainRenderPipelineDesc.SetVertexDescriptor(DomainVertexDesc);
 		
-		ns::Array<mtlpp::RenderPipelineColorAttachmentDescriptor> ColorAttachments = RenderPipelineDesc.GetColorAttachments();
-#if !PLATFORM_TVOS
-		auto DebugColorAttachements = DebugPipelineDesc.GetColorAttachments();
+#if PLATFORM_TVOS
+		if (!ConfigureRenderPipelineDescriptor(DomainRenderPipelineDesc, Key, Init, IndexType))
+#else
+		if (!ConfigureRenderPipelineDescriptor(DomainRenderPipelineDesc, DebugPipelineDesc, Key, Init, IndexType))
 #endif
+		{
+			[Pipeline release];
+			return nil;
+		}
+		
+		mtlpp::Function domainFunction = DomainShader ? DomainShader->GetFunction() : nil;
 
-		uint32 TargetWidth = 0;
-		for (uint32 i = 0; i < NumActiveTargets; i++)
-        {
-            EPixelFormat TargetFormat = (EPixelFormat)Init.RenderTargetFormats[i];
-			
-			METAL_FATAL_ASSERT(!(TargetFormat == PF_Unknown && PixelShader && (((PixelShader->Bindings.InOutMask & 0x7fff) & (1 << i)))), TEXT("Pipeline pixel shader expects target %u to be bound but it isn't: %s."), i, *FString(PixelShader->GetSourceCode()));
-
-			TargetWidth += GPixelFormats[TargetFormat].BlockBytes;
-            
-            mtlpp::PixelFormat MetalFormat = (mtlpp::PixelFormat)GPixelFormats[TargetFormat].PlatformFormat;
-            uint32 Flags = Init.RenderTargetFlags[i];
-            if (Flags & TexCreate_SRGB)
-            {
-        #if PLATFORM_MAC // Expand as R8_sRGB is iOS only.
-                if (MetalFormat == mtlpp::PixelFormat::R8Unorm)
-                {
-                    MetalFormat = mtlpp::PixelFormat::RGBA8Unorm;
-                }
-        #endif
-                MetalFormat = ToSRGBFormat(MetalFormat);
-            }
-            
-            mtlpp::RenderPipelineColorAttachmentDescriptor Attachment = ColorAttachments[i];
-            Attachment.SetPixelFormat(MetalFormat);
-			
-#if !PLATFORM_TVOS
-			auto DebugAttachment = DebugColorAttachements[i];;
-			DebugAttachment.SetPixelFormat(MetalFormat);
-#endif
-			
-            mtlpp::RenderPipelineColorAttachmentDescriptor Blend = BlendState->RenderTargetStates[i].BlendState;
-            if(TargetFormat != PF_Unknown)
-            {
-                // assign each property manually, would be nice if this was faster
-                Attachment.SetBlendingEnabled(Blend.IsBlendingEnabled());
-                Attachment.SetSourceRgbBlendFactor(Blend.GetSourceRgbBlendFactor());
-                Attachment.SetDestinationRgbBlendFactor(Blend.GetDestinationRgbBlendFactor());
-                Attachment.SetRgbBlendOperation(Blend.GetRgbBlendOperation());
-                Attachment.SetSourceAlphaBlendFactor(Blend.GetSourceAlphaBlendFactor());
-                Attachment.SetDestinationAlphaBlendFactor(Blend.GetDestinationAlphaBlendFactor());
-                Attachment.SetAlphaBlendOperation(Blend.GetAlphaBlendOperation());
-                Attachment.SetWriteMask(Blend.GetWriteMask());
-				
-#if PLATFORM_MAC
-				DebugAttachment.SetBlendingEnabled(Blend.IsBlendingEnabled());
-				DebugAttachment.SetSourceRgbBlendFactor(Blend.GetSourceRgbBlendFactor());
-				DebugAttachment.SetDestinationRgbBlendFactor(Blend.GetDestinationRgbBlendFactor());
-				DebugAttachment.SetRgbBlendOperation(Blend.GetRgbBlendOperation());
-				DebugAttachment.SetSourceAlphaBlendFactor(Blend.GetSourceAlphaBlendFactor());
-				DebugAttachment.SetDestinationAlphaBlendFactor(Blend.GetDestinationAlphaBlendFactor());
-				DebugAttachment.SetAlphaBlendOperation(Blend.GetAlphaBlendOperation());
-				DebugAttachment.SetWriteMask(Blend.GetWriteMask());
-#endif
-            }
-            else
-            {
-                Attachment.SetBlendingEnabled(NO);
-				Attachment.SetWriteMask(mtlpp::ColorWriteMask::None);
-#if PLATFORM_MAC
-				DebugAttachment.SetBlendingEnabled(NO);
-				DebugAttachment.SetWriteMask(mtlpp::ColorWriteMask::None);
+		FMetalPixelShader* PixelShader = (FMetalPixelShader*)Init.BoundShaderState.PixelShaderRHI;
+		mtlpp::Function fragmentFunction = PixelShader ? PixelShader->GetFunction() : nil;
+		
+		DomainRenderPipelineDesc.SetVertexFunction(domainFunction);
+		DomainRenderPipelineDesc.SetFragmentFunction(fragmentFunction);
+#if ENABLE_METAL_GPUPROFILE
+		ns::String DomainName = domainFunction.GetName();
+		ns::String FragmentName = fragmentFunction ? fragmentFunction.GetName() : @"";
+		DomainRenderPipelineDesc.SetLabel([NSString stringWithFormat:@"%@+%@", DomainName.GetPtr(), FragmentName.GetPtr()]);
 #endif
-            }
-        }
 		
-		// don't allow a PSO that is too wide
-		if (!GSupportsWideMRT && TargetWidth > 16)
+		mtlpp::RenderPipelineState& RenderPipelineState = Pipeline->RenderPipelineState;
 		{
-			return nil;
+			ns::AutoReleasedError RenderError;
+			METAL_GPUPROFILE(FScopedMetalCPUStats CPUStat(FString::Printf(TEXT("NewRenderPipeline: %s"), TEXT("")/**FString([RenderPipelineDesc.GetPtr() description])*/)));
+			RenderPipelineState = Device.NewRenderPipelineState(DomainRenderPipelineDesc, (mtlpp::PipelineOption)RenderOption, Reflection, &RenderError);
+			if (Reflection)
+			{
+				Pipeline->RenderPipelineReflection = *Reflection;
+#if METAL_DEBUG_OPTIONS
+				Pipeline->RenderDesc = DomainRenderPipelineDesc;
+#endif
+			}
+			Error = RenderError;
 		}
-        
-        switch(Init.DepthStencilTargetFormat)
-        {
-            case PF_DepthStencil:
-            {
-                mtlpp::PixelFormat MetalFormat = (mtlpp::PixelFormat)GPixelFormats[PF_DepthStencil].PlatformFormat;
-                if(MetalFormat == mtlpp::PixelFormat::Depth32Float)
-                {
-                    if (Init.DepthTargetLoadAction != ERenderTargetLoadAction::ENoAction || Init.DepthTargetStoreAction != ERenderTargetStoreAction::ENoAction)
-                    {
-                        RenderPipelineDesc.SetDepthAttachmentPixelFormat(MetalFormat);
-#if PLATFORM_MAC
-						DebugPipelineDesc.SetDepthAttachmentPixelFormat(MetalFormat);
+		
+		UE_CLOG((RenderPipelineState == nil), LogMetal, Error, TEXT("Failed to generate a pipeline state object: %s"), *FString(Error.GetPtr().description));
+		UE_CLOG((RenderPipelineState == nil), LogMetal, Error, TEXT("Domain shader: %s"), DomainShader ? *FString(DomainShader->GetSourceCode()) : TEXT("NULL"));
+		UE_CLOG((RenderPipelineState == nil), LogMetal, Error, TEXT("Pixel shader: %s"), PixelShader ? *FString(PixelShader->GetSourceCode()) : TEXT("NULL"));
+		UE_CLOG((RenderPipelineState == nil), LogMetal, Error, TEXT("Descriptor: %s"), *FString(DomainRenderPipelineDesc.GetPtr().description));
+		UE_CLOG((RenderPipelineState == nil), LogMetal, Error, TEXT("Failed to generate a render pipeline state object:\n\n %s\n\n"), *FString(Error.GetLocalizedDescription()));
+	}
+	
+	// We need to pass a failure up the chain, so we'll clean up here.
+	if(Pipeline->StreamPipelineState == nil || Pipeline->ComputePipelineState == nil || Pipeline->RenderPipelineState == nil)
+	{
+		[Pipeline release];
+		return nil;
+	}
+	
+#if METAL_DEBUG_OPTIONS
+	Pipeline->ComputeSource = HullShader->GetSourceCode();
+	Pipeline->VertexSource = VertexShader->GetSourceCode();
+	Pipeline->DomainSource = DomainShader->GetSourceCode();
+	FMetalPixelShader* PixelShader = (FMetalPixelShader*)Init.BoundShaderState.PixelShaderRHI;
+	Pipeline->FragmentSource = PixelShader ? PixelShader->GetSourceCode() : nil;
 #endif
-                    }
-                    if (Init.StencilTargetLoadAction != ERenderTargetLoadAction::ENoAction || Init.StencilTargetStoreAction != ERenderTargetStoreAction::ENoAction)
-                    {
-                        RenderPipelineDesc.SetStencilAttachmentPixelFormat(mtlpp::PixelFormat::Stencil8);
+	
+#if !PLATFORM_TVOS
+	if (GMetalCommandBufferDebuggingEnabled)
+	{
 #if PLATFORM_MAC
-						DebugPipelineDesc.SetStencilAttachmentPixelFormat(mtlpp::PixelFormat::Stencil8);
+		DebugPipelineDesc.SetVertexFunction(FMetalHelperFunctions::Get().GetDebugFunction());
+		DebugPipelineDesc.SetRasterizationEnabled(false);
+#else
+		DebugPipelineDesc.SetTileFunction(FMetalHelperFunctions::Get().GetDebugFunction());
+		DebugPipelineDesc.SetRasterSampleCount(RenderPipelineDesc.GetSampleCount());
+		DebugPipelineDesc.SetThreadgroupSizeMatchesTileSize(false);
 #endif
-                    }
-                }
-                else
-                {
-                    RenderPipelineDesc.SetDepthAttachmentPixelFormat(MetalFormat);
-                    RenderPipelineDesc.SetStencilAttachmentPixelFormat(MetalFormat);
-#if PLATFORM_MAC
-					DebugPipelineDesc.SetDepthAttachmentPixelFormat(MetalFormat);
-					DebugPipelineDesc.SetStencilAttachmentPixelFormat(MetalFormat);
+#if ENABLE_METAL_GPUPROFILE
+		DebugPipelineDesc.SetLabel(@"Main_Debug");
 #endif
-                }
-                break;
-            }
-            case PF_ShadowDepth:
-            {
-                RenderPipelineDesc.SetDepthAttachmentPixelFormat((mtlpp::PixelFormat)GPixelFormats[PF_ShadowDepth].PlatformFormat);
-#if PLATFORM_MAC
-				DebugPipelineDesc.SetDepthAttachmentPixelFormat((mtlpp::PixelFormat)GPixelFormats[PF_ShadowDepth].PlatformFormat);
+		
+		ns::AutoReleasedError RenderError;
+		METAL_GPUPROFILE(FScopedMetalCPUStats CPUStat(FString::Printf(TEXT("NewDebugPipeline: %s"), TEXT("")/**FString([RenderPipelineDesc.GetPtr() description])*/)));
+		Pipeline->DebugPipelineState = Device.NewRenderPipelineState(DebugPipelineDesc, mtlpp::PipelineOption::NoPipelineOption, Reflection, nullptr);
+	}
 #endif
-                break;
-            }
-            default:
-            {
-                break;
-            }
-        }
-        
-        check(Init.BoundShaderState.VertexShaderRHI != nullptr);
-#if PLATFORM_SUPPORTS_GEOMETRY_SHADERS
-        check(Init.BoundShaderState.GeometryShaderRHI == nullptr);
+	
+#if METAL_DEBUG_OPTIONS
+	if (GFrameCounter > 3)
+	{
+		UE_LOG(LogMetal, Verbose, TEXT("Created a hitchy pipeline state for hash %llx %llx %llx"), (uint64)Key.RenderPipelineHash.RasterBits, (uint64)(Key.RenderPipelineHash.TargetBits), (uint64)Key.VertexDescriptorHash.VertexDescHash);
+	}
 #endif
 
-        if(RenderPipelineDesc.GetDepthAttachmentPixelFormat() == mtlpp::PixelFormat::Invalid && PixelShader && ((PixelShader->Bindings.InOutMask & 0x8000) || (NumActiveTargets == 0 && (PixelShader->Bindings.NumUAVs > 0))))
-        {
-            RenderPipelineDesc.SetDepthAttachmentPixelFormat((mtlpp::PixelFormat)GPixelFormats[PF_DepthStencil].PlatformFormat);
-            RenderPipelineDesc.SetStencilAttachmentPixelFormat((mtlpp::PixelFormat)GPixelFormats[PF_DepthStencil].PlatformFormat);
-			
+	if (Pipeline && SafeGetRuntimeDebuggingLevel() >= EMetalDebugLevelFastValidation)
+	{
+		[Pipeline initResourceMask];
+	}
+	
+	return Pipeline;
+}
+#endif
+
+static FMetalShaderPipeline* CreateMTLRenderPipeline(bool const bSync, FMetalGraphicsPipelineKey const& Key, const FGraphicsPipelineStateInitializer& Init, EMetalIndexType const IndexType)
+{
+    FMetalVertexShader* VertexShader = (FMetalVertexShader*)Init.BoundShaderState.VertexShaderRHI;
+    FMetalPixelShader* PixelShader = (FMetalPixelShader*)Init.BoundShaderState.PixelShaderRHI;
+    
+    mtlpp::Function vertexFunction = VertexShader->GetFunction();
+    mtlpp::Function fragmentFunction = PixelShader ? PixelShader->GetFunction() : nil;
+
+#if PLATFORM_SUPPORTS_TESSELLATION_SHADERS
+	FMetalDomainShader* DomainShader = (FMetalDomainShader*)Init.BoundShaderState.DomainShaderRHI;
+    mtlpp::Function domainFunction = DomainShader ? DomainShader->GetFunction() : nil;
+	if (DomainShader && FMetalCommandQueue::SupportsFeature(EMetalFeaturesSeparateTessellation))
+	{
+		return CreateSeparateMetalTessellationPipeline(bSync, Key, Init, IndexType);
+	}
+#endif
+    
+    FMetalShaderPipeline* Pipeline = nil;
+    if (vertexFunction && ((PixelShader != nullptr) == (fragmentFunction != nil))
+#if PLATFORM_SUPPORTS_TESSELLATION_SHADERS
+		&& ((DomainShader != nullptr) == (domainFunction != nil))
+#endif
+	)
+    {
+		ns::Error Error;
+		mtlpp::Device Device = GetMetalDeviceContext().GetDevice();
+
+		uint32 const NumActiveTargets = Init.ComputeNumValidRenderTargets();
+        check(NumActiveTargets <= MaxSimultaneousRenderTargets);
+		
+		Pipeline = [FMetalShaderPipeline new];
+		METAL_DEBUG_OPTION(FMemory::Memzero(Pipeline->ResourceMask, sizeof(Pipeline->ResourceMask)));
+
+		mtlpp::RenderPipelineDescriptor RenderPipelineDesc;
+        mtlpp::ComputePipelineDescriptor ComputePipelineDesc(nil);
 #if PLATFORM_MAC
-			DebugPipelineDesc.SetDepthAttachmentPixelFormat((mtlpp::PixelFormat)GPixelFormats[PF_DepthStencil].PlatformFormat);
-			DebugPipelineDesc.SetStencilAttachmentPixelFormat((mtlpp::PixelFormat)GPixelFormats[PF_DepthStencil].PlatformFormat);
+        mtlpp::RenderPipelineDescriptor DebugPipelineDesc;
+#elif !PLATFORM_TVOS
+        mtlpp::TileRenderPipelineDescriptor DebugPipelineDesc;
 #endif
-        }
-        
-        static bool bNoMSAA = FParse::Param(FCommandLine::Get(), TEXT("nomsaa"));
-        RenderPipelineDesc.SetSampleCount(!bNoMSAA ? FMath::Max(Init.NumSamples, (uint16)1u) : (uint16)1u);
-    #if PLATFORM_MAC
-        RenderPipelineDesc.SetInputPrimitiveTopology(TranslatePrimitiveTopology(Init.PrimitiveType));
-		DebugPipelineDesc.SetSampleCount(!bNoMSAA ? FMath::Max(Init.NumSamples, (uint16)1u) : (uint16)1u);
-		DebugPipelineDesc.SetInputPrimitiveTopology(mtlpp::PrimitiveTopologyClass::Point);
-    #endif
+		
+#if PLATFORM_TVOS
+		if (!ConfigureRenderPipelineDescriptor(RenderPipelineDesc, Key, Init, IndexType))
+#else
+		if (!ConfigureRenderPipelineDescriptor(RenderPipelineDesc, DebugPipelineDesc, Key, Init, IndexType))
+#endif
+		{
+			return nil;
+		}
         
         FMetalVertexDeclaration* VertexDecl = (FMetalVertexDeclaration*)Init.BoundShaderState.VertexDeclarationRHI;
 		
@@ -944,8 +1425,6 @@ static FMetalShaderPipeline* CreateMTLRenderPipeline(bool const bSync, FMetalGra
             TessellationDesc.TessellationIndexBufferIndex = VertexShader->TessellationIndexBuffer;
             TessellationDesc.TessellationPatchConstOutSize = VertexShader->TessellationOutputAttribs.HSOutSize;
 			TessellationDesc.TessellationControlPointIndexBufferIndex = VertexShader->TessellationControlPointIndexBuffer;
-			TessellationDesc.DomainVertexDescriptor = DomainVertexDesc;
-            TessellationDesc.DSNumUniformBuffers = DomainShader->Bindings.NumUniformBuffers;
             TessellationDesc.TessellationPatchControlPointOutSize = VertexShader->TessellationOutputAttribs.PatchControlPointOutSize;
             TessellationDesc.TessellationTessFactorOutSize = VertexShader->TessellationOutputAttribs.HSTFOutSize;
             
@@ -955,14 +1434,6 @@ static FMetalShaderPipeline* CreateMTLRenderPipeline(bool const bSync, FMetalGra
             check(TessellationDesc.TessellationTessFactorOutSize == 2*4 || TessellationDesc.TessellationTessFactorOutSize == 2*6);
             
             mtlpp::VertexStepFunction stepFunction = mtlpp::VertexStepFunction::PerPatch;
-            
-            static mtlpp::VertexFormat Formats[(uint8)EMetalComponentType::Max][4] = {
-                {mtlpp::VertexFormat::UInt, mtlpp::VertexFormat::UInt2, mtlpp::VertexFormat::UInt3, mtlpp::VertexFormat::UInt4},
-                {mtlpp::VertexFormat::Int, mtlpp::VertexFormat::Int2, mtlpp::VertexFormat::Int3, mtlpp::VertexFormat::Int4},
-                {mtlpp::VertexFormat::Invalid, mtlpp::VertexFormat::Half2, mtlpp::VertexFormat::Half3, mtlpp::VertexFormat::Half4},
-                {mtlpp::VertexFormat::Float, mtlpp::VertexFormat::Float2, mtlpp::VertexFormat::Float3, mtlpp::VertexFormat::Float4},
-                {mtlpp::VertexFormat::Invalid, mtlpp::VertexFormat::UChar2, mtlpp::VertexFormat::UChar3, mtlpp::VertexFormat::UChar4},
-            };
 			
 			ns::Array<mtlpp::VertexBufferLayoutDescriptor> DomainVertexLayouts = DomainVertexDesc.GetLayouts();
 			
diff --git a/Engine/Source/Runtime/Apple/MetalRHI/Private/MetalPipeline.h b/Engine/Source/Runtime/Apple/MetalRHI/Private/MetalPipeline.h
--- a/Engine/Source/Runtime/Apple/MetalRHI/Private/MetalPipeline.h
+++ b/Engine/Source/Runtime/Apple/MetalRHI/Private/MetalPipeline.h
@@ -50,10 +50,8 @@ enum EMetalPipelineHashOffsets
 struct FMetalTessellationPipelineDesc
 {
 	FMetalTessellationPipelineDesc()
-	: DomainVertexDescriptor(nil)
 	{
 	}
-	mtlpp::VertexDescriptor DomainVertexDescriptor;
 	NSUInteger TessellationInputControlPointBufferIndex;
 	NSUInteger TessellationOutputControlPointBufferIndex;
 	NSUInteger TessellationPatchControlPointOutSize;
@@ -65,7 +63,6 @@ struct FMetalTessellationPipelineDesc
 	NSUInteger TessellationPatchCountBufferIndex;
 	NSUInteger TessellationControlPointIndexBufferIndex;
 	NSUInteger TessellationIndexBufferIndex;
-	NSUInteger DSNumUniformBuffers; // DEBUG ONLY
 };
 
 @interface FMetalShaderPipeline : FApplePlatformObject
@@ -73,18 +70,22 @@ struct FMetalTessellationPipelineDesc
 @public
 	mtlpp::RenderPipelineState RenderPipelineState;
 	mtlpp::ComputePipelineState ComputePipelineState;
+	mtlpp::RenderPipelineState StreamPipelineState;
 	mtlpp::RenderPipelineState DebugPipelineState;
 	FMetalTessellationPipelineDesc TessellationPipelineDesc;
 	TArray<uint32> BufferDataSizes[EMetalShaderStagesNum];
 	TMap<uint8, uint8> TextureTypes[EMetalShaderStagesNum];
 	FMetalDebugShaderResourceMask ResourceMask[EMetalShaderStagesNum];
 	mtlpp::RenderPipelineReflection RenderPipelineReflection;
+	mtlpp::RenderPipelineReflection StreamPipelineReflection;
 	mtlpp::ComputePipelineReflection ComputePipelineReflection;
 #if METAL_DEBUG_OPTIONS
 	ns::String VertexSource;
+	ns::String DomainSource;
 	ns::String FragmentSource;
 	ns::String ComputeSource;
 	mtlpp::RenderPipelineDescriptor RenderDesc;
+	mtlpp::RenderPipelineDescriptor StreamDesc;
 	mtlpp::ComputePipelineDescriptor ComputeDesc;
 #endif
 }
diff --git a/Engine/Source/Runtime/Apple/MetalRHI/Private/MetalRHIPrivate.h b/Engine/Source/Runtime/Apple/MetalRHI/Private/MetalRHIPrivate.h
--- a/Engine/Source/Runtime/Apple/MetalRHI/Private/MetalRHIPrivate.h
+++ b/Engine/Source/Runtime/Apple/MetalRHI/Private/MetalRHIPrivate.h
@@ -154,6 +154,12 @@ struct FMetalDebugInfo
 // Get a compute pipeline state used to implement some debug features.
 mtlpp::ComputePipelineState GetMetalDebugComputeState();
 
+// Get the copy 32-bit indices function
+mtlpp::ComputePipelineState GetMetalCopyIndex32Function();
+
+// Get the copy 16-bit indices function
+mtlpp::ComputePipelineState GetMetalCopyIndex16Function();
+
 // Access the internal context for the device-owning DynamicRHI object
 FMetalDeviceContext& GetMetalDeviceContext();
 
diff --git a/Engine/Source/Runtime/Apple/MetalRHI/Private/MetalRenderPass.cpp b/Engine/Source/Runtime/Apple/MetalRHI/Private/MetalRenderPass.cpp
--- a/Engine/Source/Runtime/Apple/MetalRHI/Private/MetalRenderPass.cpp
+++ b/Engine/Source/Runtime/Apple/MetalRHI/Private/MetalRenderPass.cpp
@@ -31,6 +31,12 @@ static FAutoConsoleVariableRef CVarMetalTessellationRunDomainStage(
 	GMetalTessellationRunDomainStage,
 	TEXT("Whether to run the DS+PS domain stage when performing tessellated draw calls in Metal or not. (Default: 1)"));
 
+static int32 GMetalTessellationFlattenTriangleFactors = 0;
+static FAutoConsoleVariableRef CVarMetalTessellationFlattenTriangleFactors(
+	TEXT("rhi.Metal.FlattenTriangleFactors"),
+	GMetalTessellationFlattenTriangleFactors,
+	TEXT("Whether to flatten the tessellation factors to 1 when performing tessellated draw calls in Metal or not - used to debug tessellation factor generation. (Default: 0)"));
+
 static int32 GMetalDeferRenderPasses = 1;
 static FAutoConsoleVariableRef CVarMetalDeferRenderPasses(
 	TEXT("rhi.Metal.DeferRenderPasses"),
@@ -637,11 +643,336 @@ void FMetalRenderPass::DrawIndexedPrimitiveIndirect(uint32 PrimitiveType,FMetalI
 }
 
 #if PLATFORM_SUPPORTS_TESSELLATION_SHADERS
+extern mtlpp::ComputePipelineState GetMetalFlattenTessState();
+
 void FMetalRenderPass::DrawPatches(uint32 PrimitiveType,FMetalBuffer const& IndexBuffer, uint32 IndexBufferStride, int32 BaseVertexIndex, uint32 FirstInstance, uint32 StartIndex,
 									uint32 NumPrimitives, uint32 NumInstances)
 {
 	if (GetMetalDeviceContext().SupportsFeature(EMetalFeaturesTessellation))
 	{
+		if (FMetalCommandQueue::SupportsFeature(EMetalFeaturesSeparateTessellation))
+		{
+			FMetalDeviceContext& deviceContext = (FMetalDeviceContext&)GetMetalDeviceContext();
+			mtlpp::Device device = deviceContext.GetDevice();
+			
+			FMetalGraphicsPipelineState* boundShaderState = State.GetGraphicsPSO();
+			FMetalShaderPipeline* Pipeline = State.GetPipelineState();
+			
+			ConditionalSwitchToSeparateTessellation();
+			check(CurrentEncoder.GetCommandBuffer());
+			check(PrologueEncoder.GetCommandBuffer());
+			check(PrologueEncoder.IsRenderCommandEncoderActive());
+			check(CurrentEncoder.IsRenderCommandEncoderActive());
+			
+			mtlpp::RenderCommandEncoder& StreamOutEncoder = PrologueEncoder.GetRenderCommandEncoder();
+			
+			PrepareToStreamOut(PrimitiveType);
+			
+			FMetalBuffer ShaderOutputBuffer = nil;
+			if (IndexBuffer)
+			{
+				uint32 NumIndices = GetVertexCountForPrimitiveCount(NumPrimitives, PrimitiveType);
+				
+				uint32 OutputBufferSize = (NumIndices * NumInstances) * boundShaderState->VertexShader->TessellationOutputAttribs.HSOutSize;
+				
+				ShaderOutputBuffer = deviceContext.CreatePooledBuffer(FMetalPooledBufferArgs(device, OutputBufferSize, BUF_Dynamic, mtlpp::StorageMode::Private));
+				FMetalCommandBufferDebugHelpers::TrackResource(CurrentEncoder.GetCommandBuffer().GetPtr(), ShaderOutputBuffer.GetPtr());
+				
+				PrologueEncoder.SetShaderBuffer(mtlpp::FunctionType::Vertex, ShaderOutputBuffer, 0, ShaderOutputBuffer, boundShaderState->VertexShader->TessellationHSOutBuffer, mtlpp::ResourceUsage::Write);
+				State.SetShaderBuffer(EMetalShaderStages::Vertex, nil, nil, 0, 0, boundShaderState->VertexShader->TessellationHSOutBuffer, mtlpp::ResourceUsage(0));
+				
+				// set the patchCount
+				FMetalSubBufferRing& Ring = PrologueEncoder.GetRingBuffer();
+				FMetalBuffer PatchData = Ring.NewBuffer(sizeof(NumIndices), sizeof(NumIndices));
+				uint32* patchCountData = (uint32*)PatchData.GetContents();
+				patchCountData[0] = NumIndices;
+				PrologueEncoder.SetShaderBuffer(mtlpp::FunctionType::Vertex, PatchData, 0, PatchData, boundShaderState->VertexShader->TessellationPatchCountBuffer, mtlpp::ResourceUsage::Read);
+				State.SetShaderBuffer(EMetalShaderStages::Vertex, nil, nil, 0, 0, boundShaderState->VertexShader->TessellationPatchCountBuffer, mtlpp::ResourceUsage(0));
+				
+				if (boundShaderState->VertexShader->SideTableBinding >= 0)
+				{
+					PrologueEncoder.SetShaderSideTable(mtlpp::FunctionType::Vertex, boundShaderState->VertexShader->SideTableBinding);
+					State.SetShaderBuffer(EMetalShaderStages::Vertex, nil, nil, 0, 0, boundShaderState->VertexShader->SideTableBinding, mtlpp::ResourceUsage(0));
+				}
+				
+				StreamOutEncoder.DrawIndexed(TranslatePrimitiveType(PrimitiveType), NumIndices, ((IndexBufferStride == 2) ? mtlpp::IndexType::UInt16 : mtlpp::IndexType::UInt32), IndexBuffer, IndexBufferStride * StartIndex, NumInstances, BaseVertexIndex, FirstInstance);
+				METAL_DEBUG_LAYER(EMetalDebugLevelFastValidation, PrologueEncoder.GetRenderCommandEncoderDebugging().DrawIndexed(TranslatePrimitiveType(PrimitiveType), NumIndices, ((IndexBufferStride == 2) ? mtlpp::IndexType::UInt16 : mtlpp::IndexType::UInt32), IndexBuffer, IndexBufferStride * StartIndex, NumInstances, BaseVertexIndex, FirstInstance));
+			}
+			else
+			{
+				uint32 NumVertices = GetVertexCountForPrimitiveCount(NumPrimitives, PrimitiveType);
+				
+				uint32 OutputBufferSize = (NumVertices * NumInstances) * boundShaderState->VertexShader->TessellationOutputAttribs.HSOutSize;
+				
+				ShaderOutputBuffer = deviceContext.CreatePooledBuffer(FMetalPooledBufferArgs(device, OutputBufferSize, BUF_Dynamic, mtlpp::StorageMode::Private));
+				FMetalCommandBufferDebugHelpers::TrackResource(CurrentEncoder.GetCommandBuffer().GetPtr(), ShaderOutputBuffer.GetPtr());
+
+				PrologueEncoder.SetShaderBuffer(mtlpp::FunctionType::Vertex, ShaderOutputBuffer, 0, ShaderOutputBuffer, boundShaderState->VertexShader->TessellationHSOutBuffer, mtlpp::ResourceUsage::Write);
+				State.SetShaderBuffer(EMetalShaderStages::Vertex, nil, nil, 0, 0, boundShaderState->VertexShader->TessellationHSOutBuffer, mtlpp::ResourceUsage(0));
+				
+				// set the patchCount
+				FMetalSubBufferRing& Ring = PrologueEncoder.GetRingBuffer();
+				FMetalBuffer PatchData = Ring.NewBuffer(sizeof(NumVertices), sizeof(NumVertices));
+				uint32* patchCountData = (uint32*)PatchData.GetContents();
+				patchCountData[0] = NumVertices;
+				PrologueEncoder.SetShaderBuffer(mtlpp::FunctionType::Vertex, PatchData, 0, PatchData, boundShaderState->VertexShader->TessellationPatchCountBuffer, mtlpp::ResourceUsage::Read);
+				State.SetShaderBuffer(EMetalShaderStages::Vertex, nil, nil, 0, 0, boundShaderState->VertexShader->TessellationPatchCountBuffer, mtlpp::ResourceUsage(0));
+				
+				if (boundShaderState->VertexShader->SideTableBinding >= 0)
+				{
+					PrologueEncoder.SetShaderSideTable(mtlpp::FunctionType::Vertex, boundShaderState->VertexShader->SideTableBinding);
+					State.SetShaderBuffer(EMetalShaderStages::Vertex, nil, nil, 0, 0, boundShaderState->VertexShader->SideTableBinding, mtlpp::ResourceUsage(0));
+				}
+				
+				METAL_GPUPROFILE(FMetalProfiler::GetProfiler()->EncodeDraw(CurrentEncoder.GetCommandBufferStats(), __FUNCTION__, NumPrimitives, NumVertices, NumInstances));
+				StreamOutEncoder.Draw(TranslatePrimitiveType(PrimitiveType), BaseVertexIndex, NumVertices, NumInstances, FirstInstance);
+				METAL_DEBUG_LAYER(EMetalDebugLevelFastValidation, PrologueEncoder.GetRenderCommandEncoderDebugging().Draw(TranslatePrimitiveType(PrimitiveType), BaseVertexIndex, NumVertices, NumInstances, FirstInstance));
+			}
+			
+			ConditionalSwitchToAsyncCompute();
+			check(CurrentEncoder.GetCommandBuffer());
+			check(PrologueEncoder.GetCommandBuffer());
+			check(PrologueEncoder.IsComputeCommandEncoderActive());
+			check(CurrentEncoder.IsRenderCommandEncoderActive());
+			
+			auto& computeEncoder = PrologueEncoder.GetComputeCommandEncoder();
+			
+			uint32 NumVertices = GetVertexCountForPrimitiveCount(NumPrimitives, PrimitiveType);
+			
+			uint32 InputCPs = boundShaderState->HullShader->TessellationInputControlPoints;
+			uint32 OutputCPs = boundShaderState->HullShader->TessellationOutputControlPoints;
+			
+			uint32 PatchCount = (NumVertices + (InputCPs - 1)) / InputCPs;
+			
+			FMetalBuffer HullIndexBuffer;
+			if (IndexBuffer)
+			{
+				HullIndexBuffer = deviceContext.CreatePooledBuffer(FMetalPooledBufferArgs(device, NumInstances * PatchCount * FMath::Max(OutputCPs, InputCPs) * sizeof(uint32), BUF_Static, mtlpp::StorageMode::Private));
+				FMetalCommandBufferDebugHelpers::TrackResource(CurrentEncoder.GetCommandBuffer().GetPtr(), HullIndexBuffer.GetPtr());
+				if (IndexBufferStride == 2)
+				{
+					computeEncoder.SetComputePipelineState(GetMetalCopyIndex16Function());
+				}
+				else
+				{
+					computeEncoder.SetComputePipelineState(GetMetalCopyIndex32Function());
+				}
+				
+				PrologueEncoder.SetShaderBuffer(mtlpp::FunctionType::Kernel, IndexBuffer, 0, IndexBuffer, 0, mtlpp::ResourceUsage::Read);
+				State.SetShaderBuffer(EMetalShaderStages::Compute, nil, nil, 0, 0, 0, mtlpp::ResourceUsage(0));
+				
+				PrologueEncoder.SetShaderBuffer(mtlpp::FunctionType::Kernel, HullIndexBuffer, 0, HullIndexBuffer, 1, mtlpp::ResourceUsage::Write);
+				State.SetShaderBuffer(EMetalShaderStages::Compute, nil, nil, 0, 0, 1, mtlpp::ResourceUsage(0));
+				
+				uint32 patchCountData[] = { InputCPs, OutputCPs };
+				if (InputCPs > OutputCPs)
+					patchCountData[1] = InputCPs;
+				
+				PrologueEncoder.SetShaderBytes(mtlpp::FunctionType::Kernel, (const uint8*)&patchCountData[0], sizeof(patchCountData), 2);
+				State.SetShaderBuffer(EMetalShaderStages::Compute, nil, nil, 0, 0, 2, mtlpp::ResourceUsage(0));
+				
+				mtlpp::DrawIndexedPrimitivesIndirectArguments Args;
+				Args.IndexCount = NumVertices;
+				Args.InstanceCount = NumInstances;
+				Args.IndexStart = StartIndex;
+				PrologueEncoder.SetShaderBytes(mtlpp::FunctionType::Kernel, (const uint8*)&Args, sizeof(Args), 3);
+				State.SetShaderBuffer(EMetalShaderStages::Compute, nil, nil, 0, 0, 3, mtlpp::ResourceUsage(0));
+				
+				if(GMetalTessellationRunTessellationStage)
+				{
+					METAL_GPUPROFILE(FMetalProfiler::GetProfiler()->EncodeDispatch(PrologueEncoder.GetCommandBufferStats(), __FUNCTION__));
+					computeEncoder.DispatchThreadgroups(mtlpp::Size(PatchCount, 1, 1), mtlpp::Size(FMath::Max(InputCPs, OutputCPs), NumInstances, 1));
+					METAL_DEBUG_LAYER(EMetalDebugLevelFastValidation, PrologueEncoder.GetComputeCommandEncoderDebugging().DispatchThreadgroups(mtlpp::Size(PatchCount, 1, 1), mtlpp::Size(FMath::Max(InputCPs, OutputCPs), NumInstances, 1)));
+				}
+			}
+			else
+			{
+				HullIndexBuffer = deviceContext.CreatePooledBuffer(FMetalPooledBufferArgs(device, NumInstances * PatchCount * FMath::Max(OutputCPs, InputCPs) * sizeof(uint32), BUF_Static, mtlpp::StorageMode::Shared));
+				FMetalCommandBufferDebugHelpers::TrackResource(CurrentEncoder.GetCommandBuffer().GetPtr(), HullIndexBuffer.GetPtr());
+				uint32* Data = (uint32*)HullIndexBuffer.GetContents();
+				uint32 Index = 0;
+				if (OutputCPs > InputCPs)
+				{
+					for (uint32 i = 0; i < NumInstances; i++)
+					{
+						for (uint32 j = 0; j < PatchCount * OutputCPs; j++)
+						{
+							if ((j % OutputCPs) < InputCPs)
+							{
+								Data[j] = Index++;
+							}
+							else
+							{
+								Data[j] = Index;
+							}
+						}
+					}
+				}
+				else
+				{
+					for (uint32 i = 0; i < NumInstances; i++)
+					{
+						for (uint32 j = 0; j < PatchCount; j++)
+						{
+							for (uint32 k = 0; k < InputCPs; k++)
+							{
+								Data[(i * NumVertices) + j * InputCPs + k] = (i * NumVertices) + j * InputCPs + k;
+							}
+						}
+					}
+				}
+			}
+			
+			PrepareToSeparateTessellate(PrimitiveType);
+			
+			if (HullIndexBuffer)
+			{
+				PrologueEncoder.SetShaderBuffer(mtlpp::FunctionType::Kernel, HullIndexBuffer, 0, HullIndexBuffer, boundShaderState->HullShader->TessellationIndexBuffer, mtlpp::ResourceUsage::Read);
+				State.SetShaderBuffer(EMetalShaderStages::Compute, nil, nil, 0, 0, boundShaderState->HullShader->TessellationIndexBuffer, mtlpp::ResourceUsage(0));
+			}
+			else
+			{
+				check(IndexBuffer);
+				PrologueEncoder.SetShaderBuffer(mtlpp::FunctionType::Kernel, IndexBuffer, 0, IndexBuffer, boundShaderState->HullShader->TessellationIndexBuffer, mtlpp::ResourceUsage::Read);
+				State.SetShaderBuffer(EMetalShaderStages::Compute, nil, nil, 0, 0, boundShaderState->HullShader->TessellationIndexBuffer, mtlpp::ResourceUsage(0));
+			}
+			
+			computeEncoder.SetStageInRegion(mtlpp::Region(0, FMath::Max(NumVertices, OutputCPs * PatchCount) * NumInstances));
+			
+			FMetalBuffer HSOutBuffer;
+			if (boundShaderState->HullShader->TessellationHSOutBuffer != UINT_MAX && boundShaderState->HullShader->TessellationOutputAttribs.HSOutSize)
+			{
+				HSOutBuffer = deviceContext.CreatePooledBuffer(FMetalPooledBufferArgs(device, Align(NumInstances * PatchCount * boundShaderState->HullShader->TessellationOutputAttribs.HSOutSize, BufferOffsetAlignment), BUF_Static, mtlpp::StorageMode::Private));
+				FMetalCommandBufferDebugHelpers::TrackResource(CurrentEncoder.GetCommandBuffer().GetPtr(), HSOutBuffer.GetPtr());
+				
+				PrologueEncoder.SetShaderBuffer(mtlpp::FunctionType::Kernel, HSOutBuffer, 0, HSOutBuffer, boundShaderState->HullShader->TessellationHSOutBuffer, mtlpp::ResourceUsage::Write);
+				State.SetShaderBuffer(EMetalShaderStages::Compute, nil, nil, 0, 0, boundShaderState->HullShader->TessellationHSOutBuffer, mtlpp::ResourceUsage(0));
+			}
+			
+			FMetalBuffer CPOutBuffer;
+			if (boundShaderState->HullShader->TessellationControlPointOutBuffer != UINT_MAX && boundShaderState->HullShader->TessellationOutputAttribs.PatchControlPointOutSize)
+			{
+				CPOutBuffer = deviceContext.CreatePooledBuffer(FMetalPooledBufferArgs(device, NumInstances * PatchCount * OutputCPs * boundShaderState->HullShader->TessellationOutputAttribs.PatchControlPointOutSize, BUF_Static, mtlpp::StorageMode::Private));
+				FMetalCommandBufferDebugHelpers::TrackResource(CurrentEncoder.GetCommandBuffer().GetPtr(), CPOutBuffer.GetPtr());
+				
+				PrologueEncoder.SetShaderBuffer(mtlpp::FunctionType::Kernel, CPOutBuffer, 0, CPOutBuffer, boundShaderState->HullShader->TessellationControlPointOutBuffer, mtlpp::ResourceUsage::Write);
+				State.SetShaderBuffer(EMetalShaderStages::Compute, nil, nil, 0, 0, boundShaderState->HullShader->TessellationControlPointOutBuffer, mtlpp::ResourceUsage(0));
+			}
+			
+			FMetalBuffer HSTFBuffer;
+			if (boundShaderState->HullShader->TessellationHSTFOutBuffer != UINT_MAX)
+			{
+				HSTFBuffer = deviceContext.CreatePooledBuffer(FMetalPooledBufferArgs(device, NumInstances * PatchCount * boundShaderState->HullShader->TessellationOutputAttribs.HSTFOutSize, BUF_Static, mtlpp::StorageMode::Private));
+				FMetalCommandBufferDebugHelpers::TrackResource(CurrentEncoder.GetCommandBuffer().GetPtr(), HSTFBuffer.GetPtr());
+				
+				PrologueEncoder.SetShaderBuffer(mtlpp::FunctionType::Kernel, HSTFBuffer, 0, HSTFBuffer, boundShaderState->HullShader->TessellationHSTFOutBuffer, mtlpp::ResourceUsage::Write);
+				State.SetShaderBuffer(EMetalShaderStages::Compute, nil, nil, 0, 0, boundShaderState->HullShader->TessellationHSTFOutBuffer, mtlpp::ResourceUsage(0));
+			}
+			
+			PrologueEncoder.SetShaderBuffer(mtlpp::FunctionType::Kernel, ShaderOutputBuffer, 0, ShaderOutputBuffer, boundShaderState->HullShader->TessellationControlPointIndexBuffer, mtlpp::ResourceUsage::Read);
+			State.SetShaderBuffer(EMetalShaderStages::Compute, nil, nil, 0, 0, boundShaderState->HullShader->TessellationControlPointIndexBuffer, mtlpp::ResourceUsage(0));
+			
+			FMetalSubBufferRing& Ring = PrologueEncoder.GetRingBuffer();
+			FMetalBuffer PatchData = Ring.NewBuffer(sizeof(InputCPs), sizeof(InputCPs));
+			uint32* patchCountData = (uint32*)PatchData.GetContents();
+			patchCountData[0] = InputCPs;
+			PrologueEncoder.SetShaderBuffer(mtlpp::FunctionType::Kernel, PatchData, 0, PatchData, boundShaderState->HullShader->TessellationPatchCountBuffer, mtlpp::ResourceUsage::Read);
+			State.SetShaderBuffer(EMetalShaderStages::Compute, nil, nil, 0, 0, boundShaderState->HullShader->TessellationPatchCountBuffer, mtlpp::ResourceUsage(0));
+			
+			if (boundShaderState->HullShader->SideTableBinding >= 0)
+			{
+				PrologueEncoder.SetShaderSideTable(mtlpp::FunctionType::Kernel, boundShaderState->HullShader->SideTableBinding);
+				State.SetShaderBuffer(EMetalShaderStages::Compute, nil, nil, 0, 0, boundShaderState->HullShader->SideTableBinding, mtlpp::ResourceUsage(0));
+			}
+			
+			computeEncoder.SetThreadgroupMemory(Align(boundShaderState->VertexShader->TessellationOutputAttribs.HSOutSize * InputCPs, 16), 0);
+			
+			if(GMetalTessellationRunTessellationStage)
+			{
+				METAL_GPUPROFILE(FMetalProfiler::GetProfiler()->EncodeDispatch(PrologueEncoder.GetCommandBufferStats(), __FUNCTION__));
+				computeEncoder.DispatchThreadgroups(mtlpp::Size(NumInstances * PatchCount, 1, 1), mtlpp::Size(FMath::Max(InputCPs, OutputCPs), 1, 1));
+				METAL_DEBUG_LAYER(EMetalDebugLevelFastValidation, PrologueEncoder.GetComputeCommandEncoderDebugging().DispatchThreadgroups(mtlpp::Size(NumInstances * PatchCount, 1, 1), mtlpp::Size(FMath::Max(InputCPs, OutputCPs), 1, 1)));
+			}
+			
+			computeEncoder.SetThreadgroupMemory(0, 0);
+			
+			if (HSTFBuffer && GMetalTessellationFlattenTriangleFactors)
+			{
+				computeEncoder.SetComputePipelineState(GetMetalFlattenTessState());
+				PrologueEncoder.SetShaderBuffer(mtlpp::FunctionType::Kernel, HSTFBuffer, 0, HSTFBuffer, 0, mtlpp::ResourceUsage::Write);
+				State.SetShaderBuffer(EMetalShaderStages::Compute, nil, nil, 0, 0, 0, mtlpp::ResourceUsage(0));
+				
+				mtlpp::DrawIndexedPrimitivesIndirectArguments Args;
+				Args.IndexCount = NumInstances * PatchCount;
+				Args.InstanceCount = 0;
+				Args.IndexStart = 0;
+				PrologueEncoder.SetShaderBytes(mtlpp::FunctionType::Kernel, (const uint8*)&Args, sizeof(Args), 1);
+				State.SetShaderBuffer(EMetalShaderStages::Compute, nil, nil, 0, 0, 1, mtlpp::ResourceUsage(0));
+				
+				METAL_GPUPROFILE(FMetalProfiler::GetProfiler()->EncodeDispatch(PrologueEncoder.GetCommandBufferStats(), __FUNCTION__));
+				computeEncoder.DispatchThreadgroups(mtlpp::Size(1, 1, 1), mtlpp::Size(1, 1, 1));
+				METAL_DEBUG_LAYER(EMetalDebugLevelFastValidation, PrologueEncoder.GetComputeCommandEncoderDebugging().DispatchThreadgroups(mtlpp::Size(1, 1, 1), mtlpp::Size(1, 1, 1)));
+			}
+			
+			auto& renderEncoder = CurrentEncoder.GetRenderCommandEncoder();
+			
+			if (HSOutBuffer)
+			{
+				CurrentEncoder.SetShaderBuffer(mtlpp::FunctionType::Vertex, HSOutBuffer, 0, HSOutBuffer, boundShaderState->DomainShader->TessellationHSOutBuffer, mtlpp::ResourceUsage::Read);
+				State.SetShaderBuffer(EMetalShaderStages::Domain, nil, nil, 0, 0, boundShaderState->DomainShader->TessellationHSOutBuffer, mtlpp::ResourceUsage(0));
+			}
+			
+			if (CPOutBuffer)
+			{
+				CurrentEncoder.SetShaderBuffer(mtlpp::FunctionType::Vertex, CPOutBuffer, 0, CPOutBuffer, boundShaderState->DomainShader->TessellationControlPointOutBuffer, mtlpp::ResourceUsage::Read);
+				State.SetShaderBuffer(EMetalShaderStages::Domain, nil, nil, 0, 0, boundShaderState->DomainShader->TessellationControlPointOutBuffer, mtlpp::ResourceUsage(0));
+			}
+			
+			if (HSTFBuffer)
+			{
+				renderEncoder.SetTessellationFactorBuffer(HSTFBuffer, 0, 0);
+				METAL_DEBUG_LAYER(EMetalDebugLevelFastValidation, CurrentEncoder.GetRenderCommandEncoderDebugging().SetTessellationFactorBuffer(HSTFBuffer, 0, 0));
+			}
+			
+			if (boundShaderState->DomainShader->SideTableBinding >= 0)
+			{
+				CurrentEncoder.SetShaderSideTable(mtlpp::FunctionType::Vertex, boundShaderState->DomainShader->SideTableBinding);
+				State.SetShaderBuffer(EMetalShaderStages::Domain, nil, nil, 0, 0, boundShaderState->DomainShader->SideTableBinding, mtlpp::ResourceUsage(0));
+			}
+			
+			if (IsValidRef(boundShaderState->PixelShader) && boundShaderState->PixelShader->SideTableBinding >= 0)
+			{
+				CurrentEncoder.SetShaderSideTable(mtlpp::FunctionType::Fragment, boundShaderState->PixelShader->SideTableBinding);
+				State.SetShaderBuffer(EMetalShaderStages::Pixel, nil, nil, 0, 0, boundShaderState->PixelShader->SideTableBinding, mtlpp::ResourceUsage(0));
+			}
+			
+			if (GMetalTessellationRunDomainStage)
+			{
+				METAL_GPUPROFILE(FMetalProfiler::GetProfiler()->EncodeDraw(CurrentEncoder.GetCommandBufferStats(), __FUNCTION__, OutputCPs, PatchCount, NumInstances));
+				renderEncoder.DrawPatches(OutputCPs, 0, NumInstances * PatchCount, nil, 0, 1, 0);
+				METAL_DEBUG_LAYER(EMetalDebugLevelFastValidation, CurrentEncoder.GetRenderCommandEncoderDebugging().DrawPatches(OutputCPs, 0, NumInstances * PatchCount, nil, 0, 1, 0));
+			}
+			
+			if (GMetalCommandBufferDebuggingEnabled)
+			{
+				FMetalCommandData Data;
+				Data.CommandType = FMetalCommandData::Type::DrawPrimitivePatch;
+				Data.DrawPatch.BaseInstance = FirstInstance;
+				Data.DrawPatch.InstanceCount = NumInstances;
+				Data.DrawPatch.PatchCount = PatchCount;
+				Data.DrawPatch.PatchStart = 0;
+				
+				InsertDebugDraw(Data);
+			}
+			
+			SafeReleaseMetalBuffer(ShaderOutputBuffer);
+			SafeReleaseMetalBuffer(HullIndexBuffer);
+			SafeReleaseMetalBuffer(HSOutBuffer);
+			SafeReleaseMetalBuffer(CPOutBuffer);
+			SafeReleaseMetalBuffer(HSTFBuffer);
+			return;
+		}
+		
 		ConditionalSwitchToTessellation();
 		check(CurrentEncoder.GetCommandBuffer());
 		check(PrologueEncoder.GetCommandBuffer());
@@ -1489,6 +1820,120 @@ void FMetalRenderPass::ConditionalSwitchToTessellation(void)
 	}
 }
 
+void FMetalRenderPass::ConditionalSwitchToSeparateTessellation(void)
+{
+	SCOPE_CYCLE_COUNTER(STAT_MetalSwitchToTessellationTime);
+	
+	check(bWithinRenderPass);
+	check(RenderPassDesc);
+	check(CurrentEncoder.GetCommandBuffer());
+	
+	// End all current encoders that don't match required compute/raster setup.
+	if (PrologueEncoder.IsBlitCommandEncoderActive() || PrologueEncoder.IsComputeCommandEncoderActive())
+	{
+		PrologueEncoderFence = PrologueEncoder.EndEncoding();
+	}
+	if (CurrentEncoder.IsComputeCommandEncoderActive() || CurrentEncoder.IsBlitCommandEncoderActive())
+	{
+		CurrentEncoderFence = CurrentEncoder.EndEncoding();
+	}
+	
+	// Create a new prologue stream-out encoder if needed
+	if (!PrologueEncoder.IsRenderCommandEncoderActive())
+	{
+		State.SetStateDirty();
+		if (!PrologueEncoder.GetCommandBuffer())
+		{
+			PrologueEncoder.StartCommandBuffer();
+		}
+
+		// Deliberately allocate a new pass with outstanding ref-count as it will be pooled.
+		mtlpp::RenderPassDescriptor Desc([MTLRenderPassDescriptor new]);
+		mtlpp::RenderPassDepthAttachmentDescriptor DepthAttachment;
+		mtlpp::RenderPassStencilAttachmentDescriptor StencilAttachment;
+		
+		CGSize FBSize = CGSizeMake(State.GetViewport(0).width, State.GetViewport(0).height);
+		
+		FTexture2DRHIRef FallbackDepthStencilSurface = State.CreateFallbackDepthStencilSurface(FBSize.width, FBSize.height);
+		check(IsValidRef(FallbackDepthStencilSurface));
+		
+		FMetalSurface& Surface = *GetMetalSurfaceFromRHITexture(FallbackDepthStencilSurface.GetReference());
+		
+		DepthAttachment.SetTexture(Surface.Texture);
+		DepthAttachment.SetLoadAction(mtlpp::LoadAction::DontCare);
+		DepthAttachment.SetStoreAction(mtlpp::StoreAction::DontCare);
+		Desc.SetDepthAttachment(DepthAttachment);
+		
+		StencilAttachment.SetTexture(Surface.Texture);
+		StencilAttachment.SetLoadAction(mtlpp::LoadAction::DontCare);
+		StencilAttachment.SetStoreAction(mtlpp::StoreAction::DontCare);
+		Desc.SetStencilAttachment(StencilAttachment);
+		
+		PrologueEncoder.SetRenderPassDescriptor(Desc);
+		PrologueEncoder.BeginRenderCommandEncoding();
+		
+		// Wait on the pass start fence to ensure proper ordering.
+		if (PrologueStartEncoderFence)
+		{
+			if (PrologueStartEncoderFence->NeedsWait(mtlpp::RenderStages::Vertex))
+			{
+				PrologueEncoder.WaitForFence(PrologueStartEncoderFence);
+			}
+			else
+			{
+				PrologueEncoder.WaitAndUpdateFence(PrologueStartEncoderFence);
+			}
+			PrologueStartEncoderFence = nullptr;
+		}
+		// Wait on previous prologue encoder fence and consume it, we'll replace it with the new one later.
+		if (PrologueEncoderFence)
+		{
+			if (PrologueEncoderFence->NeedsWait(mtlpp::RenderStages::Vertex))
+			{
+				PrologueEncoder.WaitForFence(PrologueEncoderFence);
+			}
+			else
+			{
+				PrologueEncoder.WaitAndUpdateFence(PrologueEncoderFence);
+			}
+			PrologueEncoderFence = nullptr;
+		}
+#if METAL_DEBUG_OPTIONS
+		//		if (GetEmitDrawEvents() && PrologueEncoder.GetEncoderFence())
+		//		{
+		//			for (uint32 i = mtlpp::RenderStages::Vertex; i <= mtlpp::RenderStages::Fragment && PrologueEncoder.GetEncoderFence()->Get((mtlpp::RenderStages)i).GetPtr(); i++)
+		//			{
+		//				if (CmdList.GetCommandQueue().GetRuntimeDebuggingLevel() >= EMetalDebugLevelValidation)
+		//				{
+		//					PrologueEncoder.GetEncoderFence()->Get((mtlpp::RenderStages)i).GetPtr().label = [NSString stringWithFormat:@"Prologue %@", PrologueEncoderFence->Get((mtlpp::RenderStages)i).GetLabel().GetPtr()];
+		//				}
+		//				else
+		//				{
+		//					PrologueEncoder.GetEncoderFence()->Get((mtlpp::RenderStages)i).SetLabel([NSString stringWithFormat:@"Prologue %@", PrologueEncoderFence->Get((mtlpp::RenderStages)i).GetLabel().GetPtr()]);
+		//				}
+		//			}
+		//		}
+#endif
+	}
+	
+	// Restart the render pass to ensure we have a raster encoder
+	if (!CurrentEncoder.IsRenderCommandEncoderActive())
+	{
+		RestartRenderPass(nil);
+		
+		check(CurrentEncoder.IsRenderCommandEncoderActive());
+		check(PrologueEncoder.IsRenderCommandEncoderActive());
+	}
+	else
+	{
+		check(CurrentEncoder.IsRenderCommandEncoderActive());
+		check(PrologueEncoder.IsRenderCommandEncoderActive());
+		
+		// Encode a wait to the current encoder for the necessary prologue encoder
+		CurrentEncoder.WaitForFence(PrologueEncoder.GetEncoderFence());
+	}
+}
+
 void FMetalRenderPass::ConditionalSwitchToCompute(void)
 {
 	SCOPE_CYCLE_COUNTER(STAT_MetalSwitchToComputeTime);
@@ -1594,7 +2039,7 @@ void FMetalRenderPass::ConditionalSwitchToAsyncBlit(void)
 {
 	SCOPE_CYCLE_COUNTER(STAT_MetalSwitchToAsyncBlitTime);
 	
-	if (PrologueEncoder.IsComputeCommandEncoderActive())
+	if (PrologueEncoder.IsComputeCommandEncoderActive() || PrologueEncoder.IsRenderCommandEncoderActive())
 	{
 		PrologueEncoderFence = PrologueEncoder.EndEncoding();
 	}
@@ -1660,7 +2105,7 @@ void FMetalRenderPass::ConditionalSwitchToAsyncCompute(void)
 {
 	SCOPE_CYCLE_COUNTER(STAT_MetalSwitchToComputeTime);
 	
-	if (PrologueEncoder.IsBlitCommandEncoderActive())
+	if (PrologueEncoder.IsBlitCommandEncoderActive() || PrologueEncoder.IsRenderCommandEncoderActive())
 	{
 		PrologueEncoderFence = PrologueEncoder.EndEncoding();
 	}
@@ -1780,6 +2225,27 @@ void FMetalRenderPass::CommitTessellationResourceTables(void)
 		State.CommitResourceTable(EMetalShaderStages::Pixel, mtlpp::FunctionType::Fragment, CurrentEncoder);
 	}
 }
+
+void FMetalRenderPass::CommitStreamOutResourceTables(void)
+{
+	State.CommitVertexStreamResources(&PrologueEncoder);
+	State.CommitResourceTable(EMetalShaderStages::Vertex, mtlpp::FunctionType::Vertex, PrologueEncoder);
+}
+
+void FMetalRenderPass::CommitSeparateTessellationResourceTables(void)
+{
+	State.CommitSeparateTessellationResources(&CurrentEncoder, &PrologueEncoder);
+	
+	State.CommitResourceTable(EMetalShaderStages::Hull, mtlpp::FunctionType::Kernel, PrologueEncoder);
+	
+	State.CommitResourceTable(EMetalShaderStages::Domain, mtlpp::FunctionType::Vertex, CurrentEncoder);
+	
+	TRefCountPtr<FMetalGraphicsPipelineState> CurrentBoundShaderState = State.GetGraphicsPSO();
+	if (IsValidRef(CurrentBoundShaderState->PixelShader))
+	{
+		State.CommitResourceTable(EMetalShaderStages::Pixel, mtlpp::FunctionType::Fragment, CurrentEncoder);
+	}
+}
 #endif
 
 void FMetalRenderPass::CommitDispatchResourceTables(void)
@@ -1844,6 +2310,37 @@ void FMetalRenderPass::PrepareToTessellate(uint32 PrimitiveType)
     
     State.SetRenderPipelineState(CurrentEncoder, &PrologueEncoder);
 }
+
+void FMetalRenderPass::PrepareToStreamOut(uint32 PrimType)
+{
+	SCOPE_CYCLE_COUNTER(STAT_MetalPrepareToTessellateTime);
+	
+	check(PrologueEncoder.GetCommandBuffer());
+	check(PrologueEncoder.IsRenderCommandEncoderActive());
+	
+	// Bind shader resources
+	CommitStreamOutResourceTables();
+	
+	State.SetStreamOutPipelineState(PrologueEncoder);
+}
+
+void FMetalRenderPass::PrepareToSeparateTessellate(uint32 PrimType)
+{
+	SCOPE_CYCLE_COUNTER(STAT_MetalPrepareToTessellateTime);
+	
+	check(CurrentEncoder.GetCommandBuffer());
+	check(PrologueEncoder.GetCommandBuffer());
+	check(CurrentEncoder.IsRenderCommandEncoderActive());
+	check(PrologueEncoder.IsComputeCommandEncoderActive());
+	
+	// Set raster state
+	State.SetRenderState(CurrentEncoder, &PrologueEncoder);
+	
+	// Bind shader resources
+	CommitSeparateTessellationResourceTables();
+	
+	State.SetTessellationPipelineState(CurrentEncoder, &PrologueEncoder);
+}
 #endif
 
 void FMetalRenderPass::PrepareToDispatch(void)
diff --git a/Engine/Source/Runtime/Apple/MetalRHI/Private/MetalRenderPass.h b/Engine/Source/Runtime/Apple/MetalRHI/Private/MetalRenderPass.h
--- a/Engine/Source/Runtime/Apple/MetalRHI/Private/MetalRenderPass.h
+++ b/Engine/Source/Runtime/Apple/MetalRHI/Private/MetalRenderPass.h
@@ -153,6 +153,7 @@ private:
 #pragma mark -
     void ConditionalSwitchToRender(void);
     void ConditionalSwitchToTessellation(void);
+	void ConditionalSwitchToSeparateTessellation(void);
     void ConditionalSwitchToCompute(void);
 	void ConditionalSwitchToBlit(void);
 	void ConditionalSwitchToAsyncBlit(void);
@@ -161,6 +162,8 @@ private:
     void PrepareToRender(uint32 PrimType);
 #if PLATFORM_SUPPORTS_TESSELLATION_SHADERS
     void PrepareToTessellate(uint32 PrimType);
+	void PrepareToStreamOut(uint32 PrimType);
+	void PrepareToSeparateTessellate(uint32 PrimType);
 #endif
     void PrepareToDispatch(void);
 	void PrepareToAsyncDispatch(void);
@@ -168,6 +171,8 @@ private:
     void CommitRenderResourceTables(void);
 #if PLATFORM_SUPPORTS_TESSELLATION_SHADERS
     void CommitTessellationResourceTables(void);
+	void CommitStreamOutResourceTables(void);
+	void CommitSeparateTessellationResourceTables(void);
 #endif
     void CommitDispatchResourceTables(void);
 	void CommitAsyncDispatchResourceTables(void);
diff --git a/Engine/Source/Runtime/Apple/MetalRHI/Private/MetalShaders.cpp b/Engine/Source/Runtime/Apple/MetalRHI/Private/MetalShaders.cpp
--- a/Engine/Source/Runtime/Apple/MetalRHI/Private/MetalShaders.cpp
+++ b/Engine/Source/Runtime/Apple/MetalRHI/Private/MetalShaders.cpp
@@ -1040,12 +1040,80 @@ FMetalHullShader::FMetalHullShader(const TArray<uint8>& InCode)
 {
 	FMetalCodeHeader Header;
 	Init(InCode, Header);
+
+	if (Header.Tessellation.Num())
+	{
+		auto const& Tess = Header.Tessellation[0];
+		TessellationOutputAttribs = Tess.TessellationOutputAttribs;
+		TessellationPatchCountBuffer = Tess.TessellationPatchCountBuffer;
+		TessellationIndexBuffer = Tess.TessellationIndexBuffer;
+		TessellationHSOutBuffer = Tess.TessellationHSOutBuffer;
+		TessellationHSTFOutBuffer = Tess.TessellationHSTFOutBuffer;
+		TessellationControlPointOutBuffer = Tess.TessellationControlPointOutBuffer;
+		TessellationControlPointIndexBuffer = Tess.TessellationControlPointIndexBuffer;
+		TessellationOutputControlPoints = Tess.TessellationOutputControlPoints;
+		TessellationDomain = Tess.TessellationDomain;
+		TessellationInputControlPoints = Tess.TessellationInputControlPoints;
+		TessellationMaxTessFactor = Tess.TessellationMaxTessFactor;
+		TessellationPatchesPerThreadGroup = Tess.TessellationPatchesPerThreadGroup;
+		
+		switch (Tess.TessellationOutputWinding)
+		{
+				// NOTE: cw and ccw are flipped
+			case EMetalOutputWindingMode::Clockwise:		TessellationOutputWinding = mtlpp::Winding::CounterClockwise; break;
+			case EMetalOutputWindingMode::CounterClockwise:	TessellationOutputWinding = mtlpp::Winding::Clockwise; break;
+			default: break;
+		}
+		
+		switch (Tess.TessellationPartitioning)
+		{
+			case EMetalPartitionMode::Pow2:				TessellationPartitioning = mtlpp::TessellationPartitionMode::ModePow2; break;
+			case EMetalPartitionMode::Integer:			TessellationPartitioning = mtlpp::TessellationPartitionMode::ModeInteger; break;
+			case EMetalPartitionMode::FractionalOdd:	TessellationPartitioning = mtlpp::TessellationPartitionMode::ModeFractionalOdd; break;
+			case EMetalPartitionMode::FractionalEven:	TessellationPartitioning = mtlpp::TessellationPartitionMode::ModeFractionalEven; break;
+			default: break;
+		}
+	}
 }
 
 FMetalHullShader::FMetalHullShader(const TArray<uint8>& InCode, mtlpp::Library InLibrary)
 {
 	FMetalCodeHeader Header;
 	Init(InCode, Header, InLibrary);
+	
+	if (Header.Tessellation.Num())
+	{
+		auto const& Tess = Header.Tessellation[0];
+		TessellationOutputAttribs = Tess.TessellationOutputAttribs;
+		TessellationPatchCountBuffer = Tess.TessellationPatchCountBuffer;
+		TessellationIndexBuffer = Tess.TessellationIndexBuffer;
+		TessellationHSOutBuffer = Tess.TessellationHSOutBuffer;
+		TessellationHSTFOutBuffer = Tess.TessellationHSTFOutBuffer;
+		TessellationControlPointOutBuffer = Tess.TessellationControlPointOutBuffer;
+		TessellationControlPointIndexBuffer = Tess.TessellationControlPointIndexBuffer;
+		TessellationOutputControlPoints = Tess.TessellationOutputControlPoints;
+		TessellationDomain = Tess.TessellationDomain;
+		TessellationInputControlPoints = Tess.TessellationInputControlPoints;
+		TessellationMaxTessFactor = Tess.TessellationMaxTessFactor;
+		TessellationPatchesPerThreadGroup = Tess.TessellationPatchesPerThreadGroup;
+		
+		switch (Tess.TessellationOutputWinding)
+		{
+				// NOTE: cw and ccw are flipped
+			case EMetalOutputWindingMode::Clockwise:		TessellationOutputWinding = mtlpp::Winding::CounterClockwise; break;
+			case EMetalOutputWindingMode::CounterClockwise:	TessellationOutputWinding = mtlpp::Winding::Clockwise; break;
+			default: break;
+		}
+		
+		switch (Tess.TessellationPartitioning)
+		{
+			case EMetalPartitionMode::Pow2:				TessellationPartitioning = mtlpp::TessellationPartitionMode::ModePow2; break;
+			case EMetalPartitionMode::Integer:			TessellationPartitioning = mtlpp::TessellationPartitionMode::ModeInteger; break;
+			case EMetalPartitionMode::FractionalOdd:	TessellationPartitioning = mtlpp::TessellationPartitionMode::ModeFractionalOdd; break;
+			case EMetalPartitionMode::FractionalEven:	TessellationPartitioning = mtlpp::TessellationPartitionMode::ModeFractionalEven; break;
+			default: break;
+		}
+	}
 }
 
 mtlpp::Function FMetalHullShader::GetFunction()
@@ -1079,6 +1147,9 @@ FMetalDomainShader::FMetalDomainShader(const TArray<uint8>& InCode)
 		case EMetalPartitionMode::FractionalEven:	TessellationPartitioning = mtlpp::TessellationPartitionMode::ModeFractionalEven; break;
 		default: check(0);
 	}
+	
+	TessellationDomain = Tess.TessellationDomain;
+	TessellationOutputAttribs = Tess.TessellationOutputAttribs;
 }
 
 FMetalDomainShader::FMetalDomainShader(const TArray<uint8>& InCode, mtlpp::Library InLibrary)
@@ -1107,6 +1178,9 @@ FMetalDomainShader::FMetalDomainShader(const TArray<uint8>& InCode, mtlpp::Libra
 		case EMetalPartitionMode::FractionalEven:	TessellationPartitioning = mtlpp::TessellationPartitionMode::ModeFractionalEven; break;
 		default: check(0);
 	}
+	
+	TessellationDomain = Tess.TessellationDomain;
+	TessellationOutputAttribs = Tess.TessellationOutputAttribs;
 }
 
 mtlpp::Function FMetalDomainShader::GetFunction()
diff --git a/Engine/Source/Runtime/Apple/MetalRHI/Private/MetalStateCache.cpp b/Engine/Source/Runtime/Apple/MetalRHI/Private/MetalStateCache.cpp
--- a/Engine/Source/Runtime/Apple/MetalRHI/Private/MetalStateCache.cpp
+++ b/Engine/Source/Runtime/Apple/MetalRHI/Private/MetalStateCache.cpp
@@ -1072,7 +1072,7 @@ void FMetalStateCache::SetGraphicsPipelineState(FMetalGraphicsPipelineState* Sta
 		GraphicsPSO = State;
 		
 #if PLATFORM_SUPPORTS_TESSELLATION_SHADERS
-		bool bNewUsingTessellation = (State && State->GetPipeline(IndexType)->TessellationPipelineDesc.DomainVertexDescriptor);
+		bool bNewUsingTessellation = (State && IsValidRef(State->DomainShader));
 		if (bNewUsingTessellation != bUsingTessellation)
 		{
 			for (uint32 i = 0; i < EMetalShaderStages::Num; i++)
@@ -1872,6 +1872,33 @@ void FMetalStateCache::CommitTessellationResources(FMetalCommandEncoder* Raster,
 	SetResourcesFromTables(GraphicsPSO->DomainShader, CrossCompiler::SHADER_STAGE_DOMAIN);
     GetShaderParameters(EMetalShaderStages::Domain).CommitPackedGlobals(this, Raster, EMetalShaderStages::Domain, GraphicsPSO->DomainShader->Bindings);
 }
+
+void FMetalStateCache::CommitVertexStreamResources(FMetalCommandEncoder* Raster)
+{
+	check(IsValidRef(GraphicsPSO));
+	check(IsValidRef(GraphicsPSO->HullShader) && IsValidRef(GraphicsPSO->DomainShader));
+	
+	SetResourcesFromTables(GraphicsPSO->VertexShader, CrossCompiler::SHADER_STAGE_VERTEX);
+	GetShaderParameters(EMetalShaderStages::Vertex).CommitPackedGlobals(this, Raster, EMetalShaderStages::Vertex, GraphicsPSO->VertexShader->Bindings);
+}
+
+void FMetalStateCache::CommitSeparateTessellationResources(FMetalCommandEncoder* Raster, FMetalCommandEncoder* Compute)
+{
+	check(IsValidRef(GraphicsPSO));
+	check(IsValidRef(GraphicsPSO->HullShader) && IsValidRef(GraphicsPSO->DomainShader));
+	
+	if (IsValidRef(GraphicsPSO->PixelShader))
+	{
+		SetResourcesFromTables(GraphicsPSO->PixelShader, CrossCompiler::SHADER_STAGE_PIXEL);
+		GetShaderParameters(EMetalShaderStages::Pixel).CommitPackedGlobals(this, Raster, EMetalShaderStages::Pixel, GraphicsPSO->PixelShader->Bindings);
+	}
+	
+	SetResourcesFromTables(GraphicsPSO->HullShader, CrossCompiler::SHADER_STAGE_HULL);
+	GetShaderParameters(EMetalShaderStages::Hull).CommitPackedGlobals(this, Compute, EMetalShaderStages::Hull, GraphicsPSO->HullShader->Bindings);
+	
+	SetResourcesFromTables(GraphicsPSO->DomainShader, CrossCompiler::SHADER_STAGE_DOMAIN);
+	GetShaderParameters(EMetalShaderStages::Domain).CommitPackedGlobals(this, Raster, EMetalShaderStages::Domain, GraphicsPSO->DomainShader->Bindings);
+}
 #endif
 
 void FMetalStateCache::CommitComputeResources(FMetalCommandEncoder* Compute)
@@ -2092,6 +2119,78 @@ void FMetalStateCache::SetRenderState(FMetalCommandEncoder& CommandEncoder, FMet
 	}
 }
 
+#if PLATFORM_SUPPORTS_TESSELLATION_SHADERS
+void FMetalStateCache::SetStreamOutPipelineState(FMetalCommandEncoder& CommandEncoder)
+{
+	SCOPE_CYCLE_COUNTER(STAT_MetalSetRenderPipelineStateTime);
+	FMetalShaderPipeline* Pipeline = GetPipelineState();
+	
+	auto& StreamOutEncoder = CommandEncoder.GetRenderCommandEncoder();
+	StreamOutEncoder.SetRenderPipelineState(Pipeline->StreamPipelineState);
+	
+	FMetalGraphicsPipelineState* boundShaderState = GetGraphicsPSO();
+	
+	if (SafeGetRuntimeDebuggingLevel() >= EMetalDebugLevelFastValidation)
+	{
+		EMetalShaderStages VertexStage = EMetalShaderStages::Vertex;
+
+		FMetalDebugShaderResourceMask VertexMask = Pipeline->ResourceMask[EMetalShaderStream];
+		TArray<uint32>& MinVertexBufferSizes = Pipeline->BufferDataSizes[EMetalShaderStream];
+		TMap<uint8, uint8> VertexTexTypes = Pipeline->TextureTypes[EMetalShaderStream];
+		while(VertexMask.BufferMask)
+		{
+			uint32 Index = __builtin_ctz(VertexMask.BufferMask);
+			VertexMask.BufferMask &= ~(1 << Index);
+			
+			if ((boundShaderState->VertexShader->TessellationPatchCountBuffer != Index
+																&& boundShaderState->VertexShader->TessellationHSOutBuffer != Index))
+			{
+				FMetalBufferBinding const& Binding = ShaderBuffers[VertexStage].Buffers[Index];
+				ensure(Binding.Buffer || Binding.Bytes);
+				ensure(MinVertexBufferSizes.Num() > Index);
+				ensure(Binding.Length >= MinVertexBufferSizes[Index]);
+			}
+		}
+#if PLATFORM_MAC
+		{
+			uint64 LoTextures = (uint64)VertexMask.TextureMask;
+			while(LoTextures)
+			{
+				uint32 Index = __builtin_ctzll(LoTextures);
+				LoTextures &= ~(uint64(1) << uint64(Index));
+				ensure(ShaderTextures[VertexStage].Textures[Index]);
+				ensure(ShaderTextures[VertexStage].Textures[Index].GetTextureType() == (mtlpp::TextureType)VertexTexTypes.FindRef(Index));
+			}
+			
+			uint64 HiTextures = (uint64)(VertexMask.TextureMask >> FMetalTextureMask(64));
+			while(HiTextures)
+			{
+				uint32 Index = __builtin_ctzll(HiTextures);
+				HiTextures &= ~(uint64(1) << uint64(Index));
+				ensure(ShaderTextures[VertexStage].Textures[Index + 64]);
+				ensure(ShaderTextures[VertexStage].Textures[Index + 64].GetTextureType() == (mtlpp::TextureType)VertexTexTypes.FindRef(Index + 64));
+			}
+		}
+#else
+		while(VertexMask.TextureMask)
+		{
+			uint32 Index = __builtin_ctz(VertexMask.TextureMask);
+			VertexMask.TextureMask &= ~(1 << Index);
+			
+			ensure(ShaderTextures[VertexStage].Textures[Index]);
+			ensure(ShaderTextures[VertexStage].Textures[Index].GetTextureType() == (mtlpp::TextureType)VertexTexTypes.FindRef(Index));
+		}
+#endif
+		while(VertexMask.SamplerMask)
+		{
+			uint32 Index = __builtin_ctz(VertexMask.SamplerMask);
+			VertexMask.SamplerMask &= ~(1 << Index);
+			ensure(ShaderSamplers[VertexStage].Samplers[Index]);
+		}
+	}
+}
+#endif
+
 void FMetalStateCache::SetRenderPipelineState(FMetalCommandEncoder& CommandEncoder, FMetalCommandEncoder* PrologueEncoder)
 {
 	SCOPE_CYCLE_COUNTER(STAT_MetalSetRenderPipelineStateTime);
@@ -2291,6 +2390,203 @@ void FMetalStateCache::SetRenderPipelineState(FMetalCommandEncoder& CommandEncod
 	}
 }
 
+#if PLATFORM_SUPPORTS_TESSELLATION_SHADERS
+void FMetalStateCache::SetTessellationPipelineState(FMetalCommandEncoder& CommandEncoder, FMetalCommandEncoder* PrologueEncoder)
+{
+	SCOPE_CYCLE_COUNTER(STAT_MetalSetRenderPipelineStateTime);
+	
+	if ((PipelineBits & EMetalPipelineFlagRasterMask) != 0)
+	{
+		// Some Intel drivers need RenderPipeline state to be set after DepthStencil state to work properly
+		FMetalShaderPipeline* Pipeline = GetPipelineState();
+		
+		check(Pipeline);
+		CommandEncoder.SetRenderPipelineState(Pipeline);
+		if (Pipeline->ComputePipelineState)
+		{
+			check(PrologueEncoder);
+			PrologueEncoder->SetComputePipelineState(Pipeline);
+		}
+		
+		PipelineBits &= EMetalPipelineFlagComputeMask;
+	}
+	
+	if (SafeGetRuntimeDebuggingLevel() >= EMetalDebugLevelFastValidation)
+	{
+		FMetalShaderPipeline* Pipeline = GetPipelineState();
+		EMetalShaderStages VertexStage = EMetalShaderStages::Vertex;
+#if PLATFORM_SUPPORTS_TESSELLATION_SHADERS
+		if (GraphicsPSO->DomainShader)
+		{
+			VertexStage = EMetalShaderStages::Domain;
+			
+			FMetalDebugShaderResourceMask ComputeMask = Pipeline->ResourceMask[EMetalShaderCompute];
+			TArray<uint32>& MinComputeBufferSizes = Pipeline->BufferDataSizes[EMetalShaderCompute];
+			TMap<uint8, uint8> ComputeTexTypes = Pipeline->TextureTypes[EMetalShaderCompute];
+			while(ComputeMask.BufferMask)
+			{
+				uint32 Index = __builtin_ctz(ComputeMask.BufferMask);
+				ComputeMask.BufferMask &= ~(1 << Index);
+				
+				if (GraphicsPSO->HullShader->TessellationPatchCountBuffer != Index
+					&& GraphicsPSO->HullShader->TessellationIndexBuffer != Index
+					&& GraphicsPSO->HullShader->TessellationHSOutBuffer != Index
+					&& GraphicsPSO->HullShader->TessellationHSTFOutBuffer != Index
+					&& GraphicsPSO->HullShader->TessellationControlPointOutBuffer != Index
+					&& GraphicsPSO->HullShader->TessellationControlPointIndexBuffer != Index)
+				{
+					FMetalBufferBinding const& HullBinding = ShaderBuffers[EMetalShaderStages::Hull].Buffers[Index];
+					ensure(HullBinding.Buffer || HullBinding.Bytes);
+					ensure(MinComputeBufferSizes.Num() > Index);
+					ensure(HullBinding.Length >= MinComputeBufferSizes[Index]);
+				}
+			}
+#if PLATFORM_MAC
+			{
+				uint64 LoTextures = (uint64)ComputeMask.TextureMask;
+				while(LoTextures)
+				{
+					uint32 Index = __builtin_ctzll(LoTextures);
+					LoTextures &= ~(uint64(1) << uint64(Index));
+					ensure(ShaderTextures[EMetalShaderStages::Hull].Textures[Index + 64]);
+					ensure(ShaderTextures[EMetalShaderStages::Hull].Textures[Index].GetTextureType() == (mtlpp::TextureType)ComputeTexTypes.FindRef(Index));
+				}
+				
+				uint64 HiTextures = (uint64)(ComputeMask.TextureMask >> FMetalTextureMask(64));
+				while(HiTextures)
+				{
+					uint32 Index = __builtin_ctzll(HiTextures);
+					HiTextures &= ~(uint64(1) << uint64(Index));
+					ensure(ShaderTextures[EMetalShaderStages::Hull].Textures[Index + 64]);
+					ensure(ShaderTextures[EMetalShaderStages::Hull].Textures[Index + 64].GetTextureType() == (mtlpp::TextureType)ComputeTexTypes.FindRef(Index + 64));
+				}
+			}
+#else
+			while(ComputeMask.TextureMask)
+			{
+				uint32 Index = __builtin_ctz(ComputeMask.TextureMask);
+				ComputeMask.TextureMask &= ~(1 << Index);
+				
+				ensure(ShaderTextures[EMetalShaderStages::Hull].Textures[Index]);
+				ensure(ShaderTextures[EMetalShaderStages::Hull].Textures[Index].GetTextureType() == (mtlpp::TextureType)ComputeTexTypes.FindRef(Index));
+			}
+#endif
+			while(ComputeMask.SamplerMask)
+			{
+				uint32 Index = __builtin_ctz(ComputeMask.SamplerMask);
+				ComputeMask.SamplerMask &= ~(1 << Index);
+				ensure(ShaderSamplers[EMetalShaderStages::Hull].Samplers[Index]);
+			}
+		}
+#endif
+		
+		FMetalDebugShaderResourceMask VertexMask = Pipeline->ResourceMask[EMetalShaderVertex];
+		TArray<uint32>& MinVertexBufferSizes = Pipeline->BufferDataSizes[EMetalShaderVertex];
+		TMap<uint8, uint8> VertexTexTypes = Pipeline->TextureTypes[EMetalShaderVertex];
+		while(VertexMask.BufferMask)
+		{
+			uint32 Index = __builtin_ctz(VertexMask.BufferMask);
+			VertexMask.BufferMask &= ~(1 << Index);
+			
+			if ((VertexStage == EMetalShaderStages::Vertex) || (GraphicsPSO->DomainShader->TessellationHSOutBuffer != Index
+																&& GraphicsPSO->DomainShader->TessellationControlPointOutBuffer != Index))
+			{
+				FMetalBufferBinding const& Binding = ShaderBuffers[VertexStage].Buffers[Index];
+				ensure(Binding.Buffer || Binding.Bytes);
+				ensure(MinVertexBufferSizes.Num() > Index);
+				ensure(Binding.Length >= MinVertexBufferSizes[Index]);
+			}
+		}
+#if PLATFORM_MAC
+		{
+			uint64 LoTextures = (uint64)VertexMask.TextureMask;
+			while(LoTextures)
+			{
+				uint32 Index = __builtin_ctzll(LoTextures);
+				LoTextures &= ~(uint64(1) << uint64(Index));
+				ensure(ShaderTextures[VertexStage].Textures[Index]);
+				ensure(ShaderTextures[VertexStage].Textures[Index].GetTextureType() == (mtlpp::TextureType)VertexTexTypes.FindRef(Index));
+			}
+			
+			uint64 HiTextures = (uint64)(VertexMask.TextureMask >> FMetalTextureMask(64));
+			while(HiTextures)
+			{
+				uint32 Index = __builtin_ctzll(HiTextures);
+				HiTextures &= ~(uint64(1) << uint64(Index));
+				ensure(ShaderTextures[VertexStage].Textures[Index + 64]);
+				ensure(ShaderTextures[VertexStage].Textures[Index + 64].GetTextureType() == (mtlpp::TextureType)VertexTexTypes.FindRef(Index + 64));
+			}
+		}
+#else
+		while(VertexMask.TextureMask)
+		{
+			uint32 Index = __builtin_ctz(VertexMask.TextureMask);
+			VertexMask.TextureMask &= ~(1 << Index);
+			
+			ensure(ShaderTextures[VertexStage].Textures[Index]);
+			ensure(ShaderTextures[VertexStage].Textures[Index].GetTextureType() == (mtlpp::TextureType)VertexTexTypes.FindRef(Index));
+		}
+#endif
+		while(VertexMask.SamplerMask)
+		{
+			uint32 Index = __builtin_ctz(VertexMask.SamplerMask);
+			VertexMask.SamplerMask &= ~(1 << Index);
+			ensure(ShaderSamplers[VertexStage].Samplers[Index]);
+		}
+		
+		FMetalDebugShaderResourceMask FragmentMask = Pipeline->ResourceMask[EMetalShaderFragment];
+		TArray<uint32>& MinFragmentBufferSizes = Pipeline->BufferDataSizes[EMetalShaderFragment];
+		TMap<uint8, uint8> FragmentTexTypes = Pipeline->TextureTypes[EMetalShaderFragment];
+		while(FragmentMask.BufferMask)
+		{
+			uint32 Index = __builtin_ctz(FragmentMask.BufferMask);
+			FragmentMask.BufferMask &= ~(1 << Index);
+			
+			FMetalBufferBinding const& Binding = ShaderBuffers[EMetalShaderStages::Pixel].Buffers[Index];
+			ensure(Binding.Buffer || Binding.Bytes);
+			ensure(MinFragmentBufferSizes.Num() > Index);
+			ensure(Binding.Length >= MinFragmentBufferSizes[Index]);
+		}
+#if PLATFORM_MAC
+		{
+			uint64 LoTextures = (uint64)FragmentMask.TextureMask;
+			while(LoTextures)
+			{
+				uint32 Index = __builtin_ctzll(LoTextures);
+				LoTextures &= ~(uint64(1) << uint64(Index));
+				ensure(ShaderTextures[EMetalShaderStages::Pixel].Textures[Index]);
+				ensure(ShaderTextures[EMetalShaderStages::Pixel].Textures[Index].GetTextureType() == (mtlpp::TextureType)FragmentTexTypes.FindRef(Index));
+			}
+			
+			uint64 HiTextures = (uint64)(FragmentMask.TextureMask >> FMetalTextureMask(64));
+			while(HiTextures)
+			{
+				uint32 Index = __builtin_ctzll(HiTextures);
+				HiTextures &= ~(uint64(1) << uint64(Index));
+				ensure(ShaderTextures[EMetalShaderStages::Pixel].Textures[Index + 64]);
+				ensure(ShaderTextures[EMetalShaderStages::Pixel].Textures[Index + 64].GetTextureType() == (mtlpp::TextureType)FragmentTexTypes.FindRef(Index + 64));
+			}
+		}
+#else
+		while(FragmentMask.TextureMask)
+		{
+			uint32 Index = __builtin_ctz(FragmentMask.TextureMask);
+			FragmentMask.TextureMask &= ~(1 << Index);
+			
+			ensure(ShaderTextures[EMetalShaderStages::Pixel].Textures[Index]);
+			ensure(ShaderTextures[EMetalShaderStages::Pixel].Textures[Index].GetTextureType() == (mtlpp::TextureType)FragmentTexTypes.FindRef(Index));
+		}
+#endif
+		while(FragmentMask.SamplerMask)
+		{
+			uint32 Index = __builtin_ctz(FragmentMask.SamplerMask);
+			FragmentMask.SamplerMask &= ~(1 << Index);
+			ensure(ShaderSamplers[EMetalShaderStages::Pixel].Samplers[Index]);
+		}
+	}
+}
+#endif
+
 void FMetalStateCache::SetComputePipelineState(FMetalCommandEncoder& CommandEncoder)
 {
 	if ((PipelineBits & EMetalPipelineFlagComputeMask) != 0)
diff --git a/Engine/Source/Runtime/Apple/MetalRHI/Private/MetalStateCache.h b/Engine/Source/Runtime/Apple/MetalRHI/Private/MetalStateCache.h
--- a/Engine/Source/Runtime/Apple/MetalRHI/Private/MetalStateCache.h
+++ b/Engine/Source/Runtime/Apple/MetalRHI/Private/MetalStateCache.h
@@ -108,6 +108,9 @@ public:
 
 #if PLATFORM_SUPPORTS_TESSELLATION_SHADERS
 	void CommitTessellationResources(FMetalCommandEncoder* Raster, FMetalCommandEncoder* Compute);
+	
+	void CommitVertexStreamResources(FMetalCommandEncoder* Raster);
+	void CommitSeparateTessellationResources(FMetalCommandEncoder* Raster, FMetalCommandEncoder* Compute);
 #endif
 
 	void CommitComputeResources(FMetalCommandEncoder* Compute);
@@ -154,7 +157,11 @@ public:
 	FTexture2DRHIRef CreateFallbackDepthStencilSurface(uint32 Width, uint32 Height);
 	bool GetFallbackDepthStencilBound(void) const { return bFallbackDepthStencilBound; }
 	
-    void SetRenderPipelineState(FMetalCommandEncoder& CommandEncoder, FMetalCommandEncoder* PrologueEncoder);
+	void SetRenderPipelineState(FMetalCommandEncoder& CommandEncoder, FMetalCommandEncoder* PrologueEncoder);
+#if PLATFORM_SUPPORTS_TESSELLATION_SHADERS
+	void SetTessellationPipelineState(FMetalCommandEncoder& CommandEncoder, FMetalCommandEncoder* PrologueEncoder);
+	void SetStreamOutPipelineState(FMetalCommandEncoder& CommandEncoder);
+#endif
     void SetComputePipelineState(FMetalCommandEncoder& CommandEncoder);
 	void FlushVisibilityResults(FMetalCommandEncoder& CommandEncoder);
 
diff --git a/Engine/Source/Runtime/Apple/MetalRHI/Public/MetalResources.h b/Engine/Source/Runtime/Apple/MetalRHI/Public/MetalResources.h
--- a/Engine/Source/Runtime/Apple/MetalRHI/Public/MetalResources.h
+++ b/Engine/Source/Runtime/Apple/MetalRHI/Public/MetalResources.h
@@ -210,6 +210,22 @@ public:
 	FMetalHullShader(const TArray<uint8>& InCode, mtlpp::Library InLibrary);
 	
 	mtlpp::Function GetFunction();
+	
+	// for VSHS
+	FMetalTessellationOutputs TessellationOutputAttribs;
+	float  TessellationMaxTessFactor;
+	uint32 TessellationOutputControlPoints;
+	uint32 TessellationDomain;
+	uint32 TessellationInputControlPoints;
+	uint32 TessellationPatchesPerThreadGroup;
+	uint32 TessellationPatchCountBuffer;
+	uint32 TessellationIndexBuffer;
+	uint32 TessellationHSOutBuffer;
+	uint32 TessellationHSTFOutBuffer;
+	uint32 TessellationControlPointOutBuffer;
+	uint32 TessellationControlPointIndexBuffer;
+	mtlpp::Winding TessellationOutputWinding;
+	mtlpp::TessellationPartitionMode TessellationPartitioning;
 };
 
 class FMetalDomainShader : public TMetalBaseShader<FRHIDomainShader, SF_Domain>
@@ -224,6 +240,9 @@ public:
 	mtlpp::TessellationPartitionMode TessellationPartitioning;
 	uint32 TessellationHSOutBuffer;
 	uint32 TessellationControlPointOutBuffer;
+	
+	uint32 TessellationDomain;
+	FMetalTessellationOutputs TessellationOutputAttribs;
 };
 #else
 typedef TMetalBaseShader<FRHIHullShader, SF_Hull> FMetalHullShader;
diff --git a/Engine/Source/Runtime/Apple/MetalRHI/Public/MetalShaderResources.h b/Engine/Source/Runtime/Apple/MetalRHI/Public/MetalShaderResources.h
--- a/Engine/Source/Runtime/Apple/MetalRHI/Public/MetalShaderResources.h
+++ b/Engine/Source/Runtime/Apple/MetalRHI/Public/MetalShaderResources.h
@@ -206,12 +206,14 @@ struct FMetalAttribute
 	uint32 Components;
 	uint32 Offset;
 	EMetalComponentType Type;
+	uint32 Semantic;
 	
 	FMetalAttribute()
 	: Index(0)
 	, Components(0)
 	, Offset(0)
 	, Type(EMetalComponentType::Uint)
+	, Semantic(0)
 	{
 		
 	}
@@ -222,12 +224,14 @@ struct FMetalAttribute
 		Ar << Attr.Type;
 		Ar << Attr.Components;
 		Ar << Attr.Offset;
+		Ar << Attr.Semantic;
 		return Ar;
 	}
 };
 
 struct FMetalTessellationOutputs
 {
+	TArray<FMetalAttribute> HSIn;
 	TArray<FMetalAttribute> HSOut;
 	TArray<FMetalAttribute> PatchControlPointOut;
 	uint32 HSOutSize;
@@ -244,6 +248,7 @@ struct FMetalTessellationOutputs
 	
 	friend FArchive& operator<<(FArchive& Ar, FMetalTessellationOutputs& Attrs)
 	{
+		Ar << Attrs.HSIn;
 		Ar << Attrs.HSOut;
 		Ar << Attrs.PatchControlPointOut;
 		Ar << Attrs.HSOutSize;
