Set the type of FMetalCommandEncoder when creating it so that we don't mistakenly try and create parallel render encoders when using separate tessellation stages on the prologue encoder.

diff --git a/Engine/Source/Runtime/Apple/MetalRHI/Private/MetalCommandEncoder.cpp b/Engine/Source/Runtime/Apple/MetalRHI/Private/MetalCommandEncoder.cpp
--- a/Engine/Source/Runtime/Apple/MetalRHI/Private/MetalCommandEncoder.cpp
+++ b/Engine/Source/Runtime/Apple/MetalRHI/Private/MetalCommandEncoder.cpp
@@ -209,7 +209,7 @@ FMetalCommandBufferMarkers FMetalCommandBufferMarkers::Get(mtlpp::CommandBuffer
 
 #pragma mark - Public C++ Boilerplate -
 
-FMetalCommandEncoder::FMetalCommandEncoder(FMetalCommandList& CmdList)
+FMetalCommandEncoder::FMetalCommandEncoder(FMetalCommandList& CmdList, EMetalCommandEncoderType InType)
 : CommandList(CmdList)
 , bSupportsMetalFeaturesSetBytes(CmdList.GetCommandQueue().SupportsFeature(EMetalFeaturesSetBytes))
 , RingBuffer(EncoderRingBufferSize, BufferOffsetAlignment, FMetalCommandQueue::GetCompatibleResourceOptions((mtlpp::ResourceOptions)(mtlpp::ResourceOptions::HazardTrackingModeUntracked | BUFFER_RESOURCE_STORAGE_MANAGED)))
@@ -226,6 +226,7 @@ FMetalCommandEncoder::FMetalCommandEncoder(FMetalCommandList& CmdList)
 , FenceStage(mtlpp::RenderStages::Fragment)
 , EncoderNum(0)
 , CmdBufIndex(0)
+, Type(InType)
 {
 	for (uint32 Frequency = 0; Frequency < uint32(mtlpp::FunctionType::Kernel)+1; Frequency++)
 	{
@@ -494,7 +495,7 @@ bool FMetalCommandEncoder::IsImmediate(void) const
 
 bool FMetalCommandEncoder::IsParallel(void) const
 {
-	return CommandList.IsParallel();
+	return CommandList.IsParallel() && (Type == EMetalCommandEncoderCurrent);
 }
 
 bool FMetalCommandEncoder::IsRenderPassDescriptorValid(void) const
@@ -596,7 +597,7 @@ void FMetalCommandEncoder::BeginRenderCommandEncoding(void)
 	
 	FenceResources.Append(TransitionedResources);
 	
-	if (!CommandList.IsParallel())
+	if (!CommandList.IsParallel() || Type == EMetalCommandEncoderPrologue)
 	{
 		RenderCommandEncoder = MTLPP_VALIDATE(mtlpp::CommandBuffer, CommandBuffer, SafeGetRuntimeDebuggingLevel() >= EMetalDebugLevelValidation, RenderCommandEncoder(RenderPassDesc));
 		METAL_DEBUG_LAYER(EMetalDebugLevelFastValidation, RenderEncoderDebug = FMetalRenderCommandEncoderDebugging(RenderCommandEncoder, RenderPassDesc, CommandBufferDebug));
@@ -644,7 +645,7 @@ void FMetalCommandEncoder::BeginRenderCommandEncoding(void)
 	}
 }
 
-void FMetalCommandEncoder::BeginComputeCommandEncoding(mtlpp::DispatchType Type)
+void FMetalCommandEncoder::BeginComputeCommandEncoding(mtlpp::DispatchType DispatchType)
 {
 	check(CommandBuffer);
 	check(IsRenderCommandEncoderActive() == false && IsComputeCommandEncoderActive() == false && IsBlitCommandEncoderActive() == false);
@@ -652,7 +653,7 @@ void FMetalCommandEncoder::BeginComputeCommandEncoding(mtlpp::DispatchType Type)
 	FenceResources.Append(TransitionedResources);
 	TransitionedResources.Empty();
 	
-	if (Type == mtlpp::DispatchType::Serial)
+	if (DispatchType == mtlpp::DispatchType::Serial)
 	{
 		ComputeCommandEncoder = MTLPP_VALIDATE(mtlpp::CommandBuffer, CommandBuffer, SafeGetRuntimeDebuggingLevel() >= EMetalDebugLevelValidation, ComputeCommandEncoder());
 	}
diff --git a/Engine/Source/Runtime/Apple/MetalRHI/Private/MetalCommandEncoder.h b/Engine/Source/Runtime/Apple/MetalRHI/Private/MetalCommandEncoder.h
--- a/Engine/Source/Runtime/Apple/MetalRHI/Private/MetalCommandEncoder.h
+++ b/Engine/Source/Runtime/Apple/MetalRHI/Private/MetalCommandEncoder.h
@@ -101,6 +101,17 @@ public:
 	static char const* kTableAssociationKey;
 };
 
+/**
+ * EMetalCommandEncoderType:
+ *   EMetalCommandEncoderCurrent: The primary encoder that is used for draw calls & dispatches
+ *   EMetalCommandEncoderPrologue: A secondary encoder that is used for blits & dispatches that setup resources & state for the current encoder.
+ */
+enum EMetalCommandEncoderType
+{
+	EMetalCommandEncoderCurrent,
+	EMetalCommandEncoderPrologue
+};
+
 /**
  * FMetalCommandEncoder:
  *	Wraps the details of switching between different command encoders on the command-buffer, allowing for restoration of the render encoder if needed.
@@ -114,7 +125,7 @@ public:
 #pragma mark - Public C++ Boilerplate -
 
 	/** Default constructor */
-	FMetalCommandEncoder(FMetalCommandList& CmdList);
+	FMetalCommandEncoder(FMetalCommandList& CmdList, EMetalCommandEncoderType Type);
 	
 	/** Destructor */
 	~FMetalCommandEncoder(void);
@@ -573,4 +584,5 @@ private:
 	mtlpp::RenderStages FenceStage;
 	uint32 EncoderNum;
 	uint32 CmdBufIndex;
+	EMetalCommandEncoderType Type;
 };
diff --git a/Engine/Source/Runtime/Apple/MetalRHI/Private/MetalRenderPass.cpp b/Engine/Source/Runtime/Apple/MetalRHI/Private/MetalRenderPass.cpp
--- a/Engine/Source/Runtime/Apple/MetalRHI/Private/MetalRenderPass.cpp
+++ b/Engine/Source/Runtime/Apple/MetalRHI/Private/MetalRenderPass.cpp
@@ -55,8 +55,8 @@ static FAutoConsoleVariableRef CVarMetalDebugOpsCount(
 FMetalRenderPass::FMetalRenderPass(FMetalCommandList& InCmdList, FMetalStateCache& Cache)
 : CmdList(InCmdList)
 , State(Cache)
-, CurrentEncoder(InCmdList)
-, PrologueEncoder(InCmdList)
+, CurrentEncoder(InCmdList, EMetalCommandEncoderCurrent)
+, PrologueEncoder(InCmdList, EMetalCommandEncoderPrologue)
 , RenderPassDesc(nil)
 , ComputeDispatchType(mtlpp::DispatchType::Serial)
 , NumOutstandingOps(0)
