Didn't hit save, so missed a change.

diff --git a/Engine/Source/Runtime/Apple/MetalRHI/Private/MetalCommandEncoder.cpp b/Engine/Source/Runtime/Apple/MetalRHI/Private/MetalCommandEncoder.cpp
--- a/Engine/Source/Runtime/Apple/MetalRHI/Private/MetalCommandEncoder.cpp
+++ b/Engine/Source/Runtime/Apple/MetalRHI/Private/MetalCommandEncoder.cpp
@@ -659,7 +659,7 @@ void FMetalCommandEncoder::BeginComputeCommandEncoding(mtlpp::DispatchType Dispa
 	}
 	else
 	{
-		ComputeCommandEncoder = MTLPP_VALIDATE(mtlpp::CommandBuffer, CommandBuffer, SafeGetRuntimeDebuggingLevel() >= EMetalDebugLevelValidation, ComputeCommandEncoder(Type));
+		ComputeCommandEncoder = MTLPP_VALIDATE(mtlpp::CommandBuffer, CommandBuffer, SafeGetRuntimeDebuggingLevel() >= EMetalDebugLevelValidation, ComputeCommandEncoder(DispatchType));
 	}
 	METAL_DEBUG_LAYER(EMetalDebugLevelFastValidation, ComputeEncoderDebug = FMetalComputeCommandEncoderDebugging(ComputeCommandEncoder, CommandBufferDebug));
 
