DR - Compile fix
#rnx

diff --git a/Engine/Source/Runtime/Apple/MetalRHI/Private/MetalCommands.cpp b/Engine/Source/Runtime/Apple/MetalRHI/Private/MetalCommands.cpp
--- a/Engine/Source/Runtime/Apple/MetalRHI/Private/MetalCommands.cpp
+++ b/Engine/Source/Runtime/Apple/MetalRHI/Private/MetalCommands.cpp
@@ -496,6 +496,7 @@ void FMetalRHICommandContext::RHISetShaderParameter(FRHIVertexShader* VertexShad
 
 void FMetalRHICommandContext::RHISetShaderParameter(FRHIHullShader* HullShaderRHI, uint32 BufferIndex, uint32 BaseIndex, uint32 NumBytes, const void* NewValue)
 {
+#if PLATFORM_SUPPORTS_TESSELLATION_SHADERS
 	// Just ignore Hull shader parameter sets - none of our Hull shaders have any loose parameters to bind.
 	// @todo Whenever we do put a shader parameter into a hull shader we'll need to map it into the vertex-shader parameter buffer so that it can be set on the device.
 	if (FMetalCommandQueue::SupportsFeature(EMetalFeaturesSeparateTessellation))
@@ -504,6 +505,7 @@ void FMetalRHICommandContext::RHISetShaderParameter(FRHIHullShader* HullShaderRH
 			Context->GetCurrentState().GetShaderParameters(EMetalShaderStages::Hull).Set(BufferIndex, BaseIndex, NumBytes, NewValue);
 		}
 	}
+#endif
 }
 
 void FMetalRHICommandContext::RHISetShaderParameter(FRHIPixelShader* PixelShaderRHI, uint32 BufferIndex, uint32 BaseIndex, uint32 NumBytes, const void* NewValue)
