Minimal compile fixes after changes to RHI Buffer API - compiles & runs again but MetalRHI needs some refactoring.

diff --git a/Engine/Source/Runtime/Apple/MetalRHI/Private/MetalIndexBuffer.cpp b/Engine/Source/Runtime/Apple/MetalRHI/Private/MetalIndexBuffer.cpp
--- a/Engine/Source/Runtime/Apple/MetalRHI/Private/MetalIndexBuffer.cpp
+++ b/Engine/Source/Runtime/Apple/MetalRHI/Private/MetalIndexBuffer.cpp
@@ -65,12 +65,12 @@ FIndexBufferRHIRef FMetalDynamicRHI::RHICreateIndexBuffer(uint32 Stride, uint32
 		check(Size == CreateInfo.ResourceArray->GetResourceDataSize());
 
 		// make a buffer usable by CPU
-		void* Buffer = RHILockIndexBuffer(IndexBuffer, 0, Size, RLM_WriteOnly);
+		void* Buffer = ::RHILockIndexBuffer(IndexBuffer, 0, Size, RLM_WriteOnly);
 
 		// copy the contents of the given data into the buffer
 		FMemory::Memcpy(Buffer, CreateInfo.ResourceArray->GetResourceData(), Size);
 
-		RHIUnlockIndexBuffer(IndexBuffer);
+		::RHIUnlockIndexBuffer(IndexBuffer);
 
 		// Discard the resource array's contents.
 		CreateInfo.ResourceArray->Discard();
@@ -209,12 +209,12 @@ FIndexBufferRHIRef FMetalDynamicRHI::CreateIndexBuffer_RenderThread(class FRHICo
 			else
 			{
 				// make a buffer usable by CPU
-				void* Buffer = RHILockIndexBuffer(IndexBuffer, 0, Size, RLM_WriteOnly);
+				void* Buffer = RHILockIndexBuffer(RHICmdList, IndexBuffer, 0, Size, RLM_WriteOnly);
 				
 				// copy the contents of the given data into the buffer
 				FMemory::Memcpy(Buffer, CreateInfo.ResourceArray->GetResourceData(), Size);
 				
-				RHIUnlockIndexBuffer(IndexBuffer);
+				RHIUnlockIndexBuffer(RHICmdList, IndexBuffer);
 			}
 			
 			// Discard the resource array's contents.
diff --git a/Engine/Source/Runtime/Apple/MetalRHI/Private/MetalStructuredBuffer.cpp b/Engine/Source/Runtime/Apple/MetalRHI/Private/MetalStructuredBuffer.cpp
--- a/Engine/Source/Runtime/Apple/MetalRHI/Private/MetalStructuredBuffer.cpp
+++ b/Engine/Source/Runtime/Apple/MetalRHI/Private/MetalStructuredBuffer.cpp
@@ -142,12 +142,12 @@ FStructuredBufferRHIRef FMetalDynamicRHI::CreateStructuredBuffer_RenderThread(cl
 			else
 			{
 				// make a buffer usable by CPU
-				void* Buffer = RHILockStructuredBuffer(VertexBuffer, 0, Size, RLM_WriteOnly);
+				void* Buffer = RHILockStructuredBuffer(RHICmdList, VertexBuffer, 0, Size, RLM_WriteOnly);
 				
 				// copy the contents of the given data into the buffer
 				FMemory::Memcpy(Buffer, CreateInfo.ResourceArray->GetResourceData(), Size);
 				
-				RHIUnlockStructuredBuffer(VertexBuffer);
+				RHIUnlockStructuredBuffer(RHICmdList, VertexBuffer);
 			}
 			
 			// Discard the resource array's contents.
diff --git a/Engine/Source/Runtime/Apple/MetalRHI/Private/MetalVertexBuffer.cpp b/Engine/Source/Runtime/Apple/MetalRHI/Private/MetalVertexBuffer.cpp
--- a/Engine/Source/Runtime/Apple/MetalRHI/Private/MetalVertexBuffer.cpp
+++ b/Engine/Source/Runtime/Apple/MetalRHI/Private/MetalVertexBuffer.cpp
@@ -550,12 +550,12 @@ FVertexBufferRHIRef FMetalDynamicRHI::RHICreateVertexBuffer(uint32 Size, uint32
 		check(Size >= CreateInfo.ResourceArray->GetResourceDataSize());
 
 		// make a buffer usable by CPU
-		void* Buffer = RHILockVertexBuffer(VertexBuffer, 0, Size, RLM_WriteOnly);
+		void* Buffer = ::RHILockVertexBuffer(VertexBuffer, 0, Size, RLM_WriteOnly);
 		
 		// copy the contents of the given data into the buffer
 		FMemory::Memcpy(Buffer, CreateInfo.ResourceArray->GetResourceData(), Size);
 		
-		RHIUnlockVertexBuffer(VertexBuffer);
+		::RHIUnlockVertexBuffer(VertexBuffer);
 
 		// Discard the resource array's contents.
 		CreateInfo.ResourceArray->Discard();
@@ -699,12 +699,12 @@ FVertexBufferRHIRef FMetalDynamicRHI::CreateVertexBuffer_RenderThread(class FRHI
 			else
 			{
 				// make a buffer usable by CPU
-				void* Buffer = RHILockVertexBuffer(VertexBuffer, 0, Size, RLM_WriteOnly);
+				void* Buffer = RHILockVertexBuffer(RHICmdList, VertexBuffer, 0, Size, RLM_WriteOnly);
 				
 				// copy the contents of the given data into the buffer
 				FMemory::Memcpy(Buffer, CreateInfo.ResourceArray->GetResourceData(), Size);
 				
-				RHIUnlockVertexBuffer(VertexBuffer);
+				RHIUnlockVertexBuffer(RHICmdList, VertexBuffer);
 			}
 			
 			// Discard the resource array's contents.
