Reserve index 0 for the buffer size table in IABs - it just makes life easier.

diff --git a/Engine/Source/Developer/Apple/MetalShaderFormat/Private/MetalDerivedData.cpp b/Engine/Source/Developer/Apple/MetalShaderFormat/Private/MetalDerivedData.cpp
--- a/Engine/Source/Developer/Apple/MetalShaderFormat/Private/MetalDerivedData.cpp
+++ b/Engine/Source/Developer/Apple/MetalShaderFormat/Private/MetalDerivedData.cpp
@@ -1154,7 +1154,7 @@ bool FMetalShaderOutputCooker::Build(TArray<uint8>& OutData)
 						}
 					}
 					
-					SPVRResult = Reflection.ChangeDescriptorBindingNumbers(Binding, Entry->ResourceIndex, Entry->SetIndex);
+					SPVRResult = Reflection.ChangeDescriptorBindingNumbers(Binding, Entry->ResourceIndex + 1, Entry->SetIndex);
 					check(SPVRResult == SPV_REFLECT_RESULT_SUCCESS);
 				}
 				
@@ -1177,11 +1177,12 @@ bool FMetalShaderOutputCooker::Build(TArray<uint8>& OutData)
 							{
 								IABString += TEXT(",");
 							}
-							IABString += FString::Printf(TEXT("%u"), Resource.ResourceIndex);
+							IABString += FString::Printf(TEXT("%u"), Resource.ResourceIndex + 1);
 							bComma = true;
 						}
 					}
 					IABString += TEXT("]");
+					UBOString += FString::Printf(TEXT("%s%s(%u)"), UBOString.Len() ? TEXT(",") : TEXT(""), *Name, SetIndex);
 				}
 				
 				for (auto const& Binding : TBufferSRVBindings)
diff --git a/Engine/Source/Runtime/Apple/MetalRHI/Private/MetalUniformBuffer.cpp b/Engine/Source/Runtime/Apple/MetalRHI/Private/MetalUniformBuffer.cpp
--- a/Engine/Source/Runtime/Apple/MetalRHI/Private/MetalUniformBuffer.cpp
+++ b/Engine/Source/Runtime/Apple/MetalRHI/Private/MetalUniformBuffer.cpp
@@ -204,6 +204,7 @@ FMetalUniformBuffer::FMetalUniformBuffer(const void* Contents, const FRHIUniform
 	
 	if (FMetalCommandQueue::SupportsFeature(EMetalFeaturesIABs))
 	{
+		ResourceTypes.AddZeroed(NumResources);
 		for (int32 i = 0; i < NumResources; ++i)
 		{
 			ResourceTypes[i] = Layout.Resources[i].MemberType;
@@ -263,6 +264,14 @@ void FMetalUniformBuffer::InitIAB()
 		TArray<FMetalArgumentDesc> Arguments;
 		
 		int32 Index = 0;
+		// Buffer Size table
+		{
+			FMetalArgumentDesc& Desc = Arguments.Emplace_GetRef();
+			Desc.SetIndex(Index++);
+			Desc.SetAccess(mtlpp::ArgumentAccess::ReadOnly);
+			Desc.SetDataType(mtlpp::DataType::Pointer);
+		}
+		
 		for (int32 i = 0; i < NumResources; ++i, ++Index)
 		{
 			FRHIResource* Resource = ResourceTable[i].GetReference();
@@ -309,44 +318,21 @@ void FMetalUniformBuffer::InitIAB()
 						check(!Surface->Texture.IsAliasable());
 						NewIAB->IndirectArgumentResources.Add(Argument(Surface->Texture, (mtlpp::ResourceUsage)(mtlpp::ResourceUsage::Read|mtlpp::ResourceUsage::Sample)));
 					}
+					else if (SB)
+					{
+						Desc.SetDataType(mtlpp::DataType::Pointer);
+						NewIAB->IndirectArgumentResources.Add(Argument(SB->Buffer, (mtlpp::ResourceUsage)(mtlpp::ResourceUsage::Read)));
+						
+						BufferSizes.Add(SB->GetSize());
+						BufferSizes.Add(GMetalBufferFormats[SRV->Format].DataFormat);
+					}
 					else
 					{
-						check(VB || IB || SB);
+						check(VB || IB);
 						ns::AutoReleased<FMetalTexture> Tex = SRV->GetLinearTexture(false);
 						Desc.SetDataType(mtlpp::DataType::Texture);
 						Desc.SetTextureType(Tex.GetTextureType());
 						NewIAB->IndirectArgumentResources.Add(Argument(Tex, (mtlpp::ResourceUsage)(mtlpp::ResourceUsage::Read|mtlpp::ResourceUsage::Sample)));
-						
-						FMetalArgumentDesc& BufferDesc = Arguments.Emplace_GetRef();
-						BufferDesc.SetIndex(++Index);
-						BufferDesc.SetAccess(mtlpp::ArgumentAccess::ReadOnly);
-						
-						if (VB)
-						{
-							BufferDesc.SetDataType(mtlpp::DataType::Pointer);
-							NewIAB->IndirectArgumentResources.Add(Argument(VB->Buffer, (mtlpp::ResourceUsage)(mtlpp::ResourceUsage::Read)));
-							
-							check(VB->Buffer.GetStorageMode() == mtlpp::StorageMode::Private);
-							BufferSizes.Add(VB->GetSize());
-							BufferSizes.Add(GMetalBufferFormats[SRV->Format].DataFormat);
-						}
-						else if (IB)
-						{
-							BufferDesc.SetDataType(mtlpp::DataType::Pointer);
-							NewIAB->IndirectArgumentResources.Add(Argument(IB->Buffer, (mtlpp::ResourceUsage)(mtlpp::ResourceUsage::Read)));
-							
-							check(IB->Buffer.GetStorageMode() == mtlpp::StorageMode::Private);
-							BufferSizes.Add(IB->GetSize());
-							BufferSizes.Add(GMetalBufferFormats[SRV->Format].DataFormat);
-						}
-						else if (SB)
-						{
-							BufferDesc.SetDataType(mtlpp::DataType::Pointer);
-							NewIAB->IndirectArgumentResources.Add(Argument(SB->Buffer, (mtlpp::ResourceUsage)(mtlpp::ResourceUsage::Read)));
-							
-							BufferSizes.Add(SB->GetSize());
-							BufferSizes.Add(GMetalBufferFormats[SRV->Format].DataFormat);
-						}
 					}
 					break;
 				}
@@ -487,11 +473,6 @@ void FMetalUniformBuffer::InitIAB()
 		
 		if (BufferSizes.Num() > 0)
 		{
-			FMetalArgumentDesc& Desc = Arguments.Emplace_GetRef();
-			Desc.SetIndex(Index++);
-			Desc.SetAccess(mtlpp::ArgumentAccess::ReadOnly);
-			Desc.SetDataType(mtlpp::DataType::Pointer);
-			
 			FMetalPooledBufferArgs Args(GetMetalDeviceContext().GetDevice(), BufferSizes.Num() * sizeof(uint32), BUF_Dynamic, BUFFER_STORAGE_MODE);
 			NewIAB->IndirectArgumentBufferSideTable = GetMetalDeviceContext().CreatePooledBuffer(Args);
 			
