Metal version 6 should also use IABs now that DXC can emit some functional shaders.

diff --git a/Engine/Source/Runtime/Apple/MetalRHI/Private/MetalCommandQueue.cpp b/Engine/Source/Runtime/Apple/MetalRHI/Private/MetalCommandQueue.cpp
--- a/Engine/Source/Runtime/Apple/MetalRHI/Private/MetalCommandQueue.cpp
+++ b/Engine/Source/Runtime/Apple/MetalRHI/Private/MetalCommandQueue.cpp
@@ -239,7 +239,7 @@ FMetalCommandQueue::FMetalCommandQueue(mtlpp::Device InDevice, uint32 const MaxN
 			{
 				Features |= EMetalFeaturesTextureBuffers;
             }
-            if (MaxShaderVersion == 5)
+            if (MaxShaderVersion >= 5)
             {
                 Features |= EMetalFeaturesIABs;
             }
