PF_Unknown is not a valid pixel format for creating texture views - when that is passed in use the source texture's format.

diff --git a/Engine/Source/Runtime/Apple/MetalRHI/Private/MetalUAV.cpp b/Engine/Source/Runtime/Apple/MetalRHI/Private/MetalUAV.cpp
--- a/Engine/Source/Runtime/Apple/MetalRHI/Private/MetalUAV.cpp
+++ b/Engine/Source/Runtime/Apple/MetalRHI/Private/MetalUAV.cpp
@@ -242,7 +242,7 @@ FShaderResourceViewRHIRef FMetalDynamicRHI::RHICreateShaderResourceView(FRHIText
 		SRV->SourceTexture = (FRHITexture*)Texture2DRHI;
 		
 		FMetalSurface* Surface = GetMetalSurfaceFromRHITexture(Texture2DRHI);
-		SRV->TextureView = Surface ? new FMetalSurface(*Surface, NSMakeRange(CreateInfo.MipLevel, CreateInfo.NumMipLevels), (EPixelFormat)CreateInfo.Format) : nullptr;
+		SRV->TextureView = Surface ? new FMetalSurface(*Surface, NSMakeRange(CreateInfo.MipLevel, CreateInfo.NumMipLevels), (EPixelFormat)(CreateInfo.Format == PF_Unknown ? Surface->PixelFormat : CreateInfo.Format)) : nullptr;
 		
 		SRV->SourceVertexBuffer = nullptr;
 		SRV->SourceIndexBuffer = nullptr;
