Buffer management refactoring for Metal, prompted by Luke's changes to the RHI API.
- Fixed the allocation of the ring-buffer so that it can actually shrink.
- Updated the code to shrink the ring-buffer on both the prologe and current encoders.
- Pushed more code from the RHI buffer types into the common FMetalRHIBuffer base-class.
- Wherever possible favour using a temporary Shared buffer and a blit operation to load data into the GPU - this avoids buffer renaming which is expensive on the CPU (as all the linear textures used in UAVs & SRVs have to be reallocated).
- Uniform buffers start life as Shared/Managed so they can be created from any thread, but calling Update will reallcoate them as Private memory (when available). This avoids what appears to be a race condition between didModifyRange and the blit operations that I just can't explain (didModifyRange happens on the CPU prior to commiting the command-buffer containing the blit operations!), but which is obviously the source of long-running buffer update errors.
- Shouldn't affect iOS unless you enable EfficientBufferBlits or IABs.
- Still need to apply similar logic to IABs themselves as this may be partly why they never worked.

diff --git a/Engine/Source/Runtime/Apple/MetalRHI/Private/MetalBuffer.cpp b/Engine/Source/Runtime/Apple/MetalRHI/Private/MetalBuffer.cpp
--- a/Engine/Source/Runtime/Apple/MetalRHI/Private/MetalBuffer.cpp
+++ b/Engine/Source/Runtime/Apple/MetalRHI/Private/MetalBuffer.cpp
@@ -675,7 +675,7 @@ FMetalSubBufferMagazine::FMetalSubBufferMagazine(NSUInteger Size, NSUInteger Chu
 	Options = (mtlpp::ResourceOptions)FMetalCommandQueue::GetCompatibleResourceOptions(Options);
     static bool bSupportsHeaps = GetMetalDeviceContext().SupportsFeature(EMetalFeaturesHeaps);
     mtlpp::StorageMode Storage = (mtlpp::StorageMode)((Options & mtlpp::ResourceStorageModeMask) >> mtlpp::ResourceStorageModeShift);
-    if (bSupportsHeaps && Storage == mtlpp::StorageMode::Private)
+    if (PLATFORM_IOS && bSupportsHeaps && Storage == mtlpp::StorageMode::Private)
     {
         MinAlign = GetMetalDeviceContext().GetDevice().HeapBufferSizeAndAlign(BlockSize, Options).Align;
     }
@@ -941,6 +941,7 @@ FMetalSubBufferRing::FMetalSubBufferRing(NSUInteger Size, NSUInteger Alignment,
 , CommitHead(0)
 , SubmitHead(0)
 , WriteHead(0)
+, BufferSize(0)
 , Options(InOptions)
 , Storage((mtlpp::StorageMode)((Options & mtlpp::ResourceStorageModeMask) >> mtlpp::ResourceStorageModeShift))
 {
@@ -987,6 +988,7 @@ FMetalBuffer FMetalSubBufferRing::NewBuffer(NSUInteger Size, uint32 Alignment)
 	if(!Buffer.IsValid())
 	{
 		Buffer = MakeShared<FMetalRingBufferRef, ESPMode::ThreadSafe>(GetMetalDeviceContext().GetResourceHeap().CreateBuffer(InitialSize, MinAlign, BUF_Dynamic, Options, true));
+		BufferSize = InitialSize;
 	}
 	
 	if(Buffer->LastRead <= WriteHead)
@@ -1005,7 +1007,7 @@ FMetalBuffer FMetalSubBufferRing::NewBuffer(NSUInteger Size, uint32 Alignment)
 		else if (Storage == mtlpp::StorageMode::Managed)
 		{
 			Submit();
-			Buffer = MakeShared<FMetalRingBufferRef, ESPMode::ThreadSafe>(GetMetalDeviceContext().GetResourceHeap().CreateBuffer(Buffer->Buffer.GetLength(), MinAlign, BUF_Dynamic, Options, true));
+			Buffer = MakeShared<FMetalRingBufferRef, ESPMode::ThreadSafe>(GetMetalDeviceContext().GetResourceHeap().CreateBuffer(BufferSize, MinAlign, BUF_Dynamic, Options, true));
 			WriteHead = 0;
 			CommitHead = 0;
 			SubmitHead = 0;
@@ -1017,16 +1019,16 @@ FMetalBuffer FMetalSubBufferRing::NewBuffer(NSUInteger Size, uint32 Alignment)
 		}
 	}
 	
-	const NSUInteger BufferSize = Buffer->Buffer.GetLength();
 	if(WriteHead + FullSize >= Buffer->LastRead || WriteHead + FullSize > BufferSize)
 	{
-		NSUInteger NewBufferSize = AlignArbitrary(BufferSize + Size, Align(Buffer->Buffer.GetLength() / 4, MinAlign));
+		NSUInteger NewBufferSize = AlignArbitrary(BufferSize + Size, Align(BufferSize / 4, MinAlign));
 		
 		UE_LOG(LogMetal, Verbose, TEXT("Reallocating ring-buffer from %d to %d to avoid wrapping write at offset %d into outstanding buffer region %d at frame %lld]"), (uint32)BufferSize, (uint32)NewBufferSize, (uint32)WriteHead, (uint32)Buffer->LastRead, (uint64)GFrameCounter);
 		
 		Submit();
 		
 		Buffer = MakeShared<FMetalRingBufferRef, ESPMode::ThreadSafe>(GetMetalDeviceContext().GetResourceHeap().CreateBuffer(NewBufferSize, MinAlign, BUF_Dynamic, Options, true));
+		BufferSize = NewBufferSize;
 		WriteHead = 0;
 		CommitHead = 0;
 		SubmitHead = 0;
@@ -1055,16 +1057,16 @@ void FMetalSubBufferRing::Shrink()
 		}
 		
 		NSUInteger NecessarySize = FMath::Max(FrameMax, InitialSize);
-		NSUInteger ThreeQuarterSize = Align((Buffer->Buffer.GetLength() / 4) * 3, MinAlign);
+		NSUInteger ThreeQuarterSize = Align((BufferSize / 4) * 3, MinAlign);
 		
-		if ((GFrameNumberRenderThread - LastFrameChange) >= 120 && NecessarySize < ThreeQuarterSize && NecessarySize < Buffer->Buffer.GetLength())
+		if ((GFrameNumberRenderThread - LastFrameChange) >= 120 && NecessarySize < ThreeQuarterSize && NecessarySize < BufferSize)
 		{
 			Submit();
 			
 			UE_LOG(LogMetal, Verbose, TEXT("Shrinking RingBuffer from %u to %u as max. usage is %u at frame %lld]"), (uint32)Buffer->Buffer.GetLength(), (uint32)ThreeQuarterSize, (uint32)FrameMax, GFrameNumberRenderThread);
 			
 			Buffer = MakeShared<FMetalRingBufferRef, ESPMode::ThreadSafe>(GetMetalDeviceContext().GetResourceHeap().CreateBuffer(ThreeQuarterSize, MinAlign, BUF_Dynamic, Options, true));
-			
+			BufferSize = ThreeQuarterSize;
 			WriteHead = 0;
 			CommitHead = 0;
 			SubmitHead = 0;
diff --git a/Engine/Source/Runtime/Apple/MetalRHI/Private/MetalBuffer.h b/Engine/Source/Runtime/Apple/MetalRHI/Private/MetalBuffer.h
--- a/Engine/Source/Runtime/Apple/MetalRHI/Private/MetalBuffer.h
+++ b/Engine/Source/Runtime/Apple/MetalRHI/Private/MetalBuffer.h
@@ -177,6 +177,7 @@ private:
 	NSUInteger CommitHead;
 	NSUInteger SubmitHead;
 	NSUInteger WriteHead;
+	NSUInteger BufferSize;
 	mtlpp::ResourceOptions Options;
 	mtlpp::StorageMode Storage;
 	TSharedPtr<FMetalRingBufferRef, ESPMode::ThreadSafe> Buffer;
diff --git a/Engine/Source/Runtime/Apple/MetalRHI/Private/MetalContext.cpp b/Engine/Source/Runtime/Apple/MetalRHI/Private/MetalContext.cpp
--- a/Engine/Source/Runtime/Apple/MetalRHI/Private/MetalContext.cpp
+++ b/Engine/Source/Runtime/Apple/MetalRHI/Private/MetalContext.cpp
@@ -1074,7 +1074,7 @@ void FMetalContext::InitFrame(bool const bImmediateContext, uint32 Index, uint32
 	}
 	
 	// Reallocate if necessary to ensure >= 80% usage, otherwise we're just too wasteful
-	// RenderPass.GetRingBuffer().Shrink();
+	RenderPass.ShrinkRingBuffers();
 	
 	// Begin the render pass frame.
 	RenderPass.Begin(StartFence);
diff --git a/Engine/Source/Runtime/Apple/MetalRHI/Private/MetalIndexBuffer.cpp b/Engine/Source/Runtime/Apple/MetalRHI/Private/MetalIndexBuffer.cpp
--- a/Engine/Source/Runtime/Apple/MetalRHI/Private/MetalIndexBuffer.cpp
+++ b/Engine/Source/Runtime/Apple/MetalRHI/Private/MetalIndexBuffer.cpp
@@ -75,13 +75,9 @@ FIndexBufferRHIRef FMetalDynamicRHI::RHICreateIndexBuffer(uint32 Stride, uint32
 		// Discard the resource array's contents.
 		CreateInfo.ResourceArray->Discard();
 	}
-	else if (IndexBuffer->Buffer.GetStorageMode() == mtlpp::StorageMode::Private)
+	else if (IndexBuffer->Mode == mtlpp::StorageMode::Private)
 	{
-		if (IndexBuffer->UsePrivateMemory())
-		{
-			SafeReleaseMetalBuffer(IndexBuffer->CPUBuffer);
-			IndexBuffer->CPUBuffer = nil;
-		}
+		check (!IndexBuffer->CPUBuffer);
 
 		if (GMetalBufferZeroFill && !FMetalCommandQueue::SupportsFeature(EMetalFeaturesFences))
 		{
@@ -89,7 +85,7 @@ FIndexBufferRHIRef FMetalDynamicRHI::RHICreateIndexBuffer(uint32 Stride, uint32
 		}
 	}
 #if PLATFORM_MAC
-	else if (GMetalBufferZeroFill && IndexBuffer->Buffer.GetStorageMode() == mtlpp::StorageMode::Managed)
+	else if (GMetalBufferZeroFill && IndexBuffer->Mode == mtlpp::StorageMode::Managed)
 	{
 		MTLPP_VALIDATE(mtlpp::Buffer, IndexBuffer->Buffer, SafeGetRuntimeDebuggingLevel() >= EMetalDebugLevelValidation, DidModify(ns::Range(0, IndexBuffer->Buffer.GetLength())));
 	}
@@ -123,7 +119,7 @@ void* FMetalDynamicRHI::LockIndexBuffer_BottomOfPipe(FRHICommandListImmediate& R
 	@autoreleasepool {
 	FMetalIndexBuffer* IndexBuffer = ResourceCast(IndexBufferRHI);
 
-	return (uint8*)IndexBuffer->Lock(LockMode, Offset, Size);
+	return (uint8*)IndexBuffer->Lock(true, LockMode, Offset, Size);
 	}
 }
 
@@ -136,40 +132,6 @@ void FMetalDynamicRHI::UnlockIndexBuffer_BottomOfPipe(FRHICommandListImmediate&
 	}
 }
 
-struct FMetalRHICommandInitialiseIndexBuffer : public FRHICommand<FMetalRHICommandInitialiseIndexBuffer>
-{
-	TRefCountPtr<FMetalIndexBuffer> Buffer;
-	
-	FORCEINLINE_DEBUGGABLE FMetalRHICommandInitialiseIndexBuffer(FMetalIndexBuffer* InBuffer)
-	: Buffer(InBuffer)
-	{
-	}
-	
-	virtual ~FMetalRHICommandInitialiseIndexBuffer()
-	{
-	}
-	
-	void Execute(FRHICommandListBase& CmdList)
-	{
-		if (Buffer->CPUBuffer)
-		{
-			GetMetalDeviceContext().AsyncCopyFromBufferToBuffer(Buffer->CPUBuffer, 0, Buffer->Buffer, 0, Buffer->Buffer.GetLength());
-			if (Buffer->UsePrivateMemory())
-			{
-				SafeReleaseMetalBuffer(Buffer->CPUBuffer);
-			}
-			else
-			{
-				Buffer->LastUpdate = GFrameNumberRenderThread;
-			}
-		}
-		else if (GMetalBufferZeroFill && !FMetalCommandQueue::SupportsFeature(EMetalFeaturesFences))
-		{
-			GetMetalDeviceContext().FillBuffer(Buffer->Buffer, ns::Range(0, Buffer->Buffer.GetLength()), 0);
-		}
-	}
-};
-
 FIndexBufferRHIRef FMetalDynamicRHI::CreateIndexBuffer_RenderThread(class FRHICommandListImmediate& RHICmdList, uint32 Stride, uint32 Size, uint32 InUsage, FRHIResourceCreateInfo& CreateInfo)
 {
 	@autoreleasepool {
@@ -181,72 +143,7 @@ FIndexBufferRHIRef FMetalDynamicRHI::CreateIndexBuffer_RenderThread(class FRHICo
 		// make the RHI object, which will allocate memory
 		TRefCountPtr<FMetalIndexBuffer> IndexBuffer = new FMetalIndexBuffer(Stride, Size, InUsage);
 		
-		if (CreateInfo.ResourceArray)
-		{
-			check(Size == CreateInfo.ResourceArray->GetResourceDataSize());
-			
-			if (IndexBuffer->CPUBuffer)
-			{
-				FMemory::Memcpy(IndexBuffer->CPUBuffer.GetContents(), CreateInfo.ResourceArray->GetResourceData(), Size);
-
-#if PLATFORM_MAC
-				if(IndexBuffer->CPUBuffer.GetStorageMode() == mtlpp::StorageMode::Managed)
-				{
-					MTLPP_VALIDATE(mtlpp::Buffer, IndexBuffer->CPUBuffer, SafeGetRuntimeDebuggingLevel() >= EMetalDebugLevelValidation, DidModify(ns::Range(0, GMetalBufferZeroFill ? IndexBuffer->CPUBuffer.GetLength() : Size)));
-				}
-#endif
-				
-				if (RHICmdList.Bypass() || !IsRunningRHIInSeparateThread())
-				{
-					FMetalRHICommandInitialiseIndexBuffer UpdateCommand(IndexBuffer);
-					UpdateCommand.Execute(RHICmdList);
-				}
-				else
-				{
-					new (RHICmdList.AllocCommand<FMetalRHICommandInitialiseIndexBuffer>()) FMetalRHICommandInitialiseIndexBuffer(IndexBuffer);
-				}
-			}
-			else
-			{
-				// make a buffer usable by CPU
-				void* Buffer = RHILockIndexBuffer(RHICmdList, IndexBuffer, 0, Size, RLM_WriteOnly);
-				
-				// copy the contents of the given data into the buffer
-				FMemory::Memcpy(Buffer, CreateInfo.ResourceArray->GetResourceData(), Size);
-				
-				RHIUnlockIndexBuffer(RHICmdList, IndexBuffer);
-			}
-			
-			// Discard the resource array's contents.
-			CreateInfo.ResourceArray->Discard();
-		}
-		else if (IndexBuffer->Buffer.GetStorageMode() == mtlpp::StorageMode::Private)
-		{
-			if (IndexBuffer->UsePrivateMemory())
-			{
-				SafeReleaseMetalBuffer(IndexBuffer->CPUBuffer);
-				IndexBuffer->CPUBuffer = nil;
-			}
-			
-			if (GMetalBufferZeroFill)
-			{
-				if (RHICmdList.Bypass() || !IsRunningRHIInSeparateThread())
-				{
-					FMetalRHICommandInitialiseIndexBuffer UpdateCommand(IndexBuffer);
-					UpdateCommand.Execute(RHICmdList);
-				}
-				else
-				{
-					new (RHICmdList.AllocCommand<FMetalRHICommandInitialiseIndexBuffer>()) FMetalRHICommandInitialiseIndexBuffer(IndexBuffer);
-				}
-			}
-		}
-#if PLATFORM_MAC
-		else if (GMetalBufferZeroFill && IndexBuffer->Buffer.GetStorageMode() == mtlpp::StorageMode::Managed)
-		{
-			MTLPP_VALIDATE(mtlpp::Buffer, IndexBuffer->Buffer, SafeGetRuntimeDebuggingLevel() >= EMetalDebugLevelValidation, DidModify(ns::Range(0, IndexBuffer->Buffer.GetLength())));
-		}
-#endif
+		IndexBuffer->Init_RenderThread(RHICmdList, Size, InUsage, CreateInfo, IndexBuffer);
 		
 		return IndexBuffer.GetReference();
 	}
diff --git a/Engine/Source/Runtime/Apple/MetalRHI/Private/MetalRenderPass.cpp b/Engine/Source/Runtime/Apple/MetalRHI/Private/MetalRenderPass.cpp
--- a/Engine/Source/Runtime/Apple/MetalRHI/Private/MetalRenderPass.cpp
+++ b/Engine/Source/Runtime/Apple/MetalRHI/Private/MetalRenderPass.cpp
@@ -1550,6 +1550,21 @@ void FMetalRenderPass::AsyncCopyFromBufferToBuffer(FMetalBuffer const& SourceBuf
 	METAL_DEBUG_LAYER(EMetalDebugLevelFastValidation, Debugging.Copy(SourceBuffer, SourceOffset, DestinationBuffer, DestinationOffset, Size));
 }
 
+FMetalBuffer FMetalRenderPass::AllocateTemporyBufferForCopy(FMetalBuffer const& DestinationBuffer, NSUInteger Size, NSUInteger Align)
+{
+	FMetalBuffer Buffer;
+	bool bAsync = !CurrentEncoder.HasBufferBindingHistory(DestinationBuffer);
+	if(bAsync)
+	{
+		Buffer = PrologueEncoder.GetRingBuffer().NewBuffer(Size, Align);
+	}
+	else
+	{
+		Buffer = CurrentEncoder.GetRingBuffer().NewBuffer(Size, Align);
+	}
+	return Buffer;
+}
+
 void FMetalRenderPass::AsyncGenerateMipmapsForTexture(FMetalTexture const& Texture)
 {
 	// This must be a plain old error
@@ -1697,6 +1712,12 @@ FMetalSubBufferRing& FMetalRenderPass::GetRingBuffer(void)
 	return CurrentEncoder.GetRingBuffer();
 }
 
+void FMetalRenderPass::ShrinkRingBuffers(void)
+{
+	PrologueEncoder.GetRingBuffer().Shrink();
+	CurrentEncoder.GetRingBuffer().Shrink();
+}
+
 bool FMetalRenderPass::IsWithinParallelPass(void)
 {
 	return bWithinRenderPass && CurrentEncoder.IsParallelRenderCommandEncoderActive();
diff --git a/Engine/Source/Runtime/Apple/MetalRHI/Private/MetalRenderPass.h b/Engine/Source/Runtime/Apple/MetalRHI/Private/MetalRenderPass.h
--- a/Engine/Source/Runtime/Apple/MetalRHI/Private/MetalRenderPass.h
+++ b/Engine/Source/Runtime/Apple/MetalRHI/Private/MetalRenderPass.h
@@ -80,6 +80,8 @@ public:
 	
 	void AsyncCopyFromBufferToBuffer(FMetalBuffer const& SourceBuffer, NSUInteger SourceOffset, FMetalBuffer const& DestinationBuffer, NSUInteger DestinationOffset, NSUInteger Size);
 	
+	FMetalBuffer AllocateTemporyBufferForCopy(FMetalBuffer const& DestinationBuffer, NSUInteger Size, NSUInteger Align);
+	
 	void AsyncGenerateMipmapsForTexture(FMetalTexture const& Texture);
 	
     TRefCountPtr<FMetalFence> const& Submit(EMetalSubmitFlags SubmissionFlags);
@@ -137,6 +139,11 @@ public:
 	 */
 	FMetalSubBufferRing& GetRingBuffer(void);
 	
+	/*
+	 * Attempts to shrink the ring-buffers so we don't keep very large allocations when we don't need them.
+	 */
+	void ShrinkRingBuffers(void);
+	
 	/*
 	 * Whether the render-pass is within a parallel rendering pass.
 	 * @returns True if and only if within a parallel rendering pass, otherwise false.
diff --git a/Engine/Source/Runtime/Apple/MetalRHI/Private/MetalStructuredBuffer.cpp b/Engine/Source/Runtime/Apple/MetalRHI/Private/MetalStructuredBuffer.cpp
--- a/Engine/Source/Runtime/Apple/MetalRHI/Private/MetalStructuredBuffer.cpp
+++ b/Engine/Source/Runtime/Apple/MetalRHI/Private/MetalStructuredBuffer.cpp
@@ -31,9 +31,9 @@ FStructuredBufferRHIRef FMetalDynamicRHI::RHICreateStructuredBuffer(uint32 Strid
 {
 	@autoreleasepool {
 	FMetalStructuredBuffer* Buffer = new FMetalStructuredBuffer(Stride, Size, CreateInfo.ResourceArray, InUsage);
-	if (!CreateInfo.ResourceArray && Buffer->Buffer.GetStorageMode() == mtlpp::StorageMode::Private)
+	if (!CreateInfo.ResourceArray && Buffer->Mode == mtlpp::StorageMode::Private)
 	{
-		if (Buffer->GetUsage() & (BUF_Dynamic|BUF_Static))
+		if (Buffer->CPUBuffer)
 		{
 			SafeReleaseMetalBuffer(Buffer->CPUBuffer);
 			Buffer->CPUBuffer = nil;
@@ -45,7 +45,7 @@ FStructuredBufferRHIRef FMetalDynamicRHI::RHICreateStructuredBuffer(uint32 Strid
 		}
 	}
 #if PLATFORM_MAC
-	else if (GMetalBufferZeroFill && !CreateInfo.ResourceArray && Buffer->Buffer.GetStorageMode() == mtlpp::StorageMode::Managed)
+	else if (GMetalBufferZeroFill && !CreateInfo.ResourceArray && Buffer->Mode == mtlpp::StorageMode::Managed)
 	{
 		MTLPP_VALIDATE(mtlpp::Buffer, Buffer->Buffer, SafeGetRuntimeDebuggingLevel() >= EMetalDebugLevelValidation, DidModify(ns::Range(0, Buffer->Buffer.GetLength())));
 	}
@@ -60,7 +60,7 @@ void* FMetalDynamicRHI::LockStructuredBuffer_BottomOfPipe(FRHICommandListImmedia
 	FMetalStructuredBuffer* StructuredBuffer = ResourceCast(StructuredBufferRHI);
 	
 	// just return the memory plus the offset
-	return (uint8*)StructuredBuffer->Lock(LockMode, Offset, Size);
+	return (uint8*)StructuredBuffer->Lock(true, LockMode, Offset, Size);
 	}
 }
 
@@ -72,114 +72,13 @@ void FMetalDynamicRHI::UnlockStructuredBuffer_BottomOfPipe(FRHICommandListImmedi
 	}
 }
 
-struct FMetalRHICommandInitialiseStructuredBuffer : public FRHICommand<FMetalRHICommandInitialiseStructuredBuffer>
-{
-	TRefCountPtr<FMetalStructuredBuffer> Buffer;
-	
-	FORCEINLINE_DEBUGGABLE FMetalRHICommandInitialiseStructuredBuffer(FMetalStructuredBuffer* InBuffer)
-	: Buffer(InBuffer)
-	{
-	}
-	
-	virtual ~FMetalRHICommandInitialiseStructuredBuffer()
-	{
-	}
-	
-	void Execute(FRHICommandListBase& CmdList)
-	{
-		if (Buffer->CPUBuffer)
-		{
-			GetMetalDeviceContext().AsyncCopyFromBufferToBuffer(Buffer->CPUBuffer, 0, Buffer->Buffer, 0, Buffer->Buffer.GetLength());
-			if (Buffer->GetUsage() & (BUF_Dynamic|BUF_Static))
-			{
-				SafeReleaseMetalBuffer(Buffer->CPUBuffer);
-			}
-			else
-			{
-				Buffer->LastUpdate = GFrameNumberRenderThread;
-			}
-		}
-#if PLATFORM_MAC
-		else if (GMetalBufferZeroFill && !FMetalCommandQueue::SupportsFeature(EMetalFeaturesFences))
-		{
-			GetMetalDeviceContext().FillBuffer(Buffer->Buffer, ns::Range(0, Buffer->Buffer.GetLength()), 0);
-		}
-#endif
-	}
-};
-
 FStructuredBufferRHIRef FMetalDynamicRHI::CreateStructuredBuffer_RenderThread(class FRHICommandListImmediate& RHICmdList, uint32 Stride, uint32 Size, uint32 InUsage, FRHIResourceCreateInfo& CreateInfo)
 {
 	@autoreleasepool {
 		// make the RHI object, which will allocate memory
 		TRefCountPtr<FMetalStructuredBuffer> VertexBuffer = new FMetalStructuredBuffer(Stride, Size, nullptr, InUsage);
 		
-		if (CreateInfo.ResourceArray)
-		{
-			check(Size == CreateInfo.ResourceArray->GetResourceDataSize());
-			
-			if (VertexBuffer->CPUBuffer)
-			{
-				FMemory::Memcpy(VertexBuffer->CPUBuffer.GetContents(), CreateInfo.ResourceArray->GetResourceData(), Size);
-
-#if PLATFORM_MAC
-				if(VertexBuffer->CPUBuffer.GetStorageMode() == mtlpp::StorageMode::Managed)
-				{
-					MTLPP_VALIDATE(mtlpp::Buffer, VertexBuffer->CPUBuffer, SafeGetRuntimeDebuggingLevel() >= EMetalDebugLevelValidation, DidModify(ns::Range(0, GMetalBufferZeroFill ? VertexBuffer->CPUBuffer.GetLength() : Size)));
-				}
-#endif
-				
-				if (RHICmdList.Bypass() || !IsRunningRHIInSeparateThread())
-				{
-					FMetalRHICommandInitialiseStructuredBuffer UpdateCommand(VertexBuffer);
-					UpdateCommand.Execute(RHICmdList);
-				}
-				else
-				{
-					new (RHICmdList.AllocCommand<FMetalRHICommandInitialiseStructuredBuffer>()) FMetalRHICommandInitialiseStructuredBuffer(VertexBuffer);
-				}
-			}
-			else
-			{
-				// make a buffer usable by CPU
-				void* Buffer = RHILockStructuredBuffer(RHICmdList, VertexBuffer, 0, Size, RLM_WriteOnly);
-				
-				// copy the contents of the given data into the buffer
-				FMemory::Memcpy(Buffer, CreateInfo.ResourceArray->GetResourceData(), Size);
-				
-				RHIUnlockStructuredBuffer(RHICmdList, VertexBuffer);
-			}
-			
-			// Discard the resource array's contents.
-			CreateInfo.ResourceArray->Discard();
-		}
-		else if (VertexBuffer->Buffer.GetStorageMode() == mtlpp::StorageMode::Private)
-		{
-			if (VertexBuffer->GetUsage() & (BUF_Dynamic|BUF_Static))
-			{
-				SafeReleaseMetalBuffer(VertexBuffer->CPUBuffer);
-				VertexBuffer->CPUBuffer = nil;
-			}
-
-			if (GMetalBufferZeroFill)
-			{
-				if (RHICmdList.Bypass() || !IsRunningRHIInSeparateThread())
-				{
-					FMetalRHICommandInitialiseStructuredBuffer UpdateCommand(VertexBuffer);
-					UpdateCommand.Execute(RHICmdList);
-				}
-				else
-				{
-					new (RHICmdList.AllocCommand<FMetalRHICommandInitialiseStructuredBuffer>()) FMetalRHICommandInitialiseStructuredBuffer(VertexBuffer);
-				}
-			}
-		}
-#if PLATFORM_MAC
-		else if (GMetalBufferZeroFill && VertexBuffer->Buffer.GetStorageMode() == mtlpp::StorageMode::Managed)
-		{
-			MTLPP_VALIDATE(mtlpp::Buffer, VertexBuffer->Buffer, SafeGetRuntimeDebuggingLevel() >= EMetalDebugLevelValidation, DidModify(ns::Range(0, VertexBuffer->Buffer.GetLength())));
-		}
-#endif
+		VertexBuffer->Init_RenderThread(RHICmdList, Size, InUsage, CreateInfo, VertexBuffer);
 		
 		return VertexBuffer.GetReference();
 	}
diff --git a/Engine/Source/Runtime/Apple/MetalRHI/Private/MetalUniformBuffer.cpp b/Engine/Source/Runtime/Apple/MetalRHI/Private/MetalUniformBuffer.cpp
--- a/Engine/Source/Runtime/Apple/MetalRHI/Private/MetalUniformBuffer.cpp
+++ b/Engine/Source/Runtime/Apple/MetalRHI/Private/MetalUniformBuffer.cpp
@@ -211,7 +211,28 @@ FMetalUniformBuffer::FMetalUniformBuffer(const void* Contents, const FRHIUniform
 		}
 	}
 
-    Update(Contents, ResourceTable, Validation);
+	if (ConstantSize > 0)
+	{
+		UE_CLOG(ConstantSize > 65536, LogMetal, Fatal, TEXT("Trying to allocated a uniform layout of size %d that is greater than the maximum permitted 64k."), ConstantSize);
+		
+		if (Buffer)
+		{
+			FMemory::Memcpy(Buffer.GetContents(), Contents, ConstantSize);
+#if PLATFORM_MAC
+			if(Mode == mtlpp::StorageMode::Managed)
+			{
+				MTLPP_VALIDATE(mtlpp::Buffer, Buffer, SafeGetRuntimeDebuggingLevel() >= EMetalDebugLevelValidation, DidModify(ns::Range(0, ConstantSize)));
+			}
+#endif
+		}
+		else
+		{
+			check(Data && Data->Data);
+			FMemory::Memcpy(Data->Data, Contents, ConstantSize);
+		}
+	}
+	
+    UpdateResourceTable(ResourceTable, Validation);
 	
     if (NumResources && FMetalCommandQueue::SupportsFeature(EMetalFeaturesIABs))
     {
@@ -397,7 +418,7 @@ void FMetalUniformBuffer::InitIAB()
 							BufferDesc.SetDataType(mtlpp::DataType::Pointer);
 							NewIAB->IndirectArgumentResources.Add(Argument(VB->Buffer, (mtlpp::ResourceUsage)(mtlpp::ResourceUsage::Read|mtlpp::ResourceUsage::Write)));
 							
-							check(VB->Buffer.GetStorageMode() == mtlpp::StorageMode::Private);
+							check(VB->Mode == mtlpp::StorageMode::Private);
 							BufferSizes.Add(VB->GetSize());
 							BufferSizes.Add(GMetalBufferFormats[SRV->Format].DataFormat);
 						}
@@ -406,7 +427,7 @@ void FMetalUniformBuffer::InitIAB()
 							BufferDesc.SetDataType(mtlpp::DataType::Pointer);
 							NewIAB->IndirectArgumentResources.Add(Argument(IB->Buffer, (mtlpp::ResourceUsage)(mtlpp::ResourceUsage::Read|mtlpp::ResourceUsage::Write)));
 							
-							check(IB->Buffer.GetStorageMode() == mtlpp::StorageMode::Private);
+							check(IB->Mode == mtlpp::StorageMode::Private);
 							BufferSizes.Add(IB->GetSize());
 							BufferSizes.Add(GMetalBufferFormats[SRV->Format].DataFormat);
 						}
@@ -556,17 +577,8 @@ void const* FMetalUniformBuffer::GetData()
 	}
 }
 
-void FMetalUniformBuffer::Update(const void* Contents, TArray<TRefCountPtr<FRHIResource>>& Resources, EUniformBufferValidation Validation)
+void FMetalUniformBuffer::UpdateResourceTable(TArray<TRefCountPtr<FRHIResource>>& Resources, EUniformBufferValidation Validation)
 {
-    if (ConstantSize > 0)
-    {
-        UE_CLOG(ConstantSize > 65536, LogMetal, Fatal, TEXT("Trying to allocated a uniform layout of size %d that is greater than the maximum permitted 64k."), ConstantSize);
-        
-        void* Data = Lock(RLM_WriteOnly, 0);
-        FMemory::Memcpy(Data, Contents, ConstantSize);
-        Unlock();
-    }
-	
 	ResourceTable = Resources;
 	
 	if (NumResources && FMetalCommandQueue::SupportsFeature(EMetalFeaturesIABs))
@@ -629,6 +641,20 @@ void FMetalUniformBuffer::Update(const void* Contents, TArray<TRefCountPtr<FRHIR
 	}
 }
 
+void FMetalUniformBuffer::Update(const void* Contents, TArray<TRefCountPtr<FRHIResource>>& Resources, EUniformBufferValidation Validation)
+{
+    if (ConstantSize > 0)
+    {
+        UE_CLOG(ConstantSize > 65536, LogMetal, Fatal, TEXT("Trying to allocated a uniform layout of size %d that is greater than the maximum permitted 64k."), ConstantSize);
+        
+        void* Data = Lock(true, RLM_WriteOnly, 0);
+        FMemory::Memcpy(Data, Contents, ConstantSize);
+        Unlock();
+    }
+	
+	UpdateResourceTable(Resources, Validation);
+}
+
 FUniformBufferRHIRef FMetalDynamicRHI::RHICreateUniformBuffer(const void* Contents, const FRHIUniformBufferLayout& Layout, EUniformBufferUsage Usage, EUniformBufferValidation Validation)
 {
 	@autoreleasepool {
diff --git a/Engine/Source/Runtime/Apple/MetalRHI/Private/MetalVertexBuffer.cpp b/Engine/Source/Runtime/Apple/MetalRHI/Private/MetalVertexBuffer.cpp
--- a/Engine/Source/Runtime/Apple/MetalRHI/Private/MetalVertexBuffer.cpp
+++ b/Engine/Source/Runtime/Apple/MetalRHI/Private/MetalVertexBuffer.cpp
@@ -97,10 +97,15 @@ void FMetalRHIBuffer::Swap(FMetalRHIBuffer& Other)
 	::Swap(*this, Other);
 }
 
+static bool CanUsePrivateMemory()
+{
+	return (FMetalCommandQueue::SupportsFeature(EMetalFeaturesEfficientBufferBlits) || FMetalCommandQueue::SupportsFeature(EMetalFeaturesIABs));
+}
+
 bool FMetalRHIBuffer::UsePrivateMemory() const
 {
 	return (FMetalCommandQueue::SupportsFeature(EMetalFeaturesEfficientBufferBlits) && (Usage & (BUF_Dynamic|BUF_Static)))
-	|| (FMetalCommandQueue::SupportsFeature(EMetalFeaturesIABs) && Usage & (BUF_ShaderResource|BUF_UnorderedAccess));
+	|| (FMetalCommandQueue::SupportsFeature(EMetalFeaturesIABs) && (Usage & (BUF_ShaderResource|BUF_UnorderedAccess)));
 }
 
 FMetalRHIBuffer::FMetalRHIBuffer(uint32 InSize, uint32 InUsage, ERHIResourceType InType)
@@ -110,6 +115,7 @@ FMetalRHIBuffer::FMetalRHIBuffer(uint32 InSize, uint32 InUsage, ERHIResourceType
 , LockSize(0)
 , Size(InSize)
 , Usage(InUsage)
+, Mode(BUFFER_STORAGE_MODE)
 , Type(InType)
 {
 	// No life-time usage information? Enforce Dynamic.
@@ -118,6 +124,8 @@ FMetalRHIBuffer::FMetalRHIBuffer(uint32 InSize, uint32 InUsage, ERHIResourceType
 		Usage |= BUF_Dynamic;
 	}
 	
+	Mode = UsePrivateMemory() ? mtlpp::StorageMode::Private : BUFFER_STORAGE_MODE;
+	
 	if (InSize)
 	{
 		checkf(InSize <= 1024 * 1024 * 1024, TEXT("Metal doesn't support buffers > 1GB"));
@@ -204,7 +212,7 @@ FMetalRHIBuffer::~FMetalRHIBuffer()
 
 void FMetalRHIBuffer::Alias()
 {
-	if (Buffer.GetStorageMode() == mtlpp::StorageMode::Private && Buffer.GetHeap() && !Buffer.IsAliasable())
+	if (Mode == mtlpp::StorageMode::Private && Buffer.GetHeap() && !Buffer.IsAliasable())
 	{
 		Buffer.MakeAliasable();
 #if STATS || ENABLE_LOW_LEVEL_MEM_TRACKER
@@ -215,7 +223,7 @@ void FMetalRHIBuffer::Alias()
 
 void FMetalRHIBuffer::Unalias()
 {
-	if (Buffer.GetStorageMode() == mtlpp::StorageMode::Private && Buffer.GetHeap() && Buffer.IsAliasable())
+	if (Mode == mtlpp::StorageMode::Private && Buffer.GetHeap() && Buffer.IsAliasable())
 	{
 		uint32 Len = Buffer.GetLength();
 		METAL_INC_DWORD_STAT_BY(Type, MemFreed, Len);
@@ -228,12 +236,9 @@ void FMetalRHIBuffer::Unalias()
 
 void FMetalRHIBuffer::Alloc(uint32 InSize, EResourceLockMode LockMode)
 {
-	bool const bUsePrivateMem = UsePrivateMemory();
-
 	if (!Buffer)
 	{
         check(LockMode != RLM_ReadOnly);
-		mtlpp::StorageMode Mode = (bUsePrivateMem ? mtlpp::StorageMode::Private : BUFFER_STORAGE_MODE);
         FMetalPooledBufferArgs Args(GetMetalDeviceContext().GetDevice(), InSize, Usage, Mode);
 		Buffer = GetMetalDeviceContext().CreatePooledBuffer(Args);
 		METAL_FATAL_ASSERT(Buffer, TEXT("Failed to create buffer of size %u and storage mode %u"), InSize, (uint32)Mode);
@@ -254,16 +259,19 @@ void FMetalRHIBuffer::Alloc(uint32 InSize, EResourceLockMode LockMode)
 			}
 		}
 	}
-	
-	if (bUsePrivateMem && !CPUBuffer)
+}
+
+void FMetalRHIBuffer::AllocTransferBuffer(bool bOnRHIThread, uint32 InSize, EResourceLockMode LockMode)
+{
+	if (!CPUBuffer && ((LockMode == RLM_WriteOnly && CanUsePrivateMemory()) || Mode == mtlpp::StorageMode::Private))
 	{
-        FMetalPooledBufferArgs ArgsCPU(GetMetalDeviceContext().GetDevice(), InSize, BUF_Dynamic, mtlpp::StorageMode::Shared);
-		CPUBuffer = GetMetalDeviceContext().CreatePooledBuffer(ArgsCPU);
+		FMetalPooledBufferArgs ArgsCPU(GetMetalDeviceContext().GetDevice(), InSize, BUF_Dynamic, mtlpp::StorageMode::Shared);
+        CPUBuffer = GetMetalDeviceContext().CreatePooledBuffer(ArgsCPU);
 		CPUBuffer.SetOwner(this);
-		check(CPUBuffer && CPUBuffer.GetPtr());
+        check(CPUBuffer && CPUBuffer.GetPtr());
         METAL_INC_DWORD_STAT_BY(Type, MemAlloc, InSize);
-		METAL_FATAL_ASSERT(CPUBuffer, TEXT("Failed to create buffer of size %u and storage mode %u"), InSize, (uint32)mtlpp::StorageMode::Shared);
-	}
+        METAL_FATAL_ASSERT(CPUBuffer, TEXT("Failed to create buffer of size %u and storage mode %u"), InSize, (uint32)mtlpp::StorageMode::Shared);
+    }
 }
 
 FMetalTexture FMetalRHIBuffer::AllocLinearTexture(EPixelFormat Format)
@@ -273,8 +281,8 @@ FMetalTexture FMetalRHIBuffer::AllocLinearTexture(EPixelFormat Format)
 		mtlpp::PixelFormat MTLFormat = (mtlpp::PixelFormat)GMetalBufferFormats[Format].LinearTextureFormat;
 		
 		mtlpp::TextureDescriptor Desc;
-		NSUInteger Mode = ((NSUInteger)Buffer.GetStorageMode() << mtlpp::ResourceStorageModeShift) | ((NSUInteger)Buffer.GetCpuCacheMode() << mtlpp::ResourceCpuCacheModeShift);
-		Mode = FMetalCommandQueue::GetCompatibleResourceOptions(mtlpp::ResourceOptions(Mode | mtlpp::ResourceOptions::HazardTrackingModeUntracked));
+		NSUInteger Options = ((NSUInteger)Mode << mtlpp::ResourceStorageModeShift) | ((NSUInteger)Buffer.GetCpuCacheMode() << mtlpp::ResourceCpuCacheModeShift);
+		Options = FMetalCommandQueue::GetCompatibleResourceOptions(mtlpp::ResourceOptions(Options | mtlpp::ResourceOptions::HazardTrackingModeUntracked));
 		NSUInteger TexUsage = mtlpp::TextureUsage::Unknown;
 		if (Usage & BUF_ShaderResource)
 		{
@@ -294,7 +302,7 @@ FMetalTexture FMetalRHIBuffer::AllocLinearTexture(EPixelFormat Format)
 
 		if (FMetalCommandQueue::SupportsFeature(EMetalFeaturesTextureBuffers))
 		{
-			Desc = mtlpp::TextureDescriptor::TextureBufferDescriptor(MTLFormat, NewSize / Stride, mtlpp::ResourceOptions(Mode), mtlpp::TextureUsage(TexUsage));
+			Desc = mtlpp::TextureDescriptor::TextureBufferDescriptor(MTLFormat, NewSize / Stride, mtlpp::ResourceOptions(Options), mtlpp::TextureUsage(TexUsage));
 			Desc.SetAllowGPUOptimisedContents(false);
 		}
 		else
@@ -320,10 +328,10 @@ FMetalTexture FMetalRHIBuffer::AllocLinearTexture(EPixelFormat Format)
 			NewSize = SizeX*Stride;
 			
 			Desc = mtlpp::TextureDescriptor::Texture2DDescriptor(MTLFormat, SizeX, SizeY, NO);
-			Desc.SetStorageMode(Buffer.GetStorageMode());
+			Desc.SetStorageMode(Mode);
 			Desc.SetCpuCacheMode(Buffer.GetCpuCacheMode());
 			Desc.SetUsage((mtlpp::TextureUsage)TexUsage);
-			Desc.SetResourceOptions((mtlpp::ResourceOptions)Mode);
+			Desc.SetResourceOptions((mtlpp::ResourceOptions)Options);
 		}
 		
 		FMetalTexture Texture = MTLPP_VALIDATE(mtlpp::Buffer, Buffer, SafeGetRuntimeDebuggingLevel() >= EMetalDebugLevelValidation, NewTexture(Desc, 0, NewSize));
@@ -403,7 +411,7 @@ ns::AutoReleased<FMetalTexture> FMetalRHIBuffer::GetLinearTexture(EPixelFormat F
 	return Texture;
 }
 
-void* FMetalRHIBuffer::Lock(EResourceLockMode LockMode, uint32 Offset, uint32 InSize)
+void* FMetalRHIBuffer::Lock(bool bIsOnRHIThread, EResourceLockMode LockMode, uint32 Offset, uint32 InSize)
 {
 	check(LockSize == 0 && LockOffset == 0);
 	
@@ -420,15 +428,11 @@ void* FMetalRHIBuffer::Lock(EResourceLockMode LockMode, uint32 Offset, uint32 In
 	// In order to properly synchronise the buffer access, when a dynamic buffer is locked for writing, discard the old buffer & create a new one. This prevents writing to a buffer while it is being read by the GPU & thus causing corruption. This matches the logic of other RHIs.
 	if (LockMode == RLM_WriteOnly)
 	{
-        bool const bUsePrivateMem = UsePrivateMemory();
-        if (bUsePrivateMem)
+        if (CPUBuffer)
         {
-			if (CPUBuffer)
-			{
-				METAL_INC_DWORD_STAT_BY(Type, MemFreed, Len);
-				SafeReleaseMetalBuffer(CPUBuffer);
-				CPUBuffer = nil;
-			}
+			METAL_INC_DWORD_STAT_BY(Type, MemFreed, Len);
+			SafeReleaseMetalBuffer(CPUBuffer);
+			CPUBuffer = nil;
 			
 			if (LastUpdate && LastUpdate == GFrameNumberRenderThread)
 			{
@@ -437,8 +441,11 @@ void* FMetalRHIBuffer::Lock(EResourceLockMode LockMode, uint32 Offset, uint32 In
 				Buffer = nil;
 			}
 		}
-		else
+		else if (Mode == BUFFER_STORAGE_MODE)
 		{
+			// Turns out to be better to use Shared->Private blits whenever possible
+			// Should only put write-once buffers into Shared/Managed or the cost of recreating linear textures overwhelms any other efficiency
+			Mode =  CanUsePrivateMemory() ? mtlpp::StorageMode::Private : Mode;
 			METAL_INC_DWORD_STAT_BY(Type, MemFreed, Len);
 			SafeReleaseMetalBuffer(Buffer);
 			Buffer = nil;
@@ -446,6 +453,7 @@ void* FMetalRHIBuffer::Lock(EResourceLockMode LockMode, uint32 Offset, uint32 In
 	}
     
     Alloc(Len, LockMode);
+	AllocTransferBuffer(bIsOnRHIThread, Len, LockMode);
 	
 	FMetalBuffer& theBufferToUse = CPUBuffer ? CPUBuffer : Buffer;
 	if(LockMode != RLM_ReadOnly)
@@ -458,25 +466,19 @@ void* FMetalRHIBuffer::Lock(EResourceLockMode LockMode, uint32 Offset, uint32 In
 	else if (CPUBuffer)
 	{
 		SCOPE_CYCLE_COUNTER(STAT_MetalBufferPageOffTime);
+        check(bIsOnRHIThread);
 		
 		// Synchronise the buffer with the CPU
 		GetMetalDeviceContext().CopyFromBufferToBuffer(Buffer, 0, CPUBuffer, 0, Buffer.GetLength());
-
-#if PLATFORM_MAC
-		if(CPUBuffer.GetStorageMode() == mtlpp::StorageMode::Managed)
-		{
-			// Synchronise the buffer with the CPU
-			GetMetalDeviceContext().SynchroniseResource(CPUBuffer);
-		}
-#endif
 		
 		//kick the current command buffer.
 		GetMetalDeviceContext().SubmitCommandBufferAndWait();
 	}
 #if PLATFORM_MAC
-	else if(theBufferToUse.GetStorageMode() == mtlpp::StorageMode::Managed)
+	else if(Mode == mtlpp::StorageMode::Managed)
 	{
 		SCOPE_CYCLE_COUNTER(STAT_MetalBufferPageOffTime);
+        check(bIsOnRHIThread);
 		
 		// Synchronise the buffer with the CPU
 		GetMetalDeviceContext().SynchroniseResource(Buffer);
@@ -498,19 +500,9 @@ void FMetalRHIBuffer::Unlock()
 	{
 		if (LockSize && CPUBuffer)
 		{
-#if PLATFORM_MAC
-			if(CPUBuffer.GetStorageMode() == mtlpp::StorageMode::Managed)
-			{
-				if (GMetalBufferZeroFill)
-					MTLPP_VALIDATE(mtlpp::Buffer, CPUBuffer, SafeGetRuntimeDebuggingLevel() >= EMetalDebugLevelValidation, DidModify(ns::Range(0, Buffer.GetLength())));
-				else
-					MTLPP_VALIDATE(mtlpp::Buffer, CPUBuffer, SafeGetRuntimeDebuggingLevel() >= EMetalDebugLevelValidation, DidModify(ns::Range(LockOffset, LockSize)));
-			}
-#endif
-
 			// Synchronise the buffer with the GPU
 			GetMetalDeviceContext().AsyncCopyFromBufferToBuffer(CPUBuffer, 0, Buffer, 0, Buffer.GetLength());
-			if (UsePrivateMemory())
+			if (CPUBuffer)
             {
 				SafeReleaseMetalBuffer(CPUBuffer);
 				CPUBuffer = nil;
@@ -521,7 +513,7 @@ void FMetalRHIBuffer::Unlock()
 			}
 		}
 #if PLATFORM_MAC
-		else if(LockSize && Buffer.GetStorageMode() == mtlpp::StorageMode::Managed)
+		else if(LockSize && Mode == mtlpp::StorageMode::Managed)
 		{
 			if (GMetalBufferZeroFill)
 				MTLPP_VALIDATE(mtlpp::Buffer, Buffer, SafeGetRuntimeDebuggingLevel() >= EMetalDebugLevelValidation, DidModify(ns::Range(0, Buffer.GetLength())));
@@ -560,13 +552,9 @@ FVertexBufferRHIRef FMetalDynamicRHI::RHICreateVertexBuffer(uint32 Size, uint32
 		// Discard the resource array's contents.
 		CreateInfo.ResourceArray->Discard();
 	}
-	else if (VertexBuffer->Buffer.GetStorageMode() == mtlpp::StorageMode::Private)
+	else if (VertexBuffer->Mode == mtlpp::StorageMode::Private)
 	{
-		if (VertexBuffer->UsePrivateMemory())
-		{
-			SafeReleaseMetalBuffer(VertexBuffer->CPUBuffer);
-			VertexBuffer->CPUBuffer = nil;
-		}
+		check (!VertexBuffer->CPUBuffer);
 
 		if (GMetalBufferZeroFill && !FMetalCommandQueue::SupportsFeature(EMetalFeaturesFences))
 		{
@@ -574,7 +562,7 @@ FVertexBufferRHIRef FMetalDynamicRHI::RHICreateVertexBuffer(uint32 Size, uint32
 		}
 	}
 #if PLATFORM_MAC
-	else if (GMetalBufferZeroFill && VertexBuffer->Buffer.GetStorageMode() == mtlpp::StorageMode::Managed)
+	else if (GMetalBufferZeroFill && VertexBuffer->Mode == mtlpp::StorageMode::Managed)
 	{
 		MTLPP_VALIDATE(mtlpp::Buffer, VertexBuffer->Buffer, SafeGetRuntimeDebuggingLevel() >= EMetalDebugLevelValidation, DidModify(ns::Range(0, VertexBuffer->Buffer.GetLength())));
 	}
@@ -590,7 +578,7 @@ void* FMetalDynamicRHI::LockVertexBuffer_BottomOfPipe(FRHICommandListImmediate&
 	FMetalVertexBuffer* VertexBuffer = ResourceCast(VertexBufferRHI);
 
 	// default to vertex buffer memory
-	return (uint8*)VertexBuffer->Lock(LockMode, Offset, Size);
+	return (uint8*)VertexBuffer->Lock(true, LockMode, Offset, Size);
 	}
 }
 
@@ -613,10 +601,18 @@ void FMetalDynamicRHI::RHICopyVertexBuffer(FRHIVertexBuffer* SourceBufferRHI, FR
 		{
 			GetMetalDeviceContext().CopyFromBufferToBuffer(SrcVertexBuffer->Buffer, 0, DstVertexBuffer->Buffer, 0, FMath::Min(SrcVertexBuffer->GetSize(), DstVertexBuffer->GetSize()));
 		}
+		else if (DstVertexBuffer->Buffer)
+		{
+			FMetalPooledBufferArgs ArgsCPU(GetMetalDeviceContext().GetDevice(), SrcVertexBuffer->GetSize(), BUF_Dynamic, mtlpp::StorageMode::Shared);
+			FMetalBuffer TempBuffer = GetMetalDeviceContext().CreatePooledBuffer(ArgsCPU);
+			FMemory::Memcpy(TempBuffer.GetContents(), SrcVertexBuffer->Data->Data, SrcVertexBuffer->GetSize());
+			GetMetalDeviceContext().CopyFromBufferToBuffer(TempBuffer, 0, DstVertexBuffer->Buffer, 0, FMath::Min(SrcVertexBuffer->GetSize(), DstVertexBuffer->GetSize()));
+			SafeReleaseMetalBuffer(TempBuffer);
+		}
 		else
 		{
-			void const* SrcData = SrcVertexBuffer->Lock(RLM_ReadOnly, 0);
-			void* DstData = DstVertexBuffer->Lock(RLM_WriteOnly, 0);
+			void const* SrcData = SrcVertexBuffer->Lock(true, RLM_ReadOnly, 0);
+			void* DstData = DstVertexBuffer->Lock(true, RLM_WriteOnly, 0);
 			FMemory::Memcpy(DstData, SrcData, FMath::Min(SrcVertexBuffer->GetSize(), DstVertexBuffer->GetSize()));
 			SrcVertexBuffer->Unlock();
 			DstVertexBuffer->Unlock();
@@ -624,16 +620,18 @@ void FMetalDynamicRHI::RHICopyVertexBuffer(FRHIVertexBuffer* SourceBufferRHI, FR
 	}
 }
 
-struct FMetalRHICommandInitialiseVertexBuffer : public FRHICommand<FMetalRHICommandInitialiseVertexBuffer>
+struct FMetalRHICommandInitialiseBuffer : public FRHICommand<FMetalRHICommandInitialiseBuffer>
 {
-	TRefCountPtr<FMetalVertexBuffer> Buffer;
+	TRefCountPtr<FRHIResource> Resource;
+	FMetalRHIBuffer* Buffer;
 	
-	FORCEINLINE_DEBUGGABLE FMetalRHICommandInitialiseVertexBuffer(FMetalVertexBuffer* InBuffer)
-	: Buffer(InBuffer)
+	FORCEINLINE_DEBUGGABLE FMetalRHICommandInitialiseBuffer(FMetalRHIBuffer* InBuffer, FRHIResource* InResource)
+	: Resource(InResource)
+	, Buffer(InBuffer)
 	{
 	}
 	
-	virtual ~FMetalRHICommandInitialiseVertexBuffer()
+	virtual ~FMetalRHICommandInitialiseBuffer()
 	{
 	}
 	
@@ -644,7 +642,7 @@ struct FMetalRHICommandInitialiseVertexBuffer : public FRHICommand<FMetalRHIComm
 			uint32 Size = FMath::Min(Buffer->Buffer.GetLength(), Buffer->CPUBuffer.GetLength());
 			GetMetalDeviceContext().AsyncCopyFromBufferToBuffer(Buffer->CPUBuffer, 0, Buffer->Buffer, 0, Size);
 
-			if (Buffer->UsePrivateMemory())
+			if (Buffer->CPUBuffer)
 			{
 				SafeReleaseMetalBuffer(Buffer->CPUBuffer);
 			}
@@ -660,84 +658,80 @@ struct FMetalRHICommandInitialiseVertexBuffer : public FRHICommand<FMetalRHIComm
 	}
 };
 
-FVertexBufferRHIRef FMetalDynamicRHI::CreateVertexBuffer_RenderThread(class FRHICommandListImmediate& RHICmdList, uint32 Size, uint32 InUsage, FRHIResourceCreateInfo& CreateInfo)
+void FMetalRHIBuffer::Init_RenderThread(FRHICommandListImmediate& RHICmdList, uint32 InSize, uint32 InUsage, FRHIResourceCreateInfo& CreateInfo, FRHIResource* Resource)
 {
-	@autoreleasepool {
-		if (CreateInfo.bWithoutNativeResource)
-		{
-			return new FMetalVertexBuffer(0, 0);
-		}
+	if (CreateInfo.ResourceArray)
+	{
+		check(InSize == CreateInfo.ResourceArray->GetResourceDataSize());
 		
-		// make the RHI object, which will allocate memory
-		TRefCountPtr<FMetalVertexBuffer> VertexBuffer = new FMetalVertexBuffer(Size, InUsage);
+		AllocTransferBuffer(RHICmdList.IsBottomOfPipe(), InSize, RLM_WriteOnly);
 		
-		if (CreateInfo.ResourceArray)
+		if (CPUBuffer)
 		{
-			check(Size == CreateInfo.ResourceArray->GetResourceDataSize());
+			FMemory::Memcpy(CPUBuffer.GetContents(), CreateInfo.ResourceArray->GetResourceData(), InSize);
 			
-			if (VertexBuffer->CPUBuffer)
+			if (RHICmdList.IsBottomOfPipe())
 			{
-				FMemory::Memcpy(VertexBuffer->CPUBuffer.GetContents(), CreateInfo.ResourceArray->GetResourceData(), Size);
-
-#if PLATFORM_MAC
-				if(VertexBuffer->CPUBuffer.GetStorageMode() == mtlpp::StorageMode::Managed)
-				{
-					MTLPP_VALIDATE(mtlpp::Buffer, VertexBuffer->CPUBuffer, SafeGetRuntimeDebuggingLevel() >= EMetalDebugLevelValidation, DidModify(ns::Range(0, GMetalBufferZeroFill ? VertexBuffer->CPUBuffer.GetLength() : Size)));
-				}
-#endif
-
-				if (RHICmdList.Bypass() || !IsRunningRHIInSeparateThread())
-				{
-					FMetalRHICommandInitialiseVertexBuffer UpdateCommand(VertexBuffer);
-					UpdateCommand.Execute(RHICmdList);
-				}
-				else
-				{
-					new (RHICmdList.AllocCommand<FMetalRHICommandInitialiseVertexBuffer>()) FMetalRHICommandInitialiseVertexBuffer(VertexBuffer);
-				}
+				FMetalRHICommandInitialiseBuffer UpdateCommand(this, Resource);
+				UpdateCommand.Execute(RHICmdList);
 			}
 			else
 			{
-				// make a buffer usable by CPU
-				void* Buffer = RHILockVertexBuffer(RHICmdList, VertexBuffer, 0, Size, RLM_WriteOnly);
-				
-				// copy the contents of the given data into the buffer
-				FMemory::Memcpy(Buffer, CreateInfo.ResourceArray->GetResourceData(), Size);
-				
-				RHIUnlockVertexBuffer(RHICmdList, VertexBuffer);
+				new (RHICmdList.AllocCommand<FMetalRHICommandInitialiseBuffer>()) FMetalRHICommandInitialiseBuffer(this, Resource);
 			}
-			
-			// Discard the resource array's contents.
-			CreateInfo.ResourceArray->Discard();
 		}
-		else if (VertexBuffer->Buffer)
+		else
 		{
-			if (VertexBuffer->UsePrivateMemory())
+			FMemory::Memcpy(Buffer.GetContents(), CreateInfo.ResourceArray->GetResourceData(), InSize);
+#if PLATFORM_MAC
+			if(Mode == mtlpp::StorageMode::Managed)
 			{
-				SafeReleaseMetalBuffer(VertexBuffer->CPUBuffer);
-				VertexBuffer->CPUBuffer = nil;
+				MTLPP_VALIDATE(mtlpp::Buffer, Buffer, SafeGetRuntimeDebuggingLevel() >= EMetalDebugLevelValidation, DidModify(ns::Range(0, GMetalBufferZeroFill ? Buffer.GetLength() : InSize)));
 			}
-			
-			if (GMetalBufferZeroFill && VertexBuffer->Buffer.GetStorageMode() == mtlpp::StorageMode::Private)
+#endif
+		}
+		
+		// Discard the resource array's contents.
+		CreateInfo.ResourceArray->Discard();
+	}
+	else if (Buffer)
+	{
+		check (!CPUBuffer);
+		
+		if (GMetalBufferZeroFill && Mode == mtlpp::StorageMode::Private)
+		{
+			if (RHICmdList.IsBottomOfPipe())
 			{
-				if (RHICmdList.Bypass() || !IsRunningRHIInSeparateThread())
-				{
-					FMetalRHICommandInitialiseVertexBuffer UpdateCommand(VertexBuffer);
-					UpdateCommand.Execute(RHICmdList);
-				}
-				else
-				{
-					new (RHICmdList.AllocCommand<FMetalRHICommandInitialiseVertexBuffer>()) FMetalRHICommandInitialiseVertexBuffer(VertexBuffer);
-				}
+				FMetalRHICommandInitialiseBuffer UpdateCommand(this, Resource);
+				UpdateCommand.Execute(RHICmdList);
 			}
-#if PLATFORM_MAC
-			else if (GMetalBufferZeroFill && VertexBuffer->Buffer.GetStorageMode() == mtlpp::StorageMode::Managed)
+			else
 			{
-				MTLPP_VALIDATE(mtlpp::Buffer, VertexBuffer->Buffer, SafeGetRuntimeDebuggingLevel() >= EMetalDebugLevelValidation, DidModify(ns::Range(0, VertexBuffer->Buffer.GetLength())));
+				new (RHICmdList.AllocCommand<FMetalRHICommandInitialiseBuffer>()) FMetalRHICommandInitialiseBuffer(this, Resource);
 			}
+		}
+#if PLATFORM_MAC
+		else if (GMetalBufferZeroFill && Mode == mtlpp::StorageMode::Managed)
+		{
+			MTLPP_VALIDATE(mtlpp::Buffer, Buffer, SafeGetRuntimeDebuggingLevel() >= EMetalDebugLevelValidation, DidModify(ns::Range(0, Buffer.GetLength())));
+		}
 #endif
+	}
+}
+
+FVertexBufferRHIRef FMetalDynamicRHI::CreateVertexBuffer_RenderThread(class FRHICommandListImmediate& RHICmdList, uint32 Size, uint32 InUsage, FRHIResourceCreateInfo& CreateInfo)
+{
+	@autoreleasepool {
+		if (CreateInfo.bWithoutNativeResource)
+		{
+			return new FMetalVertexBuffer(0, 0);
 		}
 		
+		// make the RHI object, which will allocate memory
+		TRefCountPtr<FMetalVertexBuffer> VertexBuffer = new FMetalVertexBuffer(Size, InUsage);
+		
+		VertexBuffer->Init_RenderThread(RHICmdList, Size, InUsage, CreateInfo, VertexBuffer);
+		
 		return VertexBuffer.GetReference();
 	}
 }
diff --git a/Engine/Source/Runtime/Apple/MetalRHI/Public/MetalResources.h b/Engine/Source/Runtime/Apple/MetalRHI/Public/MetalResources.h
--- a/Engine/Source/Runtime/Apple/MetalRHI/Public/MetalResources.h
+++ b/Engine/Source/Runtime/Apple/MetalRHI/Public/MetalResources.h
@@ -758,6 +758,11 @@ public:
 	FMetalRHIBuffer(uint32 InSize, uint32 InUsage, ERHIResourceType InType);
 	virtual ~FMetalRHIBuffer();
 	
+	/**
+	 * Initialize the buffer contents from the render-thread.
+	 */
+	void Init_RenderThread(class FRHICommandListImmediate& RHICmdList, uint32 Size, uint32 InUsage, FRHIResourceCreateInfo& CreateInfo, FRHIResource* Resource);
+	
 	/**
 	 * Alias the buffer backing store allowing the memory to be reused by another resource.
 	 */
@@ -769,10 +774,15 @@ public:
 	void Unalias();
 	
 	/**
-	 * Allocate the index buffer backing store.
+	 * Allocate the buffer backing store.
 	 */
 	void Alloc(uint32 InSize, EResourceLockMode LockMode);
 	
+	/**
+	 * Allocate the CPU accessible buffer for data transfer.
+	 */
+	void AllocTransferBuffer(bool bOnRHIThread, uint32 InSize, EResourceLockMode LockMode);
+	
 	/**
 	 * Allocate a linear texture for given format.
 	 */
@@ -791,7 +801,7 @@ public:
 	/**
 	 * Prepare a CPU accessible buffer for uploading to GPU memory
 	 */
-	void* Lock(EResourceLockMode LockMode, uint32 Offset, uint32 Size=0);
+	void* Lock(bool bIsOnRHIThread, EResourceLockMode LockMode, uint32 Offset, uint32 Size=0);
 	
 	/**
 	 * Prepare a CPU accessible buffer for uploading to GPU memory
@@ -832,6 +842,9 @@ public:
 	// Buffer usage.
 	uint32 Usage;
 	
+	// Storage mode
+	mtlpp::StorageMode Mode;
+	
 	// Resource type
 	ERHIResourceType Type;
 };
@@ -877,6 +890,7 @@ public:
 
 	void InitIAB();
 
+	void UpdateResourceTable(TArray<TRefCountPtr<FRHIResource>>& Resources, EUniformBufferValidation Validation);
 	void Update(const void* Contents, TArray<TRefCountPtr<FRHIResource>>& Resources, EUniformBufferValidation Validation);
 	
 	/** Resource table containing RHI references. */
