Make Tier 2 IABs work about as well as they used to before Xmas and changes to the shaders, compiler & runtime left them non-functional. Still not entirely correct and it is unsatisfying that updates have to be deferred until binding time but that is a result of not being able to update a single element within an IAB without shaders that specifically declare the source & target IAB structures.

diff --git a/Engine/Source/Runtime/Apple/MetalRHI/Private/MetalCommands.cpp b/Engine/Source/Runtime/Apple/MetalRHI/Private/MetalCommands.cpp
--- a/Engine/Source/Runtime/Apple/MetalRHI/Private/MetalCommands.cpp
+++ b/Engine/Source/Runtime/Apple/MetalRHI/Private/MetalCommands.cpp
@@ -544,10 +544,18 @@ void FMetalRHICommandContext::RHISetShaderUniformBuffer(FRHIVertexShader* Vertex
 
 	auto& Bindings = VertexShader->Bindings;
 	check(BufferIndex < Bindings.NumUniformBuffers);
-	if (Bindings.ConstantBuffers & 1 << BufferIndex)
+	if ((Bindings.ArgumentBuffers) & 1 << BufferIndex)
 	{
 		auto* UB = (FMetalUniformBuffer*)BufferRHI;
-		Context->GetCurrentState().SetShaderBuffer(EMetalShaderStages::Vertex, (Bindings.ArgumentBuffers & 1 << BufferIndex) ? UB->GetIAB().IndirectArgumentBuffer : UB->Buffer, UB->Data, 0, UB->GetSize(), BufferIndex, mtlpp::ResourceUsage::Read);
+		UB->UploadIAB(Context);
+		auto& IAB = UB->GetIAB();
+		check(IAB.IndirectArgumentBuffer->IndirectArgumentEncoder.GetEncodedLength() == VertexShader->ArgumentEncoders[BufferIndex].GetEncodedLength());
+		Context->GetCurrentState().SetShaderBuffer(EMetalShaderStages::Vertex, IAB.IndirectArgumentBuffer->IndirectArgumentBuffer, nil, 0, IAB.IndirectArgumentBuffer->IndirectArgumentEncoder.GetEncodedLength(), BufferIndex, mtlpp::ResourceUsage::Read);
+	}
+	else if ((Bindings.ConstantBuffers) & 1 << BufferIndex)
+	{
+		auto* UB = (FMetalUniformBuffer*)BufferRHI;
+		Context->GetCurrentState().SetShaderBuffer(EMetalShaderStages::Vertex, UB->Buffer, UB->Data, 0, UB->GetSize(), BufferIndex, mtlpp::ResourceUsage::Read);
 	}
 	}
 }
@@ -561,10 +569,18 @@ void FMetalRHICommandContext::RHISetShaderUniformBuffer(FRHIHullShader* HullShad
 
 	auto& Bindings = HullShader->Bindings;
 	check(BufferIndex < Bindings.NumUniformBuffers);
-	if (Bindings.ConstantBuffers & 1 << BufferIndex)
+	if ((Bindings.ArgumentBuffers) & 1 << BufferIndex)
+	{
+		auto* UB = (FMetalUniformBuffer*)BufferRHI;
+		UB->UploadIAB(Context);
+		auto& IAB = UB->GetIAB();
+		check(IAB.IndirectArgumentBuffer->IndirectArgumentEncoder.GetEncodedLength() == HullShader->ArgumentEncoders[BufferIndex].GetEncodedLength());
+		Context->GetCurrentState().SetShaderBuffer(EMetalShaderStages::Hull, IAB.IndirectArgumentBuffer->IndirectArgumentBuffer, nil, 0, IAB.IndirectArgumentBuffer->IndirectArgumentEncoder.GetEncodedLength(), BufferIndex, mtlpp::ResourceUsage::Read);
+	}
+	else if ((Bindings.ConstantBuffers) & 1 << BufferIndex)
 	{
 		auto* UB = (FMetalUniformBuffer*)BufferRHI;
-		Context->GetCurrentState().SetShaderBuffer(EMetalShaderStages::Hull, (Bindings.ArgumentBuffers & 1 << BufferIndex) ? UB->GetIAB().IndirectArgumentBuffer : UB->Buffer, UB->Data, 0, UB->GetSize(), BufferIndex, mtlpp::ResourceUsage::Read);
+		Context->GetCurrentState().SetShaderBuffer(EMetalShaderStages::Hull, UB->Buffer, UB->Data, 0, UB->GetSize(), BufferIndex, mtlpp::ResourceUsage::Read);
 	}
 	}
 #endif
@@ -579,10 +595,18 @@ void FMetalRHICommandContext::RHISetShaderUniformBuffer(FRHIDomainShader* Domain
 
 	auto& Bindings = DomainShader->Bindings;
 	check(BufferIndex < Bindings.NumUniformBuffers);
-	if (Bindings.ConstantBuffers & 1 << BufferIndex)
+	if ((Bindings.ArgumentBuffers) & 1 << BufferIndex)
 	{
 		auto* UB = (FMetalUniformBuffer*)BufferRHI;
-		Context->GetCurrentState().SetShaderBuffer(EMetalShaderStages::Domain, (Bindings.ArgumentBuffers & 1 << BufferIndex) ? UB->GetIAB().IndirectArgumentBuffer : UB->Buffer, UB->Data, 0, UB->GetSize(), BufferIndex, mtlpp::ResourceUsage::Read);
+		UB->UploadIAB(Context);
+		auto& IAB = UB->GetIAB();
+		check(IAB.IndirectArgumentBuffer->IndirectArgumentEncoder.GetEncodedLength() == DomainShader->ArgumentEncoders[BufferIndex].GetEncodedLength());
+		Context->GetCurrentState().SetShaderBuffer(EMetalShaderStages::Domain, IAB.IndirectArgumentBuffer->IndirectArgumentBuffer, nil, 0, IAB.IndirectArgumentBuffer->IndirectArgumentEncoder.GetEncodedLength(), BufferIndex, mtlpp::ResourceUsage::Read);
+	}
+	else if ((Bindings.ConstantBuffers) & 1 << BufferIndex)
+	{
+		auto* UB = (FMetalUniformBuffer*)BufferRHI;
+		Context->GetCurrentState().SetShaderBuffer(EMetalShaderStages::Domain, UB->Buffer, UB->Data, 0, UB->GetSize(), BufferIndex, mtlpp::ResourceUsage::Read);
 	}
 	}
 #endif
@@ -601,10 +625,18 @@ void FMetalRHICommandContext::RHISetShaderUniformBuffer(FRHIPixelShader* PixelSh
 
 	auto& Bindings = PixelShader->Bindings;
 	check(BufferIndex < Bindings.NumUniformBuffers);
-	if (Bindings.ConstantBuffers & 1 << BufferIndex)
+	if ((Bindings.ArgumentBuffers) & 1 << BufferIndex)
+	{
+		auto* UB = (FMetalUniformBuffer*)BufferRHI;
+		UB->UploadIAB(Context);
+		auto& IAB = UB->GetIAB();
+		check(IAB.IndirectArgumentBuffer->IndirectArgumentEncoder.GetEncodedLength() == PixelShader->ArgumentEncoders[BufferIndex].GetEncodedLength());
+		Context->GetCurrentState().SetShaderBuffer(EMetalShaderStages::Pixel, IAB.IndirectArgumentBuffer->IndirectArgumentBuffer, nil, 0, IAB.IndirectArgumentBuffer->IndirectArgumentEncoder.GetEncodedLength(), BufferIndex, mtlpp::ResourceUsage::Read);
+	}
+	else if ((Bindings.ConstantBuffers) & 1 << BufferIndex)
 	{
 		auto* UB = (FMetalUniformBuffer*)BufferRHI;
-		Context->GetCurrentState().SetShaderBuffer(EMetalShaderStages::Pixel, (Bindings.ArgumentBuffers & 1 << BufferIndex) ? UB->GetIAB().IndirectArgumentBuffer : UB->Buffer, UB->Data, 0, UB->GetSize(), BufferIndex, mtlpp::ResourceUsage::Read);
+		Context->GetCurrentState().SetShaderBuffer(EMetalShaderStages::Pixel, UB->Buffer, UB->Data, 0, UB->GetSize(), BufferIndex, mtlpp::ResourceUsage::Read);
 	}
 	}
 }
@@ -617,10 +649,18 @@ void FMetalRHICommandContext::RHISetShaderUniformBuffer(FRHIComputeShader* Compu
 
 	auto& Bindings = ComputeShader->Bindings;
 	check(BufferIndex < Bindings.NumUniformBuffers);
-	if (Bindings.ConstantBuffers & 1 << BufferIndex)
+	if ((Bindings.ArgumentBuffers) & 1 << BufferIndex)
+	{
+		auto* UB = (FMetalUniformBuffer*)BufferRHI;
+		UB->UploadIAB(Context);
+		auto& IAB = UB->GetIAB();
+		check(IAB.IndirectArgumentBuffer->IndirectArgumentEncoder.GetEncodedLength() == ComputeShader->ArgumentEncoders[BufferIndex].GetEncodedLength());
+		Context->GetCurrentState().SetShaderBuffer(EMetalShaderStages::Compute, IAB.IndirectArgumentBuffer->IndirectArgumentBuffer, nil, 0, IAB.IndirectArgumentBuffer->IndirectArgumentEncoder.GetEncodedLength(), BufferIndex, mtlpp::ResourceUsage::Read);
+	}
+	else if ((Bindings.ConstantBuffers) & 1 << BufferIndex)
 	{
 		auto* UB = (FMetalUniformBuffer*)BufferRHI;
-		Context->GetCurrentState().SetShaderBuffer(EMetalShaderStages::Compute, (Bindings.ArgumentBuffers & 1 << BufferIndex) ? UB->GetIAB().IndirectArgumentBuffer : UB->Buffer, UB->Data, 0, UB->GetSize(), BufferIndex, mtlpp::ResourceUsage::Read);
+		Context->GetCurrentState().SetShaderBuffer(EMetalShaderStages::Compute, UB->Buffer, UB->Data, 0, UB->GetSize(), BufferIndex, mtlpp::ResourceUsage::Read);
 	}
 	}
 }
diff --git a/Engine/Source/Runtime/Apple/MetalRHI/Private/MetalContext.cpp b/Engine/Source/Runtime/Apple/MetalRHI/Private/MetalContext.cpp
--- a/Engine/Source/Runtime/Apple/MetalRHI/Private/MetalContext.cpp
+++ b/Engine/Source/Runtime/Apple/MetalRHI/Private/MetalContext.cpp
@@ -746,10 +746,7 @@ void FMetalDeviceContext::UpdateIABs(FRHITextureReference* ModifiedRef)
 		{
 			if (UB && UB->IAB && UB->TextureReferences.Contains(ModifiedRef))
 			{
-				FMetalUniformBuffer::FMetalIndirectArgumentBuffer* IAB = UB->IAB;
-				UB->IAB = nullptr;
-				UB->InitIAB();
-				delete IAB;
+				UB->UpdateTextureReference(ModifiedRef);
 			}
 		}
 	}
diff --git a/Engine/Source/Runtime/Apple/MetalRHI/Private/MetalShaders.cpp b/Engine/Source/Runtime/Apple/MetalRHI/Private/MetalShaders.cpp
--- a/Engine/Source/Runtime/Apple/MetalRHI/Private/MetalShaders.cpp
+++ b/Engine/Source/Runtime/Apple/MetalRHI/Private/MetalShaders.cpp
@@ -854,6 +854,19 @@ mtlpp::Function TMetalBaseShader<BaseResourceType, ShaderType>::GetCompiledFunct
 		}
 	}	
 	
+	if (FMetalCommandQueue::SupportsFeature(EMetalFeaturesIABs) && Bindings.ArgumentBuffers && ArgumentEncoders.Num() == 0)
+	{
+		uint32 ArgumentBuffers = Bindings.ArgumentBuffers;
+		while(ArgumentBuffers)
+		{
+			uint32 Index = __builtin_ctz(ArgumentBuffers);
+			ArgumentBuffers &= ~(1 << Index);
+			
+			mtlpp::ArgumentEncoder ArgumentEncoder = Function.NewArgumentEncoderWithBufferIndex(Index);
+			ArgumentEncoders.Add(Index, ArgumentEncoder);
+		}
+	}
+									
     check(Func);
 	return Func;
 }
@@ -962,6 +975,11 @@ FMetalShaderPipeline* FMetalComputeShader::GetPipeline()
 
 	return Pipeline;
 }
+									
+mtlpp::Function FMetalComputeShader::GetFunction()
+{
+	return GetCompiledFunction();
+}
 
 FMetalVertexShader::FMetalVertexShader(const TArray<uint8>& InCode)
 {
diff --git a/Engine/Source/Runtime/Apple/MetalRHI/Private/MetalStateCache.cpp b/Engine/Source/Runtime/Apple/MetalRHI/Private/MetalStateCache.cpp
--- a/Engine/Source/Runtime/Apple/MetalRHI/Private/MetalStateCache.cpp
+++ b/Engine/Source/Runtime/Apple/MetalRHI/Private/MetalStateCache.cpp
@@ -2798,9 +2798,9 @@ void FMetalStateCache::CommitResourceTable(EMetalShaderStages const Frequency, m
 			}
 		}
 		
-		if (IAB.IndirectArgumentBufferSideTable)
+		if (IAB.IndirectArgumentBuffer->IndirectArgumentBufferSideTable)
 		{
-			CommandEncoder.UseIndirectArgumentResource(IAB.IndirectArgumentBufferSideTable, mtlpp::ResourceUsage::Read);
+			CommandEncoder.UseIndirectArgumentResource(IAB.IndirectArgumentBuffer->IndirectArgumentBufferSideTable, mtlpp::ResourceUsage::Read);
 		}
 		
 		if (Buffer->Buffer)
diff --git a/Engine/Source/Runtime/Apple/MetalRHI/Private/MetalUniformBuffer.cpp b/Engine/Source/Runtime/Apple/MetalRHI/Private/MetalUniformBuffer.cpp
--- a/Engine/Source/Runtime/Apple/MetalRHI/Private/MetalUniformBuffer.cpp
+++ b/Engine/Source/Runtime/Apple/MetalRHI/Private/MetalUniformBuffer.cpp
@@ -11,6 +11,23 @@
 #include "HAL/LowLevelMemTracker.h"
 #include "Misc/ScopeRWLock.h"
 
+@implementation FMetalIAB
+
+APPLE_PLATFORM_OBJECT_ALLOC_OVERRIDES(FMetalIAB)
+
+-(instancetype)init
+{
+	id Self = [super init];
+	return Self;
+}
+-(void)dealloc
+{
+	SafeReleaseMetalBuffer(IndirectArgumentBuffer);
+	SafeReleaseMetalBuffer(IndirectArgumentBufferSideTable);
+	[super dealloc];
+}
+@end
+
 struct FMetalRHICommandInitialiseUniformBufferIAB : public FRHICommand<FMetalRHICommandInitialiseUniformBufferIAB>
 {
 	TRefCountPtr<FMetalUniformBuffer> Buffer;
@@ -26,84 +43,8 @@ struct FMetalRHICommandInitialiseUniformBufferIAB : public FRHICommand<FMetalRHI
 	
 	void Execute(FRHICommandListBase& CmdList)
 	{
-		Buffer->InitIAB();
-	}
-};
-
-struct FMetalArgumentDesc
-{
-	FMetalArgumentDesc()
-	: DataType((mtlpp::DataType)0)
-	, Index(0)
-	, ArrayLength(0)
-	, Access(mtlpp::ArgumentAccess::ReadOnly)
-	, TextureType((mtlpp::TextureType)0)
-	, ConstantBlockAlignment(0)
-	{
-		
-	}
-	
-	mtlpp::DataType DataType;
-	NSUInteger Index;
-	NSUInteger ArrayLength;
-	mtlpp::ArgumentAccess Access;
-	mtlpp::TextureType TextureType;
-	NSUInteger ConstantBlockAlignment;
-	
-	void FillDescriptor(mtlpp::ArgumentDescriptor& Desc) const
-	{
-		Desc.SetDataType(DataType);
-		Desc.SetIndex(Index);
-		Desc.SetArrayLength(ArrayLength);
-		Desc.SetAccess(Access);
-		Desc.SetTextureType(TextureType);
-		Desc.SetConstantBlockAlignment(ConstantBlockAlignment);
-	}
-	
-	void SetDataType(mtlpp::DataType Type)
-	{
-		DataType = Type;
-	}
-	
-	void SetIndex(NSUInteger i)
-	{
-		Index = i;
-	}
-	
-	void SetArrayLength(NSUInteger Len)
-	{
-		ArrayLength = Len;
-	}
-	
-	void SetAccess(mtlpp::ArgumentAccess InAccess)
-	{
-		Access = InAccess;
-	}
-	
-	void SetTextureType(mtlpp::TextureType Type)
-	{
-		TextureType = Type;
-	}
-	
-	void SetConstantBlockAlignment(NSUInteger Align)
-	{
-		ConstantBlockAlignment = Align;
-	}
-	
-	bool operator==(FMetalArgumentDesc const& Other) const
-	{
-		if (this != &Other)
-		{
-			return (DataType == Other.DataType && Index == Other.Index && ArrayLength == Other.ArrayLength && Access == Other.Access && TextureType == Other.TextureType && ConstantBlockAlignment == Other.ConstantBlockAlignment);
-		}
-		return true;
-	}
-	
-	friend uint32 GetTypeHash(FMetalArgumentDesc const& Desc)
-	{
-		uint32 Hash = 0;
-		Hash ^= ((uint32)Desc.DataType * (uint32)Desc.TextureType * (uint32)Desc.Access * Desc.ArrayLength) << Desc.Index;
-		return Hash;
+		FPlatformAtomics::InterlockedIncrement(&Buffer->UpdateNum);
+		Buffer->UpdateIAB();
 	}
 };
 
@@ -179,6 +120,7 @@ FMetalUniformBuffer::FMetalUniformBuffer(const void* Contents, const FRHIUniform
     , FMetalRHIBuffer(Layout.ConstantBufferSize, (FMetalCommandQueue::SupportsFeature(EMetalFeaturesIABs) && Layout.Resources.Num() ? (EMetalBufferUsage_GPUOnly|BUF_Volatile) : BUF_Volatile), RRT_UniformBuffer)
 	, UniformUsage(InUsage)
 	, IAB(nullptr)
+	, UpdateNum(0)
 	, NumResources(0)
 	, ConstantSize(Layout.ConstantBufferSize)
 {
@@ -232,10 +174,12 @@ FMetalUniformBuffer::FMetalUniformBuffer(const void* Contents, const FRHIUniform
 		}
 	}
 	
+	InitIAB();
     UpdateResourceTable(ResourceTable, Validation);
 	
     if (NumResources && FMetalCommandQueue::SupportsFeature(EMetalFeaturesIABs))
     {
+		check(IAB);
         GetMetalDeviceContext().RegisterUB(this);
     }
 }
@@ -256,288 +200,52 @@ FMetalUniformBuffer::~FMetalUniformBuffer()
 
 FMetalUniformBuffer::FMetalIndirectArgumentBuffer::FMetalIndirectArgumentBuffer()
 {
-	
+	UpdateNum = 0;
+	UpdateIAB = 0;
+	UpdateEnc = 0;
+	IndirectArgumentBuffer = nil;
 }
 
 FMetalUniformBuffer::FMetalIndirectArgumentBuffer::~FMetalIndirectArgumentBuffer()
 {
-	SafeReleaseMetalBuffer(IndirectArgumentBuffer);
-	SafeReleaseMetalBuffer(IndirectArgumentBufferSideTable);
+	if (IndirectArgumentBuffer)
+		SafeReleaseMetalObject(IndirectArgumentBuffer);
 }
 
-FMetalUniformBuffer::FMetalIndirectArgumentBuffer& FMetalUniformBuffer::GetIAB()
+FMetalIAB* FMetalUniformBuffer::UploadIAB(FMetalContext* Ctx)
 {
-	check(ResourceTable.Num() && FMetalCommandQueue::SupportsFeature(EMetalFeaturesIABs));
-
-	InitIAB();
 	check(IAB);
 	
-	return *IAB;
-}
-
-void FMetalUniformBuffer::InitIAB()
-{
-	if (NumResources && FMetalCommandQueue::SupportsFeature(EMetalFeaturesIABs) && !IAB)
+	int64 UpdateIAB = FPlatformAtomics::InterlockedCompareExchange(&IAB->UpdateIAB, 0, IAB->UpdateIAB);
+	if (UpdateIAB && Ctx)
 	{
-		FMetalIndirectArgumentBuffer* NewIAB = new FMetalIndirectArgumentBuffer;
-		
-		TArray<uint32> BufferSizes;
-		TArray<FMetalArgumentDesc> Arguments;
-		
-		int32 Index = 0;
-		// Buffer Size table
-		{
-			FMetalArgumentDesc& Desc = Arguments.Emplace_GetRef();
-			Desc.SetIndex(Index++);
-			Desc.SetAccess(mtlpp::ArgumentAccess::ReadOnly);
-			Desc.SetDataType(mtlpp::DataType::Pointer);
-		}
-		
-		for (int32 i = 0; i < NumResources; ++i, ++Index)
-		{
-			FRHIResource* Resource = ResourceTable[i].GetReference();
-			
-			switch(ResourceTypes[i])
-			{
-				case UBMT_RDG_TEXTURE_SRV:
-				case UBMT_RDG_BUFFER_SRV:
-				case UBMT_SRV:
-				{
-					FMetalArgumentDesc& Desc = Arguments.Emplace_GetRef();
-					Desc.SetIndex(Index);
-					Desc.SetAccess(mtlpp::ArgumentAccess::ReadOnly);
-					
-					FMetalShaderResourceView* SRV = (FMetalShaderResourceView*)Resource;
-					FRHITexture* Texture = SRV->SourceTexture.GetReference();
-					FMetalVertexBuffer* VB = SRV->SourceVertexBuffer.GetReference();
-					FMetalIndexBuffer* IB = SRV->SourceIndexBuffer.GetReference();
-					FMetalStructuredBuffer* SB = SRV->SourceStructuredBuffer.GetReference();
-					if (Texture)
-					{
-						FMetalSurface* Surface = SRV->TextureView;
-						check (Surface != nullptr);
-						Desc.SetDataType(mtlpp::DataType::Texture);
-						Desc.SetTextureType(Surface->Texture.GetTextureType());
-						
-						union {
-							uint8 Components[4];
-							uint32 Packed;
-						} Swizzle;
-						Swizzle.Packed = 0;
-						assert(sizeof(Swizzle) == sizeof(uint32));
-						if (Surface->Texture.GetPixelFormat() == mtlpp::PixelFormat::X32_Stencil8
-#if PLATFORM_MAC
-						 ||	Surface->Texture.GetPixelFormat() == mtlpp::PixelFormat::X24_Stencil8
-#endif
-						)
-						{
-							Swizzle.Components[0] = Swizzle.Components[1] = Swizzle.Components[2] = Swizzle.Components[3] = 1;
-						}
-						BufferSizes.Add(Swizzle.Packed);
-						BufferSizes.Add(GMetalBufferFormats[Texture->GetFormat()].DataFormat);
-						
-						check(!Surface->Texture.IsAliasable());
-						NewIAB->IndirectArgumentResources.Add(Argument(Surface->Texture, (mtlpp::ResourceUsage)(mtlpp::ResourceUsage::Read|mtlpp::ResourceUsage::Sample)));
-					}
-					else if (SB)
-					{
-						Desc.SetDataType(mtlpp::DataType::Pointer);
-						NewIAB->IndirectArgumentResources.Add(Argument(SB->Buffer, (mtlpp::ResourceUsage)(mtlpp::ResourceUsage::Read)));
-						
-						BufferSizes.Add(SB->GetSize());
-						BufferSizes.Add(GMetalBufferFormats[SRV->Format].DataFormat);
-					}
-					else
-					{
-						check(VB || IB);
-						ns::AutoReleased<FMetalTexture> Tex = SRV->GetLinearTexture(false);
-						Desc.SetDataType(mtlpp::DataType::Texture);
-						Desc.SetTextureType(Tex.GetTextureType());
-						NewIAB->IndirectArgumentResources.Add(Argument(Tex, (mtlpp::ResourceUsage)(mtlpp::ResourceUsage::Read|mtlpp::ResourceUsage::Sample)));
-					}
-					break;
-				}
-				case UBMT_RDG_TEXTURE_UAV:
-				case UBMT_RDG_BUFFER_UAV:
-				{
-					FMetalArgumentDesc& Desc = Arguments.Emplace_GetRef();
-					Desc.SetIndex(Index);
-					Desc.SetAccess(mtlpp::ArgumentAccess::ReadWrite);
-					
-					FMetalUnorderedAccessView* UAV = (FMetalUnorderedAccessView*)Resource;
-					FMetalShaderResourceView* SRV = (FMetalShaderResourceView*)UAV->SourceView;
-					FMetalStructuredBuffer* SB = UAV->SourceView->SourceStructuredBuffer.GetReference();
-					FMetalVertexBuffer* VB = UAV->SourceView->SourceVertexBuffer.GetReference();
-					FMetalIndexBuffer* IB = UAV->SourceView->SourceIndexBuffer.GetReference();
-					FRHITexture* Texture = UAV->SourceView->SourceTexture.GetReference();
-					FMetalSurface* Surface = UAV->SourceView->TextureView;
-					if (Texture)
-					{
-						if (!Surface)
-						{
-							Surface = GetMetalSurfaceFromRHITexture(Texture);
-						}
-						check (Surface != nullptr);
-						Desc.SetDataType(mtlpp::DataType::Texture);
-						Desc.SetTextureType(Surface->Texture.GetTextureType());
-						
-						union {
-							uint8 Components[4];
-							uint32 Packed;
-						} Swizzle;
-						Swizzle.Packed = 0;
-						assert(sizeof(Swizzle) == sizeof(uint32));
-						if (Surface->Texture.GetPixelFormat() == mtlpp::PixelFormat::X32_Stencil8
-#if PLATFORM_MAC
-						 ||	Surface->Texture.GetPixelFormat() == mtlpp::PixelFormat::X24_Stencil8
-#endif
-						)
-						{
-							Swizzle.Components[0] = Swizzle.Components[1] = Swizzle.Components[2] = Swizzle.Components[3] = 1;
-						}
-						BufferSizes.Add(Swizzle.Packed);
-						BufferSizes.Add(GMetalBufferFormats[Texture->GetFormat()].DataFormat);
-						
-						check(!Surface->Texture.IsAliasable());
-						NewIAB->IndirectArgumentResources.Add(Argument(Surface->Texture, (mtlpp::ResourceUsage)(mtlpp::ResourceUsage::Read|mtlpp::ResourceUsage::Write)));
-					}
-					else
-					{
-						check(VB || IB || SB);
-						ns::AutoReleased<FMetalTexture> Tex = SRV->GetLinearTexture(false);
-						Desc.SetDataType(mtlpp::DataType::Texture);
-						Desc.SetTextureType(Tex.GetTextureType());
-						NewIAB->IndirectArgumentResources.Add(Argument(Tex, (mtlpp::ResourceUsage)(mtlpp::ResourceUsage::Read|mtlpp::ResourceUsage::Write)));
-						
-						FMetalArgumentDesc& BufferDesc = Arguments.Emplace_GetRef();
-						BufferDesc.SetIndex(++Index);
-						BufferDesc.SetAccess(mtlpp::ArgumentAccess::ReadWrite);
-						
-						if (VB)
-						{
-							BufferDesc.SetDataType(mtlpp::DataType::Pointer);
-							NewIAB->IndirectArgumentResources.Add(Argument(VB->Buffer, (mtlpp::ResourceUsage)(mtlpp::ResourceUsage::Read|mtlpp::ResourceUsage::Write)));
-							
-							check(VB->Mode == mtlpp::StorageMode::Private);
-							BufferSizes.Add(VB->GetSize());
-							BufferSizes.Add(GMetalBufferFormats[SRV->Format].DataFormat);
-						}
-						else if (IB)
-						{
-							BufferDesc.SetDataType(mtlpp::DataType::Pointer);
-							NewIAB->IndirectArgumentResources.Add(Argument(IB->Buffer, (mtlpp::ResourceUsage)(mtlpp::ResourceUsage::Read|mtlpp::ResourceUsage::Write)));
-							
-							check(IB->Mode == mtlpp::StorageMode::Private);
-							BufferSizes.Add(IB->GetSize());
-							BufferSizes.Add(GMetalBufferFormats[SRV->Format].DataFormat);
-						}
-						else if (SB)
-						{
-							BufferDesc.SetDataType(mtlpp::DataType::Pointer);
-							NewIAB->IndirectArgumentResources.Add(Argument(SB->Buffer, (mtlpp::ResourceUsage)(mtlpp::ResourceUsage::Read|mtlpp::ResourceUsage::Write)));
-							
-							BufferSizes.Add(SB->GetSize());
-							BufferSizes.Add(GMetalBufferFormats[SRV->Format].DataFormat);
-						}
-					}
-					break;
-				}
-				case UBMT_SAMPLER:
-				{
-					FMetalArgumentDesc& Desc = Arguments.Emplace_GetRef();
-					Desc.SetIndex(Index);
-					Desc.SetAccess(mtlpp::ArgumentAccess::ReadOnly);
-					
-					FMetalSamplerState* Sampler = (FMetalSamplerState*)Resource;
-					Desc.SetDataType(mtlpp::DataType::Sampler);
-					NewIAB->IndirectArgumentResources.Add(Argument(Sampler->State));
-					break;
-				}
-				case UBMT_RDG_TEXTURE:
-				case UBMT_TEXTURE:
-				{
-					FMetalArgumentDesc& Desc = Arguments.Emplace_GetRef();
-					Desc.SetIndex(Index);
-					Desc.SetAccess(mtlpp::ArgumentAccess::ReadOnly);
-					
-					FRHITexture* Texture = (FRHITexture*)Resource;
-					FMetalSurface* Surface = GetMetalSurfaceFromRHITexture(Texture);
-					check (Surface != nullptr);
-					Desc.SetDataType(mtlpp::DataType::Texture);
-					Desc.SetTextureType(Surface->Texture.GetTextureType());
-					
-					union {
-						uint8 Components[4];
-						uint32 Packed;
-					} Swizzle;
-					Swizzle.Packed = 0;
-					assert(sizeof(Swizzle) == sizeof(uint32));
-					if (Surface->Texture.GetPixelFormat() == mtlpp::PixelFormat::X32_Stencil8
-#if PLATFORM_MAC
-					 ||	Surface->Texture.GetPixelFormat() == mtlpp::PixelFormat::X24_Stencil8
-#endif
-					)
-					{
-						Swizzle.Components[0] = Swizzle.Components[1] = Swizzle.Components[2] = Swizzle.Components[3] = 1;
-					}
-					BufferSizes.Add(Swizzle.Packed);
-					BufferSizes.Add(GMetalBufferFormats[Texture->GetFormat()].DataFormat);
-					
-					check(!Surface->Texture.IsAliasable());
-					NewIAB->IndirectArgumentResources.Add(Argument(Surface->Texture, (mtlpp::ResourceUsage)(mtlpp::ResourceUsage::Read|mtlpp::ResourceUsage::Sample)));
-					break;
-				}
-				default:
-					break;
-			}
-		}
-		
-		if (BufferSizes.Num() > 0)
-		{
-			FMetalPooledBufferArgs Args(GetMetalDeviceContext().GetDevice(), BufferSizes.Num() * sizeof(uint32), BUF_Dynamic, BUFFER_STORAGE_MODE);
-			NewIAB->IndirectArgumentBufferSideTable = GetMetalDeviceContext().CreatePooledBuffer(Args);
-			
-			FMemory::Memcpy(NewIAB->IndirectArgumentBufferSideTable.GetContents(), BufferSizes.GetData(), BufferSizes.Num() * sizeof(uint32));
-			
-#if PLATFORM_MAC
-			if(NewIAB->IndirectArgumentBufferSideTable.GetStorageMode() == mtlpp::StorageMode::Managed)
-			{
-				MTLPP_VALIDATE(mtlpp::Buffer, NewIAB->IndirectArgumentBufferSideTable, SafeGetRuntimeDebuggingLevel() >= EMetalDebugLevelValidation, DidModify(ns::Range(0, BufferSizes.Num() * sizeof(uint32))));
-			}
-#endif
-			
-			NewIAB->IndirectArgumentResources.Add(Argument(NewIAB->IndirectArgumentBufferSideTable, mtlpp::ResourceUsage::Read));
-		}
-		
-		if (ConstantSize > 0)
-		{
-			FMetalArgumentDesc& Desc = Arguments.Emplace_GetRef();
-			Desc.SetIndex(Index++);
-			Desc.SetAccess(mtlpp::ArgumentAccess::ReadOnly);
-			Desc.SetDataType(mtlpp::DataType::Pointer);
-			
-			NewIAB->IndirectArgumentResources.Add(Argument(Buffer, mtlpp::ResourceUsage::Read));
-		}
+		FMetalIAB* NewIAB = IAB->IndirectArgumentBuffer;
+		mtlpp::ArgumentEncoder Encoder = FMetalArgumentEncoderCache::Get().CreateEncoder(IAB->IndirectArgumentsDecl);
+		NewIAB->IndirectArgumentEncoder = Encoder;
 		
-		mtlpp::ArgumentEncoder Encoder = FMetalArgumentEncoderCache::Get().CreateEncoder(Arguments);
+		if (!NewIAB->IndirectArgumentBuffer)
+			NewIAB->IndirectArgumentBuffer = GetMetalDeviceContext().GetResourceHeap().CreateBuffer(Encoder.GetEncodedLength(), 16, BUF_Dynamic, mtlpp::ResourceOptions(BUFFER_CACHE_MODE | ((NSUInteger)mtlpp::StorageMode::Private << mtlpp::ResourceStorageModeShift)), true);
 		
-		NewIAB->IndirectArgumentBuffer = GetMetalDeviceContext().GetResourceHeap().CreateBuffer(Encoder.GetEncodedLength(), 16, BUF_Dynamic, mtlpp::ResourceOptions(BUFFER_CACHE_MODE | ((NSUInteger)BUFFER_STORAGE_MODE << mtlpp::ResourceStorageModeShift)), true);
+		FMetalBuffer Temp = Ctx->GetCurrentRenderPass().AllocateTemporyBufferForCopy(NewIAB->IndirectArgumentBuffer, Encoder.GetEncodedLength(), 16);
 		
-		Encoder.SetArgumentBuffer(NewIAB->IndirectArgumentBuffer, 0);
+		Encoder.SetArgumentBuffer(Temp, 0);
 		
-		for (FMetalArgumentDesc const& Arg : Arguments)
+		for (FMetalArgumentDesc const& Arg : IAB->IndirectArgumentsDecl)
 		{
 			NSUInteger NewIndex = Arg.Index;
 			switch(Arg.DataType)
 			{
 				case mtlpp::DataType::Pointer:
-					Encoder.SetBuffer(NewIAB->IndirectArgumentResources[NewIndex].Buffer, 0, NewIndex);
+					if (IAB->IndirectArgumentResources[NewIndex].Buffer)
+						Encoder.SetBuffer(IAB->IndirectArgumentResources[NewIndex].Buffer, 0, NewIndex);
 					break;
 				case mtlpp::DataType::Texture:
-					Encoder.SetTexture(NewIAB->IndirectArgumentResources[NewIndex].Texture, NewIndex);
+					if (IAB->IndirectArgumentResources[NewIndex].Texture)
+						Encoder.SetTexture(IAB->IndirectArgumentResources[NewIndex].Texture, NewIndex);
 					break;
 				case mtlpp::DataType::Sampler:
-					Encoder.SetSamplerState(NewIAB->IndirectArgumentResources[NewIndex].Sampler, NewIndex);
+					if (IAB->IndirectArgumentResources[NewIndex].Sampler)
+						Encoder.SetSamplerState(IAB->IndirectArgumentResources[NewIndex].Sampler, NewIndex);
 					break;
 				default:
 					break;
@@ -545,19 +253,34 @@ void FMetalUniformBuffer::InitIAB()
 		}
 		
 #if PLATFORM_MAC
-		if(NewIAB->IndirectArgumentBuffer.GetStorageMode() == mtlpp::StorageMode::Managed)
+		if(Temp.GetStorageMode() == mtlpp::StorageMode::Managed)
 		{
-			MTLPP_VALIDATE(mtlpp::Buffer, NewIAB->IndirectArgumentBuffer, SafeGetRuntimeDebuggingLevel() >= EMetalDebugLevelValidation, DidModify(ns::Range(0, Encoder.GetEncodedLength())));
+			MTLPP_VALIDATE(mtlpp::Buffer, Temp, SafeGetRuntimeDebuggingLevel() >= EMetalDebugLevelValidation, DidModify(ns::Range(0, Encoder.GetEncodedLength())));
 		}
 #endif
 		
-		// Atomically swap so that we don't explode if multiple threads attempt to initialise at the same time.
-		void* Result = FPlatformAtomics::InterlockedCompareExchangePointer((void**)&IAB, (void*)NewIAB, nullptr);
-		if (Result != nullptr)
-		{
-			check(Result != NewIAB);
-			delete NewIAB;
-		}
+		Ctx->AsyncCopyFromBufferToBuffer(Temp, 0, NewIAB->IndirectArgumentBuffer, 0, Encoder.GetEncodedLength());
+	}
+	return IAB->IndirectArgumentBuffer;
+}
+
+FMetalUniformBuffer::FMetalIndirectArgumentBuffer& FMetalUniformBuffer::GetIAB()
+{
+	check(ResourceTable.Num() && FMetalCommandQueue::SupportsFeature(EMetalFeaturesIABs));
+
+	check(IAB);
+	
+	return *IAB;
+}
+
+void FMetalUniformBuffer::InitIAB()
+{
+	if (NumResources && FMetalCommandQueue::SupportsFeature(EMetalFeaturesIABs) && !IAB)
+	{
+		FMetalIndirectArgumentBuffer* NewIAB = new FMetalIndirectArgumentBuffer;
+		NewIAB->IndirectArgumentBuffer = [FMetalIAB new];
+		NewIAB->UpdateNum = ~0u;
+		IAB = NewIAB;
 	}
 }
 
@@ -577,18 +300,69 @@ void const* FMetalUniformBuffer::GetData()
 	}
 }
 
-void FMetalUniformBuffer::UpdateResourceTable(TArray<TRefCountPtr<FRHIResource>>& Resources, EUniformBufferValidation Validation)
+void FMetalUniformBuffer::UpdateTextureReference(FRHITextureReference* ModifiedRef)
 {
-	ResourceTable = Resources;
-	
-	if (NumResources && FMetalCommandQueue::SupportsFeature(EMetalFeaturesIABs))
+	if (NumResources && FMetalCommandQueue::SupportsFeature(EMetalFeaturesIABs) && IAB && TextureReferences.Contains(ModifiedRef))
+	{
+		bool bModified = false;
+		bool bUpdatedDesc = false;
+		FMetalSurface* Surface = GetMetalSurfaceFromRHITexture(ModifiedRef);
+		if (Surface)
+		{
+			TBitArray<>& Ref = TextureReferences.FindChecked(ModifiedRef);
+			for (uint32 i = 0; i < Ref.Num(); i++)
+			{
+				if (Ref[i])
+				{
+					bModified |= (IAB->IndirectArgumentResources[i].Texture != Surface->Texture);
+					IAB->IndirectArgumentResources[i].Texture = Surface->Texture;
+					
+					FMetalArgumentDesc& Desc = IAB->IndirectArgumentsDecl[i];
+					bUpdatedDesc |= Desc.TextureType != Surface->Texture.GetTextureType();
+					Desc.SetTextureType(Surface->Texture.GetTextureType());
+					
+					bModified |= bUpdatedDesc;
+				}
+			}
+		}
+		if (bUpdatedDesc)
+		{
+			FPlatformAtomics::InterlockedIncrement(&IAB->UpdateEnc);
+		}
+		if (bModified)
+		{
+			FPlatformAtomics::InterlockedIncrement(&IAB->UpdateIAB);
+		}
+	}
+}
+
+void FMetalUniformBuffer::UpdateIAB()
+{
+	if (NumResources && FMetalCommandQueue::SupportsFeature(EMetalFeaturesIABs) && IAB && (UpdateNum != IAB->UpdateNum))
 	{
+		IAB->IndirectArgumentResources.Empty();
+		
+		TArray<uint32>& BufferSizes = IAB->IndirectBufferSizes;
+		TArray<FMetalArgumentDesc>& Arguments = IAB->IndirectArgumentsDecl;
+		Arguments.Empty();
+		BufferSizes.Empty();
+		
+		int32 Index = 0;
+		// Buffer Size table
+		{
+			FMetalArgumentDesc& Desc = Arguments.Emplace_GetRef();
+			Desc.SetIndex(Index++);
+			Desc.SetAccess(mtlpp::ArgumentAccess::ReadOnly);
+			Desc.SetDataType(mtlpp::DataType::Pointer);
+			
+			IAB->IndirectArgumentResources.Add(Argument(FMetalBuffer(), mtlpp::ResourceUsage::Read));
+		}
+		
 		// set up an SRT-style uniform buffer
 		for (uint32 i = 0; i < NumResources; ++i)
 		{
 			FRHIResource* Resource = ResourceTable[i];
-			
-			if (FMetalCommandQueue::SupportsFeature(EMetalFeaturesIABs) && Resource)
+			if (Resource)
 			{
 				switch(ResourceTypes[i])
 				{
@@ -596,47 +370,283 @@ void FMetalUniformBuffer::UpdateResourceTable(TArray<TRefCountPtr<FRHIResource>>
 					case UBMT_RDG_BUFFER_SRV:
 					case UBMT_SRV:
 					{
+						FMetalArgumentDesc& Desc = Arguments.Emplace_GetRef();
+						Desc.SetIndex(Index++);
+						Desc.SetAccess(mtlpp::ArgumentAccess::ReadOnly);
+						
 						FMetalShaderResourceView* SRV = (FMetalShaderResourceView*)Resource;
 						FRHITexture* Texture = SRV->SourceTexture.GetReference();
-						if (Texture && Texture->GetTextureReference())
+						FMetalVertexBuffer* VB = SRV->SourceVertexBuffer.GetReference();
+						FMetalIndexBuffer* IB = SRV->SourceIndexBuffer.GetReference();
+						FMetalStructuredBuffer* SB = SRV->SourceStructuredBuffer.GetReference();
+						if (Texture)
+						{
+							Desc.SetDataType(mtlpp::DataType::Texture);
+							
+							FRHITextureReference* Reference = Texture->GetTextureReference();
+							if (Reference)
+							{
+								TBitArray<>& Ref = TextureReferences.FindOrAdd(Reference);
+								if (Ref.Num() <= IAB->IndirectArgumentResources.Num())
+								{
+									Ref.Add(false, (IAB->IndirectArgumentResources.Num() + 1) - Ref.Num());
+								}
+								Ref[IAB->IndirectArgumentResources.Num()] = true;
+							}
+							
+							union {
+								uint8 Components[4];
+								uint32 Packed;
+							} Swizzle;
+							Swizzle.Packed = 0;
+							assert(sizeof(Swizzle) == sizeof(uint32));
+							
+							FMetalSurface* Surface = SRV->TextureView;
+							if (Surface)
+							{
+								check(!Surface->Texture.IsAliasable());
+								Desc.SetTextureType(Surface->Texture.GetTextureType());
+								IAB->IndirectArgumentResources.Add(Argument(Surface->Texture, (mtlpp::ResourceUsage)(mtlpp::ResourceUsage::Read|mtlpp::ResourceUsage::Sample)));
+								if (Surface->Texture.GetPixelFormat() == mtlpp::PixelFormat::X32_Stencil8
+#if PLATFORM_MAC
+									||	Surface->Texture.GetPixelFormat() == mtlpp::PixelFormat::X24_Stencil8
+#endif
+									)
+								{
+									Swizzle.Components[0] = Swizzle.Components[1] = Swizzle.Components[2] = Swizzle.Components[3] = 1;
+								}
+							}
+							else
+							{
+								IAB->IndirectArgumentResources.Add(Argument(FMetalTexture(), (mtlpp::ResourceUsage)(mtlpp::ResourceUsage::Read|mtlpp::ResourceUsage::Sample)));
+							}
+							
+							BufferSizes.Add(Swizzle.Packed);
+							BufferSizes.Add(GMetalBufferFormats[Texture->GetFormat()].DataFormat);
+						}
+						else if (SB)
 						{
-							TextureReferences.Add(Texture->GetTextureReference());
+							Desc.SetDataType(mtlpp::DataType::Pointer);
+							
+							IAB->IndirectArgumentResources.Add(Argument(SB->Buffer, (mtlpp::ResourceUsage)(mtlpp::ResourceUsage::Read)));
+							
+							BufferSizes.Add(SB->GetSize());
+							BufferSizes.Add(GMetalBufferFormats[SRV->Format].DataFormat);
+						}
+						else
+						{
+							check(VB || IB);
+							ns::AutoReleased<FMetalTexture> Tex = SRV->GetLinearTexture(false);
+							IAB->IndirectArgumentResources.Add(Argument(Tex, (mtlpp::ResourceUsage)(mtlpp::ResourceUsage::Read|mtlpp::ResourceUsage::Sample)));
+							
+							Desc.SetDataType(mtlpp::DataType::Texture);
+							Desc.SetTextureType(Tex.GetTextureType());
+							
+							BufferSizes.Add(VB ? VB->GetSize() : IB->GetSize());
+							BufferSizes.Add(GMetalBufferFormats[SRV->Format].DataFormat);
 						}
 						break;
 					}
 					case UBMT_RDG_TEXTURE_UAV:
 					case UBMT_RDG_BUFFER_UAV:
 					{
+						FMetalArgumentDesc& Desc = Arguments.Emplace_GetRef();
+						Desc.SetIndex(Index++);
+						Desc.SetAccess(mtlpp::ArgumentAccess::ReadWrite);
+						
 						FMetalUnorderedAccessView* UAV = (FMetalUnorderedAccessView*)Resource;
+						FMetalShaderResourceView* SRV = (FMetalShaderResourceView*)UAV->SourceView;
+						FMetalStructuredBuffer* SB = UAV->SourceView->SourceStructuredBuffer.GetReference();
+						FMetalVertexBuffer* VB = UAV->SourceView->SourceVertexBuffer.GetReference();
+						FMetalIndexBuffer* IB = UAV->SourceView->SourceIndexBuffer.GetReference();
 						FRHITexture* Texture = UAV->SourceView->SourceTexture.GetReference();
-						if (Texture && Texture->GetTextureReference())
+						FMetalSurface* Surface = UAV->SourceView->TextureView;
+						if (Texture)
 						{
-							TextureReferences.Add(Texture->GetTextureReference());
+							Desc.SetDataType(mtlpp::DataType::Texture);
+							
+							FRHITextureReference* Reference = Texture->GetTextureReference();
+							if (Reference)
+							{
+								TBitArray<>& Ref = TextureReferences.FindOrAdd(Reference);
+								if (Ref.Num() <= IAB->IndirectArgumentResources.Num())
+							   	{
+								   	Ref.Add(false, (IAB->IndirectArgumentResources.Num() + 1) - Ref.Num());
+							   	}
+							   	Ref[IAB->IndirectArgumentResources.Num()] = true;
+							}
+							
+							union {
+								uint8 Components[4];
+								uint32 Packed;
+							} Swizzle;
+							Swizzle.Packed = 0;
+							assert(sizeof(Swizzle) == sizeof(uint32));
+							
+							if (!Surface)
+							{
+								Surface = GetMetalSurfaceFromRHITexture(Texture);
+							}
+							if (Surface)
+							{
+								check(!Surface->Texture.IsAliasable());
+								if (Surface->Texture.GetPixelFormat() == mtlpp::PixelFormat::X32_Stencil8
+#if PLATFORM_MAC
+									||	Surface->Texture.GetPixelFormat() == mtlpp::PixelFormat::X24_Stencil8
+#endif
+									)
+								{
+									Swizzle.Components[0] = Swizzle.Components[1] = Swizzle.Components[2] = Swizzle.Components[3] = 1;
+								}
+								Desc.SetTextureType(Surface->Texture.GetTextureType());
+								IAB->IndirectArgumentResources.Add(Argument(Surface->Texture, (mtlpp::ResourceUsage)(mtlpp::ResourceUsage::Read|mtlpp::ResourceUsage::Write)));
+							}
+							else
+							{
+								IAB->IndirectArgumentResources.Add(Argument(FMetalTexture(), (mtlpp::ResourceUsage)(mtlpp::ResourceUsage::Read|mtlpp::ResourceUsage::Write)));
+							}
+							BufferSizes.Add(Swizzle.Packed);
+							BufferSizes.Add(GMetalBufferFormats[Texture->GetFormat()].DataFormat);
+						}
+						else if (SB)
+						{
+							Desc.SetDataType(mtlpp::DataType::Pointer);
+							IAB->IndirectArgumentResources.Add(Argument(SB->Buffer, (mtlpp::ResourceUsage)(mtlpp::ResourceUsage::Read|mtlpp::ResourceUsage::Write)));
+							BufferSizes.Add(SB->GetSize());
+							BufferSizes.Add(PF_Unknown);
+						}
+						else
+						{
+							check(VB || IB);
+							ns::AutoReleased<FMetalTexture> Tex = SRV->GetLinearTexture(false);
+							IAB->IndirectArgumentResources.Add(Argument(Tex, (mtlpp::ResourceUsage)(mtlpp::ResourceUsage::Read|mtlpp::ResourceUsage::Write)));
+							Desc.SetDataType(mtlpp::DataType::Texture);
+							Desc.SetTextureType(Tex.GetTextureType());
+							BufferSizes.Add(0);
+							BufferSizes.Add(0);
+							
+							FMetalArgumentDesc& BufferDesc = Arguments.Emplace_GetRef();
+							BufferDesc.SetIndex(Index++);
+							BufferDesc.SetAccess(mtlpp::ArgumentAccess::ReadWrite);
+							BufferDesc.SetDataType(mtlpp::DataType::Pointer);
+							
+							if (VB)
+							{
+								IAB->IndirectArgumentResources.Add(Argument(VB->Buffer, (mtlpp::ResourceUsage)(mtlpp::ResourceUsage::Read|mtlpp::ResourceUsage::Write)));
+								BufferSizes.Add(VB->GetSize());
+								BufferSizes.Add(GMetalBufferFormats[SRV->Format].DataFormat);
+							}
+							else
+							{
+								IAB->IndirectArgumentResources.Add(Argument(IB->Buffer, (mtlpp::ResourceUsage)(mtlpp::ResourceUsage::Read|mtlpp::ResourceUsage::Write)));
+								BufferSizes.Add(IB->GetSize());
+								BufferSizes.Add(GMetalBufferFormats[SRV->Format].DataFormat);
+							}
 						}
 						break;
 					}
+					case UBMT_SAMPLER:
+					{
+						FMetalArgumentDesc& Desc = Arguments.Emplace_GetRef();
+						Desc.SetIndex(Index++);
+						Desc.SetAccess(mtlpp::ArgumentAccess::ReadOnly);
+						Desc.SetDataType(mtlpp::DataType::Sampler);
+						
+						FMetalSamplerState* Sampler = (FMetalSamplerState*)Resource;
+						IAB->IndirectArgumentResources.Add(Argument(Sampler->State));
+						
+						BufferSizes.Add(0);
+						BufferSizes.Add(0);
+						break;
+					}
 					case UBMT_RDG_TEXTURE:
 					case UBMT_TEXTURE:
 					{
+						FMetalArgumentDesc& Desc = Arguments.Emplace_GetRef();
+						Desc.SetIndex(Index++);
+						Desc.SetAccess(mtlpp::ArgumentAccess::ReadOnly);
+						Desc.SetDataType(mtlpp::DataType::Texture);
+						
 						FRHITexture* Texture = (FRHITexture*)Resource;
-						if (Texture && Texture->GetTextureReference())
+						FMetalSurface* Surface = GetMetalSurfaceFromRHITexture(Texture);
+						FRHITextureReference* Reference = Texture->GetTextureReference();
+						if (Reference)
+						{
+							TBitArray<>& Ref = TextureReferences.FindOrAdd(Reference);
+							if (Ref.Num() <= IAB->IndirectArgumentResources.Num())
+							{
+								Ref.Add(false, (IAB->IndirectArgumentResources.Num() + 1) - Ref.Num());
+							}
+							Ref[IAB->IndirectArgumentResources.Num()] = true;
+						}
+						union {
+							uint8 Components[4];
+							uint32 Packed;
+						} Swizzle;
+						Swizzle.Packed = 0;
+						assert(sizeof(Swizzle) == sizeof(uint32));
+						if (Surface)
+						{
+							Desc.SetTextureType(Surface->Texture.GetTextureType());
+							check(!Surface->Texture.IsAliasable());
+							IAB->IndirectArgumentResources.Add(Argument(Surface->Texture, (mtlpp::ResourceUsage)(mtlpp::ResourceUsage::Read|mtlpp::ResourceUsage::Sample)));
+							if (Surface->Texture.GetPixelFormat() == mtlpp::PixelFormat::X32_Stencil8
+#if PLATFORM_MAC
+								||	Surface->Texture.GetPixelFormat() == mtlpp::PixelFormat::X24_Stencil8
+#endif
+								)
+							{
+								Swizzle.Components[0] = Swizzle.Components[1] = Swizzle.Components[2] = Swizzle.Components[3] = 1;
+							}
+						}
+						else
 						{
-							TextureReferences.Add(Texture->GetTextureReference());
+							IAB->IndirectArgumentResources.Add(Argument(FMetalTexture(), (mtlpp::ResourceUsage)(mtlpp::ResourceUsage::Read|mtlpp::ResourceUsage::Sample)));
 						}
+						BufferSizes.Add(Swizzle.Packed);
+						BufferSizes.Add(GMetalBufferFormats[Texture->GetFormat()].DataFormat);
 						break;
 					}
 					default:
-					{
 						break;
-					}
 				}
 			}
 		}
 		
+		{
+			FMetalArgumentDesc& Desc = Arguments.Emplace_GetRef();
+			Desc.SetIndex(Index++);
+			Desc.SetAccess(mtlpp::ArgumentAccess::ReadOnly);
+			Desc.SetDataType(mtlpp::DataType::Pointer);
+			
+			if (ConstantSize > 0)
+				IAB->IndirectArgumentResources.Add(Argument(Buffer, mtlpp::ResourceUsage::Read));
+			else
+				IAB->IndirectArgumentResources.Add(Argument(FMetalBuffer(), mtlpp::ResourceUsage::Read));
+		}
+		
+		IAB->UpdateNum = UpdateNum;
+		FPlatformAtomics::InterlockedIncrement(&IAB->UpdateIAB);
+		FPlatformAtomics::InterlockedIncrement(&IAB->UpdateEnc);
+	}
+}
+
+void FMetalUniformBuffer::UpdateResourceTable(TArray<TRefCountPtr<FRHIResource>>& Resources, EUniformBufferValidation Validation)
+{
+	ResourceTable = Resources;
+	
+	if (NumResources && FMetalCommandQueue::SupportsFeature(EMetalFeaturesIABs))
+	{
 		FRHICommandListImmediate& RHICmdList = FRHICommandListExecutor::GetImmediateCommandList();
 		if (!(UniformUsage & UniformBuffer_SingleDraw) && IsRunningRHIInSeparateThread() && !RHICmdList.Bypass() && IsInRenderingThread())
 		{
 			new (RHICmdList.AllocCommand<FMetalRHICommandInitialiseUniformBufferIAB>()) FMetalRHICommandInitialiseUniformBufferIAB(this);
+			RHICmdList.RHIThreadFence(true);
+		}
+		else
+		{
+			FPlatformAtomics::InterlockedIncrement(&UpdateNum);
+			UpdateIAB();
 		}
 	}
 }
@@ -647,9 +657,14 @@ void FMetalUniformBuffer::Update(const void* Contents, TArray<TRefCountPtr<FRHIR
     {
         UE_CLOG(ConstantSize > 65536, LogMetal, Fatal, TEXT("Trying to allocated a uniform layout of size %d that is greater than the maximum permitted 64k."), ConstantSize);
         
+		ns::AutoReleased<FMetalBuffer> Buf(Buffer);
+		
         void* Data = Lock(true, RLM_WriteOnly, 0);
         FMemory::Memcpy(Data, Contents, ConstantSize);
         Unlock();
+		
+		if (Buf != Buffer && FMetalCommandQueue::SupportsFeature(EMetalFeaturesIABs))
+			FPlatformAtomics::InterlockedIncrement(&UpdateNum);
     }
 	
 	UpdateResourceTable(Resources, Validation);
diff --git a/Engine/Source/Runtime/Apple/MetalRHI/Private/MetalVertexBuffer.cpp b/Engine/Source/Runtime/Apple/MetalRHI/Private/MetalVertexBuffer.cpp
--- a/Engine/Source/Runtime/Apple/MetalRHI/Private/MetalVertexBuffer.cpp
+++ b/Engine/Source/Runtime/Apple/MetalRHI/Private/MetalVertexBuffer.cpp
@@ -446,9 +446,9 @@ void* FMetalRHIBuffer::Lock(bool bIsOnRHIThread, EResourceLockMode LockMode, uin
 			// Turns out to be better to use Shared->Private blits whenever possible
 			// Should only put write-once buffers into Shared/Managed or the cost of recreating linear textures overwhelms any other efficiency
 			Mode =  CanUsePrivateMemory() ? mtlpp::StorageMode::Private : Mode;
-			METAL_INC_DWORD_STAT_BY(Type, MemFreed, Len);
-			SafeReleaseMetalBuffer(Buffer);
-			Buffer = nil;
+//			METAL_INC_DWORD_STAT_BY(Type, MemFreed, Len);
+//			SafeReleaseMetalBuffer(Buffer);
+//			Buffer = nil;
 		}
 	}
     
diff --git a/Engine/Source/Runtime/Apple/MetalRHI/Public/MetalResources.h b/Engine/Source/Runtime/Apple/MetalRHI/Public/MetalResources.h
--- a/Engine/Source/Runtime/Apple/MetalRHI/Public/MetalResources.h
+++ b/Engine/Source/Runtime/Apple/MetalRHI/Public/MetalResources.h
@@ -134,6 +134,9 @@ public:
 	// List of memory copies from RHIUniformBuffer to packed uniforms
 	TArray<CrossCompiler::FUniformBufferCopyInfo> UniformBuffersCopyInfo;
 	
+	/* Argument encoders for shader IABs */
+	TMap<uint32, mtlpp::ArgumentEncoder> ArgumentEncoders;
+	
 	/** The binding for the buffer side-table if present */
 	int32 SideTableBinding;
 
@@ -258,6 +261,7 @@ public:
 	virtual ~FMetalComputeShader();
 	
 	FMetalShaderPipeline* GetPipeline();
+	mtlpp::Function GetFunction();
 	
 	// thread group counts
 	int32 NumThreadsX;
@@ -876,27 +880,88 @@ public:
 	void Swap(FMetalVertexBuffer& Other);
 };
 
-class FMetalUniformBuffer : public FRHIUniformBuffer, public FMetalRHIBuffer
+struct FMetalArgumentDesc
 {
-public:
-
-	// Constructor
-	FMetalUniformBuffer(const void* Contents, const FRHIUniformBufferLayout& Layout, EUniformBufferUsage Usage, EUniformBufferValidation Validation);
-
-	// Destructor 
-	virtual ~FMetalUniformBuffer();
+	FMetalArgumentDesc()
+	: DataType((mtlpp::DataType)0)
+	, Index(0)
+	, ArrayLength(0)
+	, Access(mtlpp::ArgumentAccess::ReadOnly)
+	, TextureType((mtlpp::TextureType)0)
+	, ConstantBlockAlignment(0)
+	{
+		
+	}
 	
-	void const* GetData();
-
-	void InitIAB();
-
-	void UpdateResourceTable(TArray<TRefCountPtr<FRHIResource>>& Resources, EUniformBufferValidation Validation);
-	void Update(const void* Contents, TArray<TRefCountPtr<FRHIResource>>& Resources, EUniformBufferValidation Validation);
+	mtlpp::DataType DataType;
+	NSUInteger Index;
+	NSUInteger ArrayLength;
+	mtlpp::ArgumentAccess Access;
+	mtlpp::TextureType TextureType;
+	NSUInteger ConstantBlockAlignment;
 	
-	/** Resource table containing RHI references. */
-	TArray<TRefCountPtr<FRHIResource> > ResourceTable;
+	void FillDescriptor(mtlpp::ArgumentDescriptor& Desc) const
+	{
+		Desc.SetDataType(DataType);
+		Desc.SetIndex(Index);
+		Desc.SetArrayLength(ArrayLength);
+		Desc.SetAccess(Access);
+		Desc.SetTextureType(TextureType);
+		Desc.SetConstantBlockAlignment(ConstantBlockAlignment);
+	}
+	
+	void SetDataType(mtlpp::DataType Type)
+	{
+		DataType = Type;
+	}
+	
+	void SetIndex(NSUInteger i)
+	{
+		Index = i;
+	}
 	
-	TSet<FRHITextureReference*> TextureReferences;
+	void SetArrayLength(NSUInteger Len)
+	{
+		ArrayLength = Len;
+	}
+	
+	void SetAccess(mtlpp::ArgumentAccess InAccess)
+	{
+		Access = InAccess;
+	}
+	
+	void SetTextureType(mtlpp::TextureType Type)
+	{
+		TextureType = Type;
+	}
+	
+	void SetConstantBlockAlignment(NSUInteger Align)
+	{
+		ConstantBlockAlignment = Align;
+	}
+	
+	bool operator==(FMetalArgumentDesc const& Other) const
+	{
+		if (this != &Other)
+		{
+			return (DataType == Other.DataType && Index == Other.Index && ArrayLength == Other.ArrayLength && Access == Other.Access && TextureType == Other.TextureType && ConstantBlockAlignment == Other.ConstantBlockAlignment);
+		}
+		return true;
+	}
+	
+	friend uint32 GetTypeHash(FMetalArgumentDesc const& Desc)
+	{
+		uint32 Hash = 0;
+		Hash ^= ((uint32)Desc.DataType * (uint32)Desc.TextureType * (uint32)Desc.Access * Desc.ArrayLength) << Desc.Index;
+		return Hash;
+	}
+};
+
+@class FMetalIAB;
+
+class FMetalUniformBuffer : public FRHIUniformBuffer, public FMetalRHIBuffer
+{
+public:
 
 	struct Argument
 	{
@@ -904,7 +969,7 @@ public:
 		Argument(FMetalBuffer const& InBuffer, mtlpp::ResourceUsage const InUsage) : Buffer(InBuffer), Usage(InUsage) {}
 		Argument(FMetalTexture const& InTexture, mtlpp::ResourceUsage const InUsage) : Texture(InTexture), Usage(InUsage) {}
 		Argument(FMetalSampler const& InSampler) : Sampler(InSampler), Usage(mtlpp::ResourceUsage::Read) {}
-
+		
 		FMetalBuffer Buffer;
 		FMetalTexture Texture;
 		FMetalSampler Sampler;
@@ -916,15 +981,38 @@ public:
 		FMetalIndirectArgumentBuffer();
 		~FMetalIndirectArgumentBuffer();
 		
+		int64 UpdateNum; // The resource & declarations have been updated
+		int64 UpdateEnc; // The IAB encoder needs to be updated.
+		int64 UpdateIAB; // The IAB needs to be updated.
+		TArray<FMetalArgumentDesc> IndirectArgumentsDecl;
 		TArray<Argument> IndirectArgumentResources;
-		FMetalBuffer IndirectArgumentBuffer;
-		FMetalBuffer IndirectArgumentBufferSideTable;
+		TArray<uint32> IndirectBufferSizes;
+		FMetalIAB* IndirectArgumentBuffer;
 	};
 	
-	EUniformBufferUsage UniformUsage;
+	// Constructor
+	FMetalUniformBuffer(const void* Contents, const FRHIUniformBufferLayout& Layout, EUniformBufferUsage Usage, EUniformBufferValidation Validation);
+
+	// Destructor 
+	virtual ~FMetalUniformBuffer();
+	
+	void const* GetData();
+
+	void InitIAB();
+	void UpdateIAB();
+	void UpdateTextureReference(FRHITextureReference* ModifiedRef);
+	void UpdateResourceTable(TArray<TRefCountPtr<FRHIResource>>& Resources, EUniformBufferValidation Validation);
+	void Update(const void* Contents, TArray<TRefCountPtr<FRHIResource>>& Resources, EUniformBufferValidation Validation);
+	FMetalIAB* UploadIAB(class FMetalContext* Ctx);
 	FMetalIndirectArgumentBuffer& GetIAB();
+	
+	/** Resource table containing RHI references. */
+	TArray<TRefCountPtr<FRHIResource> > ResourceTable;
+	TMap<FRHITextureReference*, TBitArray<>> TextureReferences;
+	EUniformBufferUsage UniformUsage;
 	FMetalIndirectArgumentBuffer* IAB;
 	TArray<EUniformBufferBaseType> ResourceTypes;
+	int64 UpdateNum;
 	uint32 NumResources;
 	uint32 ConstantSize;
 };
