Revert unintended change.

diff --git a/Engine/Source/Runtime/Apple/MetalRHI/Private/MetalVertexBuffer.cpp b/Engine/Source/Runtime/Apple/MetalRHI/Private/MetalVertexBuffer.cpp
--- a/Engine/Source/Runtime/Apple/MetalRHI/Private/MetalVertexBuffer.cpp
+++ b/Engine/Source/Runtime/Apple/MetalRHI/Private/MetalVertexBuffer.cpp
@@ -446,9 +446,9 @@ void* FMetalRHIBuffer::Lock(bool bIsOnRHIThread, EResourceLockMode LockMode, uin
 			// Turns out to be better to use Shared->Private blits whenever possible
 			// Should only put write-once buffers into Shared/Managed or the cost of recreating linear textures overwhelms any other efficiency
 			Mode =  CanUsePrivateMemory() ? mtlpp::StorageMode::Private : Mode;
-//			METAL_INC_DWORD_STAT_BY(Type, MemFreed, Len);
-//			SafeReleaseMetalBuffer(Buffer);
-//			Buffer = nil;
+			METAL_INC_DWORD_STAT_BY(Type, MemFreed, Len);
+			SafeReleaseMetalBuffer(Buffer);
+			Buffer = nil;
 		}
 	}
     
