Make the use of Indirect Argument Buffers a separate option from the underlying Metal shader language version.

diff --git a/Engine/Source/Developer/Apple/MetalShaderFormat/Private/MetalDerivedData.cpp b/Engine/Source/Developer/Apple/MetalShaderFormat/Private/MetalDerivedData.cpp
--- a/Engine/Source/Developer/Apple/MetalShaderFormat/Private/MetalDerivedData.cpp
+++ b/Engine/Source/Developer/Apple/MetalShaderFormat/Private/MetalDerivedData.cpp
@@ -320,6 +320,7 @@ FMetalShaderOutputCooker::FMetalShaderOutputCooker(const FShaderCompilerInput& _
 	, GUIDHash(_GUIDHash)
 	, VersionEnum(_VersionEnum)
 	, CCFlags(_CCFlags)
+	, IABTier(0)
 	, HlslCompilerTarget(_HlslCompilerTarget)
 	, MetalCompilerTarget(_MetalCompilerTarget)
 	, Semantics(_Semantics)
@@ -330,6 +331,14 @@ FMetalShaderOutputCooker::FMetalShaderOutputCooker(const FShaderCompilerInput& _
 	, Standard(_Standard)
 	, MinOSVersion(_MinOSVersion)
 {
+	FString const* IABVersion = Input.Environment.GetDefinitions().Find(TEXT("METAL_INDIRECT_ARGUMENT_BUFFERS"));
+	if (IABVersion)
+	{
+		if(IABVersion->IsNumeric())
+		{
+			LexFromString(IABTier, *(*IABVersion));
+		}
+	}
 }
 
 FMetalShaderOutputCooker::~FMetalShaderOutputCooker()
@@ -368,7 +377,7 @@ FString FMetalShaderOutputCooker::GetPluginSpecificCacheKeySuffix() const
 			Flags |= (1llu << uint64(Flag));
 		}
 
-		CachedOutputName = FString::Printf(TEXT("%s-%s_%s-%u_%hu_%llu_%hhu_%s_%s"), *Input.ShaderFormat.GetPlainNameString(), *Input.EntryPointName, *Hash.ToString(), Len, FormatVers, Flags, VersionEnum, *GUIDHash.ToString(), *Standard);
+		CachedOutputName = FString::Printf(TEXT("%s-%s_%s-%u_%hu_%llu_%hhu_%d_%s_%s"), *Input.ShaderFormat.GetPlainNameString(), *Input.EntryPointName, *Hash.ToString(), Len, FormatVers, Flags, VersionEnum, IABTier, *GUIDHash.ToString(), *Standard);
 	}
 
 	return CachedOutputName;
@@ -648,7 +657,7 @@ bool FMetalShaderOutputCooker::Build(TArray<uint8>& OutData)
 			
 			TArray<FString> TableNames;
 			TMap<FString, FMetalResourceTableEntry> ResourceTable;
-			if (VersionEnum >= 5)
+			if (IABTier >= 1)
 			{
 				for (auto Pair : Input.Environment.ResourceTableLayoutHashes)
 				{
@@ -1914,13 +1923,14 @@ bool FMetalShaderOutputCooker::Build(TArray<uint8>& OutData)
 			{
 				case 6:
 				case 5:
-				{
-					FCStringAnsi::Snprintf(ArgumentBufferOffset, 3, "%u", IABOffsetIndex);
-					Defines[TargetDesc.numOptions++] = { "argument_buffers", "1" };
-					Defines[TargetDesc.numOptions++] = { "argument_buffer_offset", ArgumentBufferOffset };
-				}
 				case 4:
 				{
+					if (IABTier > 1)
+					{
+						FCStringAnsi::Snprintf(ArgumentBufferOffset, 3, "%u", IABOffsetIndex);
+						Defines[TargetDesc.numOptions++] = { "argument_buffers", "1" };
+						Defines[TargetDesc.numOptions++] = { "argument_buffer_offset", ArgumentBufferOffset };
+					}
 					Defines[TargetDesc.numOptions++] = { "texture_buffer_native", "1" };
 					TargetDesc.version = "20100";
 					break;
@@ -2072,7 +2082,7 @@ bool FMetalShaderOutputCooker::Build(TArray<uint8>& OutData)
 		}
 		
 		// Version 6 means Tier 2 IABs for now.
-		if (VersionEnum >= 6)
+		if (IABTier >= 2)
 		{
 			char BufferIdx[3];
 			for (auto& IAB : IABs)
diff --git a/Engine/Source/Developer/Apple/MetalShaderFormat/Private/MetalDerivedData.h b/Engine/Source/Developer/Apple/MetalShaderFormat/Private/MetalDerivedData.h
--- a/Engine/Source/Developer/Apple/MetalShaderFormat/Private/MetalDerivedData.h
+++ b/Engine/Source/Developer/Apple/MetalShaderFormat/Private/MetalDerivedData.h
@@ -175,6 +175,7 @@ private:
 	FSHAHash GUIDHash;
 	uint8 VersionEnum;
 	uint32 CCFlags;
+	int32 IABTier;
 	EHlslCompileTarget HlslCompilerTarget;
 	EHlslCompileTarget MetalCompilerTarget;
 	EMetalGPUSemantics Semantics;
diff --git a/Engine/Source/Developer/Apple/MetalShaderFormat/Private/MetalShaderCompiler.cpp b/Engine/Source/Developer/Apple/MetalShaderFormat/Private/MetalShaderCompiler.cpp
--- a/Engine/Source/Developer/Apple/MetalShaderFormat/Private/MetalShaderCompiler.cpp
+++ b/Engine/Source/Developer/Apple/MetalShaderFormat/Private/MetalShaderCompiler.cpp
@@ -2283,11 +2283,8 @@ void CompileShader_Metal(const FShaderCompilerInput& _Input,FShaderCompilerOutpu
 
 
 	// This requires removing the HLSLCC_NoPreprocess flag later on!
-    if (VersionEnum < 5 || VersionEnum == 6)
-    {
-        RemoveUniformBuffersFromSource(Input.Environment, PreprocessedShader);
-    }
-
+    RemoveUniformBuffersFromSource(Input.Environment, PreprocessedShader);
+	
 	uint32 CCFlags = HLSLCC_NoPreprocess | HLSLCC_PackUniformsIntoUniformBufferWithNames | HLSLCC_FixAtomicReferences | HLSLCC_RetainSizes | HLSLCC_KeepSamplerAndImageNames;
 	if (!bDirectCompile || UE_BUILD_DEBUG)
 	{
diff --git a/Engine/Source/Developer/Mac/MacTargetPlatform/Classes/MacTargetSettings.h b/Engine/Source/Developer/Mac/MacTargetPlatform/Classes/MacTargetSettings.h
--- a/Engine/Source/Developer/Mac/MacTargetPlatform/Classes/MacTargetSettings.h
+++ b/Engine/Source/Developer/Mac/MacTargetPlatform/Classes/MacTargetSettings.h
@@ -70,6 +70,10 @@ public:
 	 */
 	UPROPERTY(EditAnywhere, config, Category=Rendering, meta = (DisplayName = "Enable Fast-Math optimisations", ConfigRestartRequired = true))
 	bool EnableMathOptimisations;
+	
+	/** Whether to compile shaders using a tier Indirect Argument Buffers. */
+	UPROPERTY(config, EditAnywhere, Category = Rendering, Meta = (DisplayName = "Tier of Indirect Argument Buffers to use when compiling shaders", ConfigRestartRequired = true))
+	int32 IndirectArgumentTier;
 
 	/** Sample rate to run the audio mixer with. */
 	UPROPERTY(config, EditAnywhere, Category = "Audio", Meta = (DisplayName = "Audio Mixer Sample Rate"))
diff --git a/Engine/Source/Developer/Mac/MacTargetPlatform/Private/MacTargetPlatformModule.cpp b/Engine/Source/Developer/Mac/MacTargetPlatform/Private/MacTargetPlatformModule.cpp
--- a/Engine/Source/Developer/Mac/MacTargetPlatform/Private/MacTargetPlatformModule.cpp
+++ b/Engine/Source/Developer/Mac/MacTargetPlatform/Private/MacTargetPlatformModule.cpp
@@ -83,6 +83,11 @@ public:
 			TargetSettings->ForceFloats = false;
 		}
 		
+		if (!GConfig->GetInt(TEXT("/Script/MacTargetPlatform.MacTargetSettings"), TEXT("IndirectArgumentTier"), TargetSettings->IndirectArgumentTier, GEngineIni))
+		{
+			TargetSettings->IndirectArgumentTier = 0;
+		}
+		
 		TargetSettings->AddToRoot();
 
 		ISettingsModule* SettingsModule = FModuleManager::GetModulePtr<ISettingsModule>("Settings");
diff --git a/Engine/Source/Runtime/Apple/MetalRHI/Private/MetalCommandQueue.cpp b/Engine/Source/Runtime/Apple/MetalRHI/Private/MetalCommandQueue.cpp
--- a/Engine/Source/Runtime/Apple/MetalRHI/Private/MetalCommandQueue.cpp
+++ b/Engine/Source/Runtime/Apple/MetalRHI/Private/MetalCommandQueue.cpp
@@ -33,6 +33,7 @@ FMetalCommandQueue::FMetalCommandQueue(mtlpp::Device InDevice, uint32 const MaxN
 , RuntimeDebuggingLevel(EMetalDebugLevelOff)
 {
 	int32 MaxShaderVersion = 0;
+	int32 IndirectArgumentTier = 0;
 #if PLATFORM_MAC
 	int32 DefaultMaxShaderVersion = 3;
 	int32 MinShaderVersion = 3;
@@ -46,6 +47,10 @@ FMetalCommandQueue::FMetalCommandQueue(mtlpp::Device InDevice, uint32 const MaxN
     {
         MaxShaderVersion = DefaultMaxShaderVersion;
     }
+	if(!GConfig->GetInt(Settings, TEXT("IndirectArgumentTier"), IndirectArgumentTier, GEngineIni))
+	{
+		IndirectArgumentTier = 0;
+	}
 	MaxShaderVersion = FMath::Max(MinShaderVersion, MaxShaderVersion);
 	ValidateVersion(MaxShaderVersion);
 
@@ -239,7 +244,7 @@ FMetalCommandQueue::FMetalCommandQueue(mtlpp::Device InDevice, uint32 const MaxN
 			{
 				Features |= EMetalFeaturesTextureBuffers;
             }
-            if (MaxShaderVersion >= 5)
+            if (IndirectArgumentTier >= 1)
             {
                 Features |= EMetalFeaturesIABs;
             }
diff --git a/Engine/Source/Runtime/Apple/MetalRHI/Private/MetalCommands.cpp b/Engine/Source/Runtime/Apple/MetalRHI/Private/MetalCommands.cpp
--- a/Engine/Source/Runtime/Apple/MetalRHI/Private/MetalCommands.cpp
+++ b/Engine/Source/Runtime/Apple/MetalRHI/Private/MetalCommands.cpp
@@ -549,7 +549,7 @@ void FMetalRHICommandContext::RHISetShaderUniformBuffer(FRHIVertexShader* Vertex
 		auto* UB = (FMetalUniformBuffer*)BufferRHI;
 		UB->UploadIAB(Context);
 		auto& IAB = UB->GetIAB();
-		check(IAB.IndirectArgumentBuffer->IndirectArgumentEncoder.GetEncodedLength() == VertexShader->ArgumentEncoders[BufferIndex].GetEncodedLength());
+		// check(IAB.IndirectArgumentBuffer->IndirectArgumentEncoder.GetEncodedLength() == VertexShader->ArgumentEncoders[BufferIndex].GetEncodedLength());
 		Context->GetCurrentState().SetShaderBuffer(EMetalShaderStages::Vertex, IAB.IndirectArgumentBuffer->IndirectArgumentBuffer, nil, 0, IAB.IndirectArgumentBuffer->IndirectArgumentEncoder.GetEncodedLength(), BufferIndex, mtlpp::ResourceUsage::Read);
 	}
 	else if ((Bindings.ConstantBuffers) & 1 << BufferIndex)
@@ -574,7 +574,7 @@ void FMetalRHICommandContext::RHISetShaderUniformBuffer(FRHIHullShader* HullShad
 		auto* UB = (FMetalUniformBuffer*)BufferRHI;
 		UB->UploadIAB(Context);
 		auto& IAB = UB->GetIAB();
-		check(IAB.IndirectArgumentBuffer->IndirectArgumentEncoder.GetEncodedLength() == HullShader->ArgumentEncoders[BufferIndex].GetEncodedLength());
+		// check(IAB.IndirectArgumentBuffer->IndirectArgumentEncoder.GetEncodedLength() == HullShader->ArgumentEncoders[BufferIndex].GetEncodedLength());
 		Context->GetCurrentState().SetShaderBuffer(EMetalShaderStages::Hull, IAB.IndirectArgumentBuffer->IndirectArgumentBuffer, nil, 0, IAB.IndirectArgumentBuffer->IndirectArgumentEncoder.GetEncodedLength(), BufferIndex, mtlpp::ResourceUsage::Read);
 	}
 	else if ((Bindings.ConstantBuffers) & 1 << BufferIndex)
@@ -600,7 +600,7 @@ void FMetalRHICommandContext::RHISetShaderUniformBuffer(FRHIDomainShader* Domain
 		auto* UB = (FMetalUniformBuffer*)BufferRHI;
 		UB->UploadIAB(Context);
 		auto& IAB = UB->GetIAB();
-		check(IAB.IndirectArgumentBuffer->IndirectArgumentEncoder.GetEncodedLength() == DomainShader->ArgumentEncoders[BufferIndex].GetEncodedLength());
+		// check(IAB.IndirectArgumentBuffer->IndirectArgumentEncoder.GetEncodedLength() == DomainShader->ArgumentEncoders[BufferIndex].GetEncodedLength());
 		Context->GetCurrentState().SetShaderBuffer(EMetalShaderStages::Domain, IAB.IndirectArgumentBuffer->IndirectArgumentBuffer, nil, 0, IAB.IndirectArgumentBuffer->IndirectArgumentEncoder.GetEncodedLength(), BufferIndex, mtlpp::ResourceUsage::Read);
 	}
 	else if ((Bindings.ConstantBuffers) & 1 << BufferIndex)
@@ -630,7 +630,7 @@ void FMetalRHICommandContext::RHISetShaderUniformBuffer(FRHIPixelShader* PixelSh
 		auto* UB = (FMetalUniformBuffer*)BufferRHI;
 		UB->UploadIAB(Context);
 		auto& IAB = UB->GetIAB();
-		check(IAB.IndirectArgumentBuffer->IndirectArgumentEncoder.GetEncodedLength() == PixelShader->ArgumentEncoders[BufferIndex].GetEncodedLength());
+		// check(IAB.IndirectArgumentBuffer->IndirectArgumentEncoder.GetEncodedLength() == PixelShader->ArgumentEncoders[BufferIndex].GetEncodedLength());
 		Context->GetCurrentState().SetShaderBuffer(EMetalShaderStages::Pixel, IAB.IndirectArgumentBuffer->IndirectArgumentBuffer, nil, 0, IAB.IndirectArgumentBuffer->IndirectArgumentEncoder.GetEncodedLength(), BufferIndex, mtlpp::ResourceUsage::Read);
 	}
 	else if ((Bindings.ConstantBuffers) & 1 << BufferIndex)
@@ -654,7 +654,7 @@ void FMetalRHICommandContext::RHISetShaderUniformBuffer(FRHIComputeShader* Compu
 		auto* UB = (FMetalUniformBuffer*)BufferRHI;
 		UB->UploadIAB(Context);
 		auto& IAB = UB->GetIAB();
-		check(IAB.IndirectArgumentBuffer->IndirectArgumentEncoder.GetEncodedLength() == ComputeShader->ArgumentEncoders[BufferIndex].GetEncodedLength());
+		// check(IAB.IndirectArgumentBuffer->IndirectArgumentEncoder.GetEncodedLength() == ComputeShader->ArgumentEncoders[BufferIndex].GetEncodedLength());
 		Context->GetCurrentState().SetShaderBuffer(EMetalShaderStages::Compute, IAB.IndirectArgumentBuffer->IndirectArgumentBuffer, nil, 0, IAB.IndirectArgumentBuffer->IndirectArgumentEncoder.GetEncodedLength(), BufferIndex, mtlpp::ResourceUsage::Read);
 	}
 	else if ((Bindings.ConstantBuffers) & 1 << BufferIndex)
diff --git a/Engine/Source/Runtime/Engine/Private/ShaderCompiler/ShaderCompiler.cpp b/Engine/Source/Runtime/Engine/Private/ShaderCompiler/ShaderCompiler.cpp
--- a/Engine/Source/Runtime/Engine/Private/ShaderCompiler/ShaderCompiler.cpp
+++ b/Engine/Source/Runtime/Engine/Private/ShaderCompiler/ShaderCompiler.cpp
@@ -3211,23 +3211,27 @@ void GlobalBeginCompileShader(
 			
 			FString AllowFastIntrinsics;
 			FString ForceFloats;
+			FString IndirectArgumentTier;
 			bool bEnableMathOptimisations = true;
 			if (IsPCPlatform(EShaderPlatform(Target.Platform)))
 			{
 				GConfig->GetString(TEXT("/Script/MacTargetPlatform.MacTargetSettings"), TEXT("UseFastIntrinsics"), AllowFastIntrinsics, GEngineIni);
 				GConfig->GetBool(TEXT("/Script/MacTargetPlatform.MacTargetSettings"), TEXT("EnableMathOptimisations"), bEnableMathOptimisations, GEngineIni);
 				GConfig->GetString(TEXT("/Script/MacTargetPlatform.MacTargetSettings"), TEXT("ForceFloats"), ForceFloats, GEngineIni);
+				GConfig->GetString(TEXT("/Script/MacTargetPlatform.MacTargetSettings"), TEXT("IndirectArgumentTier"), IndirectArgumentTier, GEngineIni);
 			}
 			else
 			{
 				GConfig->GetString(TEXT("/Script/IOSRuntimeSettings.IOSRuntimeSettings"), TEXT("UseFastIntrinsics"), AllowFastIntrinsics, GEngineIni);
 				GConfig->GetBool(TEXT("/Script/IOSRuntimeSettings.IOSRuntimeSettings"), TEXT("EnableMathOptimisations"), bEnableMathOptimisations, GEngineIni);
 				GConfig->GetString(TEXT("/Script/IOSRuntimeSettings.IOSRuntimeSettings"), TEXT("ForceFloats"), ForceFloats, GEngineIni);
+				GConfig->GetString(TEXT("/Script/IOSRuntimeSettings.IOSRuntimeSettings"), TEXT("IndirectArgumentTier"), IndirectArgumentTier, GEngineIni);
 				// Force no development shaders on iOS
 				bAllowDevelopmentShaderCompile = false;
 			}
 			Input.Environment.SetDefine(TEXT("METAL_USE_FAST_INTRINSICS"), *AllowFastIntrinsics);
 			Input.Environment.SetDefine(TEXT("FORCE_FLOATS"), *ForceFloats);
+			Input.Environment.SetDefine(TEXT("METAL_INDIRECT_ARGUMENT_BUFFERS"), *IndirectArgumentTier);
 			
 			// Same as console-variable above, but that's global and this is per-platform, per-project
 			if (!bEnableMathOptimisations)
diff --git a/Engine/Source/Runtime/IOS/IOSRuntimeSettings/Classes/IOSRuntimeSettings.h b/Engine/Source/Runtime/IOS/IOSRuntimeSettings/Classes/IOSRuntimeSettings.h
--- a/Engine/Source/Runtime/IOS/IOSRuntimeSettings/Classes/IOSRuntimeSettings.h
+++ b/Engine/Source/Runtime/IOS/IOSRuntimeSettings/Classes/IOSRuntimeSettings.h
@@ -471,6 +471,10 @@ public:
 	UPROPERTY(EditAnywhere, config, Category=Rendering, meta = (DisplayName = "Enable Fast-Math optimisations", ConfigRestartRequired = true))
 	bool EnableMathOptimisations;
 	
+	/** Whether to compile shaders using a tier Indirect Argument Buffers. */
+	UPROPERTY(config, EditAnywhere, Category = Rendering, Meta = (DisplayName = "Tier of Indirect Argument Buffers to use when compiling shaders", ConfigRestartRequired = true))
+	int32 IndirectArgumentTier;
+	
 	// Whether or not the keyboard should be usable on it's own without a UITextField
 	UPROPERTY(GlobalConfig, EditAnywhere, Category = Input)
 	bool bUseIntegratedKeyboard;
diff --git a/Engine/Source/Runtime/RenderCore/Private/Shader.cpp b/Engine/Source/Runtime/RenderCore/Private/Shader.cpp
--- a/Engine/Source/Runtime/RenderCore/Private/Shader.cpp
+++ b/Engine/Source/Runtime/RenderCore/Private/Shader.cpp
@@ -2369,17 +2369,20 @@ void ShaderMapAppendKeyString(EShaderPlatform Platform, FString& KeyString)
 		bool bAllowFastIntrinsics = false;
 		bool bEnableMathOptimisations = true;
 		bool bForceFloats = false;
+		int32 IndirectArgumentTier = 0;
 		if (IsPCPlatform(Platform))
 		{
 			GConfig->GetBool(TEXT("/Script/MacTargetPlatform.MacTargetSettings"), TEXT("UseFastIntrinsics"), bAllowFastIntrinsics, GEngineIni);
 			GConfig->GetBool(TEXT("/Script/MacTargetPlatform.MacTargetSettings"), TEXT("EnableMathOptimisations"), bEnableMathOptimisations, GEngineIni);
 			GConfig->GetBool(TEXT("/Script/MacTargetPlatform.MacTargetSettings"), TEXT("ForceFloats"), bForceFloats, GEngineIni);
+			GConfig->GetInt(TEXT("/Script/MacTargetPlatform.MacTargetSettings"), TEXT("IndirectArgumentTier"), IndirectArgumentTier, GEngineIni);
 		}
 		else
 		{
 			GConfig->GetBool(TEXT("/Script/IOSRuntimeSettings.IOSRuntimeSettings"), TEXT("UseFastIntrinsics"), bAllowFastIntrinsics, GEngineIni);
 			GConfig->GetBool(TEXT("/Script/IOSRuntimeSettings.IOSRuntimeSettings"), TEXT("EnableMathOptimisations"), bEnableMathOptimisations, GEngineIni);
 			GConfig->GetBool(TEXT("/Script/IOSRuntimeSettings.IOSRuntimeSettings"), TEXT("ForceFloats"), bForceFloats, GEngineIni);
+			GConfig->GetInt(TEXT("/Script/IOSRuntimeSettings.IOSRuntimeSettings"), TEXT("IndirectArgumentTier"), IndirectArgumentTier, GEngineIni);
 		}
 		
 		if (bAllowFastIntrinsics)
@@ -2398,6 +2401,8 @@ void ShaderMapAppendKeyString(EShaderPlatform Platform, FString& KeyString)
 			KeyString += TEXT("_FP32");
 		}
 		
+		KeyString += FString::Printf(TEXT("_IAB%d"), IndirectArgumentTier);
+		
 		// Shaders built for archiving - for Metal that requires compiling the code in a different way so that we can strip it later
 		bool bArchive = false;
 		GConfig->GetBool(TEXT("/Script/UnrealEd.ProjectPackagingSettings"), TEXT("bSharedMaterialNativeLibraries"), bArchive, GGameIni);
