Inefficient, but functional tier 1 IAB support.

diff --git a/Engine/Source/Runtime/Apple/MetalRHI/Private/MetalCommandQueue.cpp b/Engine/Source/Runtime/Apple/MetalRHI/Private/MetalCommandQueue.cpp
--- a/Engine/Source/Runtime/Apple/MetalRHI/Private/MetalCommandQueue.cpp
+++ b/Engine/Source/Runtime/Apple/MetalRHI/Private/MetalCommandQueue.cpp
@@ -247,6 +247,11 @@ FMetalCommandQueue::FMetalCommandQueue(mtlpp::Device InDevice, uint32 const MaxN
             if (IndirectArgumentTier >= 1)
             {
                 Features |= EMetalFeaturesIABs;
+				
+				if (IndirectArgumentTier >= 2)
+				{
+					Features |= EMetalFeaturesTier2IABs;
+				}
             }
             
             IConsoleVariable* GPUCrashDebuggingCVar = IConsoleManager::Get().FindConsoleVariable(TEXT("r.GPUCrashDebugging"));
diff --git a/Engine/Source/Runtime/Apple/MetalRHI/Private/MetalCommandQueue.h b/Engine/Source/Runtime/Apple/MetalRHI/Private/MetalCommandQueue.h
--- a/Engine/Source/Runtime/Apple/MetalRHI/Private/MetalCommandQueue.h
+++ b/Engine/Source/Runtime/Apple/MetalRHI/Private/MetalCommandQueue.h
@@ -72,7 +72,9 @@ typedef NS_OPTIONS(uint64, EMetalFeatures)
     /** Supports tile shaders */
     EMetalFeaturesTileShaders = 1llu << 26llu,
 	/** Supports separate tessellation shader execution */
-	EMetalFeaturesSeparateTessellation = 1llu << 27llu
+	EMetalFeaturesSeparateTessellation = 1llu << 27llu,
+	/** Supports indirect argument buffers Tier 2 */
+	EMetalFeaturesTier2IABs = 1llu << 28llu,
 };
 
 /**
diff --git a/Engine/Source/Runtime/Apple/MetalRHI/Private/MetalCommands.cpp b/Engine/Source/Runtime/Apple/MetalRHI/Private/MetalCommands.cpp
--- a/Engine/Source/Runtime/Apple/MetalRHI/Private/MetalCommands.cpp
+++ b/Engine/Source/Runtime/Apple/MetalRHI/Private/MetalCommands.cpp
@@ -547,10 +547,9 @@ void FMetalRHICommandContext::RHISetShaderUniformBuffer(FRHIVertexShader* Vertex
 	if ((Bindings.ArgumentBuffers) & 1 << BufferIndex)
 	{
 		auto* UB = (FMetalUniformBuffer*)BufferRHI;
-		UB->UploadIAB(Context);
-		auto& IAB = UB->GetIAB();
+		auto* IAB = UB->UploadIAB(Context, VertexShader->ArgumentBitmasks[BufferIndex], VertexShader->ArgumentEncoders[BufferIndex]);
 		// check(IAB.IndirectArgumentBuffer->IndirectArgumentEncoder.GetEncodedLength() == VertexShader->ArgumentEncoders[BufferIndex].GetEncodedLength());
-		Context->GetCurrentState().SetShaderBuffer(EMetalShaderStages::Vertex, IAB.IndirectArgumentBuffer->IndirectArgumentBuffer, nil, 0, IAB.IndirectArgumentBuffer->IndirectArgumentEncoder.GetEncodedLength(), BufferIndex, mtlpp::ResourceUsage::Read);
+		Context->GetCurrentState().SetShaderBuffer(EMetalShaderStages::Vertex, IAB->IndirectArgumentBuffer, nil, 0, IAB->IndirectArgumentEncoder.GetEncodedLength(), BufferIndex, mtlpp::ResourceUsage::Read);
 	}
 	else if ((Bindings.ConstantBuffers) & 1 << BufferIndex)
 	{
@@ -572,10 +571,9 @@ void FMetalRHICommandContext::RHISetShaderUniformBuffer(FRHIHullShader* HullShad
 	if ((Bindings.ArgumentBuffers) & 1 << BufferIndex)
 	{
 		auto* UB = (FMetalUniformBuffer*)BufferRHI;
-		UB->UploadIAB(Context);
-		auto& IAB = UB->GetIAB();
+		auto* IAB = UB->UploadIAB(Context, HullShader->ArgumentBitmasks[BufferIndex], HullShader->ArgumentEncoders[BufferIndex]);
 		// check(IAB.IndirectArgumentBuffer->IndirectArgumentEncoder.GetEncodedLength() == HullShader->ArgumentEncoders[BufferIndex].GetEncodedLength());
-		Context->GetCurrentState().SetShaderBuffer(EMetalShaderStages::Hull, IAB.IndirectArgumentBuffer->IndirectArgumentBuffer, nil, 0, IAB.IndirectArgumentBuffer->IndirectArgumentEncoder.GetEncodedLength(), BufferIndex, mtlpp::ResourceUsage::Read);
+		Context->GetCurrentState().SetShaderBuffer(EMetalShaderStages::Hull, IAB->IndirectArgumentBuffer, nil, 0, IAB->IndirectArgumentEncoder.GetEncodedLength(), BufferIndex, mtlpp::ResourceUsage::Read);
 	}
 	else if ((Bindings.ConstantBuffers) & 1 << BufferIndex)
 	{
@@ -598,10 +596,9 @@ void FMetalRHICommandContext::RHISetShaderUniformBuffer(FRHIDomainShader* Domain
 	if ((Bindings.ArgumentBuffers) & 1 << BufferIndex)
 	{
 		auto* UB = (FMetalUniformBuffer*)BufferRHI;
-		UB->UploadIAB(Context);
-		auto& IAB = UB->GetIAB();
+		auto* IAB = UB->UploadIAB(Context, DomainShader->ArgumentBitmasks[BufferIndex], DomainShader->ArgumentEncoders[BufferIndex]);
 		// check(IAB.IndirectArgumentBuffer->IndirectArgumentEncoder.GetEncodedLength() == DomainShader->ArgumentEncoders[BufferIndex].GetEncodedLength());
-		Context->GetCurrentState().SetShaderBuffer(EMetalShaderStages::Domain, IAB.IndirectArgumentBuffer->IndirectArgumentBuffer, nil, 0, IAB.IndirectArgumentBuffer->IndirectArgumentEncoder.GetEncodedLength(), BufferIndex, mtlpp::ResourceUsage::Read);
+		Context->GetCurrentState().SetShaderBuffer(EMetalShaderStages::Domain, IAB->IndirectArgumentBuffer, nil, 0, IAB->IndirectArgumentEncoder.GetEncodedLength(), BufferIndex, mtlpp::ResourceUsage::Read);
 	}
 	else if ((Bindings.ConstantBuffers) & 1 << BufferIndex)
 	{
@@ -628,10 +625,9 @@ void FMetalRHICommandContext::RHISetShaderUniformBuffer(FRHIPixelShader* PixelSh
 	if ((Bindings.ArgumentBuffers) & 1 << BufferIndex)
 	{
 		auto* UB = (FMetalUniformBuffer*)BufferRHI;
-		UB->UploadIAB(Context);
-		auto& IAB = UB->GetIAB();
+		auto* IAB = UB->UploadIAB(Context, PixelShader->ArgumentBitmasks[BufferIndex], PixelShader->ArgumentEncoders[BufferIndex]);
 		// check(IAB.IndirectArgumentBuffer->IndirectArgumentEncoder.GetEncodedLength() == PixelShader->ArgumentEncoders[BufferIndex].GetEncodedLength());
-		Context->GetCurrentState().SetShaderBuffer(EMetalShaderStages::Pixel, IAB.IndirectArgumentBuffer->IndirectArgumentBuffer, nil, 0, IAB.IndirectArgumentBuffer->IndirectArgumentEncoder.GetEncodedLength(), BufferIndex, mtlpp::ResourceUsage::Read);
+		Context->GetCurrentState().SetShaderBuffer(EMetalShaderStages::Pixel, IAB->IndirectArgumentBuffer, nil, 0, IAB->IndirectArgumentEncoder.GetEncodedLength(), BufferIndex, mtlpp::ResourceUsage::Read);
 	}
 	else if ((Bindings.ConstantBuffers) & 1 << BufferIndex)
 	{
@@ -652,10 +648,9 @@ void FMetalRHICommandContext::RHISetShaderUniformBuffer(FRHIComputeShader* Compu
 	if ((Bindings.ArgumentBuffers) & 1 << BufferIndex)
 	{
 		auto* UB = (FMetalUniformBuffer*)BufferRHI;
-		UB->UploadIAB(Context);
-		auto& IAB = UB->GetIAB();
+		auto* IAB = UB->UploadIAB(Context, ComputeShader->ArgumentBitmasks[BufferIndex], ComputeShader->ArgumentEncoders[BufferIndex]);
 		// check(IAB.IndirectArgumentBuffer->IndirectArgumentEncoder.GetEncodedLength() == ComputeShader->ArgumentEncoders[BufferIndex].GetEncodedLength());
-		Context->GetCurrentState().SetShaderBuffer(EMetalShaderStages::Compute, IAB.IndirectArgumentBuffer->IndirectArgumentBuffer, nil, 0, IAB.IndirectArgumentBuffer->IndirectArgumentEncoder.GetEncodedLength(), BufferIndex, mtlpp::ResourceUsage::Read);
+		Context->GetCurrentState().SetShaderBuffer(EMetalShaderStages::Compute, IAB->IndirectArgumentBuffer, nil, 0, IAB->IndirectArgumentEncoder.GetEncodedLength(), BufferIndex, mtlpp::ResourceUsage::Read);
 	}
 	else if ((Bindings.ConstantBuffers) & 1 << BufferIndex)
 	{
diff --git a/Engine/Source/Runtime/Apple/MetalRHI/Private/MetalShaders.cpp b/Engine/Source/Runtime/Apple/MetalRHI/Private/MetalShaders.cpp
--- a/Engine/Source/Runtime/Apple/MetalRHI/Private/MetalShaders.cpp
+++ b/Engine/Source/Runtime/Apple/MetalRHI/Private/MetalShaders.cpp
@@ -864,6 +864,17 @@ mtlpp::Function TMetalBaseShader<BaseResourceType, ShaderType>::GetCompiledFunct
 			
 			mtlpp::ArgumentEncoder ArgumentEncoder = Function.NewArgumentEncoderWithBufferIndex(Index);
 			ArgumentEncoders.Add(Index, ArgumentEncoder);
+			
+			TBitArray<> Resources;
+			for (uint8 Id : Bindings.ArgumentBufferMasks[Index])
+			{
+				if (Id >= Resources.Num())
+				{
+					Resources.Add(false, (Id + 1) - Resources.Num());
+				}
+				Resources[Id] = true;
+			}
+			ArgumentBitmasks.Add(Index, Resources);
 		}
 	}
 									
diff --git a/Engine/Source/Runtime/Apple/MetalRHI/Private/MetalUniformBuffer.cpp b/Engine/Source/Runtime/Apple/MetalRHI/Private/MetalUniformBuffer.cpp
--- a/Engine/Source/Runtime/Apple/MetalRHI/Private/MetalUniformBuffer.cpp
+++ b/Engine/Source/Runtime/Apple/MetalRHI/Private/MetalUniformBuffer.cpp
@@ -18,6 +18,8 @@ APPLE_PLATFORM_OBJECT_ALLOC_OVERRIDES(FMetalIAB)
 -(instancetype)init
 {
 	id Self = [super init];
+	if (Self)
+		((FMetalIAB*)Self)->UpdateIAB = 0;
 	return Self;
 }
 -(void)dealloc
@@ -210,58 +212,150 @@ FMetalUniformBuffer::FMetalIndirectArgumentBuffer::~FMetalIndirectArgumentBuffer
 {
 	if (IndirectArgumentBuffer)
 		SafeReleaseMetalObject(IndirectArgumentBuffer);
+	
+	for (auto& Pair : Tier1IABs)
+	{
+		SafeReleaseMetalObject(Pair.Value);
+	}
 }
 
-FMetalIAB* FMetalUniformBuffer::UploadIAB(FMetalContext* Ctx)
+FMetalIAB* FMetalUniformBuffer::UploadIAB(FMetalContext* Ctx, TBitArray<> const& Bitmask, mtlpp::ArgumentEncoder const& IABEncoder)
 {
 	check(IAB);
 	
 	int64 UpdateIAB = FPlatformAtomics::InterlockedCompareExchange(&IAB->UpdateIAB, 0, IAB->UpdateIAB);
-	if (UpdateIAB && Ctx)
+	if (!FMetalCommandQueue::SupportsFeature(EMetalFeaturesTier2IABs) || UpdateIAB)
 	{
-		FMetalIAB* NewIAB = IAB->IndirectArgumentBuffer;
-		mtlpp::ArgumentEncoder Encoder = FMetalArgumentEncoderCache::Get().CreateEncoder(IAB->IndirectArgumentsDecl);
-		NewIAB->IndirectArgumentEncoder = Encoder;
+		FRWScopeLock Lock(IAB->Mutex, SLT_ReadOnly);
 		
-		if (!NewIAB->IndirectArgumentBuffer)
-			NewIAB->IndirectArgumentBuffer = GetMetalDeviceContext().GetResourceHeap().CreateBuffer(Encoder.GetEncodedLength(), 16, BUF_Dynamic, mtlpp::ResourceOptions(BUFFER_CACHE_MODE | ((NSUInteger)mtlpp::StorageMode::Private << mtlpp::ResourceStorageModeShift)), true);
-		
-		FMetalBuffer Temp = Ctx->GetCurrentRenderPass().AllocateTemporyBufferForCopy(NewIAB->IndirectArgumentBuffer, Encoder.GetEncodedLength(), 16);
+		FMetalIAB* NewIAB = nullptr;
+		mtlpp::ArgumentEncoder Encoder;
+		FMetalBuffer EncodingBuffer;
+		if (FMetalCommandQueue::SupportsFeature(EMetalFeaturesTier2IABs))
+		{
+			Lock.ReleaseReadOnlyLockAndAcquireWriteLock_USE_WITH_CAUTION();
+			
+			NewIAB = IAB->IndirectArgumentBuffer;
+			Encoder = FMetalArgumentEncoderCache::Get().CreateEncoder(IAB->IndirectArgumentsDecl);
+			
+			if (!NewIAB->IndirectArgumentBuffer)
+				NewIAB->IndirectArgumentBuffer = GetMetalDeviceContext().GetResourceHeap().CreateBuffer(Encoder.GetEncodedLength(), 16, BUF_Dynamic, mtlpp::ResourceOptions(BUFFER_CACHE_MODE | ((NSUInteger)mtlpp::StorageMode::Private << mtlpp::ResourceStorageModeShift)), true);
+			
+			EncodingBuffer = Ctx->GetCurrentRenderPass().AllocateTemporyBufferForCopy(NewIAB->IndirectArgumentBuffer, Encoder.GetEncodedLength(), 16);
+			
+			if (!NewIAB->IndirectArgumentBufferSideTable)
+				NewIAB->IndirectArgumentBufferSideTable = GetMetalDeviceContext().GetResourceHeap().CreateBuffer(IAB->IndirectBufferSizes.Num() * sizeof(uint32), 16, BUF_Dynamic, mtlpp::ResourceOptions(BUFFER_CACHE_MODE | ((NSUInteger)mtlpp::StorageMode::Private << mtlpp::ResourceStorageModeShift)), true);
+			
+			FMetalBuffer Temp = Ctx->GetCurrentRenderPass().AllocateTemporyBufferForCopy(NewIAB->IndirectArgumentBufferSideTable, NewIAB->IndirectArgumentBufferSideTable.GetLength(), 16);
+			
+			FMemory::Memcpy(Temp.GetContents(), IAB->IndirectBufferSizes.GetData(), IAB->IndirectBufferSizes.Num() * sizeof(uint32));
+			
+#if PLATFORM_MAC
+			if(Temp.GetStorageMode() == mtlpp::StorageMode::Managed)
+			{
+				MTLPP_VALIDATE(mtlpp::Buffer, Temp, SafeGetRuntimeDebuggingLevel() >= EMetalDebugLevelValidation, DidModify(ns::Range(0, IAB->IndirectBufferSizes.Num() * sizeof(uint32))));
+			}
+#endif
+			Ctx->AsyncCopyFromBufferToBuffer(Temp, 0, NewIAB->IndirectArgumentBufferSideTable, 0, NewIAB->IndirectArgumentBufferSideTable.GetLength());
+		}
+		else
+		{
+			NewIAB = IAB->Tier1IABs.FindRef(Bitmask);
+			if (NewIAB && (NewIAB->UpdateIAB == UpdateIAB))
+			{
+				return NewIAB;
+			}
+			
+			Lock.ReleaseReadOnlyLockAndAcquireWriteLock_USE_WITH_CAUTION();
+			NewIAB = [FMetalIAB new];
+			NewIAB->UpdateIAB = UpdateIAB;
+			
+			uint32 SetBits = 0;
+			for (uint32 i = 0; i < Bitmask.Num(); i++)
+			{
+				if (Bitmask[i])
+				{
+					SetBits++;
+				}
+			}
+			
+			Encoder = IABEncoder;
+			
+			NewIAB->IndirectArgumentBuffer = GetMetalDeviceContext().GetResourceHeap().CreateBuffer(Encoder.GetEncodedLength(), 16, BUF_Dynamic, mtlpp::ResourceOptions(BUFFER_CACHE_MODE | ((NSUInteger)BUFFER_STORAGE_MODE << mtlpp::ResourceStorageModeShift)), true);
+			
+			NewIAB->IndirectArgumentBufferSideTable = GetMetalDeviceContext().GetResourceHeap().CreateBuffer(SetBits * sizeof(uint32) * 2, 16, BUF_Dynamic, mtlpp::ResourceOptions(BUFFER_CACHE_MODE | ((NSUInteger)BUFFER_STORAGE_MODE << mtlpp::ResourceStorageModeShift)), true);
+			
+			uint32 j = 0;
+			uint32* Ptr = (uint32*)NewIAB->IndirectArgumentBufferSideTable.GetContents();
+			for (uint32 i = 0; i < Bitmask.Num(); i++)
+			{
+				if (Bitmask[i])
+				{
+					Ptr[j * 2] = IAB->IndirectBufferSizes[i * 2];
+					Ptr[(j * 2) + 1] = IAB->IndirectBufferSizes[(i * 2) + 1];
+					j++;
+				}
+			}
+			
+			EncodingBuffer = NewIAB->IndirectArgumentBuffer;
+			
+			if (IAB->Tier1IABs.Contains(Bitmask))
+			{
+				SafeReleaseMetalObject(IAB->Tier1IABs[Bitmask]);
+			}
+			
+			IAB->Tier1IABs.Add(Bitmask, NewIAB);
+		}
+		NewIAB->IndirectArgumentEncoder = Encoder;
 		
-		Encoder.SetArgumentBuffer(Temp, 0);
+		Encoder.SetArgumentBuffer(EncodingBuffer, 0);
 		
+		NSUInteger MTLIndex = 0;
 		for (FMetalArgumentDesc const& Arg : IAB->IndirectArgumentsDecl)
 		{
 			NSUInteger NewIndex = Arg.Index;
-			switch(Arg.DataType)
+			if (FMetalCommandQueue::SupportsFeature(EMetalFeaturesTier2IABs) || (NewIndex < Bitmask.Num() && Bitmask[NewIndex]))
 			{
-				case mtlpp::DataType::Pointer:
-					if (IAB->IndirectArgumentResources[NewIndex].Buffer)
-						Encoder.SetBuffer(IAB->IndirectArgumentResources[NewIndex].Buffer, 0, NewIndex);
-					break;
-				case mtlpp::DataType::Texture:
-					if (IAB->IndirectArgumentResources[NewIndex].Texture)
-						Encoder.SetTexture(IAB->IndirectArgumentResources[NewIndex].Texture, NewIndex);
-					break;
-				case mtlpp::DataType::Sampler:
-					if (IAB->IndirectArgumentResources[NewIndex].Sampler)
-						Encoder.SetSamplerState(IAB->IndirectArgumentResources[NewIndex].Sampler, NewIndex);
-					break;
-				default:
-					break;
+				NSUInteger ResourceIndex = FMetalCommandQueue::SupportsFeature(EMetalFeaturesTier2IABs) ? NewIndex : MTLIndex++;
+				switch(Arg.DataType)
+				{
+					case mtlpp::DataType::Pointer:
+						if (IAB->IndirectArgumentResources[NewIndex].Buffer)
+							Encoder.SetBuffer(IAB->IndirectArgumentResources[NewIndex].Buffer, 0, ResourceIndex);
+						break;
+					case mtlpp::DataType::Texture:
+						if (IAB->IndirectArgumentResources[NewIndex].Texture)
+							Encoder.SetTexture(IAB->IndirectArgumentResources[NewIndex].Texture, ResourceIndex);
+						break;
+					case mtlpp::DataType::Sampler:
+						if (IAB->IndirectArgumentResources[NewIndex].Sampler)
+							Encoder.SetSamplerState(IAB->IndirectArgumentResources[NewIndex].Sampler, ResourceIndex);
+						break;
+					default:
+						break;
+				}
 			}
 		}
 		
 #if PLATFORM_MAC
-		if(Temp.GetStorageMode() == mtlpp::StorageMode::Managed)
+		if(EncodingBuffer.GetStorageMode() == mtlpp::StorageMode::Managed)
 		{
-			MTLPP_VALIDATE(mtlpp::Buffer, Temp, SafeGetRuntimeDebuggingLevel() >= EMetalDebugLevelValidation, DidModify(ns::Range(0, Encoder.GetEncodedLength())));
+			MTLPP_VALIDATE(mtlpp::Buffer, EncodingBuffer, SafeGetRuntimeDebuggingLevel() >= EMetalDebugLevelValidation, DidModify(ns::Range(0, Encoder.GetEncodedLength())));
 		}
 #endif
-		
-		Ctx->AsyncCopyFromBufferToBuffer(Temp, 0, NewIAB->IndirectArgumentBuffer, 0, Encoder.GetEncodedLength());
+		if (FMetalCommandQueue::SupportsFeature(EMetalFeaturesTier2IABs))
+		{
+			Ctx->AsyncCopyFromBufferToBuffer(EncodingBuffer, 0, NewIAB->IndirectArgumentBuffer, 0, Encoder.GetEncodedLength());
+		}
+	}
+	if (FMetalCommandQueue::SupportsFeature(EMetalFeaturesTier2IABs))
+	{
+		return IAB->IndirectArgumentBuffer;
+	}
+	else
+	{
+		return IAB->Tier1IABs[Bitmask];
 	}
-	return IAB->IndirectArgumentBuffer;
 }
 
 FMetalUniformBuffer::FMetalIndirectArgumentBuffer& FMetalUniformBuffer::GetIAB()
@@ -355,6 +449,8 @@ void FMetalUniformBuffer::UpdateIAB()
 			Desc.SetAccess(mtlpp::ArgumentAccess::ReadOnly);
 			Desc.SetDataType(mtlpp::DataType::Pointer);
 			
+			BufferSizes.AddZeroed(2);
+			
 			IAB->IndirectArgumentResources.Add(Argument(FMetalBuffer(), mtlpp::ResourceUsage::Read));
 		}
 		
@@ -619,6 +715,8 @@ void FMetalUniformBuffer::UpdateIAB()
 			Desc.SetAccess(mtlpp::ArgumentAccess::ReadOnly);
 			Desc.SetDataType(mtlpp::DataType::Pointer);
 			
+			BufferSizes.AddZeroed(2);
+			
 			if (ConstantSize > 0)
 				IAB->IndirectArgumentResources.Add(Argument(Buffer, mtlpp::ResourceUsage::Read));
 			else
diff --git a/Engine/Source/Runtime/Apple/MetalRHI/Public/MetalResources.h b/Engine/Source/Runtime/Apple/MetalRHI/Public/MetalResources.h
--- a/Engine/Source/Runtime/Apple/MetalRHI/Public/MetalResources.h
+++ b/Engine/Source/Runtime/Apple/MetalRHI/Public/MetalResources.h
@@ -137,6 +137,9 @@ public:
 	/* Argument encoders for shader IABs */
 	TMap<uint32, mtlpp::ArgumentEncoder> ArgumentEncoders;
 	
+	/* Tier1 Argument buffer bitmasks */
+	TMap<uint32, TBitArray<>> ArgumentBitmasks;
+	
 	/** The binding for the buffer side-table if present */
 	int32 SideTableBinding;
 
@@ -988,6 +991,8 @@ public:
 		TArray<Argument> IndirectArgumentResources;
 		TArray<uint32> IndirectBufferSizes;
 		FMetalIAB* IndirectArgumentBuffer;
+		TMap<TBitArray<>, FMetalIAB*> Tier1IABs;
+		FRWLock Mutex;
 	};
 	
 	// Constructor
@@ -1003,7 +1008,7 @@ public:
 	void UpdateTextureReference(FRHITextureReference* ModifiedRef);
 	void UpdateResourceTable(TArray<TRefCountPtr<FRHIResource>>& Resources, EUniformBufferValidation Validation);
 	void Update(const void* Contents, TArray<TRefCountPtr<FRHIResource>>& Resources, EUniformBufferValidation Validation);
-	FMetalIAB* UploadIAB(class FMetalContext* Ctx);
+	FMetalIAB* UploadIAB(class FMetalContext* Ctx, TBitArray<> const& Bitmask, mtlpp::ArgumentEncoder const& Encoder);
 	FMetalIndirectArgumentBuffer& GetIAB();
 	
 	/** Resource table containing RHI references. */
