Missed important Tier1 IAB change.

diff --git a/Engine/Source/Runtime/Apple/MetalRHI/Private/MetalRHIPrivate.h b/Engine/Source/Runtime/Apple/MetalRHI/Private/MetalRHIPrivate.h
--- a/Engine/Source/Runtime/Apple/MetalRHI/Private/MetalRHIPrivate.h
+++ b/Engine/Source/Runtime/Apple/MetalRHI/Private/MetalRHIPrivate.h
@@ -325,5 +325,15 @@ FORCEINLINE EMetalShaderStages GetMetalShaderFrequency(EShaderFrequency Stage)
 	}
 }
 
+@interface FMetalIAB : FApplePlatformObject<NSObject>
+{
+@public
+	FMetalBuffer IndirectArgumentBuffer;
+	FMetalBuffer IndirectArgumentBufferSideTable;
+	mtlpp::ArgumentEncoder IndirectArgumentEncoder;
+	int64 UpdateIAB;
+}
+@end
+
 #include "MetalStateCache.h"
 #include "MetalContext.h"
