[RHI][Metal] Accrue resource masks for active bindings only.

#rnx

diff --git a/Engine/Source/Runtime/Apple/MetalRHI/Private/MetalPipeline.cpp b/Engine/Source/Runtime/Apple/MetalRHI/Private/MetalPipeline.cpp
--- a/Engine/Source/Runtime/Apple/MetalRHI/Private/MetalPipeline.cpp
+++ b/Engine/Source/Runtime/Apple/MetalRHI/Private/MetalPipeline.cpp
@@ -652,6 +652,12 @@ APPLE_PLATFORM_OBJECT_ALLOC_OVERRIDES(FMetalShaderPipeline)
 	{
 		MTLArgument* Arg = [Arguments objectAtIndex:i];
 		check(Arg);
+
+		if (!Arg.active)
+		{
+			continue;
+		}
+
 		switch(Arg.type)
 		{
 			case MTLArgumentTypeBuffer:
