Fix for ios issue "audio not resuming after ios alarm goes off"

diff --git a/Engine/Source/Runtime/Apple/AudioMixerAudioUnit/Private/AudioMixerPlatformAudioUnit.cpp b/Engine/Source/Runtime/Apple/AudioMixerAudioUnit/Private/AudioMixerPlatformAudioUnit.cpp
--- a/Engine/Source/Runtime/Apple/AudioMixerAudioUnit/Private/AudioMixerPlatformAudioUnit.cpp
+++ b/Engine/Source/Runtime/Apple/AudioMixerAudioUnit/Private/AudioMixerPlatformAudioUnit.cpp
@@ -1,6 +1,7 @@
 // Copyright 1998-2019 Epic Games, Inc. All Rights Reserved.
 
 #include "AudioMixerPlatformAudioUnit.h"
+#include "AudioMixerPlatformAudioUnitUtils.h"
 #include "Modules/ModuleManager.h"
 #include "AudioMixer.h"
 #include "AudioMixerDevice.h"
@@ -24,6 +25,8 @@
 #include <AudioUnit/AudioUnit.h>
 #include <AVFoundation/AVAudioSession.h>
 
+static int32 SuspendCounter = 0;
+
 DECLARE_LOG_CATEGORY_EXTERN(LogAudioMixerAudioUnit, Log, All);
 DEFINE_LOG_CATEGORY(LogAudioMixerAudioUnit);
 
@@ -32,11 +35,20 @@ namespace Audio
 	static const int32 DefaultBufferSize = 512;
 
 	static const double DefaultSampleRate = 48000.0;
-	
-	static int32 SuspendCounter = 0;
+    
+    void AUDIOMIXERAUDIOUNIT_API IncrementIOSAudioMixerPlatformSuspendCounter()
+    {
+        FMixerPlatformAudioUnit::IncrementSuspendCounter();
+    }
+    
+    void AUDIOMIXERAUDIOUNIT_API DecrementIOSAudioMixerPlatformSuspendCounter()
+    {
+        FMixerPlatformAudioUnit::DecrementSuspendCounter();
+    }
 	
 	FMixerPlatformAudioUnit::FMixerPlatformAudioUnit()
-	: bInitialized(false)
+	: bSuspended(false)
+    , bInitialized(false)
 	, bInCallback(false)
 	, SubmittedBufferPtr(nullptr)
 	, RemainingBytesInCurrentSubmittedBuffer(0)
@@ -435,13 +447,39 @@ namespace Audio
 
 		return InternalPlatformSettings;
 	}
+    
+    void FMixerPlatformAudioUnit::IncrementSuspendCounter()
+    {
+        if(SuspendCounter == 0)
+        {
+            FPlatformAtomics::InterlockedIncrement(&SuspendCounter);
+        }
+    }
+    
+    void FMixerPlatformAudioUnit::DecrementSuspendCounter()
+    {
+        if(SuspendCounter > 0)
+        {
+            FPlatformAtomics::InterlockedDecrement(&SuspendCounter);
+        }
+    }
 	
 	void FMixerPlatformAudioUnit::ResumeContext()
 	{
 		if (SuspendCounter > 0)
 		{
 			FPlatformAtomics::InterlockedDecrement(&SuspendCounter);
-			AUGraphStart(AudioUnitGraph);
+            
+            if (AudioUnitGraph != NULL)
+            {
+                AUGraphStart(AudioUnitGraph);
+            }
+            
+            if (OutputUnit != NULL)
+            {
+                AudioOutputUnitStart(OutputUnit);
+            }
+			
 			UE_LOG(LogAudioMixerAudioUnit, Display, TEXT("Resuming Audio"));
 			bSuspended = false;
 		}
@@ -452,7 +490,17 @@ namespace Audio
 		if (SuspendCounter == 0)
 		{
 			FPlatformAtomics::InterlockedIncrement(&SuspendCounter);
-			AUGraphStop(AudioUnitGraph);
+            
+            if(AudioUnitGraph != NULL)
+            {
+                AUGraphStop(AudioUnitGraph);
+            }
+            
+            if (OutputUnit != NULL)
+            {
+                AudioOutputUnitStop(OutputUnit);
+            }
+			
 			UE_LOG(LogAudioMixerAudioUnit, Display, TEXT("Suspending Audio"));
 			bSuspended = true;
 		}
diff --git a/Engine/Source/Runtime/Apple/AudioMixerAudioUnit/Private/AudioMixerPlatformAudioUnit.h b/Engine/Source/Runtime/Apple/AudioMixerAudioUnit/Private/AudioMixerPlatformAudioUnit.h
--- a/Engine/Source/Runtime/Apple/AudioMixerAudioUnit/Private/AudioMixerPlatformAudioUnit.h
+++ b/Engine/Source/Runtime/Apple/AudioMixerAudioUnit/Private/AudioMixerPlatformAudioUnit.h
@@ -51,6 +51,9 @@ namespace Audio
 		virtual bool DisablePCMAudioCaching() const override { return true; }
 		
 		//~ End IAudioMixerPlatformInterface
+        
+        static void IncrementSuspendCounter();
+        static void DecrementSuspendCounter();
 		
 	private:
 		AudioStreamBasicDescription OutputFormat;
diff --git a/Engine/Source/Runtime/Apple/AudioMixerAudioUnit/Public/AudioMixerPlatformAudioUnitUtils.h b/Engine/Source/Runtime/Apple/AudioMixerAudioUnit/Public/AudioMixerPlatformAudioUnitUtils.h
new file mode 100644
--- /dev/null
+++ b/Engine/Source/Runtime/Apple/AudioMixerAudioUnit/Public/AudioMixerPlatformAudioUnitUtils.h
@@ -0,0 +1,9 @@
+// Copyright 1998-2019 Epic Games, Inc. All Rights Reserved.
+
+#pragma once
+
+namespace Audio
+{
+    void AUDIOMIXERAUDIOUNIT_API IncrementIOSAudioMixerPlatformSuspendCounter();
+    void AUDIOMIXERAUDIOUNIT_API DecrementIOSAudioMixerPlatformSuspendCounter();
+}
diff --git a/Engine/Source/Runtime/ApplicationCore/Public/IOS/IOSAppDelegate.h b/Engine/Source/Runtime/ApplicationCore/Public/IOS/IOSAppDelegate.h
--- a/Engine/Source/Runtime/ApplicationCore/Public/IOS/IOSAppDelegate.h
+++ b/Engine/Source/Runtime/ApplicationCore/Public/IOS/IOSAppDelegate.h
@@ -89,6 +89,8 @@ namespace FAppEntry
     void Suspend(bool bIsInterrupt = false);
     void Resume(bool bIsInterrupt = false);
 	void RestartAudio();
+    void IncrementAudioSuspendCounters();
+    void DecrementAudioSuspendCounters();
 
 	bool IsStartupMoviePlaying();
 
diff --git a/Engine/Source/Runtime/IOS/IOSAudio/Private/IOSAudioDevice.cpp b/Engine/Source/Runtime/IOS/IOSAudio/Private/IOSAudioDevice.cpp
--- a/Engine/Source/Runtime/IOS/IOSAudio/Private/IOSAudioDevice.cpp
+++ b/Engine/Source/Runtime/IOS/IOSAudio/Private/IOSAudioDevice.cpp
@@ -32,9 +32,10 @@ IMPLEMENT_MODULE(FIOSAudioDeviceModule, IOSAudio);
 	FIOSAudioDevice
  ------------------------------------------------------------------------------------*/
 
+static int32 SuspendCounter = 0;
+
 void FIOSAudioDevice::ResumeContext()
 {
-	int32& SuspendCounter = GetSuspendCounter();
 	if (SuspendCounter > 0)
 	{
 		FPlatformAtomics::InterlockedDecrement(&SuspendCounter);
@@ -55,7 +56,6 @@ void FIOSAudioDevice::ResumeContext()
 
 void FIOSAudioDevice::SuspendContext()
 {
-	int32& SuspendCounter = GetSuspendCounter();
 	if (SuspendCounter == 0)
 	{
 		FPlatformAtomics::InterlockedIncrement(&SuspendCounter);
@@ -74,10 +74,20 @@ void FIOSAudioDevice::SuspendContext()
 	}
 }
 
-int32& FIOSAudioDevice::GetSuspendCounter()
+void FIOSAudioDevice::IncrementSuspendCounter()
+{
+    if(SuspendCounter == 0)
+    {
+        FPlatformAtomics::InterlockedIncrement(&SuspendCounter);
+    }
+}
+
+void FIOSAudioDevice::DecrementSuspendCounter()
 {
-	static int32 SuspendCounter = 0;
-	return SuspendCounter;
+    if (SuspendCounter > 0)
+    {
+        FPlatformAtomics::InterlockedDecrement(&SuspendCounter);
+    }
 }
 
 void FIOSAudioDevice::UpdateDeviceDeltaTime()
@@ -239,7 +249,7 @@ bool FIOSAudioDevice::InitializeHardware()
 
 	// Initialize and start the audio unit graph
 	Status = AUGraphInitialize(AudioUnitGraph);
-	if (Status == noErr && GetSuspendCounter() == 0)
+	if (Status == noErr && SuspendCounter == 0)
 	{
 		Status = AUGraphStart(AudioUnitGraph);
 	}
diff --git a/Engine/Source/Runtime/IOS/IOSAudio/Public/IOSAudioDevice.h b/Engine/Source/Runtime/IOS/IOSAudio/Public/IOSAudioDevice.h
--- a/Engine/Source/Runtime/IOS/IOSAudio/Public/IOSAudioDevice.h
+++ b/Engine/Source/Runtime/IOS/IOSAudio/Public/IOSAudioDevice.h
@@ -212,10 +212,10 @@ public:
 	
 	virtual bool Exec( UWorld* InWorld, const TCHAR* Cmd, FOutputDevice& Ar = *GLog ) override;
 
-	static int32& GetSuspendCounter();
-
-
 	virtual void UpdateDeviceDeltaTime() override;
+    
+    static void IncrementSuspendCounter();
+    static void DecrementSuspendCounter();
 
 protected:
 	/** Starts up any platform specific hardware/APIs */
diff --git a/Engine/Source/Runtime/Launch/Private/IOS/LaunchIOS.cpp b/Engine/Source/Runtime/Launch/Private/IOS/LaunchIOS.cpp
--- a/Engine/Source/Runtime/Launch/Private/IOS/LaunchIOS.cpp
+++ b/Engine/Source/Runtime/Launch/Private/IOS/LaunchIOS.cpp
@@ -14,6 +14,7 @@
 #include "AudioDevice.h"
 #include "GenericPlatform/GenericPlatformChunkInstall.h"
 #include "IOSAudioDevice.h"
+#include "AudioMixerPlatformAudioUnitUtils.h"
 #include "LocalNotification.h"
 #include "Modules/ModuleManager.h"
 #include "RenderingThread.h"
@@ -130,11 +131,8 @@ void FAppEntry::Suspend(bool bIsInterrupt)
 		}
 		else
 		{
-			int32& SuspendCounter = FIOSAudioDevice::GetSuspendCounter();
-			if (SuspendCounter == 0)
-			{
-				FPlatformAtomics::InterlockedIncrement(&SuspendCounter);
-			}
+            // Increment
+            IncrementAudioSuspendCounters();
 		}
 	}
 }
@@ -187,11 +185,8 @@ void FAppEntry::Resume(bool bIsInterrupt)
 		}
 		else
 		{
-			int32& SuspendCounter = FIOSAudioDevice::GetSuspendCounter();
-			if (SuspendCounter > 0)
-			{
-				FPlatformAtomics::InterlockedDecrement(&SuspendCounter);
-			}
+            // Decrement
+            DecrementAudioSuspendCounters();
 		}
 	}
 }
@@ -230,13 +225,8 @@ void FAppEntry::RestartAudio()
 
 		if (FTaskGraphInterface::IsRunning())
 		{
-			int32& SuspendCounter = FIOSAudioDevice::GetSuspendCounter();
-
-			//increment the counter, otherwise ResumeContext won't work
-			if (SuspendCounter == 0)
-			{
-				FPlatformAtomics::InterlockedIncrement(&SuspendCounter);
-			}
+            //increment the counter, otherwise ResumeContext won't work
+            IncrementAudioSuspendCounters();
 
 			FFunctionGraphTask::CreateAndDispatchWhenReady([AudioDevice]()
 			{
@@ -253,6 +243,36 @@ void FAppEntry::RestartAudio()
 	}
 }
 
+void FAppEntry::IncrementAudioSuspendCounters()
+{
+    // old backend
+    if(FModuleManager::Get().IsModuleLoaded("IOSAudio"))
+    {
+        FIOSAudioDevice::IncrementSuspendCounter();
+    }
+    
+    // new backend
+    if(FModuleManager::Get().IsModuleLoaded("AudioMixerAudioUnit"))
+    {
+        Audio::IncrementIOSAudioMixerPlatformSuspendCounter();
+    }
+}
+
+void FAppEntry::DecrementAudioSuspendCounters()
+{
+    // old backend
+    if(FModuleManager::Get().IsModuleLoaded("IOSAudio"))
+    {
+        FIOSAudioDevice::DecrementSuspendCounter();
+    }
+    
+    // new backend
+    if(FModuleManager::Get().IsModuleLoaded("AudioMixerAudioUnit"))
+    {
+        Audio::DecrementIOSAudioMixerPlatformSuspendCounter();
+    }
+}
+
 void FAppEntry::PreInit(IOSAppDelegate* AppDelegate, UIApplication* Application)
 {
 	// make a controller object
