Adds missing RHI context functions so the ios metal build (with the fixed rhi class) actually compiles.

diff --git a/Engine/Source/Runtime/Apple/MetalRHI/Private/MetalRHIContext.cpp b/Engine/Source/Runtime/Apple/MetalRHI/Private/MetalRHIContext.cpp
--- a/Engine/Source/Runtime/Apple/MetalRHI/Private/MetalRHIContext.cpp
+++ b/Engine/Source/Runtime/Apple/MetalRHI/Private/MetalRHIContext.cpp
@@ -230,3 +230,8 @@ void FMetalRHICommandContext::RHIBeginComputePass(const TCHAR* InName)
 void FMetalRHICommandContext::RHIEndComputePass()
 {
 }
+
+void FMetalRHICommandContext::RHITransitionResources(FExclusiveDepthStencil DepthStencilMode, FRHITexture* DepthTexture)
+{
+	IRHICommandContext::RHITransitionResources(DepthStencilMode, DepthTexture);
+}
diff --git a/Engine/Source/Runtime/Apple/MetalRHI/Public/MetalRHIContext.h b/Engine/Source/Runtime/Apple/MetalRHI/Public/MetalRHIContext.h
--- a/Engine/Source/Runtime/Apple/MetalRHI/Public/MetalRHIContext.h
+++ b/Engine/Source/Runtime/Apple/MetalRHI/Public/MetalRHIContext.h
@@ -281,6 +281,11 @@ public:
 	 */
 	virtual void RHITransitionResources(EResourceTransitionAccess TransitionType, FRHITexture** InTextures, int32 NumTextures) final override;
 	
+	/**
+	 * Transition depth resources. Does not really do anything in metal.
+	 */
+	virtual void RHITransitionResources(FExclusiveDepthStencil DepthStencilMode, FRHITexture* DepthTexture) final override;
+	
 	/**
 	 * Compute queue will wait for the fence to be written before continuing.
 	 */
