Fix warning caused by enum values being considered global scoped
#rnx

diff --git a/Engine/Source/Developer/Apple/MetalShaderFormat/Private/MetalBackend.cpp b/Engine/Source/Developer/Apple/MetalShaderFormat/Private/MetalBackend.cpp
--- a/Engine/Source/Developer/Apple/MetalShaderFormat/Private/MetalBackend.cpp
+++ b/Engine/Source/Developer/Apple/MetalShaderFormat/Private/MetalBackend.cpp
@@ -282,24 +282,24 @@ class FGenerateMetalVisitor : public ir_visitor
 				{
 					case 1:
 					{
-						return R32Uint;
+						return EMetalBufferFormat::R32Uint;
 					}
 					case 2:
 					{
-						return RG32Uint;
+						return EMetalBufferFormat::RG32Uint;
 					}
 					case 3:
 					{
-						return RGB32Uint;
+						return EMetalBufferFormat::RGB32Uint;
 					}
 					case 4:
 					{
-						return RGBA32Uint;
+						return EMetalBufferFormat::RGBA32Uint;
 					}
 					default:
 					{
 						check(0);
-						return Unknown;
+						return EMetalBufferFormat::Unknown;
 					}
 				}
 			}
@@ -309,24 +309,24 @@ class FGenerateMetalVisitor : public ir_visitor
 				{
 					case 1:
 					{
-						return R32Sint;
+						return EMetalBufferFormat::R32Sint;
 					}
 					case 2:
 					{
-						return RG32Sint;
+						return EMetalBufferFormat::RG32Sint;
 					}
 					case 3:
 					{
-						return RGB32Sint;
+						return EMetalBufferFormat::RGB32Sint;
 					}
 					case 4:
 					{
-						return RGBA32Sint;
+						return EMetalBufferFormat::RGBA32Sint;
 					}
 					default:
 					{
 						check(0);
-						return Unknown;
+						return EMetalBufferFormat::Unknown;
 					}
 				}
 			}
@@ -336,24 +336,24 @@ class FGenerateMetalVisitor : public ir_visitor
 				{
 					case 1:
 					{
-						return R16Half;
+						return EMetalBufferFormat::R16Half;
 					}
 					case 2:
 					{
-						return RG16Half;
+						return EMetalBufferFormat::RG16Half;
 					}
 					case 3:
 					{
-						return RGB16Half;
+						return EMetalBufferFormat::RGB16Half;
 					}
 					case 4:
 					{
-						return RGBA16Half;
+						return EMetalBufferFormat::RGBA16Half;
 					}
 					default:
 					{
 						check(0);
-						return Unknown;
+						return EMetalBufferFormat::Unknown;
 					}
 				}
 			}
@@ -363,31 +363,31 @@ class FGenerateMetalVisitor : public ir_visitor
 				{
 					case 1:
 					{
-						return R32Float;
+						return EMetalBufferFormat::R32Float;
 					}
 					case 2:
 					{
-						return RG32Float;
+						return EMetalBufferFormat::RG32Float;
 					}
 					case 3:
 					{
-						return RGB32Float;
+						return EMetalBufferFormat::RGB32Float;
 					}
 					case 4:
 					{
-						return RGBA32Float;
+						return EMetalBufferFormat::RGBA32Float;
 					}
 					default:
 					{
 						check(0);
-						return Unknown;
+						return EMetalBufferFormat::Unknown;
 					}
 				}
 			}
 			default:
 			{
 				check(0);
-				return Unknown;
+				return EMetalBufferFormat::Unknown;
 			}
 		}
 	}
diff --git a/Engine/Source/Runtime/Apple/MetalRHI/Private/MetalRHI.cpp b/Engine/Source/Runtime/Apple/MetalRHI/Private/MetalRHI.cpp
--- a/Engine/Source/Runtime/Apple/MetalRHI/Private/MetalRHI.cpp
+++ b/Engine/Source/Runtime/Apple/MetalRHI/Private/MetalRHI.cpp
@@ -528,75 +528,75 @@ FMetalDynamicRHI::FMetalDynamicRHI(ERHIFeatureLevel::Type RequestedFeatureLevel)
 #if METAL_DEBUG_OPTIONS
 	FMemory::Memset(GMetalBufferFormats, 255);
 #endif
-	GMetalBufferFormats[PF_Unknown              ] = { mtlpp::PixelFormat::Invalid, EMetalBufferFormat::Unknown };
-	GMetalBufferFormats[PF_A32B32G32R32F        ] = { mtlpp::PixelFormat::RGBA32Float, EMetalBufferFormat::RGBA32Float };
-	GMetalBufferFormats[PF_B8G8R8A8             ] = { mtlpp::PixelFormat::RGBA8Unorm, EMetalBufferFormat::RGBA8Unorm }; // mtlpp::PixelFormat::BGRA8Unorm/EMetalBufferFormat::BGRA8Unorm,  < We don't support this as a vertex-format so we have code to swizzle in the shader
-	GMetalBufferFormats[PF_G8                   ] = { mtlpp::PixelFormat::R8Unorm, EMetalBufferFormat::R8Unorm };
-	GMetalBufferFormats[PF_G16                  ] = { mtlpp::PixelFormat::R16Unorm, EMetalBufferFormat::R16Unorm };
-	GMetalBufferFormats[PF_DXT1                 ] = { mtlpp::PixelFormat::Invalid, EMetalBufferFormat::Unknown };
-	GMetalBufferFormats[PF_DXT3                 ] = { mtlpp::PixelFormat::Invalid, EMetalBufferFormat::Unknown };
-	GMetalBufferFormats[PF_DXT5                 ] = { mtlpp::PixelFormat::Invalid, EMetalBufferFormat::Unknown };
-	GMetalBufferFormats[PF_UYVY                 ] = { mtlpp::PixelFormat::Invalid, EMetalBufferFormat::Unknown };
-	GMetalBufferFormats[PF_FloatRGB             ] = { mtlpp::PixelFormat::Invalid, EMetalBufferFormat::RGB16Half };
-	GMetalBufferFormats[PF_FloatRGBA            ] = { mtlpp::PixelFormat::RGBA16Float, EMetalBufferFormat::RGBA16Half };
-	GMetalBufferFormats[PF_DepthStencil         ] = { mtlpp::PixelFormat::Invalid, EMetalBufferFormat::Unknown };
-	GMetalBufferFormats[PF_ShadowDepth          ] = { mtlpp::PixelFormat::Invalid, EMetalBufferFormat::Unknown };
-	GMetalBufferFormats[PF_R32_FLOAT            ] = { mtlpp::PixelFormat::R32Float, EMetalBufferFormat::R32Float };
-	GMetalBufferFormats[PF_G16R16               ] = { mtlpp::PixelFormat::RG16Unorm, EMetalBufferFormat::RG16Unorm };
-	GMetalBufferFormats[PF_G16R16F              ] = { mtlpp::PixelFormat::RG16Float, EMetalBufferFormat::RG16Half };
-	GMetalBufferFormats[PF_G16R16F_FILTER       ] = { mtlpp::PixelFormat::RG16Float, EMetalBufferFormat::RG16Half };
-	GMetalBufferFormats[PF_G32R32F              ] = { mtlpp::PixelFormat::RG32Float, EMetalBufferFormat::RG32Float };
-	GMetalBufferFormats[PF_A2B10G10R10          ] = { mtlpp::PixelFormat::RGB10A2Unorm, EMetalBufferFormat::RGB10A2Unorm };
-	GMetalBufferFormats[PF_A16B16G16R16         ] = { mtlpp::PixelFormat::RGBA16Unorm, EMetalBufferFormat::RGBA16Half };
-	GMetalBufferFormats[PF_D24                  ] = { mtlpp::PixelFormat::Invalid, EMetalBufferFormat::Unknown };
-	GMetalBufferFormats[PF_R16F                 ] = { mtlpp::PixelFormat::R16Float, EMetalBufferFormat::RG16Half };
-	GMetalBufferFormats[PF_R16F_FILTER          ] = { mtlpp::PixelFormat::R16Float, EMetalBufferFormat::RG16Half };
-	GMetalBufferFormats[PF_BC5                  ] = { mtlpp::PixelFormat::Invalid, EMetalBufferFormat::Unknown };
-	GMetalBufferFormats[PF_V8U8                 ] = { mtlpp::PixelFormat::RG8Snorm, EMetalBufferFormat::RG8Unorm };
-	GMetalBufferFormats[PF_A1                   ] = { mtlpp::PixelFormat::Invalid, EMetalBufferFormat::Unknown };
-	GMetalBufferFormats[PF_FloatR11G11B10       ] = { mtlpp::PixelFormat::RG11B10Float, EMetalBufferFormat::RG11B10Half }; // < May not work on tvOS
-	GMetalBufferFormats[PF_A8                   ] = { mtlpp::PixelFormat::A8Unorm, EMetalBufferFormat::R8Unorm };
-	GMetalBufferFormats[PF_R32_UINT             ] = { mtlpp::PixelFormat::R32Uint, EMetalBufferFormat::R32Uint };
-	GMetalBufferFormats[PF_R32_SINT             ] = { mtlpp::PixelFormat::R32Sint, EMetalBufferFormat::R32Sint };
-	GMetalBufferFormats[PF_PVRTC2               ] = { mtlpp::PixelFormat::Invalid, EMetalBufferFormat::Unknown };
-	GMetalBufferFormats[PF_PVRTC4               ] = { mtlpp::PixelFormat::Invalid, EMetalBufferFormat::Unknown };
-	GMetalBufferFormats[PF_R16_UINT             ] = { mtlpp::PixelFormat::R16Uint, EMetalBufferFormat::R16Uint };
-	GMetalBufferFormats[PF_R16_SINT             ] = { mtlpp::PixelFormat::R16Sint, EMetalBufferFormat::R16Sint };
-	GMetalBufferFormats[PF_R16G16B16A16_UINT    ] = { mtlpp::PixelFormat::RGBA16Uint, EMetalBufferFormat::RGBA16Uint };
-	GMetalBufferFormats[PF_R16G16B16A16_SINT    ] = { mtlpp::PixelFormat::RGBA16Sint, EMetalBufferFormat::RGBA16Sint };
-	GMetalBufferFormats[PF_R5G6B5_UNORM         ] = { mtlpp::PixelFormat::Invalid, EMetalBufferFormat::R5G6B5Unorm };
-	GMetalBufferFormats[PF_R8G8B8A8             ] = { mtlpp::PixelFormat::RGBA8Unorm, EMetalBufferFormat::RGBA8Unorm };
-	GMetalBufferFormats[PF_A8R8G8B8				] = { mtlpp::PixelFormat::RGBA8Unorm, EMetalBufferFormat::RGBA8Unorm }; // mtlpp::PixelFormat::BGRA8Unorm/EMetalBufferFormat::BGRA8Unorm,  < We don't support this as a vertex-format so we have code to swizzle in the shader
-	GMetalBufferFormats[PF_BC4					] = { mtlpp::PixelFormat::Invalid, EMetalBufferFormat::Unknown };
-	GMetalBufferFormats[PF_R8G8                 ] = { mtlpp::PixelFormat::RG8Unorm, EMetalBufferFormat::RG8Unorm };
-	GMetalBufferFormats[PF_ATC_RGB				] = { mtlpp::PixelFormat::Invalid, EMetalBufferFormat::Unknown };
-	GMetalBufferFormats[PF_ATC_RGBA_E			] = { mtlpp::PixelFormat::Invalid, EMetalBufferFormat::Unknown };
-	GMetalBufferFormats[PF_ATC_RGBA_I			] = { mtlpp::PixelFormat::Invalid, EMetalBufferFormat::Unknown };
-	GMetalBufferFormats[PF_X24_G8				] = { mtlpp::PixelFormat::Invalid, EMetalBufferFormat::Unknown };
-	GMetalBufferFormats[PF_ETC1					] = { mtlpp::PixelFormat::Invalid, EMetalBufferFormat::Unknown };
-	GMetalBufferFormats[PF_ETC2_RGB				] = { mtlpp::PixelFormat::Invalid, EMetalBufferFormat::Unknown };
-	GMetalBufferFormats[PF_ETC2_RGBA			] = { mtlpp::PixelFormat::Invalid, EMetalBufferFormat::Unknown };
-	GMetalBufferFormats[PF_R32G32B32A32_UINT	] = { mtlpp::PixelFormat::RGBA32Uint, EMetalBufferFormat::RGBA32Uint };
-	GMetalBufferFormats[PF_R16G16_UINT			] = { mtlpp::PixelFormat::RG16Uint, EMetalBufferFormat::RG16Uint };
-	GMetalBufferFormats[PF_R32G32_UINT			] = { mtlpp::PixelFormat::RG32Uint, EMetalBufferFormat::RG32Uint };
-	GMetalBufferFormats[PF_ASTC_4x4             ] = { mtlpp::PixelFormat::Invalid, EMetalBufferFormat::Unknown };
-	GMetalBufferFormats[PF_ASTC_6x6             ] = { mtlpp::PixelFormat::Invalid, EMetalBufferFormat::Unknown };
-	GMetalBufferFormats[PF_ASTC_8x8             ] = { mtlpp::PixelFormat::Invalid, EMetalBufferFormat::Unknown };
-	GMetalBufferFormats[PF_ASTC_10x10           ] = { mtlpp::PixelFormat::Invalid, EMetalBufferFormat::Unknown };
-	GMetalBufferFormats[PF_ASTC_12x12           ] = { mtlpp::PixelFormat::Invalid, EMetalBufferFormat::Unknown };
-	GMetalBufferFormats[PF_BC6H					] = { mtlpp::PixelFormat::Invalid, EMetalBufferFormat::Unknown };
-	GMetalBufferFormats[PF_BC7					] = { mtlpp::PixelFormat::Invalid, EMetalBufferFormat::Unknown };
-	GMetalBufferFormats[PF_R8_UINT				] = { mtlpp::PixelFormat::R8Uint, EMetalBufferFormat::R8Uint };
-	GMetalBufferFormats[PF_L8					] = { mtlpp::PixelFormat::Invalid, EMetalBufferFormat::R8Unorm };
-	GMetalBufferFormats[PF_XGXR8				] = { mtlpp::PixelFormat::Invalid, EMetalBufferFormat::Unknown };
-	GMetalBufferFormats[PF_R8G8B8A8_UINT		] = { mtlpp::PixelFormat::RGBA8Uint, EMetalBufferFormat::RGBA8Uint };
-	GMetalBufferFormats[PF_R8G8B8A8_SNORM		] = { mtlpp::PixelFormat::RGBA8Snorm, EMetalBufferFormat::RGBA8Snorm };
-	GMetalBufferFormats[PF_R16G16B16A16_UNORM	] = { mtlpp::PixelFormat::RGBA16Unorm, EMetalBufferFormat::RGBA16Unorm };
-	GMetalBufferFormats[PF_R16G16B16A16_SNORM	] = { mtlpp::PixelFormat::RGBA16Snorm, EMetalBufferFormat::RGBA16Snorm };
-	GMetalBufferFormats[PF_PLATFORM_HDR_0		] = { mtlpp::PixelFormat::Invalid, EMetalBufferFormat::Unknown };
-	GMetalBufferFormats[PF_PLATFORM_HDR_1		] = { mtlpp::PixelFormat::Invalid, EMetalBufferFormat::Unknown };
-	GMetalBufferFormats[PF_PLATFORM_HDR_2		] = { mtlpp::PixelFormat::Invalid, EMetalBufferFormat::Unknown };
-	GMetalBufferFormats[PF_NV12					] = { mtlpp::PixelFormat::Invalid, EMetalBufferFormat::Unknown };
+	GMetalBufferFormats[PF_Unknown              ] = { mtlpp::PixelFormat::Invalid, (uint8)EMetalBufferFormat::Unknown };
+	GMetalBufferFormats[PF_A32B32G32R32F        ] = { mtlpp::PixelFormat::RGBA32Float, (uint8)EMetalBufferFormat::RGBA32Float };
+	GMetalBufferFormats[PF_B8G8R8A8             ] = { mtlpp::PixelFormat::RGBA8Unorm, (uint8)EMetalBufferFormat::RGBA8Unorm }; // mtlpp::PixelFormat::BGRA8Unorm/EMetalBufferFormat::BGRA8Unorm,  < We don't support this as a vertex-format so we have code to swizzle in the shader
+	GMetalBufferFormats[PF_G8                   ] = { mtlpp::PixelFormat::R8Unorm, (uint8)EMetalBufferFormat::R8Unorm };
+	GMetalBufferFormats[PF_G16                  ] = { mtlpp::PixelFormat::R16Unorm, (uint8)EMetalBufferFormat::R16Unorm };
+	GMetalBufferFormats[PF_DXT1                 ] = { mtlpp::PixelFormat::Invalid, (uint8)EMetalBufferFormat::Unknown };
+	GMetalBufferFormats[PF_DXT3                 ] = { mtlpp::PixelFormat::Invalid, (uint8)EMetalBufferFormat::Unknown };
+	GMetalBufferFormats[PF_DXT5                 ] = { mtlpp::PixelFormat::Invalid, (uint8)EMetalBufferFormat::Unknown };
+	GMetalBufferFormats[PF_UYVY                 ] = { mtlpp::PixelFormat::Invalid, (uint8)EMetalBufferFormat::Unknown };
+	GMetalBufferFormats[PF_FloatRGB             ] = { mtlpp::PixelFormat::Invalid, (uint8)EMetalBufferFormat::RGB16Half };
+	GMetalBufferFormats[PF_FloatRGBA            ] = { mtlpp::PixelFormat::RGBA16Float, (uint8)EMetalBufferFormat::RGBA16Half };
+	GMetalBufferFormats[PF_DepthStencil         ] = { mtlpp::PixelFormat::Invalid, (uint8)EMetalBufferFormat::Unknown };
+	GMetalBufferFormats[PF_ShadowDepth          ] = { mtlpp::PixelFormat::Invalid, (uint8)EMetalBufferFormat::Unknown };
+	GMetalBufferFormats[PF_R32_FLOAT            ] = { mtlpp::PixelFormat::R32Float, (uint8)EMetalBufferFormat::R32Float };
+	GMetalBufferFormats[PF_G16R16               ] = { mtlpp::PixelFormat::RG16Unorm, (uint8)EMetalBufferFormat::RG16Unorm };
+	GMetalBufferFormats[PF_G16R16F              ] = { mtlpp::PixelFormat::RG16Float, (uint8)EMetalBufferFormat::RG16Half };
+	GMetalBufferFormats[PF_G16R16F_FILTER       ] = { mtlpp::PixelFormat::RG16Float, (uint8)EMetalBufferFormat::RG16Half };
+	GMetalBufferFormats[PF_G32R32F              ] = { mtlpp::PixelFormat::RG32Float, (uint8)EMetalBufferFormat::RG32Float };
+	GMetalBufferFormats[PF_A2B10G10R10          ] = { mtlpp::PixelFormat::RGB10A2Unorm, (uint8)EMetalBufferFormat::RGB10A2Unorm };
+	GMetalBufferFormats[PF_A16B16G16R16         ] = { mtlpp::PixelFormat::RGBA16Unorm, (uint8)EMetalBufferFormat::RGBA16Half };
+	GMetalBufferFormats[PF_D24                  ] = { mtlpp::PixelFormat::Invalid, (uint8)EMetalBufferFormat::Unknown };
+	GMetalBufferFormats[PF_R16F                 ] = { mtlpp::PixelFormat::R16Float, (uint8)EMetalBufferFormat::RG16Half };
+	GMetalBufferFormats[PF_R16F_FILTER          ] = { mtlpp::PixelFormat::R16Float, (uint8)EMetalBufferFormat::RG16Half };
+	GMetalBufferFormats[PF_BC5                  ] = { mtlpp::PixelFormat::Invalid, (uint8)EMetalBufferFormat::Unknown };
+	GMetalBufferFormats[PF_V8U8                 ] = { mtlpp::PixelFormat::RG8Snorm, (uint8)EMetalBufferFormat::RG8Unorm };
+	GMetalBufferFormats[PF_A1                   ] = { mtlpp::PixelFormat::Invalid, (uint8)EMetalBufferFormat::Unknown };
+	GMetalBufferFormats[PF_FloatR11G11B10       ] = { mtlpp::PixelFormat::RG11B10Float, (uint8)EMetalBufferFormat::RG11B10Half }; // < May not work on tvOS
+	GMetalBufferFormats[PF_A8                   ] = { mtlpp::PixelFormat::A8Unorm, (uint8)EMetalBufferFormat::R8Unorm };
+	GMetalBufferFormats[PF_R32_UINT             ] = { mtlpp::PixelFormat::R32Uint, (uint8)EMetalBufferFormat::R32Uint };
+	GMetalBufferFormats[PF_R32_SINT             ] = { mtlpp::PixelFormat::R32Sint, (uint8)EMetalBufferFormat::R32Sint };
+	GMetalBufferFormats[PF_PVRTC2               ] = { mtlpp::PixelFormat::Invalid, (uint8)EMetalBufferFormat::Unknown };
+	GMetalBufferFormats[PF_PVRTC4               ] = { mtlpp::PixelFormat::Invalid, (uint8)EMetalBufferFormat::Unknown };
+	GMetalBufferFormats[PF_R16_UINT             ] = { mtlpp::PixelFormat::R16Uint, (uint8)EMetalBufferFormat::R16Uint };
+	GMetalBufferFormats[PF_R16_SINT             ] = { mtlpp::PixelFormat::R16Sint, (uint8)EMetalBufferFormat::R16Sint };
+	GMetalBufferFormats[PF_R16G16B16A16_UINT    ] = { mtlpp::PixelFormat::RGBA16Uint, (uint8)EMetalBufferFormat::RGBA16Uint };
+	GMetalBufferFormats[PF_R16G16B16A16_SINT    ] = { mtlpp::PixelFormat::RGBA16Sint, (uint8)EMetalBufferFormat::RGBA16Sint };
+	GMetalBufferFormats[PF_R5G6B5_UNORM         ] = { mtlpp::PixelFormat::Invalid, (uint8)EMetalBufferFormat::R5G6B5Unorm };
+	GMetalBufferFormats[PF_R8G8B8A8             ] = { mtlpp::PixelFormat::RGBA8Unorm, (uint8)EMetalBufferFormat::RGBA8Unorm };
+	GMetalBufferFormats[PF_A8R8G8B8				] = { mtlpp::PixelFormat::RGBA8Unorm, (uint8)EMetalBufferFormat::RGBA8Unorm }; // mtlpp::PixelFormat::BGRA8Unorm/EMetalBufferFormat::BGRA8Unorm,  < We don't support this as a vertex-format so we have code to swizzle in the shader
+	GMetalBufferFormats[PF_BC4					] = { mtlpp::PixelFormat::Invalid, (uint8)EMetalBufferFormat::Unknown };
+	GMetalBufferFormats[PF_R8G8                 ] = { mtlpp::PixelFormat::RG8Unorm, (uint8)EMetalBufferFormat::RG8Unorm };
+	GMetalBufferFormats[PF_ATC_RGB				] = { mtlpp::PixelFormat::Invalid, (uint8)EMetalBufferFormat::Unknown };
+	GMetalBufferFormats[PF_ATC_RGBA_E			] = { mtlpp::PixelFormat::Invalid, (uint8)EMetalBufferFormat::Unknown };
+	GMetalBufferFormats[PF_ATC_RGBA_I			] = { mtlpp::PixelFormat::Invalid, (uint8)EMetalBufferFormat::Unknown };
+	GMetalBufferFormats[PF_X24_G8				] = { mtlpp::PixelFormat::Invalid, (uint8)EMetalBufferFormat::Unknown };
+	GMetalBufferFormats[PF_ETC1					] = { mtlpp::PixelFormat::Invalid, (uint8)EMetalBufferFormat::Unknown };
+	GMetalBufferFormats[PF_ETC2_RGB				] = { mtlpp::PixelFormat::Invalid, (uint8)EMetalBufferFormat::Unknown };
+	GMetalBufferFormats[PF_ETC2_RGBA			] = { mtlpp::PixelFormat::Invalid, (uint8)EMetalBufferFormat::Unknown };
+	GMetalBufferFormats[PF_R32G32B32A32_UINT	] = { mtlpp::PixelFormat::RGBA32Uint, (uint8)EMetalBufferFormat::RGBA32Uint };
+	GMetalBufferFormats[PF_R16G16_UINT			] = { mtlpp::PixelFormat::RG16Uint, (uint8)EMetalBufferFormat::RG16Uint };
+	GMetalBufferFormats[PF_R32G32_UINT			] = { mtlpp::PixelFormat::RG32Uint, (uint8)EMetalBufferFormat::RG32Uint };
+	GMetalBufferFormats[PF_ASTC_4x4             ] = { mtlpp::PixelFormat::Invalid, (uint8)EMetalBufferFormat::Unknown };
+	GMetalBufferFormats[PF_ASTC_6x6             ] = { mtlpp::PixelFormat::Invalid, (uint8)EMetalBufferFormat::Unknown };
+	GMetalBufferFormats[PF_ASTC_8x8             ] = { mtlpp::PixelFormat::Invalid, (uint8)EMetalBufferFormat::Unknown };
+	GMetalBufferFormats[PF_ASTC_10x10           ] = { mtlpp::PixelFormat::Invalid, (uint8)EMetalBufferFormat::Unknown };
+	GMetalBufferFormats[PF_ASTC_12x12           ] = { mtlpp::PixelFormat::Invalid, (uint8)EMetalBufferFormat::Unknown };
+	GMetalBufferFormats[PF_BC6H					] = { mtlpp::PixelFormat::Invalid, (uint8)EMetalBufferFormat::Unknown };
+	GMetalBufferFormats[PF_BC7					] = { mtlpp::PixelFormat::Invalid, (uint8)EMetalBufferFormat::Unknown };
+	GMetalBufferFormats[PF_R8_UINT				] = { mtlpp::PixelFormat::R8Uint, (uint8)EMetalBufferFormat::R8Uint };
+	GMetalBufferFormats[PF_L8					] = { mtlpp::PixelFormat::Invalid, (uint8)EMetalBufferFormat::R8Unorm };
+	GMetalBufferFormats[PF_XGXR8				] = { mtlpp::PixelFormat::Invalid, (uint8)EMetalBufferFormat::Unknown };
+	GMetalBufferFormats[PF_R8G8B8A8_UINT		] = { mtlpp::PixelFormat::RGBA8Uint, (uint8)EMetalBufferFormat::RGBA8Uint };
+	GMetalBufferFormats[PF_R8G8B8A8_SNORM		] = { mtlpp::PixelFormat::RGBA8Snorm, (uint8)EMetalBufferFormat::RGBA8Snorm };
+	GMetalBufferFormats[PF_R16G16B16A16_UNORM	] = { mtlpp::PixelFormat::RGBA16Unorm, (uint8)EMetalBufferFormat::RGBA16Unorm };
+	GMetalBufferFormats[PF_R16G16B16A16_SNORM	] = { mtlpp::PixelFormat::RGBA16Snorm, (uint8)EMetalBufferFormat::RGBA16Snorm };
+	GMetalBufferFormats[PF_PLATFORM_HDR_0		] = { mtlpp::PixelFormat::Invalid, (uint8)EMetalBufferFormat::Unknown };
+	GMetalBufferFormats[PF_PLATFORM_HDR_1		] = { mtlpp::PixelFormat::Invalid, (uint8)EMetalBufferFormat::Unknown };
+	GMetalBufferFormats[PF_PLATFORM_HDR_2		] = { mtlpp::PixelFormat::Invalid, (uint8)EMetalBufferFormat::Unknown };
+	GMetalBufferFormats[PF_NV12					] = { mtlpp::PixelFormat::Invalid, (uint8)EMetalBufferFormat::Unknown };
 		
 	// Initialize the platform pixel format map.
 	GPixelFormats[PF_Unknown			].PlatformFormat	= (uint32)mtlpp::PixelFormat::Invalid;
diff --git a/Engine/Source/Runtime/Apple/MetalRHI/Public/MetalShaderResources.h b/Engine/Source/Runtime/Apple/MetalRHI/Public/MetalShaderResources.h
--- a/Engine/Source/Runtime/Apple/MetalRHI/Public/MetalShaderResources.h
+++ b/Engine/Source/Runtime/Apple/MetalRHI/Public/MetalShaderResources.h
@@ -23,7 +23,7 @@ enum
 /**
 * Buffer data-types for MetalRHI & MetalSL
 */
-enum EMetalBufferFormat
+enum class EMetalBufferFormat : uint8
 {
 	Unknown					=0,
 	
