Fixed a crash on Metal platforms caused by CL 9691350

diff --git a/Engine/Source/Runtime/Apple/MetalRHI/Private/MetalRHI.cpp b/Engine/Source/Runtime/Apple/MetalRHI/Private/MetalRHI.cpp
--- a/Engine/Source/Runtime/Apple/MetalRHI/Private/MetalRHI.cpp
+++ b/Engine/Source/Runtime/Apple/MetalRHI/Private/MetalRHI.cpp
@@ -597,6 +597,8 @@ FMetalDynamicRHI::FMetalDynamicRHI(ERHIFeatureLevel::Type RequestedFeatureLevel)
 	GMetalBufferFormats[PF_PLATFORM_HDR_1		] = { mtlpp::PixelFormat::Invalid, (uint8)EMetalBufferFormat::Unknown };
 	GMetalBufferFormats[PF_PLATFORM_HDR_2		] = { mtlpp::PixelFormat::Invalid, (uint8)EMetalBufferFormat::Unknown };
 	GMetalBufferFormats[PF_NV12					] = { mtlpp::PixelFormat::Invalid, (uint8)EMetalBufferFormat::Unknown };
+	GMetalBufferFormats[PF_ETC2_R11_EAC			] = { mtlpp::PixelFormat::Invalid, (uint8)EMetalBufferFormat::Unknown };
+	GMetalBufferFormats[PF_ETC2_RG11_EAC		] = { mtlpp::PixelFormat::Invalid, (uint8)EMetalBufferFormat::Unknown };
 		
 	// Initialize the platform pixel format map.
 	GPixelFormats[PF_Unknown			].PlatformFormat	= (uint32)mtlpp::PixelFormat::Invalid;
@@ -764,6 +766,11 @@ FMetalDynamicRHI::FMetalDynamicRHI(ERHIFeatureLevel::Type RequestedFeatureLevel)
 	GPixelFormats[PF_NV12				].PlatformFormat	= (uint32)mtlpp::PixelFormat::Invalid;
 	GPixelFormats[PF_NV12				].Supported			= false;
 
+	GPixelFormats[PF_ETC2_R11_EAC	  	].PlatformFormat	= (uint32)mtlpp::PixelFormat::Invalid;
+	GPixelFormats[PF_ETC2_R11_EAC		].Supported			= false;
+	GPixelFormats[PF_ETC2_RG11_EAC		].PlatformFormat	= (uint32)mtlpp::PixelFormat::Invalid;
+	GPixelFormats[PF_ETC2_RG11_EAC		].Supported			= false;
+	
 #if METAL_DEBUG_OPTIONS
 	for (uint32 i = 0; i < PF_MAX; i++)
 	{
