Fix crash when flushing into matches on Intel based Macs.

- Implement workaround to enable support for devices that support less than MTLReadWriteTextureTier2.
- Reset pointers to released resources.

#rnx

diff --git a/Engine/Source/Runtime/Apple/MetalRHI/Private/MetalVertexBuffer.cpp b/Engine/Source/Runtime/Apple/MetalRHI/Private/MetalVertexBuffer.cpp
--- a/Engine/Source/Runtime/Apple/MetalRHI/Private/MetalVertexBuffer.cpp
+++ b/Engine/Source/Runtime/Apple/MetalRHI/Private/MetalVertexBuffer.cpp
@@ -654,6 +654,7 @@ struct FMetalRHICommandInitialiseBuffer : public FRHICommand<FMetalRHICommandIni
 			if (Buffer->CPUBuffer)
 			{
 				SafeReleaseMetalBuffer(Buffer->CPUBuffer);
+				Buffer->CPUBuffer = nil;
 			}
 			else
 			{
@@ -797,6 +798,7 @@ FMetalStagingBuffer::~FMetalStagingBuffer()
 	if (ShadowBuffer)
 	{
 		SafeReleaseMetalBuffer(ShadowBuffer);
+		ShadowBuffer = nil;
 	}
 }
 
diff --git a/Engine/Source/Runtime/RHI/Public/RHI.h b/Engine/Source/Runtime/RHI/Public/RHI.h
--- a/Engine/Source/Runtime/RHI/Public/RHI.h
+++ b/Engine/Source/Runtime/RHI/Public/RHI.h
@@ -153,7 +153,7 @@ inline bool RHISupportsMSAA(EShaderPlatform Platform)
 
 inline bool RHISupportsBufferLoadTypeConversion(EShaderPlatform Platform)
 {
-	return true;
+	return !IsMetalPlatform(Platform);
 }
 
 /** Whether the platform supports reading from volume textures (does not cover rendering to volume textures). */
