The Apple TV HD is based on the A8 which does not allow writes to r11g11b10. So we'll promote to 16F instead of losing precision.

diff --git a/Engine/Source/Runtime/Apple/MetalRHI/Private/MetalRHI.cpp b/Engine/Source/Runtime/Apple/MetalRHI/Private/MetalRHI.cpp
--- a/Engine/Source/Runtime/Apple/MetalRHI/Private/MetalRHI.cpp
+++ b/Engine/Source/Runtime/Apple/MetalRHI/Private/MetalRHI.cpp
@@ -649,6 +649,7 @@ FMetalDynamicRHI::FMetalDynamicRHI(ERHIFeatureLevel::Type RequestedFeatureLevel)
 		GPixelFormats[PF_FloatRGBA			].BlockBytes		= 8;
 		GPixelFormats[PF_FloatR11G11B10		].PlatformFormat	= (uint32)mtlpp::PixelFormat::RGBA16Float;
 		GPixelFormats[PF_FloatR11G11B10		].BlockBytes		= 8;
+		GPixelFormats[PF_FloatR11G11B10		].Supported			= true;
 	}
 	else
 	{
