Acquire unique backing store for intra-frame async uniform buffer updates.

Note, this is only a partial fix for UE-82454.

#rnx

diff --git a/Engine/Source/Runtime/Apple/MetalRHI/Private/MetalContext.cpp b/Engine/Source/Runtime/Apple/MetalRHI/Private/MetalContext.cpp
--- a/Engine/Source/Runtime/Apple/MetalRHI/Private/MetalContext.cpp
+++ b/Engine/Source/Runtime/Apple/MetalRHI/Private/MetalContext.cpp
@@ -349,6 +349,7 @@ FMetalDeviceContext::FMetalDeviceContext(mtlpp::Device MetalDevice, uint32 InDev
 , ActiveContexts(1)
 , ActiveParallelContexts(0)
 , PSOManager(0)
+, DeviceFrameIndex(0)
 {
 	CommandQueue.SetRuntimeDebuggingLevel(GMetalRuntimeDebugLevel);
 	
@@ -418,6 +419,9 @@ void FMetalDeviceContext::BeginFrame()
 	
 	// Wait for the frame semaphore on the immediate context.
 	dispatch_semaphore_wait(CommandBufferSemaphore, DISPATCH_TIME_FOREVER);
+	
+	// Bump the frame counter.
+	DeviceFrameIndex++;
 }
 
 #if METAL_DEBUG_OPTIONS
diff --git a/Engine/Source/Runtime/Apple/MetalRHI/Private/MetalContext.h b/Engine/Source/Runtime/Apple/MetalRHI/Private/MetalContext.h
--- a/Engine/Source/Runtime/Apple/MetalRHI/Private/MetalContext.h
+++ b/Engine/Source/Runtime/Apple/MetalRHI/Private/MetalContext.h
@@ -225,6 +225,12 @@ public:
 	/** Get the index of the bound Metal device in the global list of rendering devices. */
 	uint32 GetDeviceIndex(void) const;
     
+	/** Device frame index accessor. */
+	uint64 GetDeviceFrameIndex() const
+	{
+		return DeviceFrameIndex;
+	}
+	
 #if METAL_DEBUG_OPTIONS
     void AddActiveBuffer(FMetalBuffer const& Buffer);
     void RemoveActiveBuffer(FMetalBuffer const& Buffer);
@@ -309,4 +315,7 @@ private:
 	
 	/** PSO cache manager */
 	FMetalPipelineStateCacheManager* PSOManager;
+
+	/** Device frame index, glorified frame counter in the device namespace. */
+	uint64 DeviceFrameIndex;
 };
diff --git a/Engine/Source/Runtime/Apple/MetalRHI/Private/MetalUniformBuffer.cpp b/Engine/Source/Runtime/Apple/MetalRHI/Private/MetalUniformBuffer.cpp
--- a/Engine/Source/Runtime/Apple/MetalRHI/Private/MetalUniformBuffer.cpp
+++ b/Engine/Source/Runtime/Apple/MetalRHI/Private/MetalUniformBuffer.cpp
@@ -757,13 +757,17 @@ void FMetalUniformBuffer::Update(const void* Contents, TArray<TRefCountPtr<FRHIR
         
 		ns::AutoReleased<FMetalBuffer> Buf(Buffer);
 		
-        void* Data = Lock(true, RLM_WriteOnly, 0);
+		void* Data = Lock(true, RLM_WriteOnly, 0, 0, true);
         FMemory::Memcpy(Data, Contents, ConstantSize);
         Unlock();
 		
 		if (Buf != Buffer && FMetalCommandQueue::SupportsFeature(EMetalFeaturesIABs))
+		{
 			FPlatformAtomics::InterlockedIncrement(&UpdateNum);
-    }
+		}
+		
+		ConditionalSetUniformBufferFrameIndex();
+	}
 	
 	UpdateResourceTable(Resources, Validation);
 }
diff --git a/Engine/Source/Runtime/Apple/MetalRHI/Private/MetalVertexBuffer.cpp b/Engine/Source/Runtime/Apple/MetalRHI/Private/MetalVertexBuffer.cpp
--- a/Engine/Source/Runtime/Apple/MetalRHI/Private/MetalVertexBuffer.cpp
+++ b/Engine/Source/Runtime/Apple/MetalRHI/Private/MetalVertexBuffer.cpp
@@ -234,7 +234,7 @@ void FMetalRHIBuffer::Unalias()
 	}
 }
 
-void FMetalRHIBuffer::Alloc(uint32 InSize, EResourceLockMode LockMode)
+void FMetalRHIBuffer::Alloc(uint32 InSize, EResourceLockMode LockMode, bool bIsUniformBuffer)
 {
 	if (!Buffer)
 	{
@@ -258,6 +258,8 @@ void FMetalRHIBuffer::Alloc(uint32 InSize, EResourceLockMode LockMode)
 				check(Pair.Value);
 			}
 		}
+
+		bIsUniformBufferBacking = bIsUniformBuffer;
 	}
 }
 
@@ -411,7 +413,22 @@ ns::AutoReleased<FMetalTexture> FMetalRHIBuffer::GetLinearTexture(EPixelFormat F
 	return Texture;
 }
 
-void* FMetalRHIBuffer::Lock(bool bIsOnRHIThread, EResourceLockMode LockMode, uint32 Offset, uint32 InSize)
+bool FMetalRHIBuffer::CanUseBufferAsBackingForAsyncCopy() const
+{
+	return (!bIsUniformBufferBacking ||
+			UniformBufferFrameIndex != GetMetalDeviceContext().GetDeviceFrameIndex() ||
+			UniformBufferPreviousOffset != Buffer.GetOffset());
+}
+
+void FMetalRHIBuffer::ConditionalSetUniformBufferFrameIndex()
+{
+	if (bIsUniformBufferBacking)
+	{
+		UniformBufferFrameIndex = GetMetalDeviceContext().GetDeviceFrameIndex();
+	}
+}
+
+void* FMetalRHIBuffer::Lock(bool bIsOnRHIThread, EResourceLockMode LockMode, uint32 Offset, uint32 InSize, bool bIsUniformBuffer /*= false*/)
 {
 	check(LockSize == 0 && LockOffset == 0);
 	
@@ -454,14 +471,18 @@ void* FMetalRHIBuffer::Lock(bool bIsOnRHIThread, EResourceLockMode LockMode, uin
 	
 	// When writing to a private buffer, make sure that we can perform an async copy so we don't introduce order of operation bugs
 	// When we can't we have to reallocate the backing store
-	if (LockMode != RLM_ReadOnly && Mode == mtlpp::StorageMode::Private && Buffer && !GetMetalDeviceContext().CanAsyncCopyToBuffer(Buffer))
+	if (LockMode != RLM_ReadOnly &&
+		Mode == mtlpp::StorageMode::Private &&
+		Buffer &&
+		(!GetMetalDeviceContext().CanAsyncCopyToBuffer(Buffer) ||
+		 !CanUseBufferAsBackingForAsyncCopy()))
 	{
 		METAL_INC_DWORD_STAT_BY(Type, MemFreed, Len);
 		SafeReleaseMetalBuffer(Buffer);
 		Buffer = nil;
 	}
 	
-    Alloc(Len, LockMode);
+    Alloc(Len, LockMode, bIsUniformBuffer);
 	AllocTransferBuffer(bIsOnRHIThread, Len, LockMode);
 	
 	FMetalBuffer& theBufferToUse = CPUBuffer ? CPUBuffer : Buffer;
@@ -511,6 +532,9 @@ void FMetalRHIBuffer::Unlock()
 		{
 			// Synchronise the buffer with the GPU
 			GetMetalDeviceContext().AsyncCopyFromBufferToBuffer(CPUBuffer, 0, Buffer, 0, Buffer.GetLength());
+			
+			ConditionalSetUniformBufferPreviousOffset();
+			
 			if (CPUBuffer)
             {
 				SafeReleaseMetalBuffer(CPUBuffer);
diff --git a/Engine/Source/Runtime/Apple/MetalRHI/Public/MetalResources.h b/Engine/Source/Runtime/Apple/MetalRHI/Public/MetalResources.h
--- a/Engine/Source/Runtime/Apple/MetalRHI/Public/MetalResources.h
+++ b/Engine/Source/Runtime/Apple/MetalRHI/Public/MetalResources.h
@@ -786,7 +786,7 @@ public:
 	/**
 	 * Allocate the buffer backing store.
 	 */
-	void Alloc(uint32 InSize, EResourceLockMode LockMode);
+	void Alloc(uint32 InSize, EResourceLockMode LockMode, bool bIsUniformBuffer=false);
 	
 	/**
 	 * Allocate the CPU accessible buffer for data transfer.
@@ -811,7 +811,7 @@ public:
 	/**
 	 * Prepare a CPU accessible buffer for uploading to GPU memory
 	 */
-	void* Lock(bool bIsOnRHIThread, EResourceLockMode LockMode, uint32 Offset, uint32 Size=0);
+	void* Lock(bool bIsOnRHIThread, EResourceLockMode LockMode, uint32 Offset, uint32 Size=0, bool bIsUniformBuffer=false);
 	
 	/**
 	 * Prepare a CPU accessible buffer for uploading to GPU memory
@@ -821,7 +821,7 @@ public:
 	void Swap(FMetalRHIBuffer& Other);
 	
 	/**
-	 * Whether to allocate the resource rom private memory or not.
+	 * Whether to allocate the resource from private memory.
 	 */
 	bool UsePrivateMemory() const;
 	
@@ -857,6 +857,33 @@ public:
 	
 	// Resource type
 	ERHIResourceType Type;
+
+	/**
+	 * For resources backing uniform buffers, set the device frame index.
+	 *
+	 * Calling this function on resources that do not back uniform buffers has no effect.
+	 */
+	void ConditionalSetUniformBufferFrameIndex();
+
+private:
+	bool CanUseBufferAsBackingForAsyncCopy() const;
+
+	/**
+	 * For resources backing uniform buffers, updates the last used sub-allocation offset; this is used for
+	 * internal tracking and deciding when to recycle backing stores during uniform buffer updates.
+	 */
+	void ConditionalSetUniformBufferPreviousOffset()
+	{
+		if (bIsUniformBufferBacking)
+		{
+			check(Buffer);
+			UniformBufferPreviousOffset = Buffer.GetOffset();
+		}
+	}
+
+	bool bIsUniformBufferBacking = false;
+	uint64 UniformBufferFrameIndex = 0;
+	uint64 UniformBufferPreviousOffset = uint64(-1);
 };
 
 /** Index buffer resource class that stores stride information. */
