Fix iOS build.

#rnx

diff --git a/Engine/Source/Runtime/Apple/MetalRHI/Public/MetalRHIContext.h b/Engine/Source/Runtime/Apple/MetalRHI/Public/MetalRHIContext.h
--- a/Engine/Source/Runtime/Apple/MetalRHI/Public/MetalRHIContext.h
+++ b/Engine/Source/Runtime/Apple/MetalRHI/Public/MetalRHIContext.h
@@ -281,6 +281,11 @@ public:
 	 */
 	virtual void RHITransitionResources(EResourceTransitionAccess TransitionType, FRHITexture** InTextures, int32 NumTextures) final override;
 	
+	virtual void RHITransitionResources(FExclusiveDepthStencil DepthStencilMode, FRHITexture* DepthTexture) final override
+	{
+		IRHICommandContext::RHITransitionResources(DepthStencilMode, DepthTexture);
+	}
+
 	/**
 	 * Transition depth resources. Does not really do anything in metal.
 	 */
