Fix duplicate function declaration as a result of CL# 10160290 and 9610165
#rnx

diff --git a/Engine/Source/Runtime/Apple/MetalRHI/Private/MetalRHIContext.cpp b/Engine/Source/Runtime/Apple/MetalRHI/Private/MetalRHIContext.cpp
--- a/Engine/Source/Runtime/Apple/MetalRHI/Private/MetalRHIContext.cpp
+++ b/Engine/Source/Runtime/Apple/MetalRHI/Private/MetalRHIContext.cpp
@@ -230,8 +230,3 @@ void FMetalRHICommandContext::RHIBeginComputePass(const TCHAR* InName)
 void FMetalRHICommandContext::RHIEndComputePass()
 {
 }
-
-void FMetalRHICommandContext::RHITransitionResources(FExclusiveDepthStencil DepthStencilMode, FRHITexture* DepthTexture)
-{
-	IRHICommandContext::RHITransitionResources(DepthStencilMode, DepthTexture);
-}
diff --git a/Engine/Source/Runtime/Apple/MetalRHI/Public/MetalRHIContext.h b/Engine/Source/Runtime/Apple/MetalRHI/Public/MetalRHIContext.h
--- a/Engine/Source/Runtime/Apple/MetalRHI/Public/MetalRHIContext.h
+++ b/Engine/Source/Runtime/Apple/MetalRHI/Public/MetalRHIContext.h
@@ -286,11 +286,6 @@ public:
 		IRHICommandContext::RHITransitionResources(DepthStencilMode, DepthTexture);
 	}
 
-	/**
-	 * Transition depth resources. Does not really do anything in metal.
-	 */
-	virtual void RHITransitionResources(FExclusiveDepthStencil DepthStencilMode, FRHITexture* DepthTexture) final override;
-	
 	/**
 	 * Compute queue will wait for the fence to be written before continuing.
 	 */
