Use a defined, named constant instead of a magic number to initialize a class member variable.

#rnx

diff --git a/Engine/Source/Runtime/Apple/MetalRHI/Public/MetalResources.h b/Engine/Source/Runtime/Apple/MetalRHI/Public/MetalResources.h
--- a/Engine/Source/Runtime/Apple/MetalRHI/Public/MetalResources.h
+++ b/Engine/Source/Runtime/Apple/MetalRHI/Public/MetalResources.h
@@ -883,7 +883,7 @@ private:
 
 	bool bIsUniformBufferBacking = false;
 	uint64 UniformBufferFrameIndex = 0;
-	uint64 UniformBufferPreviousOffset = uint64(-1);
+	uint64 UniformBufferPreviousOffset = UINT64_MAX;
 };
 
 /** Index buffer resource class that stores stride information. */
