Fix for the leak that causes UE-82849.

diff --git a/Engine/Source/Runtime/Apple/MetalRHI/Private/MetalTexture.cpp b/Engine/Source/Runtime/Apple/MetalRHI/Private/MetalTexture.cpp
--- a/Engine/Source/Runtime/Apple/MetalRHI/Private/MetalTexture.cpp
+++ b/Engine/Source/Runtime/Apple/MetalRHI/Private/MetalTexture.cpp
@@ -1181,21 +1181,17 @@ FMetalSurface::~FMetalSurface()
 id <MTLBuffer> FMetalSurface::AllocSurface(uint32 MipIndex, uint32 ArrayIndex, EResourceLockMode LockMode, uint32& DestStride, bool SingleLayer /*= false*/)
 {
 	check(IsInRenderingThread());
-	
-	
+
 	// get size and stride
 	uint32 MipBytes = GetMipSize(MipIndex, &DestStride, SingleLayer);
 	
 	// allocate some temporary memory
-//	mtlpp::ResourceOptions ResMode = FMetalCommandQueue::GetCompatibleResourceOptions((mtlpp::ResourceOptions)(mtlpp::ResourceOptions::StorageModeShared | (!(PLATFORM_MAC && PixelFormat == PF_G8 && (Flags & TexCreate_SRGB)) ? mtlpp::ResourceOptions::CpuCacheModeWriteCombined : 0)));
-	
-//	FMetalBuffer Buffer = GetMetalDeviceContext().GetResourceHeap().CreateBuffer(MipBytes, BufferOffsetAlignment, BUF_Dynamic, ResMode);
-	
+	// This should really be pooled and texture transfers should be their own pool
 	id <MTLDevice> Device = GetMetalDeviceContext().GetDevice();
 	id <MTLBuffer> Buffer = [Device newBufferWithLength:MipBytes options:MTLResourceStorageModeShared];
+	Buffer.label = @"Temporary Surface Backing";
 	
-	// UE objects don't auto-retain objc objects. So we'll need to hold a reference.
-	[Buffer retain];
+	// Note: while the lock is active, this map owns the backing store.
 	GRHILockTracker.Lock(this, Buffer, MipIndex, 0, LockMode, false);
 	
 #if PLATFORM_MAC
