Forcing RHI thread off for Intel HW.

Temporary fix for UE-84581. Fixes meshes jumping around in Infiltrator.
Since this only occurs on Intel (AMD and iOS are fine) I don't want to put in a 'real' fix this late in the game and risk those other platforms. Needs to be fixed for real in 4.24.1.

diff --git a/Engine/Source/Runtime/Apple/MetalRHI/Private/MetalRHI.cpp b/Engine/Source/Runtime/Apple/MetalRHI/Private/MetalRHI.cpp
--- a/Engine/Source/Runtime/Apple/MetalRHI/Private/MetalRHI.cpp
+++ b/Engine/Source/Runtime/Apple/MetalRHI/Private/MetalRHI.cpp
@@ -37,6 +37,17 @@ static TAutoConsoleVariable<int32> CVarUseRHIThread(
 													ECVF_Default | ECVF_RenderThreadSafe
 													);
 
+static TAutoConsoleVariable<int32> CVarIntelUseRHIThread(
+													TEXT("r.Metal.IntelRHIThread"),
+													0,
+													TEXT("Controls RHIThread usage for Mac Intel HW:\n")
+													TEXT("\t0: No RHIThread.\n")
+													TEXT("\t1: Use RHIThread.\n")
+													TEXT("Default is 0."),
+													ECVF_Default | ECVF_RenderThreadSafe
+													);
+
+
 static void ValidateTargetedRHIFeatureLevelExists(EShaderPlatform Platform)
 {
 	bool bSupportsShaderPlatform = false;
@@ -255,7 +266,8 @@ FMetalDynamicRHI::FMetalDynamicRHI(ERHIFeatureLevel::Type RequestedFeatureLevel)
 		bSupportsTiledReflections = false;
 		bSupportsPointLights = (FPlatformMisc::MacOSXVersionCompare(10,14,6) > 0);
 		GRHIVendorId = 0x8086;
-		bSupportsRHIThread = true;
+		// HACK: Meshes jump around in Infiltrator with RHI thread on. Needs further investigation and a real fix.
+		bSupportsRHIThread = false || CVarIntelUseRHIThread.GetValueOnAnyThread() > 0;
 		bSupportsDistanceFields = (FPlatformMisc::MacOSXVersionCompare(10,12,2) >= 0);
 		bIsIntelHaswell = (GRHIAdapterName == TEXT("Intel HD Graphics 5000") || GRHIAdapterName == TEXT("Intel Iris Graphics") || GRHIAdapterName == TEXT("Intel Iris Pro Graphics"));
 	}
