Fix bad merge
#rnx

diff --git a/Engine/Source/Runtime/Apple/MetalRHI/Private/MetalStructuredBuffer.cpp b/Engine/Source/Runtime/Apple/MetalRHI/Private/MetalStructuredBuffer.cpp
--- a/Engine/Source/Runtime/Apple/MetalRHI/Private/MetalStructuredBuffer.cpp
+++ b/Engine/Source/Runtime/Apple/MetalRHI/Private/MetalStructuredBuffer.cpp
@@ -31,9 +31,9 @@ FStructuredBufferRHIRef FMetalDynamicRHI::RHICreateStructuredBuffer(uint32 Strid
 {
 	@autoreleasepool {
 	FMetalStructuredBuffer* Buffer = new FMetalStructuredBuffer(Stride, Size, CreateInfo.ResourceArray, InUsage);
-	if (!CreateInfo.ResourceArray && Buffer->Buffer.GetStorageMode() == mtlpp::StorageMode::Private)
+	if (!CreateInfo.ResourceArray && Buffer->Mode == mtlpp::StorageMode::Private)
 	{
-		if (Buffer->GetUsage() & (BUF_Dynamic|BUF_Static))
+		if (Buffer->CPUBuffer)
 		{
 			SafeReleaseMetalBuffer(Buffer->CPUBuffer);
 			Buffer->CPUBuffer = nil;
@@ -45,7 +45,7 @@ FStructuredBufferRHIRef FMetalDynamicRHI::RHICreateStructuredBuffer(uint32 Strid
 		}
 	}
 #if PLATFORM_MAC
-	else if (GMetalBufferZeroFill && !CreateInfo.ResourceArray && Buffer->Buffer.GetStorageMode() == mtlpp::StorageMode::Managed)
+	else if (GMetalBufferZeroFill && !CreateInfo.ResourceArray && Buffer->Mode == mtlpp::StorageMode::Managed)
 	{
 		MTLPP_VALIDATE(mtlpp::Buffer, Buffer->Buffer, SafeGetRuntimeDebuggingLevel() >= EMetalDebugLevelValidation, DidModify(ns::Range(0, Buffer->Buffer.GetLength())));
 	}
@@ -54,17 +54,17 @@ FStructuredBufferRHIRef FMetalDynamicRHI::RHICreateStructuredBuffer(uint32 Strid
 	}
 }
 
-void* FMetalDynamicRHI::RHILockStructuredBuffer(FRHIStructuredBuffer* StructuredBufferRHI,uint32 Offset,uint32 Size,EResourceLockMode LockMode)
+void* FMetalDynamicRHI::LockStructuredBuffer_BottomOfPipe(FRHICommandListImmediate& RHICmdList, FRHIStructuredBuffer* StructuredBufferRHI, uint32 Offset, uint32 Size, EResourceLockMode LockMode)
 {
 	@autoreleasepool {
 	FMetalStructuredBuffer* StructuredBuffer = ResourceCast(StructuredBufferRHI);
 	
 	// just return the memory plus the offset
-	return (uint8*)StructuredBuffer->Lock(LockMode, Offset, Size);
+	return (uint8*)StructuredBuffer->Lock(true, LockMode, Offset, Size);
 	}
 }
 
-void FMetalDynamicRHI::RHIUnlockStructuredBuffer(FRHIStructuredBuffer* StructuredBufferRHI)
+void FMetalDynamicRHI::UnlockStructuredBuffer_BottomOfPipe(FRHICommandListImmediate& RHICmdList, FRHIStructuredBuffer* StructuredBufferRHI)
 {
 	@autoreleasepool {
 	FMetalStructuredBuffer* StructuredBuffer = ResourceCast(StructuredBufferRHI);
@@ -78,72 +78,7 @@ FStructuredBufferRHIRef FMetalDynamicRHI::CreateStructuredBuffer_RenderThread(cl
 		// make the RHI object, which will allocate memory
 		TRefCountPtr<FMetalStructuredBuffer> VertexBuffer = new FMetalStructuredBuffer(Stride, Size, nullptr, InUsage);
 		
-		if (CreateInfo.ResourceArray)
-		{
-			check(Size == CreateInfo.ResourceArray->GetResourceDataSize());
-			
-			if (VertexBuffer->CPUBuffer)
-			{
-				FMemory::Memcpy(VertexBuffer->CPUBuffer.GetContents(), CreateInfo.ResourceArray->GetResourceData(), Size);
-
-#if PLATFORM_MAC
-				if(VertexBuffer->CPUBuffer.GetStorageMode() == mtlpp::StorageMode::Managed)
-				{
-					MTLPP_VALIDATE(mtlpp::Buffer, VertexBuffer->CPUBuffer, SafeGetRuntimeDebuggingLevel() >= EMetalDebugLevelValidation, DidModify(ns::Range(0, GMetalBufferZeroFill ? VertexBuffer->CPUBuffer.GetLength() : Size)));
-				}
-#endif
-				
-				if (RHICmdList.Bypass() || !IsRunningRHIInSeparateThread())
-				{
-					FMetalRHICommandInitialiseStructuredBuffer UpdateCommand(VertexBuffer);
-					UpdateCommand.Execute(RHICmdList);
-				}
-				else
-				{
-					new (RHICmdList.AllocCommand<FMetalRHICommandInitialiseStructuredBuffer>()) FMetalRHICommandInitialiseStructuredBuffer(VertexBuffer);
-				}
-			}
-			else
-			{
-				// make a buffer usable by CPU
-				void* Buffer = RHILockStructuredBuffer(VertexBuffer, 0, Size, RLM_WriteOnly);
-				
-				// copy the contents of the given data into the buffer
-				FMemory::Memcpy(Buffer, CreateInfo.ResourceArray->GetResourceData(), Size);
-				
-				RHIUnlockStructuredBuffer(VertexBuffer);
-			}
-			
-			// Discard the resource array's contents.
-			CreateInfo.ResourceArray->Discard();
-		}
-		else if (VertexBuffer->Buffer.GetStorageMode() == mtlpp::StorageMode::Private)
-		{
-			if (VertexBuffer->GetUsage() & (BUF_Dynamic|BUF_Static))
-			{
-				SafeReleaseMetalBuffer(VertexBuffer->CPUBuffer);
-				VertexBuffer->CPUBuffer = nil;
-			}
-
-			if (GMetalBufferZeroFill)
-			{
-				if (RHICmdList.Bypass() || !IsRunningRHIInSeparateThread())
-				{
-					FMetalRHICommandInitialiseStructuredBuffer UpdateCommand(VertexBuffer);
-					UpdateCommand.Execute(RHICmdList);
-				}
-				else
-				{
-					new (RHICmdList.AllocCommand<FMetalRHICommandInitialiseStructuredBuffer>()) FMetalRHICommandInitialiseStructuredBuffer(VertexBuffer);
-				}
-			}
-		}
-#if PLATFORM_MAC
-		else if (GMetalBufferZeroFill && VertexBuffer->Buffer.GetStorageMode() == mtlpp::StorageMode::Managed)
-		{
-			MTLPP_VALIDATE(mtlpp::Buffer, VertexBuffer->Buffer, SafeGetRuntimeDebuggingLevel() >= EMetalDebugLevelValidation, DidModify(ns::Range(0, VertexBuffer->Buffer.GetLength())));
-		}
-#endif
+		VertexBuffer->Init_RenderThread(RHICmdList, Size, InUsage, CreateInfo, VertexBuffer);
 		
 		return VertexBuffer.GetReference();
 	}
