Increase the number of bits for pixel formats stored in the Metal pipeline key

diff --git a/Engine/Source/Runtime/Apple/MetalRHI/Private/MetalPipeline.h b/Engine/Source/Runtime/Apple/MetalRHI/Private/MetalPipeline.h
--- a/Engine/Source/Runtime/Apple/MetalRHI/Private/MetalPipeline.h
+++ b/Engine/Source/Runtime/Apple/MetalRHI/Private/MetalPipeline.h
@@ -11,7 +11,7 @@
 
 enum EMetalPipelineHashBits
 {
-	NumBits_RenderTargetFormat = 5, //(x8=40),
+	NumBits_RenderTargetFormat = 6, //(x8=48),
 	NumBits_DepthFormat = 3, //(x1=3),
 	NumBits_StencilFormat = 3, //(x1=3),
 	NumBits_SampleCount = 3, //(x1=3),
diff --git a/Engine/Source/Runtime/Apple/MetalRHI/Private/MetalTexture.cpp b/Engine/Source/Runtime/Apple/MetalRHI/Private/MetalTexture.cpp
--- a/Engine/Source/Runtime/Apple/MetalRHI/Private/MetalTexture.cpp
+++ b/Engine/Source/Runtime/Apple/MetalRHI/Private/MetalTexture.cpp
@@ -652,8 +652,8 @@ uint8 GetMetalPixelFormatKey(mtlpp::PixelFormat Format)
 		if (Key == NULL)
 		{
 			Key = &GetMetalPixelFormatKeyMap().Add((uint64)Format, NextKey++);
-			// only giving 5 bits to the key
-			checkf(NextKey < 32, TEXT("Too many unique pixel formats to fit into the PipelineStateHash"));
+			// only giving 6 bits to the key
+			checkf(NextKey < 64, TEXT("Too many unique pixel formats to fit into the PipelineStateHash"));
 		}
 	}
 	return *Key;
@@ -680,7 +680,7 @@ FMetalSurface::FMetalSurface(ERHIResourceType ResourceType, EPixelFormat Format,
 	TMap<uint64, uint8>& PixelFormatKeyMap = GetMetalPixelFormatKeyMap();
 	if (PixelFormatKeyMap.Num() == 0)
 	{
-		// Add depth stencil formats fist, so we don't have to use 5 bits for them in the pipeline hash
+		// Add depth stencil formats first, so we don't have to use 6 bits for them in the pipeline hash
 		GetMetalPixelFormatKey(mtlpp::PixelFormat::Depth32Float);
 		GetMetalPixelFormatKey(mtlpp::PixelFormat::Stencil8);
 		GetMetalPixelFormatKey(mtlpp::PixelFormat::Depth32Float_Stencil8);
