Fix FMetalShaderPipeline initresourcemask when we use many textures.

diff --git a/Engine/Source/Runtime/Apple/MetalRHI/Private/MetalPipeline.cpp b/Engine/Source/Runtime/Apple/MetalRHI/Private/MetalPipeline.cpp
--- a/Engine/Source/Runtime/Apple/MetalRHI/Private/MetalPipeline.cpp
+++ b/Engine/Source/Runtime/Apple/MetalRHI/Private/MetalPipeline.cpp
@@ -687,7 +687,7 @@ APPLE_PLATFORM_OBJECT_ALLOC_OVERRIDES(FMetalShaderPipeline)
 			case MTLArgumentTypeTexture:
 			{
 				checkf(Arg.index < ML_MaxTextures, TEXT("Metal texture index exceeded!"));
-				ResourceMask[Frequency].TextureMask |= (1 << Arg.index);
+				ResourceMask[Frequency].TextureMask |= (FMetalTextureMask(1) << Arg.index);
 				TextureTypes[Frequency].Add(Arg.index, (uint8)Arg.textureType);
 				break;
 			}
