Manually merge CL 10504521: Remove manual command buffer internal resouce tracking and debug group association.

#rnx

diff --git a/Engine/Source/Runtime/Apple/MetalRHI/Private/MetalBuffer.cpp b/Engine/Source/Runtime/Apple/MetalRHI/Private/MetalBuffer.cpp
--- a/Engine/Source/Runtime/Apple/MetalRHI/Private/MetalBuffer.cpp
+++ b/Engine/Source/Runtime/Apple/MetalRHI/Private/MetalBuffer.cpp
@@ -1215,18 +1215,6 @@ void FMetalBufferPoolPolicyData::FreeResource(FMetalBuffer& Resource)
 {
 	DEC_MEMORY_STAT_BY(STAT_MetalBufferUnusedMemory, Resource.GetLength());
 	DEC_MEMORY_STAT_BY(STAT_MetalPooledBufferUnusedMemory, Resource.GetLength());
-#if METAL_DEBUG_OPTIONS // Helps to track down incorrect resource retain/release behaviour
-	if([Resource.GetPtr() retainCount] > 1)
-	{
-		UE_LOG(LogMetal, Warning, TEXT("Attempting to free an over-retained Buffer: %p: %s"), Resource.GetPtr(), *FString([Resource.GetPtr() debugDescription]));
-		void* Ptr = (void*)Resource.GetPtr();
-		objc_setAssociatedObject(Resource.GetPtr(), (void*)FMetalBufferPoolPolicyData::BucketSizes,
-			[[[FMetalDeallocHandler alloc] initWithBlock:^{
-			UE_LOG(LogMetal, Warning, TEXT("Released over-retained Buffer: %p"), Ptr);
-		}] autorelease],
-		OBJC_ASSOCIATION_RETAIN);
-	}
-#endif
 	Resource = nil;
 }
 
diff --git a/Engine/Source/Runtime/Apple/MetalRHI/Private/MetalCommandBuffer.cpp b/Engine/Source/Runtime/Apple/MetalRHI/Private/MetalCommandBuffer.cpp
--- a/Engine/Source/Runtime/Apple/MetalRHI/Private/MetalCommandBuffer.cpp
+++ b/Engine/Source/Runtime/Apple/MetalRHI/Private/MetalCommandBuffer.cpp
@@ -33,156 +33,6 @@ uint32 SafeGetRuntimeDebuggingLevel()
 	return GIsRHIInitialized ? GetMetalDeviceContext().GetCommandQueue().GetRuntimeDebuggingLevel() : GMetalRuntimeDebugLevel;
 }
 
-@implementation FMetalResourceTracker
-
-APPLE_PLATFORM_OBJECT_ALLOC_OVERRIDES(FMetalResourceTracker)
-
--(instancetype)init
-{
-	id Self = [super init];
-	return Self;
-}
-
--(void)dealloc
-{
-	for (id Entry : Resources)
-	{
-		FMetalResourceTrackCount* TrackCount = ((NSObject*)Entry).resourceTrackCount;
-		FPlatformAtomics::InterlockedDecrement((int64*)&TrackCount->RetainCount);
-		
-		[Entry release];
-	}
-	Resources.Empty();
-	[super dealloc];
-}
-
-@end
-
-@implementation FMetalResourceTrackCount
-
-APPLE_PLATFORM_OBJECT_ALLOC_OVERRIDES(FMetalResourceTrackCount)
-
--(instancetype)init
-{
-	id Self = [super init];
-	return Self;
-}
-
--(void)dealloc
-{
-	[super dealloc];
-}
-
-@end
-
-@implementation NSObject (IMetalDebugGroupAssociation)
-@dynamic debugGroups;
-- (void)setDebugGroups:(NSMutableArray<NSString*>*)Data
-{
-	objc_setAssociatedObject(self, @selector(debugGroups), Data, OBJC_ASSOCIATION_RETAIN_NONATOMIC);
-}
-
-- (NSMutableArray<NSString*>*)debugGroups
-{
-	return (NSMutableArray<NSString*>*)objc_getAssociatedObject(self, @selector(debugGroups));
-}
-@dynamic resourceTracker;
-- (void)setResourceTracker:(FMetalResourceTracker*)resourceTracker
-{
-	objc_setAssociatedObject(self, @selector(resourceTracker), resourceTracker, OBJC_ASSOCIATION_RETAIN_NONATOMIC);
-}
-
-- (FMetalResourceTracker*)resourceTracker
-{
-	return (FMetalResourceTracker*)objc_getAssociatedObject(self, @selector(resourceTracker));
-}
-@dynamic resourceTrackCount;
-- (void)setResourceTrackCount:(FMetalResourceTrackCount*)resourceTracker
-{
-	// OBJC_ASSOCIATION_RETAIN_ATOMIC -> 1
-	objc_setAssociatedObject(self, @selector(resourceTrackCount), resourceTracker, (objc_AssociationPolicy)1);
-}
-
-- (FMetalResourceTrackCount*)resourceTrackCount
-{
-	return (FMetalResourceTrackCount*)objc_getAssociatedObject(self, @selector(resourceTrackCount));
-}
-@end
-
-namespace FMetalCommandBufferDebugHelpers
-{
-	void TrackResource(id<MTLCommandBuffer> Buffer, id Ptr)
-	{
-		if (SafeGetRuntimeDebuggingLevel() >= EMetalDebugLevelTrackResources)
-		{
-			FMetalResourceTracker* ResourceTracker = ((NSObject<MTLCommandBuffer>*)Buffer).resourceTracker;
-			if (ResourceTracker)
-			{
-#if PLATFORM_MAC
-				if (GRHISupportsParallelRHIExecute)
-				{
-					FScopeLock Lock(&ResourceTracker->Mutex);
-					if (!ResourceTracker->Resources.Contains(Ptr))
-					{
-						ResourceTracker->Resources.Add(Ptr);
-						[Ptr retain];
-						
-						FMetalResourceTrackCount* TrackCount = ((NSObject*)Ptr).resourceTrackCount;
-						if (!TrackCount)
-						{
-							((NSObject*)Ptr).resourceTrackCount = [[FMetalResourceTrackCount new] autorelease];;
-							TrackCount = ((NSObject*)Ptr).resourceTrackCount;
-						}
-						FPlatformAtomics::InterlockedIncrement((int64*)&TrackCount->RetainCount);
-					}
-				}
-				else
-#endif
-				{
-					if (!ResourceTracker->Resources.Contains(Ptr))
-					{
-						ResourceTracker->Resources.Add(Ptr);
-						[Ptr retain];
-						
-						FMetalResourceTrackCount* TrackCount = ((NSObject*)Ptr).resourceTrackCount;
-						if (!TrackCount)
-						{
-							((NSObject*)Ptr).resourceTrackCount = [[FMetalResourceTrackCount new] autorelease];
-							TrackCount = ((NSObject*)Ptr).resourceTrackCount;
-						}
-						FPlatformAtomics::InterlockedIncrement((int64*)&TrackCount->RetainCount);
-					}
-				}
-			}
-		}
-	}
-	
-	void DumpResources(id<MTLCommandBuffer> Buffer)
-	{
-		if (SafeGetRuntimeDebuggingLevel() >= EMetalDebugLevelTrackResources)
-		{
-			FMetalResourceTracker* ResourceTracker = ((NSObject<MTLCommandBuffer>*)Buffer).resourceTracker;
-			if (ResourceTracker)
-			{
-				UE_LOG(LogMetal, Error, TEXT("Resources:"));
-				for (auto& Entry : ResourceTracker->Resources)
-				{
-					FMetalResourceTrackCount* TrackCount = ((NSObject*)Entry).resourceTrackCount;
-
-					if ([Entry retainCount] <= TrackCount->RetainCount)
-					{
-						UE_LOG(LogMetal, Error, TEXT("**** OVER-RELEASED RESOURCE **** : %s"), *FString([Entry debugDescription]));
-					}
-					else
-					{
-						UE_LOG(LogMetal, Error, TEXT("\t%s"), *FString([Entry debugDescription]));
-					}
-				}
-			}
-		}
-    }
-}
-
 #if MTLPP_CONFIG_VALIDATE && METAL_DEBUG_OPTIONS
 
 @implementation FMetalDebugCommandBuffer
diff --git a/Engine/Source/Runtime/Apple/MetalRHI/Private/MetalCommandBuffer.h b/Engine/Source/Runtime/Apple/MetalRHI/Private/MetalCommandBuffer.h
--- a/Engine/Source/Runtime/Apple/MetalRHI/Private/MetalCommandBuffer.h
+++ b/Engine/Source/Runtime/Apple/MetalRHI/Private/MetalCommandBuffer.h
@@ -31,9 +31,7 @@ enum EMetalDebugCommandType
 enum EMetalDebugLevel
 {
 	EMetalDebugLevelOff,
-	EMetalDebugLevelLogDebugGroups,
 	EMetalDebugLevelFastValidation,
-	EMetalDebugLevelTrackResources,
 	EMetalDebugLevelResetOnBind,
 	EMetalDebugLevelConditionalSubmit,
 	EMetalDebugLevelValidation,
@@ -52,38 +50,6 @@ struct FMetalDebugCommand
 	MTLRenderPassDescriptor* PassDesc;
 };
 
-@interface FMetalResourceTracker : FApplePlatformObject
-{
-@public
-	TSet<id> Resources;
-	FCriticalSection Mutex;
-};
-@end
-
-@interface FMetalResourceTrackCount : FApplePlatformObject
-{
-@public
-	NSUInteger RetainCount;
-};
-@end
-
-/**
- * Simpler NSObject extension that provides for an associated object to track debug groups in a command-buffer.
- * This doesn't interfere with objc_msgSend invocation so doesn't cost as much on the CPU.
- */
-@interface NSObject (IMetalDebugGroupAssociation)
-@property (nonatomic, strong) NSMutableArray<NSString*>* debugGroups;
-@property (nonatomic, strong) FMetalResourceTracker* resourceTracker;
-@property (atomic, strong) FMetalResourceTrackCount* resourceTrackCount;
-@end
-
-
-namespace FMetalCommandBufferDebugHelpers
-{
-	void TrackResource(id<MTLCommandBuffer> Buffer, id Ptr);
-	void DumpResources(id<MTLCommandBuffer> Buffer);
-}
-
 #if MTLPP_CONFIG_VALIDATE && METAL_DEBUG_OPTIONS
 
 /**
diff --git a/Engine/Source/Runtime/Apple/MetalRHI/Private/MetalCommandEncoder.cpp b/Engine/Source/Runtime/Apple/MetalRHI/Private/MetalCommandEncoder.cpp
--- a/Engine/Source/Runtime/Apple/MetalRHI/Private/MetalCommandEncoder.cpp
+++ b/Engine/Source/Runtime/Apple/MetalRHI/Private/MetalCommandEncoder.cpp
@@ -570,10 +570,6 @@ void FMetalCommandEncoder::BeginParallelRenderCommandEncoding(uint32 NumChildren
 		{
 			for (NSString* Group in DebugGroups)
 			{
-				if (CommandList.GetCommandQueue().GetRuntimeDebuggingLevel() >= EMetalDebugLevelLogDebugGroups)
-				{
-					[((NSObject<MTLCommandBuffer>*)CommandBuffer.GetPtr()).debugGroups addObject:Group];
-				}
 				ParallelRenderCommandEncoder.PushDebugGroup(Group);
 				METAL_DEBUG_LAYER(EMetalDebugLevelFastValidation, RenderEncoderDebug.PushDebugGroup(Group));
 			}
@@ -621,17 +617,6 @@ void FMetalCommandEncoder::BeginRenderCommandEncoding(void)
 		{
 			for (NSString* Group in DebugGroups)
 			{
-				if (CommandList.GetCommandQueue().GetRuntimeDebuggingLevel() >= EMetalDebugLevelLogDebugGroups)
-				{
-					if (!IsParallel())
-					{
-						[((NSObject<MTLCommandBuffer>*)CommandBuffer.GetPtr()).debugGroups addObject:Group];
-					}
-					else if (RenderCommandEncoder)
-					{
-						[((NSObject<MTLRenderCommandEncoder>*)RenderCommandEncoder.GetPtr()).debugGroups addObject:Group];
-					}
-				}
 				RenderCommandEncoder.PushDebugGroup(Group);
 				METAL_DEBUG_LAYER(EMetalDebugLevelFastValidation, RenderEncoderDebug.PushDebugGroup(Group));
 			}
@@ -677,10 +662,6 @@ void FMetalCommandEncoder::BeginComputeCommandEncoding(mtlpp::DispatchType Dispa
 		{
 			for (NSString* Group in DebugGroups)
 			{
-				if (CommandList.GetCommandQueue().GetRuntimeDebuggingLevel() >= EMetalDebugLevelLogDebugGroups)
-				{
-					[((NSObject<MTLCommandBuffer>*)CommandBuffer.GetPtr()).debugGroups addObject:Group];
-				}
 				ComputeCommandEncoder.PushDebugGroup(Group);
 				METAL_DEBUG_LAYER(EMetalDebugLevelFastValidation, ComputeEncoderDebug.PushDebugGroup(Group));
 			}
@@ -715,10 +696,6 @@ void FMetalCommandEncoder::BeginBlitCommandEncoding(void)
 		{
 			for (NSString* Group in DebugGroups)
 			{
-				if (CommandList.GetCommandQueue().GetRuntimeDebuggingLevel() >= EMetalDebugLevelLogDebugGroups)
-				{
-					[((NSObject<MTLCommandBuffer>*)CommandBuffer.GetPtr()).debugGroups addObject:Group];
-				}
 				BlitCommandEncoder.PushDebugGroup(Group);
 				METAL_DEBUG_LAYER(EMetalDebugLevelFastValidation, BlitEncoderDebug.PushDebugGroup(Group));
 			}
@@ -873,11 +850,6 @@ TRefCountPtr<FMetalFence> FMetalCommandEncoder::EndEncoding(void)
 					}
 				}
 
-				for (mtlpp::RenderCommandEncoder& Encoder : ChildRenderCommandEncoders)
-				{
-					[((NSObject<MTLRenderCommandEncoder>*)Encoder.GetPtr()).debugGroups addObjectsFromArray:((NSObject<MTLCommandBuffer>*)CommandBuffer.GetPtr()).debugGroups];
-				}
-
 //				METAL_STATISTIC(FMetalProfiler::GetProfiler()->EndEncoder(CommandBufferStats, ParallelRenderCommandEncoder));
 				ParallelRenderCommandEncoder.EndEncoding();
 				METAL_DEBUG_LAYER(EMetalDebugLevelFastValidation, ParallelEncoderDebug.EndEncoder());
@@ -1239,17 +1211,6 @@ void FMetalCommandEncoder::InsertDebugSignpost(ns::String const& String)
 {
 	if (String)
 	{
-		if (CommandBuffer && CommandList.GetCommandQueue().GetRuntimeDebuggingLevel() >= EMetalDebugLevelLogDebugGroups)
-		{
-			if (!IsParallel())
-			{
-				[((NSObject<MTLCommandBuffer>*)CommandBuffer.GetPtr()).debugGroups addObject:String];
-			}
-			else if (RenderCommandEncoder)
-			{
-				[((NSObject<MTLRenderCommandEncoder>*)RenderCommandEncoder.GetPtr()).debugGroups addObject:String];
-			}
-		}
 		if (RenderCommandEncoder)
 		{
 			RenderCommandEncoder.InsertDebugSignpost(String);
@@ -1277,17 +1238,6 @@ void FMetalCommandEncoder::PushDebugGroup(ns::String const& String)
 {
 	if (String)
 	{
-		if (CommandBuffer && CommandList.GetCommandQueue().GetRuntimeDebuggingLevel() >= EMetalDebugLevelLogDebugGroups)
-		{
-			if (!IsParallel())
-			{
-				[((NSObject<MTLCommandBuffer>*)CommandBuffer.GetPtr()).debugGroups addObject:String];
-			}
-			else if (RenderCommandEncoder)
-			{
-				[((NSObject<MTLRenderCommandEncoder>*)RenderCommandEncoder.GetPtr()).debugGroups addObject:String];
-			}
-		}
 		[DebugGroups addObject:String];
 		if (RenderCommandEncoder)
 		{
@@ -1531,7 +1481,6 @@ void FMetalCommandEncoder::SetShaderBuffer(mtlpp::FunctionType const FunctionTyp
 		if(Buffer)
 		{
 			ShaderBuffers[uint32(FunctionType)].Bound |= (1 << index);
-			FMetalCommandBufferDebugHelpers::TrackResource(CommandBuffer.GetPtr(), Buffer.GetPtr());
 		}
 		else
 		{
@@ -1615,7 +1564,6 @@ void FMetalCommandEncoder::SetShaderBytes(mtlpp::FunctionType const FunctionType
 			FMetalBuffer Buffer = RingBuffer.NewBuffer(Length, BufferOffsetAlignment);
 			FMemory::Memcpy(((uint8*)Buffer.GetContents()), Bytes, Length);
 			ShaderBuffers[uint32(FunctionType)].Buffers[Index] = Buffer;
-			FMetalCommandBufferDebugHelpers::TrackResource(CommandBuffer.GetPtr(), Buffer.GetPtr());
 		}
 		ShaderBuffers[uint32(FunctionType)].Bytes[Index] = nil;
 		ShaderBuffers[uint32(FunctionType)].Offsets[Index] = 0;
@@ -1701,7 +1649,6 @@ void FMetalCommandEncoder::SetShaderTexture(mtlpp::FunctionType FunctionType, FM
 	
 	if (Texture)
 	{
-		FMetalCommandBufferDebugHelpers::TrackResource(CommandBuffer.GetPtr(), Texture.GetPtr());
 		uint8 Swizzle[4] = {0,0,0,0};
 		assert(sizeof(Swizzle) == sizeof(uint32));
 		if (Texture.GetPixelFormat() == mtlpp::PixelFormat::X32_Stencil8
@@ -1742,11 +1689,6 @@ void FMetalCommandEncoder::SetShaderSamplerState(mtlpp::FunctionType FunctionTyp
 			check(false);
 			break;
 	}
-	
-	if (Sampler)
-	{
-		FMetalCommandBufferDebugHelpers::TrackResource(CommandBuffer.GetPtr(), Sampler.GetPtr());
-	}
 }
 
 void FMetalCommandEncoder::SetShaderSideTable(mtlpp::FunctionType const FunctionType, NSUInteger const Index)
diff --git a/Engine/Source/Runtime/Apple/MetalRHI/Private/MetalCommandList.cpp b/Engine/Source/Runtime/Apple/MetalRHI/Private/MetalCommandList.cpp
--- a/Engine/Source/Runtime/Apple/MetalRHI/Private/MetalCommandList.cpp
+++ b/Engine/Source/Runtime/Apple/MetalRHI/Private/MetalCommandList.cpp
@@ -52,23 +52,8 @@ static void ReportMetalCommandBufferFailure(mtlpp::CommandBuffer const& Complete
 	FString FailureString = FailureDesc ? FString(FailureDesc) : FString(TEXT("Unknown"));
 	FString RecoveryString = RecoveryDesc ? FString(RecoveryDesc) : FString(TEXT("Unknown"));
 	
-	if (GetMetalDeviceContext().GetCommandQueue().GetRuntimeDebuggingLevel() >= EMetalDebugLevelLogDebugGroups)
-	{
-		NSMutableString* DescString = [NSMutableString new];
-		[DescString appendFormat:@"Command Buffer %p %@:", CompletedBuffer.GetPtr(), Label ? Label : @"Unknown"];
-
-		for (NSString* String in ((NSObject<MTLCommandBuffer>*)CompletedBuffer.GetPtr()).debugGroups)
-		{
-			[DescString appendFormat:@"\n\tDebugGroup: %@", String];
-		}
-		
-		UE_LOG(LogMetal, Warning, TEXT("Command Buffer %p %s:%s"), CompletedBuffer.GetPtr(), *LabelString, *FString(DescString));
-	}
-	else
-	{
-		NSString* Desc = CompletedBuffer.GetPtr().debugDescription;
-		UE_LOG(LogMetal, Warning, TEXT("%s"), *FString(Desc));
-	}
+	NSString* Desc = CompletedBuffer.GetPtr().debugDescription;
+	UE_LOG(LogMetal, Warning, TEXT("%s"), *FString(Desc));
 	
 #if PLATFORM_IOS
     if (bDoCheck && !GIsSuspended && !GIsRenderingThreadSuspended)
@@ -182,8 +167,6 @@ static void ReportMetalCommandBufferFailure(mtlpp::CommandBuffer const& Complete
                 }
 			}
 		}
-
-        FMetalCommandBufferDebugHelpers::DumpResources(CompletedBuffer.GetPtr());
 		
 #if PLATFORM_IOS
         UE_LOG(LogMetal, Warning, TEXT("Command Buffer %s Failed with %s Error! Error Domain: %s Code: %d Description %s %s %s"), *LabelString, ErrorType, *DomainString, Code, *ErrorString, *FailureString, *RecoveryString);
diff --git a/Engine/Source/Runtime/Apple/MetalRHI/Private/MetalCommandQueue.cpp b/Engine/Source/Runtime/Apple/MetalRHI/Private/MetalCommandQueue.cpp
--- a/Engine/Source/Runtime/Apple/MetalRHI/Private/MetalCommandQueue.cpp
+++ b/Engine/Source/Runtime/Apple/MetalRHI/Private/MetalCommandQueue.cpp
@@ -391,16 +391,11 @@ mtlpp::CommandBuffer FMetalCommandQueue::CreateCommandBuffer(void)
 	{
 		CmdBuffer = bUnretainedRefs ? MTLPP_VALIDATE(mtlpp::CommandQueue, CommandQueue, SafeGetRuntimeDebuggingLevel() >= EMetalDebugLevelValidation, CommandBufferWithUnretainedReferences()) : MTLPP_VALIDATE(mtlpp::CommandQueue, CommandQueue, SafeGetRuntimeDebuggingLevel() >= EMetalDebugLevelValidation, CommandBuffer());
 		
-		if (RuntimeDebuggingLevel >= EMetalDebugLevelLogDebugGroups)
+		if (RuntimeDebuggingLevel > EMetalDebugLevelOff)
 		{			
-			((NSObject<MTLCommandBuffer>*)CmdBuffer.GetPtr()).debugGroups = [[NSMutableArray new] autorelease];
 			METAL_DEBUG_ONLY(FMetalCommandBufferDebugging AddDebugging(CmdBuffer));
 			MTLPP_VALIDATION(mtlpp::CommandBufferValidationTable ValidatedCommandBuffer(CmdBuffer));
 		}
-		if (RuntimeDebuggingLevel >= EMetalDebugLevelTrackResources)
-		{
-			((NSObject<MTLCommandBuffer>*)CmdBuffer.GetPtr()).resourceTracker = [[FMetalResourceTracker new] autorelease];
-		}
 	}
 	CommandBufferFences.Push(new mtlpp::CommandBufferFence(CmdBuffer.GetCompletionFence()));
 	INC_DWORD_STAT(STAT_MetalCommandBufferCreatedPerFrame);
diff --git a/Engine/Source/Runtime/Apple/MetalRHI/Private/MetalContext.cpp b/Engine/Source/Runtime/Apple/MetalRHI/Private/MetalContext.cpp
--- a/Engine/Source/Runtime/Apple/MetalRHI/Private/MetalContext.cpp
+++ b/Engine/Source/Runtime/Apple/MetalRHI/Private/MetalContext.cpp
@@ -76,27 +76,21 @@ static FAutoConsoleVariableRef CVarMetalResourceDeferDeleteNumFrames(
 #if UE_BUILD_SHIPPING
 int32 GMetalRuntimeDebugLevel = 0;
 #else
-#if PLATFORM_MAC
-int32 GMetalRuntimeDebugLevel = 3;
-#else
 int32 GMetalRuntimeDebugLevel = 1;
 #endif
-#endif
 static FAutoConsoleVariableRef CVarMetalRuntimeDebugLevel(
 	TEXT("rhi.Metal.RuntimeDebugLevel"),
 	GMetalRuntimeDebugLevel,
 	TEXT("The level of debug validation performed by MetalRHI in addition to the underlying Metal API & validation layer.\n")
 	TEXT("Each subsequent level adds more tests and reporting in addition to the previous level.\n")
-	TEXT("*LEVELS >= 5 ARE IGNORED IN SHIPPING AND TEST BUILDS*. (Default: 3 (Debug, Development), 0 (Test, Shipping))\n")
+	TEXT("*LEVELS >= 3 ARE IGNORED IN SHIPPING AND TEST BUILDS*. (Default: 1 (Debug, Development), 0 (Test, Shipping))\n")
 	TEXT("\t0: Off,\n")
-	TEXT("\t1: Record the debug-groups issued into a command-buffer and report them on failure,\n")
-	TEXT("\t2: Enable light-weight validation of resource bindings & API usage,\n")
-	TEXT("\t3: Track resources and validate lifetime on command-buffer failure,\n")
-	TEXT("\t4: Reset resource bindings when binding a PSO/Compute-Shader to simplify GPU debugging,\n")
-	TEXT("\t5: Allow rhi.Metal.CommandBufferCommitThreshold to break command-encoders (except when MSAA is enabled),\n")
-	TEXT("\t6: Enable slower, more extensive validation checks for resource types & encoder usage,\n")
-    TEXT("\t7: Record the draw, blit & dispatch commands issued into a command-buffer and report them on failure,\n")
-    TEXT("\t8: Wait for each command-buffer to complete immediately after submission."));
+	TEXT("\t1: Enable light-weight validation of resource bindings & API usage,\n")
+	TEXT("\t2: Reset resource bindings when binding a PSO/Compute-Shader to simplify GPU debugging,\n")
+	TEXT("\t3: Allow rhi.Metal.CommandBufferCommitThreshold to break command-encoders (except when MSAA is enabled),\n")
+	TEXT("\t4: Enable slower, more extensive validation checks for resource types & encoder usage,\n")
+    TEXT("\t5: Record the draw, blit & dispatch commands issued into a command-buffer and report them on failure,\n")
+    TEXT("\t6: Wait for each command-buffer to complete immediately after submission."));
 
 float GMetalPresentFramePacing = 0.0f;
 #if !PLATFORM_MAC
diff --git a/Engine/Source/Runtime/Apple/MetalRHI/Private/MetalRenderPass.cpp b/Engine/Source/Runtime/Apple/MetalRHI/Private/MetalRenderPass.cpp
--- a/Engine/Source/Runtime/Apple/MetalRHI/Private/MetalRenderPass.cpp
+++ b/Engine/Source/Runtime/Apple/MetalRHI/Private/MetalRenderPass.cpp
@@ -454,8 +454,6 @@ void FMetalRenderPass::DrawPrimitiveIndirect(uint32 PrimitiveType, FMetalVertexB
 		
 		PrepareToRender(PrimitiveType);
 		
-		FMetalCommandBufferDebugHelpers::TrackResource(CurrentEncoder.GetCommandBuffer().GetPtr(), VertexBuffer->Buffer.GetPtr());
-		
 		METAL_GPUPROFILE(FMetalProfiler::GetProfiler()->EncodeDraw(CurrentEncoder.GetCommandBufferStats(), __FUNCTION__, 1, 1, 1));
 		CurrentEncoder.GetRenderCommandEncoder().Draw(TranslatePrimitiveType(PrimitiveType), VertexBuffer->Buffer, ArgumentOffset);
 		METAL_DEBUG_LAYER(EMetalDebugLevelFastValidation, CurrentEncoder.GetRenderCommandEncoderDebugging().Draw(TranslatePrimitiveType(PrimitiveType), VertexBuffer->Buffer, ArgumentOffset));
@@ -538,8 +536,6 @@ void FMetalRenderPass::DrawIndexedPrimitive(FMetalBuffer const& IndexBuffer, uin
 		
 		uint32 NumIndices = GetVertexCountForPrimitiveCount(NumPrimitives, PrimitiveType);
 		
-		FMetalCommandBufferDebugHelpers::TrackResource(CurrentEncoder.GetCommandBuffer().GetPtr(), IndexBuffer.GetPtr());
-		
 		METAL_GPUPROFILE(FMetalProfiler::GetProfiler()->EncodeDraw(CurrentEncoder.GetCommandBufferStats(), __FUNCTION__, NumPrimitives, NumVertices, NumInstances));
 		if (GRHISupportsBaseVertexIndex && GRHISupportsFirstInstance)
 		{
@@ -590,9 +586,6 @@ void FMetalRenderPass::DrawIndexedIndirect(FMetalIndexBuffer* IndexBuffer, uint3
 		// finalize any pending state
 		PrepareToRender(PrimitiveType);
 		
-		FMetalCommandBufferDebugHelpers::TrackResource(CurrentEncoder.GetCommandBuffer().GetPtr(), IndexBuffer->Buffer.GetPtr());
-		FMetalCommandBufferDebugHelpers::TrackResource(CurrentEncoder.GetCommandBuffer().GetPtr(), VertexBuffer->Buffer.GetPtr());
-		
 		METAL_GPUPROFILE(FMetalProfiler::GetProfiler()->EncodeDraw(CurrentEncoder.GetCommandBufferStats(), __FUNCTION__, 1, 1, 1));
 		CurrentEncoder.GetRenderCommandEncoder().DrawIndexed(TranslatePrimitiveType(PrimitiveType), (mtlpp::IndexType)IndexBuffer->IndexType, IndexBuffer->Buffer, 0, VertexBuffer->Buffer, (DrawArgumentsIndex * 5 * sizeof(uint32)));
 		METAL_DEBUG_LAYER(EMetalDebugLevelFastValidation, CurrentEncoder.GetRenderCommandEncoderDebugging().DrawIndexed(TranslatePrimitiveType(PrimitiveType), (mtlpp::IndexType)IndexBuffer->IndexType, IndexBuffer->Buffer, 0, VertexBuffer->Buffer, (DrawArgumentsIndex * 5 * sizeof(uint32))));
@@ -624,9 +617,6 @@ void FMetalRenderPass::DrawIndexedPrimitiveIndirect(uint32 PrimitiveType,FMetalI
 		
 		PrepareToRender(PrimitiveType);
 		
-		FMetalCommandBufferDebugHelpers::TrackResource(CurrentEncoder.GetCommandBuffer().GetPtr(), IndexBuffer->Buffer.GetPtr());
-		FMetalCommandBufferDebugHelpers::TrackResource(CurrentEncoder.GetCommandBuffer().GetPtr(), VertexBuffer->Buffer.GetPtr());
-		
 		METAL_GPUPROFILE(FMetalProfiler::GetProfiler()->EncodeDraw(CurrentEncoder.GetCommandBufferStats(), __FUNCTION__, 1, 1, 1));
 		CurrentEncoder.GetRenderCommandEncoder().DrawIndexed(TranslatePrimitiveType(PrimitiveType), (mtlpp::IndexType)IndexBuffer->IndexType, IndexBuffer->Buffer, 0, VertexBuffer->Buffer, ArgumentOffset);
 		METAL_DEBUG_LAYER(EMetalDebugLevelFastValidation, CurrentEncoder.GetRenderCommandEncoderDebugging().DrawIndexed(TranslatePrimitiveType(PrimitiveType), (mtlpp::IndexType)IndexBuffer->IndexType, IndexBuffer->Buffer, 0, VertexBuffer->Buffer, ArgumentOffset));
@@ -681,7 +671,6 @@ void FMetalRenderPass::DrawPatches(uint32 PrimitiveType,FMetalBuffer const& Inde
 				uint32 OutputBufferSize = (NumIndices * NumInstances) * boundShaderState->VertexShader->TessellationOutputAttribs.HSOutSize;
 				
 				ShaderOutputBuffer = deviceContext.CreatePooledBuffer(FMetalPooledBufferArgs(device, OutputBufferSize, BUF_Dynamic, mtlpp::StorageMode::Private));
-				FMetalCommandBufferDebugHelpers::TrackResource(CurrentEncoder.GetCommandBuffer().GetPtr(), ShaderOutputBuffer.GetPtr());
 				
 				PrologueEncoder.SetShaderBuffer(mtlpp::FunctionType::Vertex, ShaderOutputBuffer, 0, ShaderOutputBuffer, boundShaderState->VertexShader->TessellationHSOutBuffer, mtlpp::ResourceUsage::Write);
 				State.SetShaderBuffer(EMetalShaderStages::Vertex, nil, nil, 0, 0, boundShaderState->VertexShader->TessellationHSOutBuffer, mtlpp::ResourceUsage(0));
@@ -710,7 +699,6 @@ void FMetalRenderPass::DrawPatches(uint32 PrimitiveType,FMetalBuffer const& Inde
 				uint32 OutputBufferSize = (NumVertices * NumInstances) * boundShaderState->VertexShader->TessellationOutputAttribs.HSOutSize;
 				
 				ShaderOutputBuffer = deviceContext.CreatePooledBuffer(FMetalPooledBufferArgs(device, OutputBufferSize, BUF_Dynamic, mtlpp::StorageMode::Private));
-				FMetalCommandBufferDebugHelpers::TrackResource(CurrentEncoder.GetCommandBuffer().GetPtr(), ShaderOutputBuffer.GetPtr());
 
 				PrologueEncoder.SetShaderBuffer(mtlpp::FunctionType::Vertex, ShaderOutputBuffer, 0, ShaderOutputBuffer, boundShaderState->VertexShader->TessellationHSOutBuffer, mtlpp::ResourceUsage::Write);
 				State.SetShaderBuffer(EMetalShaderStages::Vertex, nil, nil, 0, 0, boundShaderState->VertexShader->TessellationHSOutBuffer, mtlpp::ResourceUsage(0));
@@ -753,7 +741,6 @@ void FMetalRenderPass::DrawPatches(uint32 PrimitiveType,FMetalBuffer const& Inde
 			if (IndexBuffer)
 			{
 				HullIndexBuffer = deviceContext.CreatePooledBuffer(FMetalPooledBufferArgs(device, NumInstances * PatchCount * FMath::Max(OutputCPs, InputCPs) * sizeof(uint32), BUF_Static, mtlpp::StorageMode::Private));
-				FMetalCommandBufferDebugHelpers::TrackResource(CurrentEncoder.GetCommandBuffer().GetPtr(), HullIndexBuffer.GetPtr());
 				if (IndexBufferStride == 2)
 				{
 					computeEncoder.SetComputePipelineState(GetMetalCopyIndex16Function());
@@ -793,7 +780,6 @@ void FMetalRenderPass::DrawPatches(uint32 PrimitiveType,FMetalBuffer const& Inde
 			else
 			{
 				HullIndexBuffer = deviceContext.CreatePooledBuffer(FMetalPooledBufferArgs(device, NumInstances * PatchCount * FMath::Max(OutputCPs, InputCPs) * sizeof(uint32), BUF_Static, mtlpp::StorageMode::Shared));
-				FMetalCommandBufferDebugHelpers::TrackResource(CurrentEncoder.GetCommandBuffer().GetPtr(), HullIndexBuffer.GetPtr());
 				uint32* Data = (uint32*)HullIndexBuffer.GetContents();
 				uint32 Index = 0;
 				if (OutputCPs > InputCPs)
@@ -848,7 +834,6 @@ void FMetalRenderPass::DrawPatches(uint32 PrimitiveType,FMetalBuffer const& Inde
 			if (boundShaderState->HullShader->TessellationHSOutBuffer != UINT_MAX && boundShaderState->HullShader->TessellationOutputAttribs.HSOutSize)
 			{
 				HSOutBuffer = deviceContext.CreatePooledBuffer(FMetalPooledBufferArgs(device, Align(NumInstances * PatchCount * boundShaderState->HullShader->TessellationOutputAttribs.HSOutSize, BufferOffsetAlignment), BUF_Static, mtlpp::StorageMode::Private));
-				FMetalCommandBufferDebugHelpers::TrackResource(CurrentEncoder.GetCommandBuffer().GetPtr(), HSOutBuffer.GetPtr());
 				
 				PrologueEncoder.SetShaderBuffer(mtlpp::FunctionType::Kernel, HSOutBuffer, 0, HSOutBuffer, boundShaderState->HullShader->TessellationHSOutBuffer, mtlpp::ResourceUsage::Write);
 				State.SetShaderBuffer(EMetalShaderStages::Compute, nil, nil, 0, 0, boundShaderState->HullShader->TessellationHSOutBuffer, mtlpp::ResourceUsage(0));
@@ -858,7 +843,6 @@ void FMetalRenderPass::DrawPatches(uint32 PrimitiveType,FMetalBuffer const& Inde
 			if (boundShaderState->HullShader->TessellationControlPointOutBuffer != UINT_MAX && boundShaderState->HullShader->TessellationOutputAttribs.PatchControlPointOutSize)
 			{
 				CPOutBuffer = deviceContext.CreatePooledBuffer(FMetalPooledBufferArgs(device, NumInstances * PatchCount * OutputCPs * boundShaderState->HullShader->TessellationOutputAttribs.PatchControlPointOutSize, BUF_Static, mtlpp::StorageMode::Private));
-				FMetalCommandBufferDebugHelpers::TrackResource(CurrentEncoder.GetCommandBuffer().GetPtr(), CPOutBuffer.GetPtr());
 				
 				PrologueEncoder.SetShaderBuffer(mtlpp::FunctionType::Kernel, CPOutBuffer, 0, CPOutBuffer, boundShaderState->HullShader->TessellationControlPointOutBuffer, mtlpp::ResourceUsage::Write);
 				State.SetShaderBuffer(EMetalShaderStages::Compute, nil, nil, 0, 0, boundShaderState->HullShader->TessellationControlPointOutBuffer, mtlpp::ResourceUsage(0));
@@ -868,7 +852,6 @@ void FMetalRenderPass::DrawPatches(uint32 PrimitiveType,FMetalBuffer const& Inde
 			if (boundShaderState->HullShader->TessellationHSTFOutBuffer != UINT_MAX)
 			{
 				HSTFBuffer = deviceContext.CreatePooledBuffer(FMetalPooledBufferArgs(device, NumInstances * PatchCount * boundShaderState->HullShader->TessellationOutputAttribs.HSTFOutSize, BUF_Static, mtlpp::StorageMode::Private));
-				FMetalCommandBufferDebugHelpers::TrackResource(CurrentEncoder.GetCommandBuffer().GetPtr(), HSTFBuffer.GetPtr());
 				
 				PrologueEncoder.SetShaderBuffer(mtlpp::FunctionType::Kernel, HSTFBuffer, 0, HSTFBuffer, boundShaderState->HullShader->TessellationHSTFOutBuffer, mtlpp::ResourceUsage::Write);
 				State.SetShaderBuffer(EMetalShaderStages::Compute, nil, nil, 0, 0, boundShaderState->HullShader->TessellationHSTFOutBuffer, mtlpp::ResourceUsage(0));
@@ -1003,21 +986,18 @@ void FMetalRenderPass::DrawPatches(uint32 PrimitiveType,FMetalBuffer const& Inde
 		if(hullShaderOutputBufferSize)
 		{
 			hullShaderOutputBuffer = deviceContext.GetResourceHeap().CreateBuffer(hullShaderOutputBufferSize, 16, BUF_Dynamic, FMetalCommandQueue::GetCompatibleResourceOptions((mtlpp::ResourceOptions)( mtlpp::ResourceOptions::HazardTrackingModeUntracked | ((NSUInteger)mtlpp::StorageMode::Private << mtlpp::ResourceStorageModeShift))), true);
-			FMetalCommandBufferDebugHelpers::TrackResource(CurrentEncoder.GetCommandBuffer().GetPtr(), hullShaderOutputBuffer.GetPtr());
 		}
 		
 		FMetalBuffer hullConstShaderOutputBuffer = nil;
 		if(hullConstShaderOutputBufferSize)
 		{
 			hullConstShaderOutputBuffer = deviceContext.GetResourceHeap().CreateBuffer(hullConstShaderOutputBufferSize, 16, BUF_Dynamic, FMetalCommandQueue::GetCompatibleResourceOptions((mtlpp::ResourceOptions)( mtlpp::ResourceOptions::HazardTrackingModeUntracked | ((NSUInteger)mtlpp::StorageMode::Private << mtlpp::ResourceStorageModeShift))), true);
-			FMetalCommandBufferDebugHelpers::TrackResource(CurrentEncoder.GetCommandBuffer().GetPtr(), hullConstShaderOutputBuffer.GetPtr());
 		}
 		
 		FMetalBuffer tessellationFactorBuffer = nil;
 		if(tessellationFactorBufferSize)
 		{
 			tessellationFactorBuffer = deviceContext.GetResourceHeap().CreateBuffer(tessellationFactorBufferSize, 16, BUF_Dynamic, FMetalCommandQueue::GetCompatibleResourceOptions((mtlpp::ResourceOptions)( mtlpp::ResourceOptions::HazardTrackingModeUntracked | ((NSUInteger)mtlpp::StorageMode::Private << mtlpp::ResourceStorageModeShift))), true);
-			FMetalCommandBufferDebugHelpers::TrackResource(CurrentEncoder.GetCommandBuffer().GetPtr(), tessellationFactorBuffer.GetPtr());
 		}
 	
 		auto& computeEncoder = PrologueEncoder.GetComputeCommandEncoder();
@@ -1037,7 +1017,6 @@ void FMetalRenderPass::DrawPatches(uint32 PrimitiveType,FMetalBuffer const& Inde
 		{
 			if (IndexBuffer)
 			{
-                FMetalCommandBufferDebugHelpers::TrackResource(CurrentEncoder.GetCommandBuffer().GetPtr(), IndexBuffer.GetPtr());
 				PrologueEncoder.SetShaderBuffer(mtlpp::FunctionType::Kernel, IndexBuffer, StartIndex * IndexBufferStride, IndexBuffer.GetLength() - (StartIndex * IndexBufferStride), Pipeline->TessellationPipelineDesc.TessellationIndexBufferIndex, mtlpp::ResourceUsage::Read, IndexBufferStride == 2 ? PF_R16_UINT : PF_R32_UINT);
 		}
 			else
@@ -1202,17 +1181,17 @@ void FMetalRenderPass::Dispatch(uint32 ThreadGroupCountX, uint32 ThreadGroupCoun
 		CurrentEncoder.GetComputeCommandEncoder().DispatchThreadgroups(Threadgroups, ThreadgroupCounts);
 		METAL_DEBUG_LAYER(EMetalDebugLevelFastValidation, CurrentEncoder.GetComputeCommandEncoderDebugging().DispatchThreadgroups(Threadgroups, ThreadgroupCounts));
 		
-	if (GMetalCommandBufferDebuggingEnabled)
-	{
-		FMetalCommandData Data;
-		Data.CommandType = FMetalCommandData::Type::Dispatch;
-		Data.Dispatch.threadgroupsPerGrid[0] = ThreadGroupCountX;
-		Data.Dispatch.threadgroupsPerGrid[1] = ThreadGroupCountY;
-		Data.Dispatch.threadgroupsPerGrid[2] = ThreadGroupCountZ;
+		if (GMetalCommandBufferDebuggingEnabled)
+		{
+			FMetalCommandData Data;
+			Data.CommandType = FMetalCommandData::Type::Dispatch;
+			Data.Dispatch.threadgroupsPerGrid[0] = ThreadGroupCountX;
+			Data.Dispatch.threadgroupsPerGrid[1] = ThreadGroupCountY;
+			Data.Dispatch.threadgroupsPerGrid[2] = ThreadGroupCountZ;
+			
+			InsertDebugDispatch(Data);
+		}
 		
-		InsertDebugDispatch(Data);
-	}
-	
 		ConditionalSubmit();
 	}
 }
@@ -1256,21 +1235,19 @@ void FMetalRenderPass::DispatchIndirect(FMetalVertexBuffer* ArgumentBuffer, uint
 		METAL_GPUPROFILE(FMetalProfiler::GetProfiler()->EncodeDispatch(CurrentEncoder.GetCommandBufferStats(), __FUNCTION__));
 		mtlpp::Size ThreadgroupCounts = mtlpp::Size(ComputeShader->NumThreadsX, ComputeShader->NumThreadsY, ComputeShader->NumThreadsZ);
 		check(ComputeShader->NumThreadsX > 0 && ComputeShader->NumThreadsY > 0 && ComputeShader->NumThreadsZ > 0);
-		
-	FMetalCommandBufferDebugHelpers::TrackResource(CurrentEncoder.GetCommandBuffer().GetPtr(), ArgumentBuffer->Buffer.GetPtr());
-	
+
 		CurrentEncoder.GetComputeCommandEncoder().DispatchThreadgroupsWithIndirectBuffer(ArgumentBuffer->Buffer, ArgumentOffset, ThreadgroupCounts);
 		METAL_DEBUG_LAYER(EMetalDebugLevelFastValidation, CurrentEncoder.GetComputeCommandEncoderDebugging().DispatchThreadgroupsWithIndirectBuffer(ArgumentBuffer->Buffer, ArgumentOffset, ThreadgroupCounts));
 
-	if (GMetalCommandBufferDebuggingEnabled)
-	{
-		FMetalCommandData Data;
-		Data.CommandType = FMetalCommandData::Type::DispatchIndirect;
-		Data.DispatchIndirect.ArgumentBuffer = ArgumentBuffer->Buffer;
-		Data.DispatchIndirect.ArgumentOffset = ArgumentOffset;
-		
-		InsertDebugDispatch(Data);
-	}
+		if (GMetalCommandBufferDebuggingEnabled)
+		{
+			FMetalCommandData Data;
+			Data.CommandType = FMetalCommandData::Type::DispatchIndirect;
+			Data.DispatchIndirect.ArgumentBuffer = ArgumentBuffer->Buffer;
+			Data.DispatchIndirect.ArgumentOffset = ArgumentOffset;
+			
+			InsertDebugDispatch(Data);
+		}
 	
 		ConditionalSubmit();
 	}
@@ -1295,9 +1272,6 @@ void FMetalRenderPass::CopyFromTextureToBuffer(FMetalTexture const& Texture, uin
 	mtlpp::BlitCommandEncoder& Encoder = CurrentEncoder.GetBlitCommandEncoder();
 	check(Encoder.GetPtr());
 	
-	FMetalCommandBufferDebugHelpers::TrackResource(CurrentEncoder.GetCommandBuffer().GetPtr(), Texture.GetPtr());
-	FMetalCommandBufferDebugHelpers::TrackResource(CurrentEncoder.GetCommandBuffer().GetPtr(), toBuffer.GetPtr());
-	
 	METAL_GPUPROFILE(FMetalProfiler::GetProfiler()->EncodeBlit(CurrentEncoder.GetCommandBufferStats(), __FUNCTION__));
 	{
 		MTLPP_VALIDATE(mtlpp::BlitCommandEncoder, Encoder, SafeGetRuntimeDebuggingLevel() >= EMetalDebugLevelValidation, Copy(Texture, sourceSlice, sourceLevel, sourceOrigin, sourceSize, toBuffer, destinationOffset, destinationBytesPerRow, destinationBytesPerImage, options));
@@ -1313,9 +1287,6 @@ void FMetalRenderPass::CopyFromBufferToTexture(FMetalBuffer const& Buffer, uint3
 	check(Encoder.GetPtr());
 	
 	
-	FMetalCommandBufferDebugHelpers::TrackResource(CurrentEncoder.GetCommandBuffer().GetPtr(), Buffer.GetPtr());
-	FMetalCommandBufferDebugHelpers::TrackResource(CurrentEncoder.GetCommandBuffer().GetPtr(), toTexture.GetPtr());
-	
 	METAL_GPUPROFILE(FMetalProfiler::GetProfiler()->EncodeBlit(CurrentEncoder.GetCommandBufferStats(), __FUNCTION__));
 	if (options == mtlpp::BlitOption::None)
 	{
@@ -1336,9 +1307,6 @@ void FMetalRenderPass::CopyFromTextureToTexture(FMetalTexture const& Texture, ui
 	mtlpp::BlitCommandEncoder& Encoder = CurrentEncoder.GetBlitCommandEncoder();
 	check(Encoder.GetPtr());
 	
-	FMetalCommandBufferDebugHelpers::TrackResource(CurrentEncoder.GetCommandBuffer().GetPtr(), Texture.GetPtr());
-	FMetalCommandBufferDebugHelpers::TrackResource(CurrentEncoder.GetCommandBuffer().GetPtr(), toTexture.GetPtr());
-	
 	METAL_GPUPROFILE(FMetalProfiler::GetProfiler()->EncodeBlit(CurrentEncoder.GetCommandBufferStats(), __FUNCTION__));
 	MTLPP_VALIDATE(mtlpp::BlitCommandEncoder, Encoder, SafeGetRuntimeDebuggingLevel() >= EMetalDebugLevelValidation, Copy(Texture, sourceSlice, sourceLevel, sourceOrigin, sourceSize, toTexture, destinationSlice, destinationLevel, destinationOrigin));
 	METAL_DEBUG_LAYER(EMetalDebugLevelFastValidation, CurrentEncoder.GetBlitCommandEncoderDebugging().Copy(Texture, sourceSlice, sourceLevel, sourceOrigin, sourceSize, toTexture, destinationSlice, destinationLevel, destinationOrigin));
@@ -1351,9 +1319,6 @@ void FMetalRenderPass::CopyFromBufferToBuffer(FMetalBuffer const& SourceBuffer,
 	mtlpp::BlitCommandEncoder& Encoder = CurrentEncoder.GetBlitCommandEncoder();
 	check(Encoder.GetPtr());
 	
-	FMetalCommandBufferDebugHelpers::TrackResource(CurrentEncoder.GetCommandBuffer().GetPtr(), SourceBuffer.GetPtr());
-	FMetalCommandBufferDebugHelpers::TrackResource(CurrentEncoder.GetCommandBuffer().GetPtr(), DestinationBuffer.GetPtr());
-	
 	METAL_GPUPROFILE(FMetalProfiler::GetProfiler()->EncodeBlit(CurrentEncoder.GetCommandBufferStats(), __FUNCTION__));
 	MTLPP_VALIDATE(mtlpp::BlitCommandEncoder, Encoder, SafeGetRuntimeDebuggingLevel() >= EMetalDebugLevelValidation, Copy(SourceBuffer, SourceOffset, DestinationBuffer, DestinationOffset, Size));
 	METAL_DEBUG_LAYER(EMetalDebugLevelFastValidation, CurrentEncoder.GetBlitCommandEncoderDebugging().Copy(SourceBuffer, SourceOffset, DestinationBuffer, DestinationOffset, Size));
@@ -1366,9 +1331,6 @@ void FMetalRenderPass::PresentTexture(FMetalTexture const& Texture, uint32 sourc
 	mtlpp::BlitCommandEncoder& Encoder = CurrentEncoder.GetBlitCommandEncoder();
 	check(Encoder.GetPtr());
 	
-	FMetalCommandBufferDebugHelpers::TrackResource(CurrentEncoder.GetCommandBuffer().GetPtr(), Texture.GetPtr());
-	FMetalCommandBufferDebugHelpers::TrackResource(CurrentEncoder.GetCommandBuffer().GetPtr(), toTexture.GetPtr());
-	
 	METAL_GPUPROFILE(FMetalProfiler::GetProfiler()->EncodeBlit(CurrentEncoder.GetCommandBufferStats(), __FUNCTION__));
 	MTLPP_VALIDATE(mtlpp::BlitCommandEncoder, Encoder, SafeGetRuntimeDebuggingLevel() >= EMetalDebugLevelValidation, Copy(Texture, sourceSlice, sourceLevel, sourceOrigin, sourceSize, toTexture, destinationSlice, destinationLevel, destinationOrigin));
 	METAL_DEBUG_LAYER(EMetalDebugLevelFastValidation, CurrentEncoder.GetBlitCommandEncoderDebugging().Copy(Texture, sourceSlice, sourceLevel, sourceOrigin, sourceSize, toTexture, destinationSlice, destinationLevel, destinationOrigin));
@@ -1382,8 +1344,6 @@ void FMetalRenderPass::SynchronizeTexture(FMetalTexture const& Texture, uint32 S
 	mtlpp::BlitCommandEncoder& Encoder = CurrentEncoder.GetBlitCommandEncoder();
 	check(Encoder.GetPtr());
 	
-	FMetalCommandBufferDebugHelpers::TrackResource(CurrentEncoder.GetCommandBuffer().GetPtr(), Texture.GetPtr());
-	
 	// METAL_GPUPROFILE(FMetalProfiler::GetProfiler()->EncodeBlit(CurrentEncoder.GetCommandBufferStats(), __FUNCTION__));
 	MTLPP_VALIDATE(mtlpp::BlitCommandEncoder, Encoder, SafeGetRuntimeDebuggingLevel() >= EMetalDebugLevelValidation, Synchronize(Texture, Slice, Level));
 	METAL_DEBUG_LAYER(EMetalDebugLevelFastValidation, CurrentEncoder.GetBlitCommandEncoderDebugging().Synchronize(Texture, Slice, Level));
@@ -1399,8 +1359,6 @@ void FMetalRenderPass::SynchroniseResource(mtlpp::Resource const& Resource)
 	mtlpp::BlitCommandEncoder& Encoder = CurrentEncoder.GetBlitCommandEncoder();
 	check(Encoder.GetPtr());
 	
-	FMetalCommandBufferDebugHelpers::TrackResource(CurrentEncoder.GetCommandBuffer().GetPtr(), Resource.GetPtr());
-	
 	// METAL_GPUPROFILE(FMetalProfiler::GetProfiler()->EncodeBlit(CurrentEncoder.GetCommandBufferStats(), __FUNCTION__));
 	MTLPP_VALIDATE(mtlpp::BlitCommandEncoder, Encoder, SafeGetRuntimeDebuggingLevel() >= EMetalDebugLevelValidation, Synchronize(Resource));
 	METAL_DEBUG_LAYER(EMetalDebugLevelFastValidation, CurrentEncoder.GetBlitCommandEncoderDebugging().Synchronize(Resource));
@@ -1412,8 +1370,6 @@ void FMetalRenderPass::FillBuffer(FMetalBuffer const& Buffer, ns::Range Range, u
 {
 	check(Buffer);
 	
-	FMetalCommandBufferDebugHelpers::TrackResource(CurrentEncoder.GetCommandBuffer().GetPtr(), Buffer.GetPtr());
-	
 	mtlpp::BlitCommandEncoder TargetEncoder;
 	METAL_DEBUG_ONLY(FMetalBlitCommandEncoderDebugging Debugging);
 	bool bAsync = !CurrentEncoder.HasBufferBindingHistory(Buffer);
@@ -1423,8 +1379,6 @@ void FMetalRenderPass::FillBuffer(FMetalBuffer const& Buffer, ns::Range Range, u
 		TargetEncoder = PrologueEncoder.GetBlitCommandEncoder();
 		METAL_GPUPROFILE(FMetalProfiler::GetProfiler()->EncodeBlit(PrologueEncoder.GetCommandBufferStats(), FString::Printf(TEXT("FillBuffer: %p %llu %llu"), Buffer.GetPtr(), Buffer.GetOffset() + Range.Location, Range.Length)));
 		METAL_DEBUG_LAYER(EMetalDebugLevelFastValidation, Debugging = PrologueEncoder.GetBlitCommandEncoderDebugging());
-		
-		FMetalCommandBufferDebugHelpers::TrackResource(PrologueEncoder.GetCommandBuffer().GetPtr(), Buffer.GetPtr());
 	}
 	else
 	{
@@ -1432,8 +1386,6 @@ void FMetalRenderPass::FillBuffer(FMetalBuffer const& Buffer, ns::Range Range, u
 		TargetEncoder = CurrentEncoder.GetBlitCommandEncoder();
 		METAL_GPUPROFILE(FMetalProfiler::GetProfiler()->EncodeBlit(CurrentEncoder.GetCommandBufferStats(), FString::Printf(TEXT("FillBuffer: %p %llu %llu"), Buffer.GetPtr(), Buffer.GetOffset() + Range.Location, Range.Length)));
 		METAL_DEBUG_LAYER(EMetalDebugLevelFastValidation, Debugging = CurrentEncoder.GetBlitCommandEncoderDebugging());
-		
-		FMetalCommandBufferDebugHelpers::TrackResource(PrologueEncoder.GetCommandBuffer().GetPtr(), Buffer.GetPtr());
 	}
 	
 	check(TargetEncoder.GetPtr());
@@ -1449,9 +1401,6 @@ void FMetalRenderPass::FillBuffer(FMetalBuffer const& Buffer, ns::Range Range, u
 
 bool FMetalRenderPass::AsyncCopyFromBufferToTexture(FMetalBuffer const& Buffer, uint32 sourceOffset, uint32 sourceBytesPerRow, uint32 sourceBytesPerImage, mtlpp::Size sourceSize, FMetalTexture const& toTexture, uint32 destinationSlice, uint32 destinationLevel, mtlpp::Origin destinationOrigin, mtlpp::BlitOption options)
 {
-	FMetalCommandBufferDebugHelpers::TrackResource(CurrentEncoder.GetCommandBuffer().GetPtr(), Buffer.GetPtr());
-	FMetalCommandBufferDebugHelpers::TrackResource(CurrentEncoder.GetCommandBuffer().GetPtr(), toTexture.GetPtr());
-	
 	mtlpp::BlitCommandEncoder TargetEncoder;
 	METAL_DEBUG_ONLY(FMetalBlitCommandEncoderDebugging Debugging);
 	bool bAsync = !CurrentEncoder.HasTextureBindingHistory(toTexture);
@@ -1461,9 +1410,6 @@ bool FMetalRenderPass::AsyncCopyFromBufferToTexture(FMetalBuffer const& Buffer,
 		TargetEncoder = PrologueEncoder.GetBlitCommandEncoder();
 		METAL_GPUPROFILE(FMetalProfiler::GetProfiler()->EncodeBlit(PrologueEncoder.GetCommandBufferStats(), __FUNCTION__));
 		METAL_DEBUG_LAYER(EMetalDebugLevelFastValidation, Debugging = PrologueEncoder.GetBlitCommandEncoderDebugging());
-		
-		FMetalCommandBufferDebugHelpers::TrackResource(PrologueEncoder.GetCommandBuffer().GetPtr(), Buffer.GetPtr());
-		FMetalCommandBufferDebugHelpers::TrackResource(PrologueEncoder.GetCommandBuffer().GetPtr(), toTexture.GetPtr());
 	}
 	else
 	{
@@ -1471,9 +1417,6 @@ bool FMetalRenderPass::AsyncCopyFromBufferToTexture(FMetalBuffer const& Buffer,
 		TargetEncoder = CurrentEncoder.GetBlitCommandEncoder();
 		METAL_GPUPROFILE(FMetalProfiler::GetProfiler()->EncodeBlit(CurrentEncoder.GetCommandBufferStats(), __FUNCTION__));
 		METAL_DEBUG_LAYER(EMetalDebugLevelFastValidation, Debugging = CurrentEncoder.GetBlitCommandEncoderDebugging());
-		
-		FMetalCommandBufferDebugHelpers::TrackResource(CurrentEncoder.GetCommandBuffer().GetPtr(), Buffer.GetPtr());
-		FMetalCommandBufferDebugHelpers::TrackResource(CurrentEncoder.GetCommandBuffer().GetPtr(), toTexture.GetPtr());
 	}
 	
 	check(TargetEncoder.GetPtr());
@@ -1503,9 +1446,6 @@ bool FMetalRenderPass::AsyncCopyFromTextureToTexture(FMetalTexture const& Textur
 		TargetEncoder = PrologueEncoder.GetBlitCommandEncoder();
 		METAL_GPUPROFILE(FMetalProfiler::GetProfiler()->EncodeBlit(PrologueEncoder.GetCommandBufferStats(), __FUNCTION__));
 		METAL_DEBUG_LAYER(EMetalDebugLevelFastValidation, Debugging = PrologueEncoder.GetBlitCommandEncoderDebugging());
-		
-		FMetalCommandBufferDebugHelpers::TrackResource(PrologueEncoder.GetCommandBuffer().GetPtr(), Texture.GetPtr());
-		FMetalCommandBufferDebugHelpers::TrackResource(PrologueEncoder.GetCommandBuffer().GetPtr(), toTexture.GetPtr());
 	}
 	else
 	{
@@ -1513,9 +1453,6 @@ bool FMetalRenderPass::AsyncCopyFromTextureToTexture(FMetalTexture const& Textur
 		TargetEncoder = CurrentEncoder.GetBlitCommandEncoder();
 		METAL_GPUPROFILE(FMetalProfiler::GetProfiler()->EncodeBlit(CurrentEncoder.GetCommandBufferStats(), __FUNCTION__));
 		METAL_DEBUG_LAYER(EMetalDebugLevelFastValidation, Debugging = CurrentEncoder.GetBlitCommandEncoderDebugging());
-		
-		FMetalCommandBufferDebugHelpers::TrackResource(CurrentEncoder.GetCommandBuffer().GetPtr(), Texture.GetPtr());
-		FMetalCommandBufferDebugHelpers::TrackResource(CurrentEncoder.GetCommandBuffer().GetPtr(), toTexture.GetPtr());
 	}
 	
 	check(TargetEncoder.GetPtr());
@@ -1542,9 +1479,6 @@ void FMetalRenderPass::AsyncCopyFromBufferToBuffer(FMetalBuffer const& SourceBuf
 		TargetEncoder = PrologueEncoder.GetBlitCommandEncoder();
 		METAL_GPUPROFILE(FMetalProfiler::GetProfiler()->EncodeBlit(PrologueEncoder.GetCommandBufferStats(), FString::Printf(TEXT("AsyncCopyFromBufferToBuffer: %p %llu %llu"), DestinationBuffer.GetPtr(), DestinationBuffer.GetOffset() + DestinationOffset, Size)));
 		METAL_DEBUG_LAYER(EMetalDebugLevelFastValidation, Debugging = PrologueEncoder.GetBlitCommandEncoderDebugging());
-		
-		FMetalCommandBufferDebugHelpers::TrackResource(PrologueEncoder.GetCommandBuffer().GetPtr(), SourceBuffer.GetPtr());
-		FMetalCommandBufferDebugHelpers::TrackResource(PrologueEncoder.GetCommandBuffer().GetPtr(), DestinationBuffer.GetPtr());
 	}
 	else
 	{
@@ -1552,9 +1486,6 @@ void FMetalRenderPass::AsyncCopyFromBufferToBuffer(FMetalBuffer const& SourceBuf
 		TargetEncoder = CurrentEncoder.GetBlitCommandEncoder();
 		METAL_GPUPROFILE(FMetalProfiler::GetProfiler()->EncodeBlit(CurrentEncoder.GetCommandBufferStats(), FString::Printf(TEXT("AsyncCopyFromBufferToBuffer: %p %llu %llu"), DestinationBuffer.GetPtr(), DestinationBuffer.GetOffset() + DestinationOffset, Size)));
 		METAL_DEBUG_LAYER(EMetalDebugLevelFastValidation, Debugging = CurrentEncoder.GetBlitCommandEncoderDebugging());
-		
-		FMetalCommandBufferDebugHelpers::TrackResource(CurrentEncoder.GetCommandBuffer().GetPtr(), SourceBuffer.GetPtr());
-		FMetalCommandBufferDebugHelpers::TrackResource(CurrentEncoder.GetCommandBuffer().GetPtr(), DestinationBuffer.GetPtr());
 	}
 	
 	check(TargetEncoder.GetPtr());
@@ -1586,8 +1517,6 @@ void FMetalRenderPass::AsyncGenerateMipmapsForTexture(FMetalTexture const& Textu
 	mtlpp::BlitCommandEncoder Encoder = PrologueEncoder.GetBlitCommandEncoder();
 	check(Encoder.GetPtr());
 	
-	FMetalCommandBufferDebugHelpers::TrackResource(PrologueEncoder.GetCommandBuffer().GetPtr(), Texture.GetPtr());
-	
 	METAL_GPUPROFILE(FMetalProfiler::GetProfiler()->EncodeBlit(CurrentEncoder.GetCommandBufferStats(), __FUNCTION__));
 	MTLPP_VALIDATE(mtlpp::BlitCommandEncoder, Encoder, SafeGetRuntimeDebuggingLevel() >= EMetalDebugLevelValidation, GenerateMipmaps(Texture));
 	METAL_DEBUG_LAYER(EMetalDebugLevelFastValidation, PrologueEncoder.GetBlitCommandEncoderDebugging().GenerateMipmaps(Texture));
