Fixed Mac editor link failure in FF5PlayerMediaTexConvApple

diff --git a/Engine/Source/Runtime/Apple/MetalRHI/Private/MetalRHIPrivate.h b/Engine/Source/Runtime/Apple/MetalRHI/Private/MetalRHIPrivate.h
--- a/Engine/Source/Runtime/Apple/MetalRHI/Private/MetalRHIPrivate.h
+++ b/Engine/Source/Runtime/Apple/MetalRHI/Private/MetalRHIPrivate.h
@@ -166,7 +166,7 @@ mtlpp::ComputePipelineState GetMetalCopyIndex16Function();
 FMetalDeviceContext& GetMetalDeviceContext();
 
 // Safely release a metal object, correctly handling the case where the RHI has been destructed first
-void SafeReleaseMetalObject(id Object);
+void METALRHI_API SafeReleaseMetalObject(id Object);
 
 // Safely release a metal texture, correctly handling the case where the RHI has been destructed first
 void SafeReleaseMetalTexture(FMetalTexture& Object);
