A8 textures are not writable in metal. This change switches the backing to R8 and leverages Texture2DSample_A8 so we sample the correct channel.

diff --git a/Engine/Shaders/Private/SSRT/SSRTPrevFrameReduction.usf b/Engine/Shaders/Private/SSRT/SSRTPrevFrameReduction.usf
--- a/Engine/Shaders/Private/SSRT/SSRTPrevFrameReduction.usf
+++ b/Engine/Shaders/Private/SSRT/SSRTPrevFrameReduction.usf
@@ -124,7 +124,7 @@ void MainCS(
 
 				float4 SampleColor = float4(
 					HigherMipTexture.SampleLevel(HigherMipTextureSampler, SampleUV, 0).rgb,
-					HigherAlphaMipTexture.SampleLevel(HigherAlphaMipTextureSampler, SampleUV, 0).a);
+					Texture2DSample_A8(HigherAlphaMipTexture,HigherAlphaMipTextureSampler, SampleUV));
 
 				PrevColor += SampleColor * SampleWeight;
 			}
diff --git a/Engine/Shaders/Public/Platform.ush b/Engine/Shaders/Public/Platform.ush
--- a/Engine/Shaders/Public/Platform.ush
+++ b/Engine/Shaders/Public/Platform.ush
@@ -325,6 +325,11 @@
 
 #endif
 
+#if COMPILER_METAL
+	// Metal does not allow writes to A8 textures so we are faking it by making them all R8.
+	// WARNING: If this changes or the type in MetalRHI changes both must be updated!
+	#define A8_SAMPLE_MASK .r
+#endif
 
 // ---------------------------------------------------- Swizzle defaults
 
diff --git a/Engine/Shaders/Public/Platform/Metal/MetalCommon.ush b/Engine/Shaders/Public/Platform/Metal/MetalCommon.ush
--- a/Engine/Shaders/Public/Platform/Metal/MetalCommon.ush
+++ b/Engine/Shaders/Public/Platform/Metal/MetalCommon.ush
@@ -5,7 +5,7 @@
 =============================================================================*/
 
 // Update this GUID to improve shader recompilation for Metal only shaders
-// GUID = 03B3D368-6D07-43AF-9EF6-3FAC517881D5
+// GUID = c5ad1ae7-36ce-47c3-b5a9-41e9093577fe
 
 #pragma once
 
diff --git a/Engine/Source/Runtime/Apple/MetalRHI/Private/MetalRHI.cpp b/Engine/Source/Runtime/Apple/MetalRHI/Private/MetalRHI.cpp
--- a/Engine/Source/Runtime/Apple/MetalRHI/Private/MetalRHI.cpp
+++ b/Engine/Source/Runtime/Apple/MetalRHI/Private/MetalRHI.cpp
@@ -760,7 +760,10 @@ FMetalDynamicRHI::FMetalDynamicRHI(ERHIFeatureLevel::Type RequestedFeatureLevel)
 	GPixelFormats[PF_R16F_FILTER		].PlatformFormat	= (uint32)mtlpp::PixelFormat::R16Float;
 	GPixelFormats[PF_V8U8				].PlatformFormat	= (uint32)mtlpp::PixelFormat::RG8Snorm;
 	GPixelFormats[PF_A1					].PlatformFormat	= (uint32)mtlpp::PixelFormat::Invalid;
-	GPixelFormats[PF_A8					].PlatformFormat	= (uint32)mtlpp::PixelFormat::A8Unorm;
+	// A8 does not allow writes in Metal. So we will fake it with R8.
+	// If you change this you must also change the swizzle pattern in Platform.ush
+	// See Texture2DSample_A8 in Common.ush and A8_SAMPLE_MASK in Platform.ush
+	GPixelFormats[PF_A8					].PlatformFormat	= (uint32)mtlpp::PixelFormat::R8Unorm;
 	GPixelFormats[PF_R32_UINT			].PlatformFormat	= (uint32)mtlpp::PixelFormat::R32Uint;
 	GPixelFormats[PF_R32_SINT			].PlatformFormat	= (uint32)mtlpp::PixelFormat::R32Sint;
 	GPixelFormats[PF_R16G16B16A16_UINT	].PlatformFormat	= (uint32)mtlpp::PixelFormat::RGBA16Uint;
