[IOS] Re-enable iOS Metal GPUtime by removing any overlaps between commandlists, including overlaps across frames.

- Also uses a referenced counted timings array to track each frame's timing rather than relying on a static array and frame index.

diff --git a/Engine/Source/Runtime/Apple/MetalRHI/Private/MetalCommandList.cpp b/Engine/Source/Runtime/Apple/MetalRHI/Private/MetalCommandList.cpp
--- a/Engine/Source/Runtime/Apple/MetalRHI/Private/MetalCommandList.cpp
+++ b/Engine/Source/Runtime/Apple/MetalRHI/Private/MetalCommandList.cpp
@@ -302,8 +302,20 @@ void FMetalCommandList::SetParallelIndex(uint32 InIndex, uint32 InNum)
 void FMetalCommandList::Commit(mtlpp::CommandBuffer& Buffer, TArray<ns::Object<mtlpp::CommandBufferHandler>> CompletionHandlers, bool const bWait, bool const bIsLastCommandBuffer)
 {
 	check(Buffer);
-	
-	Buffer.AddCompletedHandler([CompletionHandlers, bIsLastCommandBuffer](mtlpp::CommandBuffer const& CompletedBuffer)
+
+	// The lifetime of this array is per frame
+	if (!FrameCommitedBufferTimings.IsValid())
+	{
+		FrameCommitedBufferTimings = MakeShared<TArray<FMetalCommandBufferTiming>, ESPMode::ThreadSafe>();
+	}
+
+	// The lifetime of this should be for the entire game
+	if (!LastCompletedBufferTiming.IsValid())
+	{
+		LastCompletedBufferTiming = MakeShared<FMetalCommandBufferTiming, ESPMode::ThreadSafe>();
+	}
+
+	Buffer.AddCompletedHandler([CompletionHandlers, FrameCommitedBufferTimingsLocal = FrameCommitedBufferTimings, LastCompletedBufferTimingLocal = LastCompletedBufferTiming](mtlpp::CommandBuffer const& CompletedBuffer)
 	{
 		if (CompletedBuffer.GetStatus() == mtlpp::CommandBufferStatus::Error)
 		{
@@ -316,17 +328,25 @@ void FMetalCommandList::Commit(mtlpp::CommandBuffer& Buffer, TArray<ns::Object<m
 				Handler.GetPtr()(CompletedBuffer);
 			}
 		}
-		
-		FMetalGPUProfiler::RecordCommandBuffer(CompletedBuffer);
-		
-		// The final command buffer in a frame will publish its frame
-		// stats and reset the counters for the next frame.
-		if(bIsLastCommandBuffer)
+
+		if (CompletedBuffer.GetStatus() == mtlpp::CommandBufferStatus::Completed)
+		{
+			FrameCommitedBufferTimingsLocal->Add({CompletedBuffer.GetGpuStartTime(), CompletedBuffer.GetGpuEndTime()});
+		}
+
+		// If this is the last reference, then it is the last command buffer to return, so record the frame
+		if (FrameCommitedBufferTimingsLocal.IsUnique())
 		{
-			FMetalGPUProfiler::RecordFrame();
+			FMetalGPUProfiler::RecordFrame(*FrameCommitedBufferTimingsLocal, *LastCompletedBufferTimingLocal);
 		}
 	});
 
+	// If bIsLastCommandBuffer is set then this is the end of the "frame".
+	if (bIsLastCommandBuffer)
+	{
+		FrameCommitedBufferTimings = MakeShared<TArray<FMetalCommandBufferTiming>, ESPMode::ThreadSafe>();
+	}
+
 	if (bImmediate)
 	{
 		CommandQueue.CommitCommandBuffer(Buffer);
diff --git a/Engine/Source/Runtime/Apple/MetalRHI/Private/MetalCommandList.h b/Engine/Source/Runtime/Apple/MetalRHI/Private/MetalCommandList.h
--- a/Engine/Source/Runtime/Apple/MetalRHI/Private/MetalCommandList.h
+++ b/Engine/Source/Runtime/Apple/MetalRHI/Private/MetalCommandList.h
@@ -92,6 +92,8 @@ public:
 private:
 #pragma mark - Private Member Variables -
 	FMetalCommandQueue& CommandQueue;
+	TSharedPtr<TArray<FMetalCommandBufferTiming>, ESPMode::ThreadSafe> FrameCommitedBufferTimings;
+	TSharedPtr<FMetalCommandBufferTiming, ESPMode::ThreadSafe> LastCompletedBufferTiming;
 	TArray<mtlpp::CommandBuffer> SubmittedBuffers;
 	uint32 Index;
 	uint32 Num;
diff --git a/Engine/Source/Runtime/Apple/MetalRHI/Private/MetalProfiler.cpp b/Engine/Source/Runtime/Apple/MetalRHI/Private/MetalProfiler.cpp
--- a/Engine/Source/Runtime/Apple/MetalRHI/Private/MetalProfiler.cpp
+++ b/Engine/Source/Runtime/Apple/MetalRHI/Private/MetalProfiler.cpp
@@ -333,49 +333,57 @@ void FMetalGPUProfiler::EndFrame()
 // All these recording functions MUST be called from within scheduled/completion handlers.
 // Ordering is enforced by libdispatch so calling these outside of that context WILL result in
 // incorrect values.
+void FMetalGPUProfiler::RecordFrame(TArray<FMetalCommandBufferTiming>& CommandBufferTimings, FMetalCommandBufferTiming& LastBufferTiming)
+{
+	const double CyclesPerSecond = 1.0 / FPlatformTime::GetSecondsPerCycle();
 
-// There overall flow should be this per frame:
-// Each command buffer finishes and accumulates its time.
-// On frame end these times are converted into cycles and published.
-// Present is a different case because it occurs on its own command buffer outside of
-// the normal submission scheme and Metal presents when a command buffer is scheduled.
-// This timing is not included in the time published in RecordFrame.
+	double RunningFrameTimeSeconds = 0.0;
+	uint64 FrameStartGPUCycles = 0;
+	uint64 FrameEndGPUCycles = 0;
 
-void FMetalGPUProfiler::IncrementFrameIndex()
-{
-	// This is called from RecordFrame.
-	// Prepare the next frame for recording.
-	FrameTimeGPUIndex = ((FrameTimeGPUIndex + 1) % MAX_FRAME_HISTORY);
-	FrameStartGPUCycles[FrameTimeGPUIndex] = 0;
-	FrameEndGPUCycles[FrameTimeGPUIndex] = 0;
-	FrameGPUTimeCycles[FrameTimeGPUIndex] = 0;
-	FrameIdleTimeCycles[FrameTimeGPUIndex] = 0;
-	FramePresentTimeCycles[FrameTimeGPUIndex] = 0;
-	RunningFrameTimeSeconds = 0.0;
-}
+	// Sort the timings
+	CommandBufferTimings.Sort();
 
-void FMetalGPUProfiler::RecordFrame()
-{
-	// Note: We do not call RecordCommandBuffer here.
-	// This is only applied by FMetalCommandList::Commit, which will also apply RecordCommandBuffer.
-	const double CyclesPerSecond = 1.0 / FPlatformTime::GetSecondsPerCycle();
-	FrameGPUTimeCycles[FrameTimeGPUIndex] = uint64(CyclesPerSecond * RunningFrameTimeSeconds);
-	// for now set the GPU time to 0 until we figure out why it isn't always accurate
-#if PLATFORM_IOS
-	FPlatformAtomics::AtomicStore_Relaxed((int32*)&GGPUFrameTime, (int32)0);
-#else
-	FPlatformAtomics::AtomicStore_Relaxed((int32*)&GGPUFrameTime, int32(FrameGPUTimeCycles[FrameTimeGPUIndex]));
-#endif
+	CFTimeInterval FirstStartTime = 0.0;
+
+	// Add the timings excluding any overlapping time
+	for (const FMetalCommandBufferTiming& Timing : CommandBufferTimings)
+	{
+		if (FirstStartTime == 0.0)
+		{
+			FirstStartTime = Timing.StartTime;
+		}
+		
+		// Only process if the previous buffer finished before the end of this one
+		if (LastBufferTiming.EndTime < Timing.EndTime)
+		{
+			// Check if the end of the previous buffer finished before the start of this one
+			if (LastBufferTiming.EndTime > Timing.StartTime)
+			{
+				// Segment from end of last buffer to end of current
+				RunningFrameTimeSeconds += Timing.EndTime - LastBufferTiming.EndTime;
+			}
+			else
+			{
+				// Full timing of this buffer
+				RunningFrameTimeSeconds += Timing.EndTime - Timing.StartTime;
+			}
+
+			LastBufferTiming = Timing;
+		}
+	}
+    
+	FrameStartGPUCycles = FirstStartTime * CyclesPerSecond;
+	FrameEndGPUCycles = LastBufferTiming.EndTime * CyclesPerSecond;
+    
+	uint64 FrameGPUTimeCycles = uint64(CyclesPerSecond * RunningFrameTimeSeconds);
+	FPlatformAtomics::AtomicStore_Relaxed((int32*)&GGPUFrameTime, int32(FrameGPUTimeCycles));
 	
 #if STATS
-	FPlatformAtomics::AtomicStore_Relaxed(&GMetalGPUWorkTime, FrameGPUTimeCycles[FrameTimeGPUIndex]);
-	
-	int64 TimeAsCycles = int64(FrameEndGPUCycles[FrameTimeGPUIndex] - FrameStartGPUCycles[FrameTimeGPUIndex] - FrameGPUTimeCycles[FrameTimeGPUIndex]);
-	FrameIdleTimeCycles[FrameTimeGPUIndex] = TimeAsCycles;
-	FPlatformAtomics::AtomicStore_Relaxed(&GMetalGPUIdleTime, TimeAsCycles);
+	FPlatformAtomics::AtomicStore_Relaxed(&GMetalGPUWorkTime, FrameGPUTimeCycles);
+	int64 FrameIdleTimeCycles = int64(FrameEndGPUCycles - FrameStartGPUCycles - FrameGPUTimeCycles);
+	FPlatformAtomics::AtomicStore_Relaxed(&GMetalGPUIdleTime, FrameIdleTimeCycles);
 #endif //STATS
-	
-	IncrementFrameIndex();
 }
 
 void FMetalGPUProfiler::RecordPresent(const mtlpp::CommandBuffer& Buffer)
@@ -386,42 +394,8 @@ void FMetalGPUProfiler::RecordPresent(const mtlpp::CommandBuffer& Buffer)
 	uint64 StartTimeCycles = uint64(GpuStartTimeSeconds * CyclesPerSecond);
 	uint64 EndTimeCycles = uint64(GpuEndTimeSeconds * CyclesPerSecond);
 	int64 Time = int64(EndTimeCycles - StartTimeCycles);
-	FramePresentTimeCycles[FrameTimeGPUIndex] = Time;
 	FPlatformAtomics::AtomicStore_Relaxed(&GMetalPresentTime, Time);
 }
-
-void FMetalGPUProfiler::RecordCommandBuffer(const mtlpp::CommandBuffer& Buffer)
-{
-	const double CyclesPerSecond = 1.0 / FPlatformTime::GetSecondsPerCycle();
-	
-	const CFTimeInterval GpuEndTimeSeconds   = Buffer.GetGpuEndTime();
-	const CFTimeInterval GpuStartTimeSeconds = Buffer.GetGpuStartTime();
-	
-	uint64 StartTimeCycles = uint64(GpuStartTimeSeconds * CyclesPerSecond);
-	uint64 EndTimeCycles   = uint64(GpuEndTimeSeconds * CyclesPerSecond);
-	
-	{
-		uint64 Existing = FrameEndGPUCycles[FrameTimeGPUIndex];
-		FrameEndGPUCycles[FrameTimeGPUIndex] = Existing > 0 ? FMath::Max(Existing, EndTimeCycles) : EndTimeCycles;
-	}
-	
-	{
-		uint64 Existing = FrameStartGPUCycles[FrameTimeGPUIndex];
-		FrameStartGPUCycles[FrameTimeGPUIndex] = Existing > 0 ? FMath::Min(Existing, StartTimeCycles) : StartTimeCycles;
-	}
-	
-	RunningFrameTimeSeconds += (GpuEndTimeSeconds - GpuStartTimeSeconds);
-}
-	
-int32 FMetalGPUProfiler::FrameTimeGPUIndex = 0;
-int64 FMetalGPUProfiler::FrameStartGPUCycles[MAX_FRAME_HISTORY];
-int64 FMetalGPUProfiler::FrameEndGPUCycles[MAX_FRAME_HISTORY];
-int64 FMetalGPUProfiler::FrameGPUTimeCycles[MAX_FRAME_HISTORY];
-int64 FMetalGPUProfiler::FrameIdleTimeCycles[MAX_FRAME_HISTORY];
-int64 FMetalGPUProfiler::FramePresentTimeCycles[MAX_FRAME_HISTORY];
-
-double FMetalGPUProfiler::RunningFrameTimeSeconds = 0.0;
-
 // END WARNING
 
 IMetalStatsScope::~IMetalStatsScope()
diff --git a/Engine/Source/Runtime/Apple/MetalRHI/Private/MetalProfiler.h b/Engine/Source/Runtime/Apple/MetalRHI/Private/MetalProfiler.h
--- a/Engine/Source/Runtime/Apple/MetalRHI/Private/MetalProfiler.h
+++ b/Engine/Source/Runtime/Apple/MetalRHI/Private/MetalProfiler.h
@@ -339,6 +339,29 @@ struct FMetalCommandBufferStats : public IMetalStatsScope
 #endif
 };
 
+/**
+ * Simple struct to hold sortable command buffer start and end timestamps.
+ */
+struct FMetalCommandBufferTiming
+{
+	CFTimeInterval StartTime;
+	CFTimeInterval EndTime;
+
+	bool operator<(const FMetalCommandBufferTiming& RHS) const
+	{
+		// Sort by start time and then by length if the commandbuffer started at the same time
+		if (this->StartTime < RHS.StartTime)
+		{
+			return true;
+		}
+		else if ((this->StartTime == RHS.StartTime) && (this->EndTime > RHS.EndTime))
+		{
+			return true;
+		}
+		return false;
+	}
+};
+
 /**
  * Encapsulates GPU profiling logic and data.
  * There's only one global instance of this struct so it should only contain global data, nothing specific to a frame.
@@ -352,14 +375,7 @@ struct FMetalGPUProfiler : public FGPUProfiler
 	:	FGPUProfiler()
 	,	Context(InContext)
 	,   NumNestedFrames(0)
-	{
-		FMemory::Memzero((void*)&FrameStartGPUCycles[0], sizeof(FrameStartGPUCycles));
-		FMemory::Memzero((void*)&FrameEndGPUCycles[0], sizeof(FrameEndGPUCycles));
-		FMemory::Memzero((void*)&FrameGPUTimeCycles[0], sizeof(FrameGPUTimeCycles));
-		FMemory::Memzero((void*)&FrameIdleTimeCycles[0], sizeof(FrameIdleTimeCycles));
-		FMemory::Memzero((void*)&FramePresentTimeCycles[0], sizeof(FramePresentTimeCycles));
-		RunningFrameTimeSeconds = 0.0;
-	}
+	{}
 	
 	virtual ~FMetalGPUProfiler() {}
 	
@@ -376,26 +392,13 @@ struct FMetalGPUProfiler : public FGPUProfiler
 	// WARNING:
 	// These functions MUST be called from within Metal scheduled/completion handlers
 	// since they depend on libdispatch to enforce ordering.
-	static void RecordFrame();
+	static void RecordFrame(TArray<FMetalCommandBufferTiming>& CommandBufferTimings, FMetalCommandBufferTiming& LastPresentBufferTiming);
 	static void RecordPresent(const mtlpp::CommandBuffer& Buffer);
-	static void RecordCommandBuffer(const mtlpp::CommandBuffer& Buffer);
 	// END WARNING
 	
-#define MAX_FRAME_HISTORY 3
-	static int32 FrameTimeGPUIndex;
-	static int64 FrameStartGPUCycles[MAX_FRAME_HISTORY];
-	static int64 FrameEndGPUCycles[MAX_FRAME_HISTORY];
-	static int64 FrameGPUTimeCycles[MAX_FRAME_HISTORY];
-	static int64 FrameIdleTimeCycles[MAX_FRAME_HISTORY];
-	static int64 FramePresentTimeCycles[MAX_FRAME_HISTORY];
-	
 	FMetalGPUTiming TimingSupport;
 	FMetalContext* Context;
 	int32 NumNestedFrames;
-private:
-	// These must only be accessed from within Metal scheduled/completion handlers. See above.
-	static void IncrementFrameIndex();
-	static double RunningFrameTimeSeconds;
 };
 
 class FMetalProfiler : public FMetalGPUProfiler
