Fix Metal validation error for fallback depth/stencil texture binding.

diff --git a/Engine/Source/Runtime/Apple/MetalRHI/Private/MetalContext.cpp b/Engine/Source/Runtime/Apple/MetalRHI/Private/MetalContext.cpp
--- a/Engine/Source/Runtime/Apple/MetalRHI/Private/MetalContext.cpp
+++ b/Engine/Source/Runtime/Apple/MetalRHI/Private/MetalContext.cpp
@@ -1322,6 +1322,8 @@ bool FMetalContext::PrepareToDraw(uint32 PrimitiveType, EMetalIndexType IndexTyp
 		{
 			check(!bBindDepthStencilForUAVRaster);
 			Info.DepthStencilRenderTarget.DepthStencilTarget = FallbackDepthStencilSurface;
+			Info.DepthStencilRenderTarget.ExclusiveDepthStencil = FExclusiveDepthStencil::DepthWrite_StencilWrite;
+			Info.DepthStencilRenderTarget.Action = MakeDepthStencilTargetActions(MakeRenderTargetActions(ERenderTargetLoadAction::ELoad, ERenderTargetStoreAction::EStore), MakeRenderTargetActions(ERenderTargetLoadAction::ELoad, ERenderTargetStoreAction::EStore));
 		}
 		else
 		{
