Reset texture binding at compute pipeline state set time.

diff --git a/Engine/Source/Runtime/Apple/MetalRHI/Private/MetalStateCache.cpp b/Engine/Source/Runtime/Apple/MetalRHI/Private/MetalStateCache.cpp
--- a/Engine/Source/Runtime/Apple/MetalRHI/Private/MetalStateCache.cpp
+++ b/Engine/Source/Runtime/Apple/MetalRHI/Private/MetalStateCache.cpp
@@ -6,6 +6,15 @@
 #include "MetalProfiler.h"
 #include "MetalCommandBuffer.h"
 
+#if PLATFORM_MAC
+	#ifndef UINT128_MAX
+		#define UINT128_MAX (((__uint128_t)1 << 127) - (__uint128_t)1 + ((__uint128_t)1 << 127))
+	#endif
+	#define FMETALTEXTUREMASK_MAX UINT128_MAX
+#else
+	#define FMETALTEXTUREMASK_MAX UINT32_MAX
+#endif
+
 static mtlpp::TriangleFillMode TranslateFillMode(ERasterizerFillMode FillMode)
 {
 	switch (FillMode)
@@ -391,6 +400,13 @@ void FMetalStateCache::SetComputeShader(FMetalComputeShader* InComputeShader)
 		
 		DirtyUniformBuffers[EMetalShaderStages::Compute] = 0xffffffff;
 
+		for (uint32 Index = 0; Index < ML_MaxTextures; ++Index)
+		{
+			ShaderTextures[EMetalShaderStages::Compute].Textures[Index] = nil;
+			ShaderTextures[EMetalShaderStages::Compute].Usage[Index] = mtlpp::ResourceUsage(0);
+		}
+		ShaderTextures[EMetalShaderStages::Compute].Bound = 0;
+
 		for (const auto& PackedGlobalArray : InComputeShader->Bindings.PackedGlobalArrays)
 		{
 			ShaderParameters[EMetalShaderStages::Compute].PrepareGlobalUniforms(CrossCompiler::PackedTypeNameToTypeIndex(PackedGlobalArray.TypeName), PackedGlobalArray.Size);
@@ -1086,14 +1102,7 @@ void FMetalStateCache::SetGraphicsPipelineState(FMetalGraphicsPipelineState* Sta
 			for (uint32 i = 0; i < EMetalShaderStages::Num; i++)
 			{
 				ShaderBuffers[i].Bound = UINT32_MAX;
-#if PLATFORM_MAC
-#ifndef UINT128_MAX
-#define UINT128_MAX (((__uint128_t)1 << 127) - (__uint128_t)1 + ((__uint128_t)1 << 127))
-#endif
-				ShaderTextures[i].Bound = UINT128_MAX;
-#else
-				ShaderTextures[i].Bound = UINT32_MAX;
-#endif
+				ShaderTextures[i].Bound = FMETALTEXTUREMASK_MAX;
 				ShaderSamplers[i].Bound = UINT16_MAX;
 			}
 		}
@@ -1102,11 +1111,7 @@ void FMetalStateCache::SetGraphicsPipelineState(FMetalGraphicsPipelineState* Sta
 		if (bNewUsingTessellation)
 		{
 			ShaderBuffers[EMetalShaderStages::Hull].Bound = UINT32_MAX;
-#if PLATFORM_MAC
-			ShaderTextures[EMetalShaderStages::Hull].Bound = UINT128_MAX;
-#else
-			ShaderTextures[EMetalShaderStages::Hull].Bound = UINT32_MAX;
-#endif
+			ShaderTextures[EMetalShaderStages::Hull].Bound = FMETALTEXTUREMASK_MAX;
 			ShaderSamplers[EMetalShaderStages::Hull].Bound = UINT16_MAX;
 			
 			for (uint32 i = 0; i < ML_MaxBuffers; i++)
@@ -1158,14 +1163,7 @@ void FMetalStateCache::SetGraphicsPipelineState(FMetalGraphicsPipelineState* Sta
             for (uint32 i = 0; i < EMetalShaderStages::Num; i++)
             {
                 ShaderBuffers[i].Bound = UINT32_MAX;
-#if PLATFORM_MAC
-#ifndef UINT128_MAX
-#define UINT128_MAX (((__uint128_t)1 << 127) - (__uint128_t)1 + ((__uint128_t)1 << 127))
-#endif
-                ShaderTextures[i].Bound = UINT128_MAX;
-#else
-                ShaderTextures[i].Bound = UINT32_MAX;
-#endif
+                ShaderTextures[i].Bound = FMETALTEXTUREMASK_MAX;
                 ShaderSamplers[i].Bound = UINT16_MAX;
             }
         }
@@ -1986,14 +1984,7 @@ void FMetalStateCache::SetStateDirty(void)
 	for (uint32 i = 0; i < EMetalShaderStages::Num; i++)
 	{
 		ShaderBuffers[i].Bound = UINT32_MAX;
-#if PLATFORM_MAC
-#ifndef UINT128_MAX
-#define UINT128_MAX (((__uint128_t)1 << 127) - (__uint128_t)1 + ((__uint128_t)1 << 127))
-#endif
-		ShaderTextures[i].Bound = UINT128_MAX;
-#else
-		ShaderTextures[i].Bound = UINT32_MAX;
-#endif
+		ShaderTextures[i].Bound = FMETALTEXTUREMASK_MAX;
 		ShaderSamplers[i].Bound = UINT16_MAX;
 	}
 }
