Copying Private-LoadTimes-4.24 stream to Main. Biggest changes are in Materials/Shader memory freezing.

diff --git a/Engine/Source/Editor/UnrealEd/Private/PreviewMaterial.cpp b/Engine/Source/Editor/UnrealEd/Private/PreviewMaterial.cpp
--- a/Engine/Source/Editor/UnrealEd/Private/PreviewMaterial.cpp
+++ b/Engine/Source/Editor/UnrealEd/Private/PreviewMaterial.cpp
@@ -239,22 +239,22 @@ public:
 		}
 	}
 
-	virtual bool GetVectorValue(const FMaterialParameterInfo& ParameterInfo, FLinearColor* OutValue, const FMaterialRenderContext& Context) const
+	virtual bool GetVectorValue(const FHashedMaterialParameterInfo& ParameterInfo, FLinearColor* OutValue, const FMaterialRenderContext& Context) const
 	{
 		return Material->GetRenderProxy()->GetVectorValue(ParameterInfo, OutValue, Context);
 	}
 
-	virtual bool GetScalarValue(const FMaterialParameterInfo& ParameterInfo, float* OutValue, const FMaterialRenderContext& Context) const
+	virtual bool GetScalarValue(const FHashedMaterialParameterInfo& ParameterInfo, float* OutValue, const FMaterialRenderContext& Context) const
 	{
 		return Material->GetRenderProxy()->GetScalarValue(ParameterInfo, OutValue, Context);
 	}
 
-	virtual bool GetTextureValue(const FMaterialParameterInfo& ParameterInfo, const UTexture** OutValue, const FMaterialRenderContext& Context) const
+	virtual bool GetTextureValue(const FHashedMaterialParameterInfo& ParameterInfo, const UTexture** OutValue, const FMaterialRenderContext& Context) const
 	{
 		return Material->GetRenderProxy()->GetTextureValue(ParameterInfo,OutValue,Context);
 	}
 
-	virtual bool GetTextureValue(const FMaterialParameterInfo& ParameterInfo, const URuntimeVirtualTexture** OutValue, const FMaterialRenderContext& Context) const
+	virtual bool GetTextureValue(const FHashedMaterialParameterInfo& ParameterInfo, const URuntimeVirtualTexture** OutValue, const FMaterialRenderContext& Context) const
 	{
 		return Material->GetRenderProxy()->GetTextureValue(ParameterInfo, OutValue, Context);
 	}
diff --git a/Engine/Source/Editor/UnrealEd/Private/Texture2DPreview.cpp b/Engine/Source/Editor/UnrealEd/Private/Texture2DPreview.cpp
--- a/Engine/Source/Editor/UnrealEd/Private/Texture2DPreview.cpp
+++ b/Engine/Source/Editor/UnrealEd/Private/Texture2DPreview.cpp
@@ -29,11 +29,9 @@ namespace
 
 class FSimpleElementTexture2DPreviewPS : public FGlobalShader
 {
-	DECLARE_GLOBAL_SHADER(FSimpleElementTexture2DPreviewPS);
-
-	using FPermutationDomain = TShaderPermutationDomain<FTexture2DPreviewVirtualTexture, FTexture2DPreviewTexture2DArray>;
-
+	DECLARE_SHADER_TYPE(FSimpleElementTexture2DPreviewPS, Global);
 public:
+	using FPermutationDomain = TShaderPermutationDomain<FTexture2DPreviewVirtualTexture, FTexture2DPreviewTexture2DArray>;
 
 	FSimpleElementTexture2DPreviewPS(const ShaderMetaType::CompiledShaderInitializerType& Initializer) :
 		FGlobalShader(Initializer)
@@ -65,23 +63,24 @@ public:
 
 	void SetParameters(FRHICommandList& RHICmdList, const FTexture* TextureValue, const FMatrix& ColorWeightsValue, float GammaValue, float MipLevel, float LayerIndex, bool bIsNormalMap, bool bIsSingleVTPhysicalSpace, bool bIsVirtualTexture, bool bIsTextureArray)
 	{
+		FRHIPixelShader* ShaderRHI = RHICmdList.GetBoundPixelShader();
 		if (bIsVirtualTexture)
 		{
 			FVirtualTexture2DResource* VirtualTextureValue = (FVirtualTexture2DResource*)TextureValue;
 			IAllocatedVirtualTexture* AllocatedVT = VirtualTextureValue->AcquireAllocatedVT();
 	
 			FRHIShaderResourceView* PhysicalView = AllocatedVT->GetPhysicalTextureSRV((uint32)LayerIndex, TextureValue->bSRGB);
-			SetSRVParameter(RHICmdList, GetPixelShader(), InTexture, PhysicalView);
-			SetSamplerParameter(RHICmdList, GetPixelShader(), InTextureSampler, VirtualTextureValue->SamplerStateRHI);
+			SetSRVParameter(RHICmdList, ShaderRHI, InTexture, PhysicalView);
+			SetSamplerParameter(RHICmdList, ShaderRHI, InTextureSampler, VirtualTextureValue->SamplerStateRHI);
 
-			SetTextureParameter(RHICmdList, GetPixelShader(), InPageTableTexture0, AllocatedVT->GetPageTableTexture(0u));
+			SetTextureParameter(RHICmdList, ShaderRHI, InPageTableTexture0, AllocatedVT->GetPageTableTexture(0u));
 			if (AllocatedVT->GetNumPageTableTextures() > 1u)
 			{
-				SetTextureParameter(RHICmdList, GetPixelShader(), InPageTableTexture1, AllocatedVT->GetPageTableTexture(1u));
+				SetTextureParameter(RHICmdList, ShaderRHI, InPageTableTexture1, AllocatedVT->GetPageTableTexture(1u));
 			}
 			else
 			{
-				SetTextureParameter(RHICmdList, GetPixelShader(), InPageTableTexture1, GBlackTexture->TextureRHI);
+				SetTextureParameter(RHICmdList, ShaderRHI, InPageTableTexture1, GBlackTexture->TextureRHI);
 			}
 
 			FUintVector4 PageTableUniform[2];
@@ -90,61 +89,44 @@ public:
 			AllocatedVT->GetPackedPageTableUniform(PageTableUniform, false);
 			AllocatedVT->GetPackedUniform(&Uniform, (uint32)LayerIndex);
 
-			SetShaderValueArray(RHICmdList, GetPixelShader(), VTPackedPageTableUniform, PageTableUniform, UE_ARRAY_COUNT(PageTableUniform));
-			SetShaderValue(RHICmdList, GetPixelShader(), VTPackedUniform, Uniform);
+			SetShaderValueArray(RHICmdList, ShaderRHI, VTPackedPageTableUniform, PageTableUniform, UE_ARRAY_COUNT(PageTableUniform));
+			SetShaderValue(RHICmdList, ShaderRHI, VTPackedUniform, Uniform);
 		}
 		else
 		{
-			SetTextureParameter(RHICmdList, GetPixelShader(), InPageTableTexture0, GBlackTexture->TextureRHI);
-			SetTextureParameter(RHICmdList, GetPixelShader(), InPageTableTexture1, GBlackTexture->TextureRHI);
-			SetTextureParameter(RHICmdList, GetPixelShader(), InTexture, InTextureSampler, TextureValue);
+			SetTextureParameter(RHICmdList, ShaderRHI, InPageTableTexture0, GBlackTexture->TextureRHI);
+			SetTextureParameter(RHICmdList, ShaderRHI, InPageTableTexture1, GBlackTexture->TextureRHI);
+			SetTextureParameter(RHICmdList, ShaderRHI, InTexture, InTextureSampler, TextureValue);
 		}
 		
-		SetShaderValue(RHICmdList, GetPixelShader(),ColorWeights,ColorWeightsValue);
+		SetShaderValue(RHICmdList, ShaderRHI,ColorWeights,ColorWeightsValue);
 		FVector4 PackedParametersValue(GammaValue, MipLevel, bIsNormalMap ? 1.0 : -1.0f, bIsSingleVTPhysicalSpace ? 0 : LayerIndex);
-		SetShaderValue(RHICmdList, GetPixelShader(), PackedParameters, PackedParametersValue);
+		SetShaderValue(RHICmdList, ShaderRHI, PackedParameters, PackedParametersValue);
 
 		// Store slice count for texture array
 		if (bIsTextureArray)
 		{
 			const FTexture2DArrayResource* TextureValue2DArray = (FTexture2DArrayResource*)(TextureValue);
 			float NumSlicesData = TextureValue2DArray ? float(TextureValue2DArray->GetNumSlices()) : 1;
-			SetShaderValue(RHICmdList, GetPixelShader(), NumSlices, NumSlicesData);
+			SetShaderValue(RHICmdList, ShaderRHI, NumSlices, NumSlicesData);
 		}
 
-		SetShaderValue(RHICmdList, GetPixelShader(),TextureComponentReplicate, (TextureValue && TextureValue->bGreyScaleFormat) ? FLinearColor(1,0,0,0) : FLinearColor(0,0,0,0));
-		SetShaderValue(RHICmdList, GetPixelShader(),TextureComponentReplicateAlpha, (TextureValue && TextureValue->bGreyScaleFormat) ? FLinearColor(1,0,0,0) : FLinearColor(0,0,0,1));
-	}
-
-	virtual bool Serialize(FArchive& Ar) override
-	{
-		bool bShaderHasOutdatedParameters = FGlobalShader::Serialize(Ar);
-		Ar << InTexture;
-		Ar << InTextureSampler;
-		Ar << InPageTableTexture0;
-		Ar << InPageTableTexture1;
-		Ar << VTPackedPageTableUniform;
-		Ar << VTPackedUniform;
-		Ar << TextureComponentReplicate;
-		Ar << TextureComponentReplicateAlpha;
-		Ar << ColorWeights;
-		Ar << NumSlices;
-		Ar << PackedParameters;
-		return bShaderHasOutdatedParameters;
+		SetShaderValue(RHICmdList, ShaderRHI, TextureComponentReplicate, (TextureValue && TextureValue->bGreyScaleFormat) ? FLinearColor(1,0,0,0) : FLinearColor(0,0,0,0));
+		SetShaderValue(RHICmdList, ShaderRHI, TextureComponentReplicateAlpha, (TextureValue && TextureValue->bGreyScaleFormat) ? FLinearColor(1,0,0,0) : FLinearColor(0,0,0,1));
 	}
 
 private:
-	FShaderResourceParameter InTexture;			// if VT, this used as the physical texture
-	FShaderResourceParameter InTextureSampler;	// if VT, this is the physical sampler
-	FShaderResourceParameter InPageTableTexture0;
-	FShaderResourceParameter InPageTableTexture1;
-	FShaderParameter VTPackedPageTableUniform;
-	FShaderParameter VTPackedUniform;
-	FShaderParameter TextureComponentReplicate;
-	FShaderParameter TextureComponentReplicateAlpha;
-	FShaderParameter ColorWeights; 
-	FShaderParameter PackedParameters;
-	FShaderParameter NumSlices;
+	LAYOUT_FIELD(FShaderResourceParameter, InTexture);			// if VT, this used as the physical texture
+	LAYOUT_FIELD(FShaderResourceParameter, InTextureSampler);	// if VT, this is the physical sampler
+	LAYOUT_FIELD(FShaderResourceParameter, InPageTableTexture0);
+	LAYOUT_FIELD(FShaderResourceParameter, InPageTableTexture1);
+	LAYOUT_FIELD(FShaderParameter, VTPackedPageTableUniform);
+	LAYOUT_FIELD(FShaderParameter, VTPackedUniform);
+	LAYOUT_FIELD(FShaderParameter, TextureComponentReplicate);
+	LAYOUT_FIELD(FShaderParameter, TextureComponentReplicateAlpha);
+	LAYOUT_FIELD(FShaderParameter, ColorWeights);
+	LAYOUT_FIELD(FShaderParameter, PackedParameters);
+	LAYOUT_FIELD(FShaderParameter, NumSlices);
 };
 
 IMPLEMENT_GLOBAL_SHADER(FSimpleElementTexture2DPreviewPS, "/Engine/Private/SimpleElementTexture2DPreviewPixelShader.usf", "Main", SF_Pixel);
@@ -162,13 +144,13 @@ void FBatchedElementTexture2DPreviewParameters::BindShaders(
 	TShaderMapRef<FSimpleElementVS> VertexShader(GetGlobalShaderMap(InFeatureLevel));
 
 	FSimpleElementTexture2DPreviewPS::FPermutationDomain PermutationVector;
-	PermutationVector.Set<FTexture2DPreviewVirtualTexture>(bIsVirtualTexture);
+	PermutationVector.Set<FTexture2DPreviewVirtualTexture>(bIsVirtualTexture);	
 	PermutationVector.Set<FTexture2DPreviewTexture2DArray>(bIsTextureArray);
 	TShaderMapRef<FSimpleElementTexture2DPreviewPS> PixelShader(GetGlobalShaderMap(InFeatureLevel), PermutationVector);
 
 	GraphicsPSOInit.BoundShaderState.VertexDeclarationRHI = GSimpleElementVertexDeclaration.VertexDeclarationRHI;
-	GraphicsPSOInit.BoundShaderState.VertexShaderRHI = GETSAFERHISHADER_VERTEX(*VertexShader);
-	GraphicsPSOInit.BoundShaderState.PixelShaderRHI = GETSAFERHISHADER_PIXEL(*PixelShader);
+	GraphicsPSOInit.BoundShaderState.VertexShaderRHI = VertexShader.GetVertexShader();
+	GraphicsPSOInit.BoundShaderState.PixelShaderRHI = PixelShader.GetPixelShader();
 	GraphicsPSOInit.PrimitiveType = PT_TriangleList;
 
 	if (bIsSingleChannelFormat)
diff --git a/Engine/Source/Editor/UnrealEd/Private/VolumeTexturePreview.cpp b/Engine/Source/Editor/UnrealEd/Private/VolumeTexturePreview.cpp
--- a/Engine/Source/Editor/UnrealEd/Private/VolumeTexturePreview.cpp
+++ b/Engine/Source/Editor/UnrealEd/Private/VolumeTexturePreview.cpp
@@ -44,6 +44,7 @@ UNREALED_API void GetBestFitForNumberOfTiles(int32 InSize, int32& OutNumTilesX,
  */
 class FSimpleElementVolumeTexturePreviewPS : public FGlobalShader
 {
+	DECLARE_INLINE_TYPE_LAYOUT(FSimpleElementVolumeTexturePreviewPS, NonVirtual);
 public:
 
 	FSimpleElementVolumeTexturePreviewPS(const ShaderMetaType::CompiledShaderInitializerType& Initializer) :
@@ -73,74 +74,59 @@ public:
 	
 	void SetParameters(FRHICommandList& RHICmdList, const FTexture* TextureValue, int32 SizeZ, const FMatrix& ColorWeightsValue, float GammaValue, float MipLevel, float Opacity, const FRotator& TraceOrientation)
 	{
-		SetTextureParameter(RHICmdList, GetPixelShader(), InTexture, InTextureSampler, TextureValue);
+		FRHIPixelShader* ShaderRHI = RHICmdList.GetBoundPixelShader();
+
+		SetTextureParameter(RHICmdList, ShaderRHI, InTexture, InTextureSampler, TextureValue);
 		if (GEditor && GEditor->Bad)
 		{
-			SetTextureParameter(RHICmdList, GetPixelShader(), BadTexture, BadTextureSampler, GEditor->Bad->Resource);
+			SetTextureParameter(RHICmdList, ShaderRHI, BadTexture, BadTextureSampler, GEditor->Bad->Resource);
 		}
 		else
 		{
-			SetTextureParameter(RHICmdList, GetPixelShader(), BadTexture, GWhiteTexture->TextureRHI);
+			SetTextureParameter(RHICmdList, ShaderRHI, BadTexture, GWhiteTexture->TextureRHI);
 		}
-		SetShaderValue(RHICmdList, GetPixelShader(),ColorWeights,ColorWeightsValue);
+		SetShaderValue(RHICmdList, ShaderRHI,ColorWeights,ColorWeightsValue);
 
 		const int32 MipSizeZ = MipLevel >= 0 ? FMath::Max<int32>(SizeZ >> FMath::FloorToInt(MipLevel), 1) : SizeZ;
 		FVector4 PackedParametersValue(GammaValue, MipLevel, (float)MipSizeZ, Opacity);
-		SetShaderValue(RHICmdList, GetPixelShader(), PackedParameters, PackedParametersValue);
+		SetShaderValue(RHICmdList, ShaderRHI, PackedParameters, PackedParametersValue);
 
 		int32 NumTilesX = 0;
 		int32 NumTilesY = 0;
 		GetBestFitForNumberOfTiles(MipSizeZ, NumTilesX, NumTilesY);
-		SetShaderValue(RHICmdList, GetPixelShader(), NumTilesPerSideParameter, FVector4((float)NumTilesX, (float)NumTilesY, 0 ,0));
+		SetShaderValue(RHICmdList, ShaderRHI, NumTilesPerSideParameter, FVector4((float)NumTilesX, (float)NumTilesY, 0 ,0));
 
-		SetShaderValue(RHICmdList, GetPixelShader(),TextureComponentReplicate,TextureValue->bGreyScaleFormat ? FLinearColor(1,0,0,0) : FLinearColor(0,0,0,0));
-		SetShaderValue(RHICmdList, GetPixelShader(),TextureComponentReplicateAlpha,TextureValue->bGreyScaleFormat ? FLinearColor(1,0,0,0) : FLinearColor(0,0,0,1));
+		SetShaderValue(RHICmdList, ShaderRHI,TextureComponentReplicate,TextureValue->bGreyScaleFormat ? FLinearColor(1,0,0,0) : FLinearColor(0,0,0,0));
+		SetShaderValue(RHICmdList, ShaderRHI,TextureComponentReplicateAlpha,TextureValue->bGreyScaleFormat ? FLinearColor(1,0,0,0) : FLinearColor(0,0,0,1));
 
 		const FVector TextureDimension((float)TextureValue->GetSizeX(), (float)TextureValue->GetSizeY(), (float)SizeZ);
 		const float OneOverMinDimension = 1.f / FMath::Max(TextureDimension.GetMin(), 1.f);
-		SetShaderValue(RHICmdList, GetPixelShader(), TraceVolumeScalingParameter, FVector4(
+		SetShaderValue(RHICmdList, ShaderRHI, TraceVolumeScalingParameter, FVector4(
 				TextureDimension.X * OneOverMinDimension, 
 				TextureDimension.Y * OneOverMinDimension, 
 				TextureDimension.Z * OneOverMinDimension, 
 				TextureDimension.GetMax() * OneOverMinDimension * .5f) // Extent
 			);
 
-		SetShaderValue(RHICmdList, GetPixelShader(), TextureDimensionParameter, FVector(TextureDimension.X, TextureDimension.Y, TextureDimension.Z));
-
-		SetShaderValue(RHICmdList, GetPixelShader(), TraceViewMatrixParameter, FMatrix(FRotationMatrix::Make(TraceOrientation)));
-	}
+		SetShaderValue(RHICmdList, ShaderRHI, TextureDimensionParameter, FVector(TextureDimension.X, TextureDimension.Y, TextureDimension.Z));
 
-	virtual bool Serialize(FArchive& Ar) override
-	{
-		bool bShaderHasOutdatedParameters = FGlobalShader::Serialize(Ar);
-		Ar << InTexture;
-		Ar << InTextureSampler;
-		Ar << BadTexture;
-		Ar << BadTextureSampler;
-		Ar << TextureComponentReplicate;
-		Ar << TextureComponentReplicateAlpha;
-		Ar << ColorWeights;
-		Ar << PackedParameters;
-		Ar << NumTilesPerSideParameter;
-		Ar << TraceVolumeScalingParameter;
-		Ar << TextureDimensionParameter;
-		Ar << TraceViewMatrixParameter;
-		return bShaderHasOutdatedParameters;
+		SetShaderValue(RHICmdList, ShaderRHI, TraceViewMatrixParameter, FMatrix(FRotationMatrix::Make(TraceOrientation)));
 	}
 
 private:
-	FShaderResourceParameter InTexture;
-	FShaderResourceParameter InTextureSampler;
-	FShaderResourceParameter BadTexture;
-	FShaderResourceParameter BadTextureSampler;
-	FShaderParameter TextureComponentReplicate;
-	FShaderParameter TextureComponentReplicateAlpha;
-	FShaderParameter ColorWeights; 
-	FShaderParameter PackedParameters;
-	FShaderParameter NumTilesPerSideParameter;
-	FShaderParameter TraceVolumeScalingParameter;
-	FShaderParameter TextureDimensionParameter;
-	FShaderParameter TraceViewMatrixParameter;
+	
+	LAYOUT_FIELD(FShaderResourceParameter, InTexture);
+	LAYOUT_FIELD(FShaderResourceParameter, InTextureSampler);
+	LAYOUT_FIELD(FShaderResourceParameter, BadTexture);
+	LAYOUT_FIELD(FShaderResourceParameter, BadTextureSampler);
+	LAYOUT_FIELD(FShaderParameter, TextureComponentReplicate);
+	LAYOUT_FIELD(FShaderParameter, TextureComponentReplicateAlpha);
+	LAYOUT_FIELD(FShaderParameter, ColorWeights);
+	LAYOUT_FIELD(FShaderParameter, PackedParameters);
+	LAYOUT_FIELD(FShaderParameter, NumTilesPerSideParameter);
+	LAYOUT_FIELD(FShaderParameter, TraceVolumeScalingParameter);
+	LAYOUT_FIELD(FShaderParameter, TextureDimensionParameter);
+	LAYOUT_FIELD(FShaderParameter, TraceViewMatrixParameter);
 };
 
 class FVolumeTextureTilePreviewPS : public FSimpleElementVolumeTexturePreviewPS
@@ -175,23 +161,22 @@ void FBatchedElementVolumeTexturePreviewParameters::BindShaders(
 {
 	TShaderMapRef<FSimpleElementVS> VertexShader(GetGlobalShaderMap(InFeatureLevel));
 
-	FSimpleElementVolumeTexturePreviewPS* PixelShader = nullptr;
-	
+	TShaderRef<FSimpleElementVolumeTexturePreviewPS> PixelShader ;
 	if (bViewModeAsDepthSlices)
 	{
 		TShaderMapRef<FVolumeTextureTilePreviewPS> TileShader(GetGlobalShaderMap(InFeatureLevel));
-		PixelShader = *TileShader;
+		PixelShader = TileShader;
 	}
 	else
 	{
 		TShaderMapRef<FVolumeTextureTracePreviewPS> TileShader(GetGlobalShaderMap(InFeatureLevel));
-		PixelShader = *TileShader;
+		PixelShader = TileShader;
 	}
 
 
 	GraphicsPSOInit.BoundShaderState.VertexDeclarationRHI = GSimpleElementVertexDeclaration.VertexDeclarationRHI;
-	GraphicsPSOInit.BoundShaderState.VertexShaderRHI = GETSAFERHISHADER_VERTEX(*VertexShader);
-	GraphicsPSOInit.BoundShaderState.PixelShaderRHI = GETSAFERHISHADER_PIXEL(PixelShader);
+	GraphicsPSOInit.BoundShaderState.VertexShaderRHI = VertexShader.GetVertexShader();
+	GraphicsPSOInit.BoundShaderState.PixelShaderRHI = PixelShader.GetPixelShader();
 	GraphicsPSOInit.PrimitiveType = PT_TriangleList;
 
 	if (!bViewModeAsDepthSlices)
diff --git a/Engine/Source/Programs/AutomationTool/AutomationUtils/CommandUtils.cs b/Engine/Source/Programs/AutomationTool/AutomationUtils/CommandUtils.cs
--- a/Engine/Source/Programs/AutomationTool/AutomationUtils/CommandUtils.cs
+++ b/Engine/Source/Programs/AutomationTool/AutomationUtils/CommandUtils.cs
@@ -661,11 +661,44 @@ namespace AutomationTool
 			}
 		}
 
+		/// <summary>
+		/// Renames/moves a directory.
+		/// If the rename of the directory fails, this function throws an Exception.
+		/// </summary>
+		/// <param name="OldName">Old name</param>
+		/// <param name="NewName">new name</param>
+		public static void RenameDirectory(string OldName, string NewName)
+		{
+			var OldNormalized = ConvertSeparators(PathSeparator.Default, OldName);
+			var NewNormalized = ConvertSeparators(PathSeparator.Default, NewName);
+			Directory.Move(OldNormalized, NewNormalized);
+		}
+
+		/// <summary>
+		/// Renames/moves a directory.
+		/// If the rename of the directory fails, this function prints a warning.
+		/// </summary>
+		/// <param name="OldName">Old name</param>
+		/// <param name="NewName">new name</param>
+		public static bool RenameDirectory_NoExceptions(string OldName, string NewName)
+		{
+			try
+			{
+				RenameDirectory(OldName, NewName);
+			}
+			catch (Exception)
+			{
+				LogWarning("Failed to rename/move file '{0}' to '{1}'", OldName, NewName);
+				return false;
+			}
+			return true;
+		}
+
 		/// <summary>
 		/// Creates a directory. Throws an exception on failure.
 		/// </summary>
-        /// <param name="DirectoryName">Name of the directory to create</param>
-        public static void CreateDirectory(string DirectoryName)
+		/// <param name="DirectoryName">Name of the directory to create</param>
+		public static void CreateDirectory(string DirectoryName)
 		{
 			string NormalizedDirectory = ConvertSeparators(PathSeparator.Default, DirectoryName);
 			try
diff --git a/Engine/Source/Programs/AutomationTool/AutomationUtils/ProjectParams.cs b/Engine/Source/Programs/AutomationTool/AutomationUtils/ProjectParams.cs
--- a/Engine/Source/Programs/AutomationTool/AutomationUtils/ProjectParams.cs
+++ b/Engine/Source/Programs/AutomationTool/AutomationUtils/ProjectParams.cs
@@ -456,6 +456,7 @@ namespace AutomationTool
 			bool? Package = null,
 			bool? Pak = null,
 			bool? IoStore = null,
+			bool? SkipIoStore = null,
 			bool? Prereqs = null,
 			string AppLocalDirectory = null,
 			bool? NoBootstrapExe = null,
@@ -629,6 +630,7 @@ namespace AutomationTool
 			this.PakAlignForMemoryMapping = GetParamValueIfNotSpecified(Command, PakAlignForMemoryMapping, this.PakAlignForMemoryMapping, "PakAlignForMemoryMapping");
 			this.Pak = GetParamValueIfNotSpecified(Command, Pak, this.Pak, "pak");
 			this.IoStore = GetParamValueIfNotSpecified(Command, IoStore, this.IoStore, "iostore");
+			this.SkipIoStore = GetParamValueIfNotSpecified(Command, IoStore, this.SkipIoStore, "skipiostore");
 			this.SkipPak = GetParamValueIfNotSpecified(Command, SkipPak, this.SkipPak, "skippak");
 			if (this.SkipPak)
 			{
@@ -1251,6 +1253,12 @@ namespace AutomationTool
 		[Help("skippak", "use a pak file, but assume it is already built, implies pak")]
 		public bool SkipPak { private set; get; }
 
+		/// <summary>
+		/// Shared: true if we want to skip iostore, even if -iostore is specified
+		/// </summary>
+		[Help("skipiostore", "override the -iostore commandline option to not run it")]
+		public bool SkipIoStore { private set; get; }
+
 		/// <summary>
 		/// Shared: true if this build is staged, command line: -stage
 		/// </summary>
@@ -2806,6 +2814,7 @@ namespace AutomationTool
 				CommandUtils.LogLog("MapIniSectionsToCook={0}", MapIniSectionsToCook.ToString());
 				CommandUtils.LogLog("Pak={0}", Pak);
 				CommandUtils.LogLog("IoStore={0}", IoStore);
+				CommandUtils.LogLog("SkipIoStore={0}", SkipIoStore);
 				CommandUtils.LogLog("Package={0}", Package);
 				CommandUtils.LogLog("ForcePackageData={0}", ForcePackageData);
 				CommandUtils.LogLog("NullRHI={0}", NullRHI);
diff --git a/Engine/Source/Programs/AutomationTool/Scripts/RunProjectCommand.Automation.cs b/Engine/Source/Programs/AutomationTool/Scripts/RunProjectCommand.Automation.cs
--- a/Engine/Source/Programs/AutomationTool/Scripts/RunProjectCommand.Automation.cs
+++ b/Engine/Source/Programs/AutomationTool/Scripts/RunProjectCommand.Automation.cs
@@ -790,10 +790,6 @@ public partial class Project : CommandUtils
 				{
 					TempCmdLine += "-signedpak ";
 				}
-				else
-				{
-					TempCmdLine += "-pak ";
-				}
 			}
 			else if (!Params.Stage)
 			{
@@ -985,10 +981,6 @@ public partial class Project : CommandUtils
 			{
 				Args += " -signedpak";
 			}
-			else
-			{
-				Args += " -pak";
-			}
 		}
 		if (IsBuildMachine || Params.Unattended)
 		{
diff --git a/Engine/Source/Programs/UnrealBuildTool/ProjectFiles/ProjectFileGenerator.cs b/Engine/Source/Programs/UnrealBuildTool/ProjectFiles/ProjectFileGenerator.cs
--- a/Engine/Source/Programs/UnrealBuildTool/ProjectFiles/ProjectFileGenerator.cs
+++ b/Engine/Source/Programs/UnrealBuildTool/ProjectFiles/ProjectFileGenerator.cs
@@ -1319,11 +1319,15 @@ namespace UnrealBuildTool
 		/// Adds additional files from the platform extensions folder
 		protected virtual void AddPlatformExtensionFiles( ProjectFile EngineProject )
 		{
-			DirectoryReference PlatformExtensionsDirectory = UnrealBuildTool.ProjectPlatformExtensionsDirectory(EngineProject.ProjectFilePath);
+			// @todo: this will add the same files to the solution (like the UBT source files that also get added to UnrealBuildTool project).
+			// not sure of a good filtering method here
+			DirectoryReference PlatformExtensionsDirectory = UnrealBuildTool.EnginePlatformExtensionsDirectory;
 			if (DirectoryReference.Exists(PlatformExtensionsDirectory))
 			{
 				List<string> SubdirectoryNamesToExclude = new List<string>();
 				SubdirectoryNamesToExclude.Add("AutomationTool"); //automation files are added separately to the AutomationTool project
+				SubdirectoryNamesToExclude.Add("Binaries");
+				SubdirectoryNamesToExclude.Add("Content");
 
 				EngineProject.AddFilesToProject(SourceFileSearch.FindFiles(PlatformExtensionsDirectory, SubdirectoryNamesToExclude), UnrealBuildTool.RootDirectory);
 			}
diff --git a/Engine/Source/Programs/UnrealBuildTool/System/PluginDescriptor.cs b/Engine/Source/Programs/UnrealBuildTool/System/PluginDescriptor.cs
--- a/Engine/Source/Programs/UnrealBuildTool/System/PluginDescriptor.cs
+++ b/Engine/Source/Programs/UnrealBuildTool/System/PluginDescriptor.cs
@@ -434,17 +434,5 @@ namespace UnrealBuildTool
 		{
 			return SupportedTargetPlatforms == null || SupportedTargetPlatforms.Length == 0 || SupportedTargetPlatforms.Contains(Platform);
 		}
-
-		/// <summary>
-		/// Combines the given target platforms with the current ones, if appropriate
-		/// </summary>
-		/// <param name="AdditionalSupportedTargetPlatforms"></param>
-		internal void MergeSupportedTargetPlatforms(UnrealTargetPlatform[] AdditionalSupportedTargetPlatforms)
-		{
-			if (SupportedTargetPlatforms != null && (SupportedTargetPlatforms.Length > 0) && AdditionalSupportedTargetPlatforms != null && AdditionalSupportedTargetPlatforms.Length > 0)
-			{
-				SupportedTargetPlatforms = SupportedTargetPlatforms.Union(AdditionalSupportedTargetPlatforms).ToArray();
-			}
-		}
 	}
 }
diff --git a/Engine/Source/Programs/UnrealBuildTool/System/Plugins.cs b/Engine/Source/Programs/UnrealBuildTool/System/Plugins.cs
--- a/Engine/Source/Programs/UnrealBuildTool/System/Plugins.cs
+++ b/Engine/Source/Programs/UnrealBuildTool/System/Plugins.cs
@@ -368,15 +368,27 @@ namespace UnrealBuildTool
 			// add our uplugin file to the existing plugin to be used to search for modules later
 			Parent.ChildFiles.Add(Child.File);
 
+			// this should cause an error if it's invalid platform name
+			UnrealTargetPlatform Platform = UnrealTargetPlatform.Parse(PlatformName);
+
 			// merge the supported platforms
-			Parent.Descriptor.MergeSupportedTargetPlatforms(Child.Descriptor.SupportedTargetPlatforms);
+			if (Parent.Descriptor.SupportedTargetPlatforms != null)
+			{
+				List<UnrealTargetPlatform> SupportedTargetPlatforms = Parent.Descriptor.SupportedTargetPlatforms.ToList();
+				if (Child.Descriptor.SupportedTargetPlatforms != null && Child.Descriptor.SupportedTargetPlatforms.Length > 0)
+				{
+					SupportedTargetPlatforms = SupportedTargetPlatforms.Union(Child.Descriptor.SupportedTargetPlatforms).ToList();
+				}
+				else if (!Parent.Descriptor.SupportedTargetPlatforms.Contains(Platform))
+				{
+					SupportedTargetPlatforms.Add(Platform);
+				}
+				Parent.Descriptor.SupportedTargetPlatforms = SupportedTargetPlatforms.ToArray();
+			}
 
 			// make sure we are whitelisted for any modules we list, if the parent had a whitelist
 			if (Child.Descriptor.Modules != null)
 			{
-				// this should cause an error if it's invalid platform name
-				UnrealTargetPlatform Platform = UnrealTargetPlatform.Parse(PlatformName);
-
 				foreach (ModuleDescriptor ChildModule in Child.Descriptor.Modules)
 				{
 					ModuleDescriptor ParentModule = Parent.Descriptor.Modules.FirstOrDefault(x => x.Name.Equals(ChildModule.Name) && x.Type == ChildModule.Type);
diff --git a/Engine/Source/Programs/UnrealBuildTool/System/RulesAssembly.cs b/Engine/Source/Programs/UnrealBuildTool/System/RulesAssembly.cs
--- a/Engine/Source/Programs/UnrealBuildTool/System/RulesAssembly.cs
+++ b/Engine/Source/Programs/UnrealBuildTool/System/RulesAssembly.cs
@@ -397,7 +397,11 @@ namespace UnrealBuildTool
 						throw new BuildException("Found multiple platform group overrides ({0} and {1}) for module {2} without a platform specific override. Create a platform override with the class hierarchy as needed.", 
 							GroupRulesObjectType.Name, PlatformRulesObjectType.Name, ModuleName);
 					}
-					PlatformRulesObjectType = GroupRulesObjectType;
+					// remember the platform group if we found it, but keep searching to verify there isn't more than one
+					if (GroupRulesObjectType != null)
+					{
+						PlatformRulesObjectType = GroupRulesObjectType;
+					}
 				}
 			}
 
@@ -445,6 +449,10 @@ namespace UnrealBuildTool
 							RulesObject.DirectoriesForModuleSubClasses.Add(SubType, SubTypeFileName.Directory);
 							RulesObject.SubclassRules.Add(SubTypeFileName.FullName);
 						}
+						if (SubType.BaseType == null)
+						{
+							throw new BuildException("{0} is not derived from {1}", RulesObjectType.Name, BaseRulesObjectType.Name);
+						}
 						SubType = SubType.BaseType;
 					}
 				}
diff --git a/Engine/Source/Programs/UnrealHeaderTool/Private/ClassMaps.cpp b/Engine/Source/Programs/UnrealHeaderTool/Private/ClassMaps.cpp
--- a/Engine/Source/Programs/UnrealHeaderTool/Private/ClassMaps.cpp
+++ b/Engine/Source/Programs/UnrealHeaderTool/Private/ClassMaps.cpp
@@ -19,3 +19,4 @@ TSet<FProperty*> GUnsizedProperties;
 TSet<UField*> GEditorOnlyDataTypes;
 TMap<UStruct*, TTuple<TSharedRef<FUnrealSourceFile>, int32>> GStructToSourceLine;
 TMap<UClass*, FArchiveTypeDefinePair> GClassSerializerMap;
+TSet<FProperty*> GPropertyUsesMemoryImageAllocator;
diff --git a/Engine/Source/Programs/UnrealHeaderTool/Private/ClassMaps.h b/Engine/Source/Programs/UnrealHeaderTool/Private/ClassMaps.h
--- a/Engine/Source/Programs/UnrealHeaderTool/Private/ClassMaps.h
+++ b/Engine/Source/Programs/UnrealHeaderTool/Private/ClassMaps.h
@@ -47,6 +47,7 @@ extern TSet<FProperty*> GUnsizedProperties;
 extern TSet<UField*> GEditorOnlyDataTypes;
 extern TMap<UStruct*, TTuple<TSharedRef<FUnrealSourceFile>, int32>> GStructToSourceLine;
 extern TMap<UClass*, FArchiveTypeDefinePair> GClassSerializerMap;
+extern TSet<FProperty*> GPropertyUsesMemoryImageAllocator;
 
 /** Types access specifiers. */
 enum EAccessSpecifier
diff --git a/Engine/Source/Programs/UnrealHeaderTool/Private/CodeGenerator.cpp b/Engine/Source/Programs/UnrealHeaderTool/Private/CodeGenerator.cpp
--- a/Engine/Source/Programs/UnrealHeaderTool/Private/CodeGenerator.cpp
+++ b/Engine/Source/Programs/UnrealHeaderTool/Private/CodeGenerator.cpp
@@ -1547,7 +1547,7 @@ void FNativeClassHeaderGenerator::PropertyNew(FOutputDevice& DeclOut, FOutputDev
 		DeclOut.Logf(TEXT("%sstatic const UE4CodeGen_Private::FArrayPropertyParams %s;\r\n"), DeclSpaces, *NameWithoutScope);
 
 		Out.Logf(
-			TEXT("%sconst UE4CodeGen_Private::FArrayPropertyParams %s = { %s, %s, (EPropertyFlags)0x%016llx, UE4CodeGen_Private::EPropertyGenFlags::Array, %s, %s, %s, %s };%s\r\n"),
+			TEXT("%sconst UE4CodeGen_Private::FArrayPropertyParams %s = { %s, %s, (EPropertyFlags)0x%016llx, UE4CodeGen_Private::EPropertyGenFlags::Array, %s, %s, %s, %s, %s };%s\r\n"),
 			Spaces,
 			Name,
 			*PropName,
@@ -1556,6 +1556,7 @@ void FNativeClassHeaderGenerator::PropertyNew(FOutputDevice& DeclOut, FOutputDev
 			FPropertyObjectFlags,
 			*ArrayDim,
 			OffsetStr,
+			GPropertyUsesMemoryImageAllocator.Contains(TypedProp) ? TEXT("EArrayPropertyFlags::UsesMemoryImageAllocator") : TEXT("EArrayPropertyFlags::None"),
 			*MetaDataParams,
 			*PropTag
 		);
@@ -1568,7 +1569,7 @@ void FNativeClassHeaderGenerator::PropertyNew(FOutputDevice& DeclOut, FOutputDev
 		DeclOut.Logf(TEXT("%sstatic const UE4CodeGen_Private::FMapPropertyParams %s;\r\n"), DeclSpaces, *NameWithoutScope);
 
 		Out.Logf(
-			TEXT("%sconst UE4CodeGen_Private::FMapPropertyParams %s = { %s, %s, (EPropertyFlags)0x%016llx, UE4CodeGen_Private::EPropertyGenFlags::Map, %s, %s, %s, %s };%s\r\n"),
+			TEXT("%sconst UE4CodeGen_Private::FMapPropertyParams %s = { %s, %s, (EPropertyFlags)0x%016llx, UE4CodeGen_Private::EPropertyGenFlags::Map, %s, %s, %s, %s, %s };%s\r\n"),
 			Spaces,
 			Name,
 			*PropName,
@@ -1577,6 +1578,7 @@ void FNativeClassHeaderGenerator::PropertyNew(FOutputDevice& DeclOut, FOutputDev
 			FPropertyObjectFlags,
 			*ArrayDim,
 			OffsetStr,
+			GPropertyUsesMemoryImageAllocator.Contains(TypedProp) ? TEXT("EMapPropertyFlags::UsesMemoryImageAllocator") : TEXT("EMapPropertyFlags::None"),
 			*MetaDataParams,
 			*PropTag
 		);
diff --git a/Engine/Source/Programs/UnrealHeaderTool/Private/HeaderParser.cpp b/Engine/Source/Programs/UnrealHeaderTool/Private/HeaderParser.cpp
--- a/Engine/Source/Programs/UnrealHeaderTool/Private/HeaderParser.cpp
+++ b/Engine/Source/Programs/UnrealHeaderTool/Private/HeaderParser.cpp
@@ -23,6 +23,7 @@
 #include "Containers/EnumAsByte.h"
 #include "Algo/AllOf.h"
 #include "Algo/FindSortedStringCaseInsensitive.h"
+#include "Misc/ScopeExit.h"
 
 #include "Specifiers/CheckedMetadataSpecifiers.h"
 #include "Specifiers/FunctionSpecifiers.h"
@@ -1281,6 +1282,16 @@ namespace
 
 		Parser.RequireSymbol(TEXT(')'), ErrorMessageGetter);
 	}
+
+	static const TCHAR* GLayoutMacroNames[] = {
+		TEXT("LAYOUT_ARRAY"),
+		TEXT("LAYOUT_ARRAY_EDITORONLY"),
+		TEXT("LAYOUT_BITFIELD"),
+		TEXT("LAYOUT_BITFIELD_EDITORONLY"),
+		TEXT("LAYOUT_FIELD"),
+		TEXT("LAYOUT_FIELD_EDITORONLY"),
+		TEXT("LAYOUT_FIELD_INITIALIZED"),
+	};
 }
 
 /////////////////////////////////////////////////////
@@ -3564,7 +3575,8 @@ void FHeaderParser::GetVarType(
 	const FToken*                   OuterPropertyType,
 	EPropertyDeclarationStyle::Type PropertyDeclarationStyle,
 	EVariableCategory::Type         VariableCategory,
-	FIndexRange*                    ParsedVarIndexRange
+	FIndexRange*                    ParsedVarIndexRange,
+	ELayoutMacroType*               OutLayoutMacroType
 )
 {
 	UStruct* OwnerStruct = Scope->IsFileScope() ? nullptr : ((FStructScope*)Scope)->GetStruct();
@@ -4093,6 +4105,48 @@ void FHeaderParser::GetVarType(
 	bool bUnconsumedEnumKeyword   = false;
 	bool bUnconsumedConstKeyword  = false;
 
+	// Handle MemoryLayout.h macros
+	ELayoutMacroType LayoutMacroType     = ELayoutMacroType::None;
+	bool             bHasWrapperBrackets = false;
+	ON_SCOPE_EXIT
+	{
+		if (OutLayoutMacroType)
+		{
+			*OutLayoutMacroType = LayoutMacroType;
+			if (bHasWrapperBrackets)
+			{
+				RequireSymbol(TEXT(')'), GLayoutMacroNames[(int32)LayoutMacroType]);
+			}
+		}
+	};
+
+	if (OutLayoutMacroType)
+	{
+		*OutLayoutMacroType = ELayoutMacroType::None;
+
+		FToken LayoutToken;
+		if (GetToken(LayoutToken))
+		{
+			if (LayoutToken.TokenType == TOKEN_Identifier)
+			{
+				LayoutMacroType = (ELayoutMacroType)Algo::FindSortedStringCaseInsensitive(LayoutToken.Identifier, GLayoutMacroNames, UE_ARRAY_COUNT(GLayoutMacroNames));
+				if (LayoutMacroType != ELayoutMacroType::None)
+				{
+					RequireSymbol(TEXT('('), GLayoutMacroNames[(int32)LayoutMacroType]);
+					if (LayoutMacroType == ELayoutMacroType::ArrayEditorOnly || LayoutMacroType == ELayoutMacroType::FieldEditorOnly || LayoutMacroType == ELayoutMacroType::BitfieldEditorOnly)
+					{
+						Flags |= CPF_EditorOnly;
+					}
+					bHasWrapperBrackets = MatchSymbol(TEXT("("));
+				}
+				else
+				{
+					UngetToken(LayoutToken);
+				}
+			}
+		}
+	}
+
 	if (MatchIdentifier(TEXT("const"), ESearchCase::CaseSensitive))
 	{
 		//@TODO: UCREMOVAL: Should use this to set the new (currently non-existent) CPF_Const flag appropriately!
@@ -4148,22 +4202,22 @@ void FHeaderParser::GetVarType(
 	{
 		VarProperty = FPropertyBase(CPT_Int64);
 	}
-	else if ( VarType.Matches(TEXT("uint64"), ESearchCase::CaseSensitive) && IsBitfieldProperty() )
+	else if ( VarType.Matches(TEXT("uint64"), ESearchCase::CaseSensitive) && IsBitfieldProperty(LayoutMacroType) )
 	{
 		// 64-bit bitfield (bool) type, treat it like 8 bit type
 		VarProperty = FPropertyBase(CPT_Bool8);
 	}
-	else if ( VarType.Matches(TEXT("uint32"), ESearchCase::CaseSensitive) && IsBitfieldProperty() )
+	else if ( VarType.Matches(TEXT("uint32"), ESearchCase::CaseSensitive) && IsBitfieldProperty(LayoutMacroType) )
 	{
 		// 32-bit bitfield (bool) type, treat it like 8 bit type
 		VarProperty = FPropertyBase(CPT_Bool8);
 	}
-	else if ( VarType.Matches(TEXT("uint16"), ESearchCase::CaseSensitive) && IsBitfieldProperty() )
+	else if ( VarType.Matches(TEXT("uint16"), ESearchCase::CaseSensitive) && IsBitfieldProperty(LayoutMacroType) )
 	{
 		// 16-bit bitfield (bool) type, treat it like 8 bit type.
 		VarProperty = FPropertyBase(CPT_Bool8);
 	}
-	else if ( VarType.Matches(TEXT("uint8"), ESearchCase::CaseSensitive) && IsBitfieldProperty() )
+	else if ( VarType.Matches(TEXT("uint8"), ESearchCase::CaseSensitive) && IsBitfieldProperty(LayoutMacroType) )
 	{
 		// 8-bit bitfield (bool) type
 		VarProperty = FPropertyBase(CPT_Bool8);
@@ -4184,7 +4238,7 @@ void FHeaderParser::GetVarType(
 	}
 	else if ( VarType.Matches(TEXT("bool"), ESearchCase::CaseSensitive) )
 	{
-		if (IsBitfieldProperty())
+		if (IsBitfieldProperty(LayoutMacroType))
 		{
 			UE_LOG_ERROR_UHT(TEXT("bool bitfields are not supported."));
 		}
@@ -4262,10 +4316,54 @@ void FHeaderParser::GetVarType(
 			FToken AllocatorToken;
 			if (!GetToken(AllocatorToken, /*bNoConsts=*/ true, ESymbolParseOption::CloseTemplateBracket))
 			{
-				FError::Throwf(TEXT("Expected '>' but found '%s'"), CloseTemplateToken.Identifier);
+				FError::Throwf(TEXT("Unexpected end of file when parsing TArray allocator."));
+			}
+
+			if (AllocatorToken.TokenType != TOKEN_Identifier)
+			{
+				FError::Throwf(TEXT("Found '%s' - expected a '>' or ','."), AllocatorToken.Identifier);
+			}
+
+			if (FCString::Strcmp(AllocatorToken.Identifier, TEXT("FMemoryImageAllocator")) == 0)
+			{
+				if (EnumHasAnyFlags(Flags, CPF_Net))
+				{
+					FError::Throwf(TEXT("Replicated arrays with MemoryImageAllocators are not yet supported"));
+				}
+
+				RequireSymbol(TEXT('>'), TEXT("TArray template arguments"), ESymbolParseOption::CloseTemplateBracket);
+
+				VarProperty.AllocatorType = EAllocatorType::MemoryImage;
 			}
+			else if (FCString::Strcmp(AllocatorToken.Identifier, TEXT("TMemoryImageAllocator")) == 0)
+			{
+				if (EnumHasAnyFlags(Flags, CPF_Net))
+				{
+					FError::Throwf(TEXT("Replicated arrays with MemoryImageAllocators are not yet supported"));
+				}
+
+				RequireSymbol(TEXT('<'), TEXT("TMemoryImageAllocator template arguments"));
 
-			FError::Throwf(TEXT("Found '%s' - explicit allocators are not supported in TArray properties."), AllocatorToken.Identifier);
+				FToken SkipToken;
+				for (;;)
+				{
+					if (!GetToken(SkipToken, /*bNoConsts=*/ false, ESymbolParseOption::CloseTemplateBracket))
+					{
+						FError::Throwf(TEXT("Unexpected end of file when parsing TMemoryImageAllocator template arguments."));
+					}
+
+					if (SkipToken.TokenType == TOKEN_Symbol && FCString::Strcmp(SkipToken.Identifier, TEXT(">")) == 0)
+					{
+						RequireSymbol(TEXT('>'), TEXT("TArray template arguments"), ESymbolParseOption::CloseTemplateBracket);
+						VarProperty.AllocatorType = EAllocatorType::MemoryImage;
+						break;
+					}
+				}
+			}
+			else
+			{
+				FError::Throwf(TEXT("Found '%s' - explicit allocators are not supported in TArray properties."), AllocatorToken.Identifier);
+			}
 		}
 	}
 	else if ( VarType.Matches(TEXT("TMap"), ESearchCase::CaseSensitive) )
@@ -4328,10 +4426,29 @@ void FHeaderParser::GetVarType(
 			FToken AllocatorToken;
 			if (!GetToken(AllocatorToken, /*bNoConsts=*/ true, ESymbolParseOption::CloseTemplateBracket))
 			{
-				FError::Throwf(TEXT("Expected '>' but found '%s'"), CloseTemplateToken.Identifier);
+				FError::Throwf(TEXT("Unexpected end of file when parsing TArray allocator."));
+			}
+
+			if (AllocatorToken.TokenType != TOKEN_Identifier)
+			{
+				FError::Throwf(TEXT("Found '%s' - expected a '>' or ','."), AllocatorToken.Identifier);
 			}
 
-			FError::Throwf(TEXT("Found '%s' - explicit allocators are not supported in TMap properties."), AllocatorToken.Identifier);
+			if (FCString::Strcmp(AllocatorToken.Identifier, TEXT("FMemoryImageSetAllocator")) == 0)
+			{
+				if (EnumHasAnyFlags(Flags, CPF_Net))
+				{
+					FError::Throwf(TEXT("Replicated maps with MemoryImageSetAllocators are not yet supported"));
+				}
+
+				RequireSymbol(TEXT('>'), TEXT("TMap template arguments"), ESymbolParseOption::CloseTemplateBracket);
+
+				VarProperty.AllocatorType = EAllocatorType::MemoryImage;
+			}
+			else
+			{
+				FError::Throwf(TEXT("Found '%s' - explicit allocators are not supported in TMap properties."), AllocatorToken.Identifier);
+			}
 		}
 	}
 	else if ( VarType.Matches(TEXT("TSet"), ESearchCase::CaseSensitive) )
@@ -4384,7 +4501,7 @@ void FHeaderParser::GetVarType(
 			FError::Throwf(TEXT("Found '%s' - explicit KeyFuncs are not supported in TSet properties."), AllocatorToken.Identifier);
 		}
 	}
-	else if ( VarType.Matches(TEXT("FString"), ESearchCase::CaseSensitive) )
+	else if ( VarType.Matches(TEXT("FString"), ESearchCase::CaseSensitive) || VarType.Matches(TEXT("FMemoryImageString"), ESearchCase::CaseSensitive))
 	{
 		VarProperty = FPropertyBase(CPT_String);
 
@@ -5012,7 +5129,8 @@ FProperty* FHeaderParser::GetVarNameAndDim
 (
 	UStruct*                Scope,
 	FToken&                 VarProperty,
-	EVariableCategory::Type VariableCategory
+	EVariableCategory::Type VariableCategory,
+	ELayoutMacroType        LayoutMacroType
 )
 {
 	check(Scope);
@@ -5043,6 +5161,20 @@ FProperty* FHeaderParser::GetVarNameAndDim
 			FError::Throwf(TEXT("Missing variable name") );
 		}
 
+		switch (LayoutMacroType)
+		{
+			case ELayoutMacroType::Array:
+			case ELayoutMacroType::ArrayEditorOnly:
+			case ELayoutMacroType::Bitfield:
+			case ELayoutMacroType::BitfieldEditorOnly:
+			case ELayoutMacroType::FieldInitialized:
+				RequireSymbol(TEXT(','), GLayoutMacroNames[(int32)LayoutMacroType]);
+				break;
+
+			default:
+				break;
+		}
+
 		VarProperty.TokenType = TOKEN_Identifier;
 		FCString::Strcpy(VarProperty.Identifier, VarToken.Identifier);
 	}
@@ -5132,7 +5264,7 @@ FProperty* FHeaderParser::GetVarNameAndDim
 
 	// Get optional dimension immediately after name.
 	FToken Dimensions;
-	if (MatchSymbol(TEXT('[')))
+	if ((LayoutMacroType == ELayoutMacroType::None && MatchSymbol(TEXT('['))) || LayoutMacroType == ELayoutMacroType::Array || LayoutMacroType == ELayoutMacroType::ArrayEditorOnly)
 	{
 		switch (VariableCategory)
 		{
@@ -5158,10 +5290,21 @@ FProperty* FHeaderParser::GetVarNameAndDim
 			FError::Throwf(TEXT("Bool arrays are not allowed") );
 		}
 
-		// Ignore how the actual array dimensions are actually defined - we'll calculate those with the compiler anyway.
-		if (!GetRawToken(Dimensions, TEXT(']')))
+		if (LayoutMacroType == ELayoutMacroType::None)
+		{
+			// Ignore how the actual array dimensions are actually defined - we'll calculate those with the compiler anyway.
+			if (!GetRawToken(Dimensions, TEXT(']')))
+			{
+				FError::Throwf(TEXT("%s %s: Missing ']'"), HintText, VarProperty.Identifier);
+			}
+		}
+		else
 		{
-			FError::Throwf(TEXT("%s %s: Missing ']'"), HintText, VarProperty.Identifier );
+			// Ignore how the actual array dimensions are actually defined - we'll calculate those with the compiler anyway.
+			if (!GetRawToken(Dimensions, TEXT(')')))
+			{
+				FError::Throwf(TEXT("%s %s: Missing ']'"), HintText, VarProperty.Identifier);
+			}
 		}
 
 		// Only static arrays are declared with [].  Dynamic arrays use TArray<> instead.
@@ -5231,7 +5374,10 @@ FProperty* FHeaderParser::GetVarNameAndDim
 			VarProperty.MetaData.Add(NAME_ArraySizeEnum, Enum->GetPathName());
 		}
 
-		MatchSymbol(TEXT(']'));
+		if (LayoutMacroType == ELayoutMacroType::None)
+		{
+			MatchSymbol(TEXT(']'));
+		}
 	}
 
 	// Try gathering metadata for member fields
@@ -5294,6 +5440,11 @@ FProperty* FHeaderParser::GetVarNameAndDim
 		PropagateFlagsFromInnerAndHandlePersistentInstanceMetadata(Array->PropertyFlags, VarProperty.MetaData, InnerProp);
 
 		Result = Array;
+
+		if (VarProperty.AllocatorType == EAllocatorType::MemoryImage)
+		{
+			GPropertyUsesMemoryImageAllocator.Add(Array);
+		}
 	}
 	else if (VarProperty.ArrayType == EArrayType::Set)
 	{
@@ -5328,6 +5479,11 @@ FProperty* FHeaderParser::GetVarNameAndDim
 		PropagateFlagsFromInnerAndHandlePersistentInstanceMetadata(Map->PropertyFlags, VarProperty.MetaData,             ValueProp);
 
 		Result = Map;
+
+		if (VarProperty.AllocatorType == EAllocatorType::MemoryImage)
+		{
+			GPropertyUsesMemoryImageAllocator.Add(Map);
+		}
 	}
 	else
 	{
@@ -7689,8 +7845,13 @@ void FHeaderParser::ParseFieldMetaData(TMap<FName, FString>& MetaData, const TCH
 	}
 }
 
-bool FHeaderParser::IsBitfieldProperty()
+bool FHeaderParser::IsBitfieldProperty(ELayoutMacroType LayoutMacroType)
 {
+	if (LayoutMacroType == ELayoutMacroType::Bitfield || LayoutMacroType == ELayoutMacroType::BitfieldEditorOnly)
+	{
+		return true;
+	}
+
 	bool bIsBitfield = false;
 
 	// The current token is the property type (uin32, uint16, etc).
@@ -7772,7 +7933,8 @@ void FHeaderParser::CompileVariableDeclaration(FClasses& AllClasses, UStruct* St
 	// Get variable type.
 	FPropertyBase OriginalProperty(CPT_None);
 	FIndexRange TypeRange;
-	GetVarType( AllClasses, &FScope::GetTypeScope(Struct).Get(), OriginalProperty, DisallowFlags, /*OuterPropertyType=*/ NULL, EPropertyDeclarationStyle::UPROPERTY, EVariableCategory::Member, &TypeRange );
+	ELayoutMacroType LayoutMacroType = ELayoutMacroType::None;
+	GetVarType( AllClasses, &FScope::GetTypeScope(Struct).Get(), OriginalProperty, DisallowFlags, /*OuterPropertyType=*/ NULL, EPropertyDeclarationStyle::UPROPERTY, EVariableCategory::Member, &TypeRange, &LayoutMacroType);
 	OriginalProperty.PropertyFlags |= EdFlags;
 
 	FString* Category = OriginalProperty.MetaData.Find(NAME_Category);
@@ -7832,20 +7994,28 @@ void FHeaderParser::CompileVariableDeclaration(FClasses& AllClasses, UStruct* St
 		}
 	}
 
+	if (LayoutMacroType != ELayoutMacroType::None)
+	{
+		RequireSymbol(TEXT(','), GLayoutMacroNames[(int32)LayoutMacroType]);
+	}
+
 	// Process all variables of this type.
 	TArray<FProperty*> NewProperties;
-	do
+	for (;;)
 	{
 		FToken     Property    = OriginalProperty;
-		FProperty* NewProperty = GetVarNameAndDim(Struct, Property, EVariableCategory::Member);
+		FProperty* NewProperty = GetVarNameAndDim(Struct, Property, EVariableCategory::Member, LayoutMacroType);
 
 		// Optionally consume the :1 at the end of a bitfield boolean declaration
-		if (Property.IsBool() && MatchSymbol(TEXT(':')))
+		if (Property.IsBool())
 		{
-			int32 BitfieldSize = 0;
-			if (!GetConstInt(/*out*/ BitfieldSize) || (BitfieldSize != 1))
+			if (LayoutMacroType == ELayoutMacroType::Bitfield || LayoutMacroType == ELayoutMacroType::BitfieldEditorOnly || MatchSymbol(TEXT(':')))
 			{
-				FError::Throwf(TEXT("Bad or missing bitfield size for '%s', must be 1."), *NewProperty->GetName());
+				int32 BitfieldSize = 0;
+				if (!GetConstInt(/*out*/ BitfieldSize) || (BitfieldSize != 1))
+				{
+					FError::Throwf(TEXT("Bad or missing bitfield size for '%s', must be 1."), *NewProperty->GetName());
+				}
 			}
 		}
 
@@ -7897,10 +8067,36 @@ void FHeaderParser::CompileVariableDeclaration(FClasses& AllClasses, UStruct* St
 			}
 		}
 
-	} while( MatchSymbol(TEXT(',')) );
+		if (LayoutMacroType != ELayoutMacroType::None || !MatchSymbol(TEXT(',')))
+		{
+			break;
+		}
+	}
 
 	// Optional member initializer.
-	if (MatchSymbol(TEXT('=')))
+	if (LayoutMacroType == ELayoutMacroType::FieldInitialized)
+	{
+		// Skip past the specified member initializer; we make no attempt to parse it
+		FToken SkipToken;
+		int Nesting = 1;
+		while (GetToken(SkipToken))
+		{
+			if (SkipToken.Matches(TEXT('(')))
+			{
+				++Nesting;
+			}
+			else if (SkipToken.Matches(TEXT(')')))
+			{
+				--Nesting;
+				if (Nesting == 0)
+				{
+					UngetToken(SkipToken);
+					break;
+				}
+			}
+		}
+	}
+	else if (MatchSymbol(TEXT('=')))
 	{
 		// Skip past the specified member initializer; we make no attempt to parse it
 		FToken SkipToken;
@@ -7936,8 +8132,16 @@ void FHeaderParser::CompileVariableDeclaration(FClasses& AllClasses, UStruct* St
 		}
 	}
 
-	// Expect a semicolon.
-	RequireSymbol( TEXT(';'), TEXT("'variable declaration'") );
+	if (LayoutMacroType == ELayoutMacroType::None)
+	{
+		// Expect a semicolon.
+		RequireSymbol(TEXT(';'), TEXT("'variable declaration'"));
+	}
+	else
+	{
+		// Expect a close bracket.
+		RequireSymbol(TEXT(')'), GLayoutMacroNames[(int32)LayoutMacroType]);
+	}
 
 	// Skip redundant semi-colons
 	for (;;)
diff --git a/Engine/Source/Programs/UnrealHeaderTool/Private/HeaderParser.h b/Engine/Source/Programs/UnrealHeaderTool/Private/HeaderParser.h
--- a/Engine/Source/Programs/UnrealHeaderTool/Private/HeaderParser.h
+++ b/Engine/Source/Programs/UnrealHeaderTool/Private/HeaderParser.h
@@ -81,6 +81,18 @@ namespace EVariableCategory
 	};
 }
 
+enum class ELayoutMacroType
+{
+	None = -1,
+	Array,
+	ArrayEditorOnly,
+	Bitfield,
+	BitfieldEditorOnly,
+	Field,
+	FieldEditorOnly,
+	FieldInitialized,
+};
+
 /** Information for a particular nesting level. */
 class FNestInfo
 {
@@ -671,7 +683,7 @@ protected:
 	bool IsValidDelegateDeclaration(const FToken& Token) const;
 
 	// Returns true if the current token is a bitfield type
-	bool IsBitfieldProperty();
+	bool IsBitfieldProperty(ELayoutMacroType LayoutMacroType);
 
 	// Parse the parameter list of a function or delegate declaration
 	void ParseParameterList(FClasses& AllClasses, UFunction* Function, bool bExpectCommaBeforeName = false, TMap<FName, FString>* MetaData = NULL);
@@ -843,7 +855,8 @@ protected:
 		const FToken*                   OuterPropertyType,
 		EPropertyDeclarationStyle::Type PropertyDeclarationStyle,
 		EVariableCategory::Type         VariableCategory,
-		FIndexRange*                    ParsedVarIndexRange = nullptr);
+		FIndexRange*                    ParsedVarIndexRange = nullptr,
+		ELayoutMacroType*               OutLayoutMacroType = nullptr);
 
 	/**
 	 * Parses a variable name declaration and creates a new FProperty object.
@@ -857,7 +870,8 @@ protected:
 	FProperty* GetVarNameAndDim(
 		UStruct* Struct,
 		FToken& VarProperty,
-		EVariableCategory::Type VariableCategory);
+		EVariableCategory::Type VariableCategory,
+		ELayoutMacroType LayoutMacroType = ELayoutMacroType::None);
 	
 	/**
 	 * Returns whether the specified class can be referenced from the class currently being compiled.
diff --git a/Engine/Source/Programs/UnrealHeaderTool/Private/ParserHelper.h b/Engine/Source/Programs/UnrealHeaderTool/Private/ParserHelper.h
--- a/Engine/Source/Programs/UnrealHeaderTool/Private/ParserHelper.h
+++ b/Engine/Source/Programs/UnrealHeaderTool/Private/ParserHelper.h
@@ -69,6 +69,12 @@ struct EArrayType
 	};
 };
 
+enum class EAllocatorType
+{
+	Default,
+	MemoryImage
+};
+
 struct ERefQualifier
 {
 	enum Type
@@ -99,6 +105,7 @@ public:
 	// Variables.
 	EPropertyType       Type;
 	EArrayType::Type    ArrayType;
+	EAllocatorType      AllocatorType = EAllocatorType::Default;
 	EPropertyFlags      PropertyFlags;
 	EPropertyFlags      ImpliedPropertyFlags;
 	ERefQualifier::Type RefQualifier; // This is needed because of legacy stuff - FString mangles the flags for reasons that have become lost in time but we need this info for testing for invalid replicated function signatures.
diff --git a/Engine/Source/Programs/UnrealHeaderTool/Resources/UHTDebugging/TestObject.h b/Engine/Source/Programs/UnrealHeaderTool/Resources/UHTDebugging/TestObject.h
--- a/Engine/Source/Programs/UnrealHeaderTool/Resources/UHTDebugging/TestObject.h
+++ b/Engine/Source/Programs/UnrealHeaderTool/Resources/UHTDebugging/TestObject.h
@@ -42,13 +42,22 @@ public:
 	TArray<FContainsInstancedProperty> InstancedPropertyArray;
 
 	UPROPERTY()
-	TArray<TWeakObjectPtr<UObject>> ObjectWrapperArray;
+	TArray<TWeakObjectPtr<UObject>, FMemoryImageAllocator> ObjectWrapperArray;
+
+	UPROPERTY()
+	TArray<TWeakObjectPtr<UObject>, TMemoryImageAllocator<64>> ObjectWrapperArrayTemplated;
 
 	UPROPERTY()
 	TSet<FContainsInstancedProperty> InstancedPropertySet;
 
 	UPROPERTY()
-	TMap<FContainsInstancedProperty, TWeakObjectPtr<UObject>> InstancedPropertyToObjectWrapperMap;
+	LAYOUT_FIELD((TMap<FContainsInstancedProperty, TWeakObjectPtr<UObject>>), InstancedPropertyToObjectWrapperMap);
+
+	UPROPERTY()
+	TMap<FContainsInstancedProperty, TWeakObjectPtr<UObject>, FMemoryImageSetAllocator> InstancedPropertyToObjectWrapperMapFrozen;
+
+	UPROPERTY()
+	LAYOUT_FIELD((TMap<FContainsInstancedProperty, TWeakObjectPtr<UObject>, FMemoryImageSetAllocator>), InstancedPropertyToObjectWrapperMapFrozenWithLayout);
 
 	UPROPERTY()
 	TMap<TWeakObjectPtr<UObject>, FContainsInstancedProperty> ObjectWrapperToInstancedPropertyMap;
@@ -60,7 +69,7 @@ public:
 	void TestPassingArrayOfInterfaces(const TArray<TScriptInterface<ITestInterface> >& ArrayOfInterfaces);
 
 	UPROPERTY()
-	int32 Cpp11Init = 123;
+	LAYOUT_FIELD_INITIALIZED(int32, Cpp11Init, 123);
 
 	UPROPERTY()
 	TArray<int> Cpp11BracedInit { 1, 2, 3 };
@@ -72,13 +81,19 @@ public:
 	int RawInt;
 
 	UPROPERTY()
-	unsigned int RawUint;
+	LAYOUT_FIELD_EDITORONLY(int32, EditorOnlyField);
+
+	UPROPERTY()
+	LAYOUT_ARRAY_EDITORONLY(int32, EditorOnlyArray, 20);
+
+	UPROPERTY()
+	LAYOUT_ARRAY(unsigned int, RawUint, 20);
 
 	UFUNCTION()
 	void FuncTakingRawInts(int Signed, unsigned int Unsigned);
 
 	UPROPERTY()
-	ECppEnum EnumProperty;
+	LAYOUT_FIELD(ECppEnum, EnumProperty);
 
 	UPROPERTY()
 	TMap<int32, bool> TestMap;
@@ -92,6 +107,9 @@ public:
 	UPROPERTY()
 	FSimpleClassDelegate DelegateProperty;
 
+	UPROPERTY()
+	LAYOUT_BITFIELD(uint32, bThing, 1);
+
 	UFUNCTION()
 	void CodeGenTestForEnumClasses(ECppEnum Val);
 
diff --git a/Engine/Source/Runtime/AIModule/Private/EnvironmentQuery/EQSRenderingComponent.cpp b/Engine/Source/Runtime/AIModule/Private/EnvironmentQuery/EQSRenderingComponent.cpp
--- a/Engine/Source/Runtime/AIModule/Private/EnvironmentQuery/EQSRenderingComponent.cpp
+++ b/Engine/Source/Runtime/AIModule/Private/EnvironmentQuery/EQSRenderingComponent.cpp
@@ -348,9 +348,9 @@ FBoxSphereBounds UEQSRenderingComponent::CalcBounds(const FTransform& LocalToWor
 	return FBoxSphereBounds(BoundingSphere).TransformBy(LocalToWorld);
 }
 
-void UEQSRenderingComponent::CreateRenderState_Concurrent()
+void UEQSRenderingComponent::CreateRenderState_Concurrent(FRegisterComponentContext* Context)
 {
-	Super::CreateRenderState_Concurrent();
+	Super::CreateRenderState_Concurrent(Context);
 
 #if USE_EQS_DEBUGGER
 	EQSRenderingDebugDrawDelegateHelper.RegisterDebugDrawDelgate();
diff --git a/Engine/Source/Runtime/AVEncoder/Private/AVEncoderCommon.cpp b/Engine/Source/Runtime/AVEncoder/Private/AVEncoderCommon.cpp
--- a/Engine/Source/Runtime/AVEncoder/Private/AVEncoderCommon.cpp
+++ b/Engine/Source/Runtime/AVEncoder/Private/AVEncoderCommon.cpp
@@ -81,13 +81,13 @@ void CopyTextureImpl(const FTexture2DRHIRef& Src, FTexture2DRHIRef& Dst, FRHIGPU
 			GraphicsPSOInit.RasterizerState = TStaticRasterizerState<>::GetRHI();
 			GraphicsPSOInit.DepthStencilState = TStaticDepthStencilState<false, CF_Always>::GetRHI();
 
-			TShaderMap<FGlobalShaderType>* ShaderMap = GetGlobalShaderMap(GMaxRHIFeatureLevel);
+			FGlobalShaderMap* ShaderMap = GetGlobalShaderMap(GMaxRHIFeatureLevel);
 			TShaderMapRef<FScreenVS> VertexShader(ShaderMap);
 			TShaderMapRef<FScreenPS> PixelShader(ShaderMap);
 
 			GraphicsPSOInit.BoundShaderState.VertexDeclarationRHI = GFilterVertexDeclaration.VertexDeclarationRHI;
-			GraphicsPSOInit.BoundShaderState.VertexShaderRHI = GETSAFERHISHADER_VERTEX(*VertexShader);
-			GraphicsPSOInit.BoundShaderState.PixelShaderRHI = GETSAFERHISHADER_PIXEL(*PixelShader);
+			GraphicsPSOInit.BoundShaderState.VertexShaderRHI = VertexShader.GetVertexShader();
+			GraphicsPSOInit.BoundShaderState.PixelShaderRHI = PixelShader.GetPixelShader();
 			GraphicsPSOInit.PrimitiveType = PT_TriangleList;
 
 			SetGraphicsPipelineState(RHICmdList, GraphicsPSOInit);
@@ -110,7 +110,7 @@ void CopyTextureImpl(const FTexture2DRHIRef& Src, FTexture2DRHIRef& Dst, FRHIGPU
 				1, 1,						// Source USize, VSize
 				Dst->GetSizeXY(),			// Target buffer size
 				FIntPoint(1, 1),			// Source texture size
-				*VertexShader,
+				VertexShader,
 				EDRF_Default);
 		}
 		RHICmdList.EndRenderPass();
diff --git a/Engine/Source/Runtime/Apple/MetalRHI/Private/MetalCommands.cpp b/Engine/Source/Runtime/Apple/MetalRHI/Private/MetalCommands.cpp
--- a/Engine/Source/Runtime/Apple/MetalRHI/Private/MetalCommands.cpp
+++ b/Engine/Source/Runtime/Apple/MetalRHI/Private/MetalCommands.cpp
@@ -91,16 +91,16 @@ mtlpp::PrimitiveType TranslatePrimitiveType(uint32 PrimitiveType)
 static FORCEINLINE EMetalShaderStages GetShaderStage(FRHIGraphicsShader* ShaderRHI)
 {
 	EMetalShaderStages Stage = EMetalShaderStages::Num;
-	switch (ShaderRHI->GetType())
+	switch (ShaderRHI->GetFrequency())
 	{
-	case FRHIShader::EType::Vertex:		Stage = EMetalShaderStages::Vertex; break;
+	case SF_Vertex:		Stage = EMetalShaderStages::Vertex; break;
 #if PLATFORM_SUPPORTS_TESSELLATION_SHADERS
-	case FRHIShader::EType::Hull:		Stage = EMetalShaderStages::Hull; break;
-	case FRHIShader::EType::Domain:		Stage = EMetalShaderStages::Domain; break;
+	case SF_Hull:		Stage = EMetalShaderStages::Hull; break;
+	case SF_Domain:		Stage = EMetalShaderStages::Domain; break;
 #endif
-	case FRHIShader::EType::Pixel:		Stage = EMetalShaderStages::Pixel; break;
+	case SF_Pixel:		Stage = EMetalShaderStages::Pixel; break;
 	default:
-		checkf(0, TEXT("FRHIShader Type %d is invalid or unsupported!"), (int32)ShaderRHI->GetType());
+		checkf(0, TEXT("FRHIShader Type %d is invalid or unsupported!"), (int32)ShaderRHI->GetFrequency());
 		NOT_SUPPORTED("RHIShaderStage");
 		break;
 	}
@@ -406,9 +406,9 @@ void FMetalRHICommandContext::RHISetShaderUniformBuffer(FRHIGraphicsShader* Shad
 	TMap<uint32, TBitArray<>>* ArgumentBitmasks = nullptr;
 	TMap<uint32, mtlpp::ArgumentEncoder>* ArgumentEncoders = nullptr;
 	FMetalShaderBindings* Bindings = nullptr;
-	switch (ShaderRHI->GetType())
+	switch (ShaderRHI->GetFrequency())
 	{
-	case FRHIShader::EType::Vertex:
+	case SF_Vertex:
 	{
 		FMetalVertexShader* VertexShader = ResourceCast(static_cast<FRHIVertexShader*>(ShaderRHI));
 		ArgumentBitmasks = &VertexShader->ArgumentBitmasks;
@@ -418,7 +418,7 @@ void FMetalRHICommandContext::RHISetShaderUniformBuffer(FRHIGraphicsShader* Shad
 	}
 		break;
 #if PLATFORM_SUPPORTS_TESSELLATION_SHADERS
-	case FRHIShader::EType::Hull:
+	case SF_Hull:
 	{
 		Stage = EMetalShaderStages::Hull;
 		FMetalHullShader* HullShader = ResourceCast(static_cast<FRHIHullShader*>(ShaderRHI));
@@ -427,7 +427,7 @@ void FMetalRHICommandContext::RHISetShaderUniformBuffer(FRHIGraphicsShader* Shad
 		Bindings = &HullShader->Bindings;
 	}
 		break;
-	case FRHIShader::EType::Domain:
+	case SF_Domain:
 	{
 		Stage = EMetalShaderStages::Domain;
 		FMetalDomainShader* DomainShader = ResourceCast(static_cast<FRHIDomainShader*>(ShaderRHI));
@@ -437,7 +437,7 @@ void FMetalRHICommandContext::RHISetShaderUniformBuffer(FRHIGraphicsShader* Shad
 	}
 		break;
 #endif
-	case FRHIShader::EType::Pixel:
+	case SF_Pixel:
 	{
 		Stage = EMetalShaderStages::Pixel;
 		FMetalPixelShader* PixelShader = ResourceCast(static_cast<FRHIPixelShader*>(ShaderRHI));
@@ -447,7 +447,7 @@ void FMetalRHICommandContext::RHISetShaderUniformBuffer(FRHIGraphicsShader* Shad
 	}
 		break;
 	default:
-		checkf(0, TEXT("FRHIShader Type %d is invalid or unsupported!"), (int32)ShaderRHI->GetType());
+		checkf(0, TEXT("FRHIShader Type %d is invalid or unsupported!"), (int32)ShaderRHI->GetFrequency());
 		NOT_SUPPORTED("RHIShaderStage");
 		break;
 	}
diff --git a/Engine/Source/Runtime/Apple/MetalRHI/Private/MetalPipeline.h b/Engine/Source/Runtime/Apple/MetalRHI/Private/MetalPipeline.h
--- a/Engine/Source/Runtime/Apple/MetalRHI/Private/MetalPipeline.h
+++ b/Engine/Source/Runtime/Apple/MetalRHI/Private/MetalPipeline.h
@@ -16,7 +16,7 @@ enum EMetalPipelineHashBits
 	NumBits_StencilFormat = 3, //(x1=3),
 	NumBits_SampleCount = 3, //(x1=3),
 
-	NumBits_BlendState = 5, //(x8=40),
+	NumBits_BlendState = 7, //(x8=56),
 	NumBits_PrimitiveTopology = 2, //(x1=2)
 	NumBits_IndexType = 2, //(x1=2)
 };
diff --git a/Engine/Source/Runtime/Apple/MetalRHI/Private/MetalRHI.cpp b/Engine/Source/Runtime/Apple/MetalRHI/Private/MetalRHI.cpp
--- a/Engine/Source/Runtime/Apple/MetalRHI/Private/MetalRHI.cpp
+++ b/Engine/Source/Runtime/Apple/MetalRHI/Private/MetalRHI.cpp
@@ -854,21 +854,11 @@ FMetalDynamicRHI::~FMetalDynamicRHI()
 {
 	check(IsInGameThread() && IsInRenderingThread());
 	
-	// Ask all initialized FRenderResources to release their RHI resources.
-	for (TLinkedList<FRenderResource*>::TIterator ResourceIt(FRenderResource::GetResourceList()); ResourceIt; ResourceIt.Next())
-	{
-		FRenderResource* Resource = *ResourceIt;
-		check(Resource->IsInitialized());
-		Resource->ReleaseRHI();
-	}
-	
-	for (TLinkedList<FRenderResource*>::TIterator ResourceIt(FRenderResource::GetResourceList()); ResourceIt; ResourceIt.Next())
-	{
-		ResourceIt->ReleaseDynamicRHI();
-	}
-	
 	GIsMetalInitialized = false;
 	GIsRHIInitialized = false;
+
+	// Ask all initialized FRenderResources to release their RHI resources.
+	FRenderResource::ReleaseRHIForAllResources();	
 	
 #if ENABLE_METAL_GPUPROFILE
 	FMetalProfiler::DestroyProfiler();
diff --git a/Engine/Source/Runtime/Apple/MetalRHI/Private/MetalShaders.cpp b/Engine/Source/Runtime/Apple/MetalRHI/Private/MetalShaders.cpp
--- a/Engine/Source/Runtime/Apple/MetalRHI/Private/MetalShaders.cpp
+++ b/Engine/Source/Runtime/Apple/MetalRHI/Private/MetalShaders.cpp
@@ -489,11 +489,11 @@ mtlpp::LanguageVersion ValidateVersion(uint8 Version)
 
 /** Initialization constructor. */
 template<typename BaseResourceType, int32 ShaderType>
-void TMetalBaseShader<BaseResourceType, ShaderType>::Init(const TArray<uint8>& InShaderCode, FMetalCodeHeader& Header, mtlpp::Library InLibrary)
+void TMetalBaseShader<BaseResourceType, ShaderType>::Init(TArrayView<const uint8> InShaderCode, FMetalCodeHeader& Header, mtlpp::Library InLibrary)
 {
 	FShaderCodeReader ShaderCode(InShaderCode);
 
-	FMemoryReader Ar(InShaderCode, true);
+	FMemoryReaderView Ar(InShaderCode, true);
 	
 	Ar.SetLimitSize(ShaderCode.GetActualShaderCodeSize());
 
@@ -896,7 +896,7 @@ mtlpp::Function TMetalBaseShader<BaseResourceType, ShaderType>::GetCompiledFunct
 	return Func;
 }
 
-FMetalComputeShader::FMetalComputeShader(const TArray<uint8>& InCode, mtlpp::Library InLibrary)
+FMetalComputeShader::FMetalComputeShader(TArrayView<const uint8> InCode, mtlpp::Library InLibrary)
 : NumThreadsX(0)
 , NumThreadsY(0)
 , NumThreadsZ(0)
@@ -1006,7 +1006,7 @@ mtlpp::Function FMetalComputeShader::GetFunction()
 	return GetCompiledFunction();
 }
 
-FMetalVertexShader::FMetalVertexShader(const TArray<uint8>& InCode)
+FMetalVertexShader::FMetalVertexShader(TArrayView<const uint8> InCode)
 {
 	FMetalCodeHeader Header;
 	Init(InCode, Header);
@@ -1031,7 +1031,7 @@ FMetalVertexShader::FMetalVertexShader(const TArray<uint8>& InCode)
 #endif
 }
 
-FMetalVertexShader::FMetalVertexShader(const TArray<uint8>& InCode, mtlpp::Library InLibrary)
+FMetalVertexShader::FMetalVertexShader(TArrayView<const uint8> InCode, mtlpp::Library InLibrary)
 {
 	FMetalCodeHeader Header;
 	Init(InCode, Header, InLibrary);
@@ -1061,13 +1061,13 @@ mtlpp::Function FMetalVertexShader::GetFunction()
 	return GetCompiledFunction();
 }
 
-FMetalPixelShader::FMetalPixelShader(const TArray<uint8>& InCode)
+FMetalPixelShader::FMetalPixelShader(TArrayView<const uint8> InCode)
 {
 	FMetalCodeHeader Header;
 	Init(InCode, Header);
 }
 
-FMetalPixelShader::FMetalPixelShader(const TArray<uint8>& InCode, mtlpp::Library InLibrary)
+FMetalPixelShader::FMetalPixelShader(TArrayView<const uint8> InCode, mtlpp::Library InLibrary)
 {
 	FMetalCodeHeader Header;
 	Init(InCode, Header, InLibrary);
@@ -1079,7 +1079,7 @@ mtlpp::Function FMetalPixelShader::GetFunction()
 }
 
 #if PLATFORM_SUPPORTS_TESSELLATION_SHADERS
-FMetalHullShader::FMetalHullShader(const TArray<uint8>& InCode)
+FMetalHullShader::FMetalHullShader(TArrayView<const uint8> InCode)
 {
 	FMetalCodeHeader Header;
 	Init(InCode, Header);
@@ -1119,7 +1119,7 @@ FMetalHullShader::FMetalHullShader(const TArray<uint8>& InCode)
 	}
 }
 
-FMetalHullShader::FMetalHullShader(const TArray<uint8>& InCode, mtlpp::Library InLibrary)
+FMetalHullShader::FMetalHullShader(TArrayView<const uint8> InCode, mtlpp::Library InLibrary)
 {
 	FMetalCodeHeader Header;
 	Init(InCode, Header, InLibrary);
@@ -1164,7 +1164,7 @@ mtlpp::Function FMetalHullShader::GetFunction()
 	return GetCompiledFunction();
 }
 
-FMetalDomainShader::FMetalDomainShader(const TArray<uint8>& InCode)
+FMetalDomainShader::FMetalDomainShader(TArrayView<const uint8> InCode)
 {
 	FMetalCodeHeader Header;
 	Init(InCode, Header);
@@ -1195,7 +1195,7 @@ FMetalDomainShader::FMetalDomainShader(const TArray<uint8>& InCode)
 	TessellationOutputAttribs = Tess.TessellationOutputAttribs;
 }
 
-FMetalDomainShader::FMetalDomainShader(const TArray<uint8>& InCode, mtlpp::Library InLibrary)
+FMetalDomainShader::FMetalDomainShader(TArrayView<const uint8> InCode, mtlpp::Library InLibrary)
 {
 	FMetalCodeHeader Header;
 	Init(InCode, Header, InLibrary);
@@ -1232,7 +1232,7 @@ mtlpp::Function FMetalDomainShader::GetFunction()
 }
 #endif
 
-FVertexShaderRHIRef FMetalDynamicRHI::RHICreateVertexShader(const TArray<uint8>& Code)
+FVertexShaderRHIRef FMetalDynamicRHI::RHICreateVertexShader(TArrayView<const uint8> Code, const FSHAHash& Hash)
 {
     @autoreleasepool {
 	FMetalVertexShader* Shader = new FMetalVertexShader(Code);
@@ -1240,19 +1240,7 @@ FVertexShaderRHIRef FMetalDynamicRHI::RHICreateVertexShader(const TArray<uint8>&
 	}
 }
 
-FVertexShaderRHIRef FMetalDynamicRHI::RHICreateVertexShader(FRHIShaderLibrary* Library, FSHAHash Hash)
-{
-	@autoreleasepool {
-
-	checkSlow(Library && Library->IsNativeLibrary() && IsMetalPlatform(Library->GetPlatform()) && Library->GetPlatform() <= GMaxRHIShaderPlatform);
-	
-	FMetalShaderLibrary* MetalLibrary = ResourceCast(Library);
-	return MetalLibrary->CreateVertexShader(Hash);
-	
-	}
-}
-
-FPixelShaderRHIRef FMetalDynamicRHI::RHICreatePixelShader(const TArray<uint8>& Code)
+FPixelShaderRHIRef FMetalDynamicRHI::RHICreatePixelShader(TArrayView<const uint8> Code, const FSHAHash& Hash)
 {
 	@autoreleasepool {
 	FMetalPixelShader* Shader = new FMetalPixelShader(Code);
@@ -1260,19 +1248,7 @@ FPixelShaderRHIRef FMetalDynamicRHI::RHICreatePixelShader(const TArray<uint8>& C
 	}
 }
 
-FPixelShaderRHIRef FMetalDynamicRHI::RHICreatePixelShader(FRHIShaderLibrary* Library, FSHAHash Hash)
-{
-	@autoreleasepool {
-	
-	checkSlow(Library && Library->IsNativeLibrary() && IsMetalPlatform(Library->GetPlatform()) && Library->GetPlatform() <= GMaxRHIShaderPlatform);
-	
-	FMetalShaderLibrary* MetalLibrary = ResourceCast(Library);
-	return MetalLibrary->CreatePixelShader(Hash);
-	
-	}
-}
-
-FHullShaderRHIRef FMetalDynamicRHI::RHICreateHullShader(const TArray<uint8>& Code) 
+FHullShaderRHIRef FMetalDynamicRHI::RHICreateHullShader(TArrayView<const uint8> Code, const FSHAHash& Hash)
 {
 	@autoreleasepool {
 #if PLATFORM_SUPPORTS_TESSELLATION_SHADERS
@@ -1286,19 +1262,7 @@ FHullShaderRHIRef FMetalDynamicRHI::RHICreateHullShader(const TArray<uint8>& Cod
 	}
 }
 
-FHullShaderRHIRef FMetalDynamicRHI::RHICreateHullShader(FRHIShaderLibrary* Library, FSHAHash Hash)
-{
-	@autoreleasepool {
-	
-	checkSlow(Library && Library->IsNativeLibrary() && IsMetalPlatform(Library->GetPlatform()) && Library->GetPlatform() <= GMaxRHIShaderPlatform);
-	
-	FMetalShaderLibrary* MetalLibrary = ResourceCast(Library);
-	return MetalLibrary->CreateHullShader(Hash);
-	
-	}
-}
-
-FDomainShaderRHIRef FMetalDynamicRHI::RHICreateDomainShader(const TArray<uint8>& Code) 
+FDomainShaderRHIRef FMetalDynamicRHI::RHICreateDomainShader(TArrayView<const uint8> Code, const FSHAHash& Hash)
 {
 	@autoreleasepool {
 #if PLATFORM_SUPPORTS_TESSELLATION_SHADERS
@@ -1312,19 +1276,7 @@ FDomainShaderRHIRef FMetalDynamicRHI::RHICreateDomainShader(const TArray<uint8>&
 	}
 }
 
-FDomainShaderRHIRef FMetalDynamicRHI::RHICreateDomainShader(FRHIShaderLibrary* Library, FSHAHash Hash)
-{
-	@autoreleasepool {
-	
-	checkSlow(Library && Library->IsNativeLibrary() && IsMetalPlatform(Library->GetPlatform()) && Library->GetPlatform() <= GMaxRHIShaderPlatform);
-	
-	FMetalShaderLibrary* MetalLibrary = ResourceCast(Library);
-	return MetalLibrary->CreateDomainShader(Hash);
-	
-	}
-}
-
-FGeometryShaderRHIRef FMetalDynamicRHI::RHICreateGeometryShader(const TArray<uint8>& Code) 
+FGeometryShaderRHIRef FMetalDynamicRHI::RHICreateGeometryShader(TArrayView<const uint8> Code, const FSHAHash& Hash)
 {
 	@autoreleasepool {
 	FMetalGeometryShader* Shader = new FMetalGeometryShader;
@@ -1334,100 +1286,59 @@ FGeometryShaderRHIRef FMetalDynamicRHI::RHICreateGeometryShader(const TArray<uin
 	}
 }
 
-FGeometryShaderRHIRef FMetalDynamicRHI::RHICreateGeometryShader(FRHIShaderLibrary* Library, FSHAHash Hash)
-{
-	@autoreleasepool {
-	
-	checkSlow(Library && Library->IsNativeLibrary() && IsMetalPlatform(Library->GetPlatform()) && Library->GetPlatform() <= GMaxRHIShaderPlatform);
-
-	FMetalShaderLibrary* MetalLibrary = ResourceCast(Library);
-	return MetalLibrary->CreateGeometryShader(Hash);
-	
-	}
-}
-
-FComputeShaderRHIRef FMetalDynamicRHI::RHICreateComputeShader(const TArray<uint8>& Code) 
+FComputeShaderRHIRef FMetalDynamicRHI::RHICreateComputeShader(TArrayView<const uint8> Code, const FSHAHash& Hash)
 {
 	@autoreleasepool {
 	return new FMetalComputeShader(Code);
 	}
 }
 
-FComputeShaderRHIRef FMetalDynamicRHI::RHICreateComputeShader(FRHIShaderLibrary* Library, FSHAHash Hash)
-{ 
-	@autoreleasepool {
-	
-	checkSlow(Library && Library->IsNativeLibrary() && IsMetalPlatform(Library->GetPlatform()) && Library->GetPlatform() <= GMaxRHIShaderPlatform);
-
-	FMetalShaderLibrary* MetalLibrary = ResourceCast(Library);
-	return MetalLibrary->CreateComputeShader(Hash);
-		
-	}
-}
-
-FVertexShaderRHIRef FMetalDynamicRHI::CreateVertexShader_RenderThread(class FRHICommandListImmediate& RHICmdList, const TArray<uint8>& Code)
-{
-	return RHICreateVertexShader(Code);
-}
-FVertexShaderRHIRef FMetalDynamicRHI::CreateVertexShader_RenderThread(class FRHICommandListImmediate& RHICmdList, FRHIShaderLibrary* Library, FSHAHash Hash)
-{
-	return RHICreateVertexShader(Library, Hash);
-}
-FPixelShaderRHIRef FMetalDynamicRHI::CreatePixelShader_RenderThread(class FRHICommandListImmediate& RHICmdList, const TArray<uint8>& Code)
-{
-	return RHICreatePixelShader(Code);
-}
-FPixelShaderRHIRef FMetalDynamicRHI::CreatePixelShader_RenderThread(class FRHICommandListImmediate& RHICmdList, FRHIShaderLibrary* Library, FSHAHash Hash)
-{
-	return RHICreatePixelShader(Library, Hash);
-}
-FGeometryShaderRHIRef FMetalDynamicRHI::CreateGeometryShader_RenderThread(class FRHICommandListImmediate& RHICmdList, const TArray<uint8>& Code)
-{
-	return RHICreateGeometryShader(Code);
-}
-FGeometryShaderRHIRef FMetalDynamicRHI::CreateGeometryShader_RenderThread(class FRHICommandListImmediate& RHICmdList, FRHIShaderLibrary* Library, FSHAHash Hash)
-{
-	return RHICreateGeometryShader(Library, Hash);
-}
-
-FComputeShaderRHIRef FMetalDynamicRHI::CreateComputeShader_RenderThread(class FRHICommandListImmediate& RHICmdList, const TArray<uint8>& Code)
+FVertexShaderRHIRef FMetalDynamicRHI::CreateVertexShader_RenderThread(class FRHICommandListImmediate& RHICmdList, TArrayView<const uint8> Code, const FSHAHash& Hash)
 {
-	return RHICreateComputeShader(Code);
+	return RHICreateVertexShader(Code, Hash);
 }
-FComputeShaderRHIRef FMetalDynamicRHI::CreateComputeShader_RenderThread(class FRHICommandListImmediate& RHICmdList, FRHIShaderLibrary* Library, FSHAHash Hash)
+FPixelShaderRHIRef FMetalDynamicRHI::CreatePixelShader_RenderThread(class FRHICommandListImmediate& RHICmdList, TArrayView<const uint8> Code, const FSHAHash& Hash)
 {
-	return RHICreateComputeShader(Library, Hash);
+	return RHICreatePixelShader(Code, Hash);
 }
-FHullShaderRHIRef FMetalDynamicRHI::CreateHullShader_RenderThread(class FRHICommandListImmediate& RHICmdList, const TArray<uint8>& Code)
+FGeometryShaderRHIRef FMetalDynamicRHI::CreateGeometryShader_RenderThread(class FRHICommandListImmediate& RHICmdList, TArrayView<const uint8> Code, const FSHAHash& Hash)
 {
-	return RHICreateHullShader(Code);
+	return RHICreateGeometryShader(Code, Hash);
 }
-FHullShaderRHIRef FMetalDynamicRHI::CreateHullShader_RenderThread(class FRHICommandListImmediate& RHICmdList, FRHIShaderLibrary* Library, FSHAHash Hash)
+FComputeShaderRHIRef FMetalDynamicRHI::CreateComputeShader_RenderThread(class FRHICommandListImmediate& RHICmdList, TArrayView<const uint8> Code, const FSHAHash& Hash)
 {
-	return RHICreateHullShader(Library, Hash);
+	return RHICreateComputeShader(Code, Hash);
 }
-FDomainShaderRHIRef FMetalDynamicRHI::CreateDomainShader_RenderThread(class FRHICommandListImmediate& RHICmdList, const TArray<uint8>& Code)
+FHullShaderRHIRef FMetalDynamicRHI::CreateHullShader_RenderThread(class FRHICommandListImmediate& RHICmdList, TArrayView<const uint8> Code, const FSHAHash& Hash)
 {
-	return RHICreateDomainShader(Code);
+	return RHICreateHullShader(Code, Hash);
 }
-FDomainShaderRHIRef FMetalDynamicRHI::CreateDomainShader_RenderThread(class FRHICommandListImmediate& RHICmdList, FRHIShaderLibrary* Library, FSHAHash Hash)
+FDomainShaderRHIRef FMetalDynamicRHI::CreateDomainShader_RenderThread(class FRHICommandListImmediate& RHICmdList, TArrayView<const uint8> Code, const FSHAHash& Hash)
 {
-	return RHICreateDomainShader(Library, Hash);
+	return RHICreateDomainShader(Code, Hash);
 }
 
 static FCriticalSection LoadedShaderLibraryMutex;
 static TMap<FString, FRHIShaderLibrary*> LoadedShaderLibraryMap;
 
-FMetalShaderLibrary::FMetalShaderLibrary(EShaderPlatform InPlatform, FString const& Name, TArray<mtlpp::Library> InLibrary, FMetalShaderMap const& InMap, const FString& InShaderLibraryFilename)
-: FRHIShaderLibrary(InPlatform, Name)
-, Library(InLibrary)
-, Map(InMap)
+FMetalShaderLibrary::FMetalShaderLibrary(EShaderPlatform Platform,
+	FString const& Name,
+	const FString& InShaderLibraryFilename,
+	const FMetalShaderLibraryHeader& InHeader,
+	const FSerializedShaderArchive& InSerializedShaders,
+	const TArray<uint8>& InShaderCode,
+	const TArray<mtlpp::Library>& InLibrary)
+: FRHIShaderLibrary(Platform, Name)
 , ShaderLibraryFilename(InShaderLibraryFilename)
+, Library(InLibrary)
+, Header(InHeader)
+, SerializedShaders(InSerializedShaders)
+, ShaderCode(InShaderCode)
 {
 #if !UE_BUILD_SHIPPING
 	DebugFile = nullptr;
 
-	FName PlatformName = LegacyShaderPlatformToShaderFormat(InPlatform);
+	FName PlatformName = LegacyShaderPlatformToShaderFormat(Platform);
 	FString LibName = FString::Printf(TEXT("%s_%s"), *Name, *PlatformName.GetPlainNameString());
 	LibName.ToLowerInline();
 	FString Path = FPaths::ProjectContentDir() / LibName + TEXT(".zip");
@@ -1445,142 +1356,47 @@ FMetalShaderLibrary::~FMetalShaderLibrary()
 	LoadedShaderLibraryMap.Remove(ShaderLibraryFilename);
 }
 
-bool FMetalShaderLibrary::ContainsEntry(const FSHAHash& Hash)
+template<typename ShaderType>
+static TRefCountPtr<FRHIShader> CreateMetalShader(TArrayView<const uint8> InCode, mtlpp::Library InLibrary)
 {
-	FMetalShadeEntry* Code = Map.HashMap.Find(Hash);
-	return (Code != nullptr);
-}
-
-bool FMetalShaderLibrary::RequestEntry(const FSHAHash& Hash, FArchive* Ar)
-{
-	FMetalShadeEntry* Code = Map.HashMap.Find(Hash);
-	return (Code != nullptr);
-}
-
-FPixelShaderRHIRef FMetalShaderLibrary::CreatePixelShader(const FSHAHash& Hash)
-{
-	FMetalShadeEntry* Code = Map.HashMap.Find(Hash);
-	if (Code)
+	ShaderType* Shader = new ShaderType(InCode, InLibrary);
+	if (!Shader->GetFunction())
 	{
-		FMetalPixelShader* Shader = new FMetalPixelShader(Code->Data, Library[Code->Index]);
-		if (Shader->GetFunction())
-		{
-			return Shader;
-		}
-		else
-		{
-			delete Shader;
-		}
+		delete Shader;
+		Shader = nullptr;
 	}
-	
-	UE_LOG(LogMetal, Error, TEXT("Failed to find Pixel Shader with SHA: %s"), *Hash.ToString());
-	return FPixelShaderRHIRef();
-}
 
-FVertexShaderRHIRef FMetalShaderLibrary::CreateVertexShader(const FSHAHash& Hash)
-{
-	FMetalShadeEntry* Code = Map.HashMap.Find(Hash);
-	if (Code)
-	{
-		FMetalVertexShader* Shader = new FMetalVertexShader(Code->Data, Library[Code->Index]);
-		if (Shader->GetFunction())
-		{
-			return Shader;
-		}
-		else
-		{
-			delete Shader;
-		}
-	}
-	UE_LOG(LogMetal, Error, TEXT("Failed to find Vertex Shader with SHA: %s"), *Hash.ToString());
-	return FVertexShaderRHIRef();
+	return TRefCountPtr<FRHIShader>(Shader);
 }
 
-FHullShaderRHIRef FMetalShaderLibrary::CreateHullShader(const FSHAHash& Hash)
+TRefCountPtr<FRHIShader> FMetalShaderLibrary::CreateShader(int32 Index)
 {
-#if PLATFORM_SUPPORTS_TESSELLATION_SHADERS
-	FMetalShadeEntry* Code = Map.HashMap.Find(Hash);
-	if (Code)
-	{
-		FMetalHullShader* Shader = new FMetalHullShader(Code->Data, Library[Code->Index]);
-		if(Shader->GetFunction())
-		{
-			return Shader;
-		}
-		else
-		{
-			delete Shader;
-		}
-	}
-	UE_LOG(LogMetal, Error, TEXT("Failed to find Hull Shader with SHA: %s"), *Hash.ToString());
-	return FHullShaderRHIRef();
-#else
-	checkf(0, TEXT("Not supported"));
-	return NULL;
-#endif
-}
+	const FShaderCodeEntry& ShaderEntry = SerializedShaders.ShaderEntries[Index];
+	check(ShaderEntry.Size == ShaderEntry.UncompressedSize); // don't handle compressed shaders here, since typically these are just tiny headers
+	const TArrayView<uint8> Code = MakeArrayView(ShaderCode.GetData() + ShaderEntry.Offset, ShaderEntry.Size);
+	const int32 LibraryIndex = Index / Header.NumShadersPerLibrary;
 
-FDomainShaderRHIRef FMetalShaderLibrary::CreateDomainShader(const FSHAHash& Hash)
-{
-#if PLATFORM_SUPPORTS_TESSELLATION_SHADERS
-	FMetalShadeEntry* Code = Map.HashMap.Find(Hash);
-	if (Code)
+	TRefCountPtr<FRHIShader> Shader;
+	switch (ShaderEntry.Frequency)
 	{
-		FMetalDomainShader* Shader = new FMetalDomainShader(Code->Data, Library[Code->Index]);
-		if (Shader->GetFunction())
-		{
-			return Shader;
-		}
-		else
-		{
-			delete Shader;
-		}
+	case SF_Vertex: Shader = CreateMetalShader<FMetalVertexShader>(Code, Library[LibraryIndex]); break;
+	case SF_Pixel: Shader = CreateMetalShader<FMetalPixelShader>(Code, Library[LibraryIndex]); break;
+	case SF_Geometry: checkf(false, TEXT("Geometry shaders not supported")); break;
+#if PLATFORM_SUPPORTS_TESSELLATION_SHADERS
+	case SF_Hull: Shader = CreateMetalShader<FMetalHullShader>(Code, Library[LibraryIndex]); break;
+	case SF_Domain: Shader = CreateMetalShader<FMetalDomainShader>(Code, Library[LibraryIndex]); break;
+#endif // PLATFORM_SUPPORTS_TESSELLATION_SHADERS
+	case SF_Compute: Shader = CreateMetalShader<FMetalComputeShader>(Code, Library[LibraryIndex]); break;
+	default: checkNoEntry(); break;
 	}
-	UE_LOG(LogMetal, Error, TEXT("Failed to find Domain Shader with SHA: %s"), *Hash.ToString());
-	return FDomainShaderRHIRef();
-#else
-	checkf(0, TEXT("Not supported"));
-	return NULL;
-#endif
-}
-
-FGeometryShaderRHIRef FMetalShaderLibrary::CreateGeometryShader(const FSHAHash& Hash)
-{
-	checkf(0, TEXT("Not supported"));
-	return NULL;
-}
 
-FComputeShaderRHIRef FMetalShaderLibrary::CreateComputeShader(const FSHAHash& Hash)
-{
-	FMetalShadeEntry* Code = Map.HashMap.Find(Hash);
-	if (Code)
+	if (Shader)
 	{
-		FMetalComputeShader* Shader = new FMetalComputeShader(Code->Data, Library[Code->Index]);
-		if (Shader->GetPipeline())
-		{
-			return Shader;
-		}
-		else
-		{
-			delete Shader;
-		}
+		Shader->SetHash(SerializedShaders.ShaderHashes[Index]);
 	}
-	UE_LOG(LogMetal, Error, TEXT("Failed to find Compute Shader with SHA: %s"), *Hash.ToString());
-	return FComputeShaderRHIRef();
-}
 
-//
-//Library Iterator
-//
-FRHIShaderLibrary::FShaderLibraryEntry FMetalShaderLibrary::FMetalShaderLibraryIterator::operator*() const
-{
-	FShaderLibraryEntry Entry;
-	
-	Entry.Hash = IteratorImpl->Key;
-	Entry.Frequency = (EShaderFrequency)IteratorImpl->Value.Type;
-	Entry.Platform = GetLibrary()->GetPlatform();
-	
-	return Entry;
+	return Shader;
+
 }
 
 FRHIShaderLibraryRef FMetalDynamicRHI::RHICreateShaderLibrary_RenderThread(class FRHICommandListImmediate& RHICmdList, EShaderPlatform Platform, FString FilePath, FString Name)
@@ -1596,8 +1412,7 @@ FRHIShaderLibraryRef FMetalDynamicRHI::RHICreateShaderLibrary(EShaderPlatform Pl
 	FName PlatformName = LegacyShaderPlatformToShaderFormat(Platform);
 	FString LibName = FString::Printf(TEXT("%s_%s"), *Name, *PlatformName.GetPlainNameString());
 	LibName.ToLowerInline();
-	
-	FMetalShaderMap Map;
+
 	FString BinaryShaderFile = FilePath / LibName + METAL_MAP_EXTENSION;
 
 	if ( IFileManager::Get().FileExists(*BinaryShaderFile) == false )
@@ -1620,15 +1435,25 @@ FRHIShaderLibraryRef FMetalDynamicRHI::RHICreateShaderLibrary(EShaderPlatform Pl
 
 	if( BinaryShaderAr != NULL )
 	{
-		*BinaryShaderAr << Map;
+		FMetalShaderLibraryHeader Header;
+		FSerializedShaderArchive SerializedShaders;
+		TArray<uint8> ShaderCode;
+
+		*BinaryShaderAr << Header;
+		*BinaryShaderAr << SerializedShaders;
+		*BinaryShaderAr << ShaderCode;
 		BinaryShaderAr->Flush();
 		delete BinaryShaderAr;
 		
 		// Would be good to check the language version of the library with the archive format here.
-		if (Map.Format == PlatformName.GetPlainNameString())
+		if (Header.Format == PlatformName.GetPlainNameString())
 		{
+			check(SerializedShaders.GetNumShaders() / Header.NumShadersPerLibrary == Header.NumLibraries);
+
 			TArray<mtlpp::Library> Libraries;
-			for (uint32 i = 0; i < Map.Count; i++)
+			Libraries.Empty(Header.NumLibraries);
+
+			for (uint32 i = 0; i < Header.NumLibraries; i++)
 			{
 				FString MetalLibraryFilePath = (FilePath / LibName) + FString::Printf(TEXT(".%d"), i) + METAL_LIB_EXTENSION;
 				MetalLibraryFilePath = FPaths::ConvertRelativePathToFull(MetalLibraryFilePath);
@@ -1650,13 +1475,13 @@ FRHIShaderLibraryRef FMetalDynamicRHI::RHICreateShaderLibrary(EShaderPlatform Pl
 				}
 			}
 
-			Result = new FMetalShaderLibrary(Platform, Name, Libraries, Map, BinaryShaderFile);
+			Result = new FMetalShaderLibrary(Platform, Name, BinaryShaderFile, Header, SerializedShaders, ShaderCode, Libraries);
 			LoadedShaderLibraryMap.Add(BinaryShaderFile, Result.GetReference());
 		}
-		else
-		{
-			UE_LOG(LogMetal, Display, TEXT("Wrong shader platform wanted: %s, got: %s"), *LibName, *Map.Format);
-		}
+		//else
+		//{
+		//	UE_LOG(LogMetal, Display, TEXT("Wrong shader platform wanted: %s, got: %s"), *LibName, *Map.Format);
+		//}
 	}
 	else
 	{
diff --git a/Engine/Source/Runtime/Apple/MetalRHI/Private/MetalState.cpp b/Engine/Source/Runtime/Apple/MetalRHI/Private/MetalState.cpp
--- a/Engine/Source/Runtime/Apple/MetalRHI/Private/MetalState.cpp
+++ b/Engine/Source/Runtime/Apple/MetalRHI/Private/MetalState.cpp
@@ -530,8 +530,8 @@ FMetalBlendState::FMetalBlendState(const FBlendStateInitializerRHI& Initializer)
 		{
 			Key = &BlendSettingsToUniqueKeyMap.Add(BlendBitMask, NextKey++);
 
-			// only giving 5 bits to the key, since we need to pack a bunch of them into 64 bits
-			checkf(NextKey < 32, TEXT("Too many unique blend states to fit into the PipelineStateHash"));
+			// only giving limited bits to the key, since we need to pack 8 of them into a key
+			checkf(NextKey < (1 << NumBits_BlendState), TEXT("Too many unique blend states to fit into the PipelineStateHash [%d allowed]"), 1 << NumBits_BlendState);
 		}
 		// set the key
 		RenderTargetStates[RenderTargetIndex].BlendStateKey = *Key;
diff --git a/Engine/Source/Runtime/Apple/MetalRHI/Public/MetalRHI.h b/Engine/Source/Runtime/Apple/MetalRHI/Public/MetalRHI.h
--- a/Engine/Source/Runtime/Apple/MetalRHI/Public/MetalRHI.h
+++ b/Engine/Source/Runtime/Apple/MetalRHI/Public/MetalRHI.h
@@ -45,18 +45,12 @@ public:
 	virtual FDepthStencilStateRHIRef RHICreateDepthStencilState(const FDepthStencilStateInitializerRHI& Initializer) final override;
 	virtual FBlendStateRHIRef RHICreateBlendState(const FBlendStateInitializerRHI& Initializer) final override;
 	virtual FVertexDeclarationRHIRef RHICreateVertexDeclaration(const FVertexDeclarationElementList& Elements) final override;
-	virtual FPixelShaderRHIRef RHICreatePixelShader(const TArray<uint8>& Code) final override;
-	virtual FPixelShaderRHIRef RHICreatePixelShader(FRHIShaderLibrary* Library, FSHAHash Hash) final override;
-	virtual FVertexShaderRHIRef RHICreateVertexShader(const TArray<uint8>& Code) final override;
-	virtual FVertexShaderRHIRef RHICreateVertexShader(FRHIShaderLibrary* Library, FSHAHash Hash) final override;
-	virtual FHullShaderRHIRef RHICreateHullShader(const TArray<uint8>& Code) final override;
-	virtual FHullShaderRHIRef RHICreateHullShader(FRHIShaderLibrary* Library, FSHAHash Hash) final override;
-	virtual FDomainShaderRHIRef RHICreateDomainShader(const TArray<uint8>& Code) final override;
-	virtual FDomainShaderRHIRef RHICreateDomainShader(FRHIShaderLibrary* Library, FSHAHash Hash) final override;
-	virtual FGeometryShaderRHIRef RHICreateGeometryShader(const TArray<uint8>& Code) final override;
-	virtual FGeometryShaderRHIRef RHICreateGeometryShader(FRHIShaderLibrary* Library, FSHAHash Hash) final override;
-	virtual FComputeShaderRHIRef RHICreateComputeShader(const TArray<uint8>& Code) final override;
-	virtual FComputeShaderRHIRef RHICreateComputeShader(FRHIShaderLibrary* Library, FSHAHash Hash) final override;
+	virtual FPixelShaderRHIRef RHICreatePixelShader(TArrayView<const uint8> Code, const FSHAHash& Hash) final override;
+	virtual FVertexShaderRHIRef RHICreateVertexShader(TArrayView<const uint8> Code, const FSHAHash& Hash) final override;
+	virtual FHullShaderRHIRef RHICreateHullShader(TArrayView<const uint8> Code, const FSHAHash& Hash) final override;
+	virtual FDomainShaderRHIRef RHICreateDomainShader(TArrayView<const uint8> Code, const FSHAHash& Hash) final override;
+	virtual FGeometryShaderRHIRef RHICreateGeometryShader(TArrayView<const uint8> Code, const FSHAHash& Hash) final override;
+	virtual FComputeShaderRHIRef RHICreateComputeShader(TArrayView<const uint8> Code, const FSHAHash& Hash) final override;
 	virtual FBoundShaderStateRHIRef RHICreateBoundShaderState(FRHIVertexDeclaration* VertexDeclaration, FRHIVertexShader* VertexShader, FRHIHullShader* HullShader, FRHIDomainShader* DomainShader, FRHIPixelShader* PixelShader, FRHIGeometryShader* GeometryShader) final override;
 	virtual FGraphicsPipelineStateRHIRef RHICreateGraphicsPipelineState(const FGraphicsPipelineStateInitializer& Initializer) final override;
 	virtual TRefCountPtr<FRHIComputePipelineState> RHICreateComputePipelineState(FRHIComputeShader* ComputeShader) final override;
@@ -199,18 +193,12 @@ public:
 	virtual void RHICopySubTextureRegion_RenderThread(class FRHICommandListImmediate& RHICmdList, FRHITexture2D* SourceTexture, FRHITexture2D* DestinationTexture, FBox2D SourceBox, FBox2D DestinationBox) final override;
 	virtual void RHICopySubTextureRegion(FRHITexture2D* SourceTexture, FRHITexture2D* DestinationTexture, FBox2D SourceBox, FBox2D DestinationBox) final override;
 
-	virtual FVertexShaderRHIRef CreateVertexShader_RenderThread(class FRHICommandListImmediate& RHICmdList, const TArray<uint8>& Code) final override;
-	virtual FVertexShaderRHIRef CreateVertexShader_RenderThread(class FRHICommandListImmediate& RHICmdList, FRHIShaderLibrary* Library, FSHAHash Hash) final override;
-	virtual FPixelShaderRHIRef CreatePixelShader_RenderThread(class FRHICommandListImmediate& RHICmdList, const TArray<uint8>& Code) final override;
-	virtual FPixelShaderRHIRef CreatePixelShader_RenderThread(class FRHICommandListImmediate& RHICmdList, FRHIShaderLibrary* Library, FSHAHash Hash) final override;
-	virtual FGeometryShaderRHIRef CreateGeometryShader_RenderThread(class FRHICommandListImmediate& RHICmdList, const TArray<uint8>& Code) final override;
-	virtual FGeometryShaderRHIRef CreateGeometryShader_RenderThread(class FRHICommandListImmediate& RHICmdList, FRHIShaderLibrary* Library, FSHAHash Hash) final override;
-	virtual FComputeShaderRHIRef CreateComputeShader_RenderThread(class FRHICommandListImmediate& RHICmdList, const TArray<uint8>& Code) final override;
-	virtual FComputeShaderRHIRef CreateComputeShader_RenderThread(class FRHICommandListImmediate& RHICmdList, FRHIShaderLibrary* Library, FSHAHash Hash) final override;
-	virtual FHullShaderRHIRef CreateHullShader_RenderThread(class FRHICommandListImmediate& RHICmdList, const TArray<uint8>& Code) final override;
-	virtual FHullShaderRHIRef CreateHullShader_RenderThread(class FRHICommandListImmediate& RHICmdList, FRHIShaderLibrary* Library, FSHAHash Hash) final override;
-	virtual FDomainShaderRHIRef CreateDomainShader_RenderThread(class FRHICommandListImmediate& RHICmdList, const TArray<uint8>& Code) final override;
-	virtual FDomainShaderRHIRef CreateDomainShader_RenderThread(class FRHICommandListImmediate& RHICmdList, FRHIShaderLibrary* Library, FSHAHash Hash) final override;
+	virtual FVertexShaderRHIRef CreateVertexShader_RenderThread(class FRHICommandListImmediate& RHICmdList, TArrayView<const uint8> Code, const FSHAHash& Hash) final override;
+	virtual FPixelShaderRHIRef CreatePixelShader_RenderThread(class FRHICommandListImmediate& RHICmdList, TArrayView<const uint8> Code, const FSHAHash& Hash) final override;
+	virtual FGeometryShaderRHIRef CreateGeometryShader_RenderThread(class FRHICommandListImmediate& RHICmdList, TArrayView<const uint8> Code, const FSHAHash& Hash) final override;
+	virtual FComputeShaderRHIRef CreateComputeShader_RenderThread(class FRHICommandListImmediate& RHICmdList, TArrayView<const uint8> Code, const FSHAHash& Hash) final override;
+	virtual FHullShaderRHIRef CreateHullShader_RenderThread(class FRHICommandListImmediate& RHICmdList, TArrayView<const uint8> Code, const FSHAHash& Hash) final override;
+	virtual FDomainShaderRHIRef CreateDomainShader_RenderThread(class FRHICommandListImmediate& RHICmdList, TArrayView<const uint8> Code, const FSHAHash& Hash) final override;
 	virtual FRHIShaderLibraryRef RHICreateShaderLibrary_RenderThread(class FRHICommandListImmediate& RHICmdList, EShaderPlatform Platform, FString FilePath, FString Name) final override;
 
 	virtual void RHIUpdateUniformBuffer(FRHIUniformBuffer* UniformBufferRHI, const void* Contents) final override;
diff --git a/Engine/Source/Runtime/Apple/MetalRHI/Public/MetalResources.h b/Engine/Source/Runtime/Apple/MetalRHI/Public/MetalResources.h
--- a/Engine/Source/Runtime/Apple/MetalRHI/Public/MetalResources.h
+++ b/Engine/Source/Runtime/Apple/MetalRHI/Public/MetalResources.h
@@ -8,6 +8,8 @@
 
 #include "BoundShaderStateCache.h"
 #include "MetalShaderResources.h"
+#include "ShaderCodeArchive.h"
+
 THIRD_PARTY_INCLUDES_START
 #include "mtlpp.hpp"
 THIRD_PARTY_INCLUDES_END
@@ -95,7 +97,7 @@ public:
 		Function = nil;
 	}
 	
-	void Init(const TArray<uint8>& InCode, FMetalCodeHeader& Header, mtlpp::Library InLibrary = nil);
+	void Init(TArrayView<const uint8> InCode, FMetalCodeHeader& Header, mtlpp::Library InLibrary = nil);
 
 	/** Destructor */
 	virtual ~TMetalBaseShader();
@@ -180,8 +182,8 @@ private:
 class FMetalVertexShader : public TMetalBaseShader<FRHIVertexShader, SF_Vertex>
 {
 public:
-	FMetalVertexShader(const TArray<uint8>& InCode);
-	FMetalVertexShader(const TArray<uint8>& InCode, mtlpp::Library InLibrary);
+	FMetalVertexShader(TArrayView<const uint8> InCode);
+	FMetalVertexShader(TArrayView<const uint8> InCode, mtlpp::Library InLibrary);
 	
 	mtlpp::Function GetFunction();
 	
@@ -205,8 +207,8 @@ public:
 class FMetalPixelShader : public TMetalBaseShader<FRHIPixelShader, SF_Pixel>
 {
 public:
-	FMetalPixelShader(const TArray<uint8>& InCode);
-	FMetalPixelShader(const TArray<uint8>& InCode, mtlpp::Library InLibrary);
+	FMetalPixelShader(TArrayView<const uint8> InCode);
+	FMetalPixelShader(TArrayView<const uint8> InCode, mtlpp::Library InLibrary);
 	
 	mtlpp::Function GetFunction();
 };
@@ -215,8 +217,8 @@ public:
 class FMetalHullShader : public TMetalBaseShader<FRHIHullShader, SF_Hull>
 {
 public:
-	FMetalHullShader(const TArray<uint8>& InCode);
-	FMetalHullShader(const TArray<uint8>& InCode, mtlpp::Library InLibrary);
+	FMetalHullShader(TArrayView<const uint8> InCode);
+	FMetalHullShader(TArrayView<const uint8> InCode, mtlpp::Library InLibrary);
 	
 	mtlpp::Function GetFunction();
 	
@@ -240,8 +242,8 @@ public:
 class FMetalDomainShader : public TMetalBaseShader<FRHIDomainShader, SF_Domain>
 {
 public:
-	FMetalDomainShader(const TArray<uint8>& InCode);
-	FMetalDomainShader(const TArray<uint8>& InCode, mtlpp::Library InLibrary);
+	FMetalDomainShader(TArrayView<const uint8> InCode);
+	FMetalDomainShader(TArrayView<const uint8> InCode, mtlpp::Library InLibrary);
 	
 	mtlpp::Function GetFunction();
 	
@@ -263,7 +265,7 @@ typedef TMetalBaseShader<FRHIGeometryShader, SF_Geometry> FMetalGeometryShader;
 class FMetalComputeShader : public TMetalBaseShader<FRHIComputeShader, SF_Compute>
 {
 public:
-	FMetalComputeShader(const TArray<uint8>& InCode, mtlpp::Library InLibrary = nil);
+	FMetalComputeShader(TArrayView<const uint8> InCode, mtlpp::Library InLibrary = nil);
 	virtual ~FMetalComputeShader();
 	
 	FMetalShaderPipeline* GetPipeline();
@@ -1223,60 +1225,52 @@ private:
 class FMetalShaderLibrary final : public FRHIShaderLibrary
 {	
 public:
-	FMetalShaderLibrary(EShaderPlatform Platform, FString const& Name, TArray<mtlpp::Library> Library, FMetalShaderMap const& Map, const FString& InShaderLibraryFilename);
+	FMetalShaderLibrary(EShaderPlatform Platform,
+		FString const& Name,
+		const FString& InShaderLibraryFilename,
+		const FMetalShaderLibraryHeader& InHeader,
+		const FSerializedShaderArchive& InSerializedShaders,
+		const TArray<uint8>& InShaderCode,
+		const TArray<mtlpp::Library>& InLibrary);
 	virtual ~FMetalShaderLibrary();
 	
 	virtual bool IsNativeLibrary() const override final {return true;}
-		
-	class FMetalShaderLibraryIterator : public FRHIShaderLibrary::FShaderLibraryIterator
+
+	virtual int32 GetNumShaders() const override { return SerializedShaders.ShaderEntries.Num(); }
+	virtual int32 GetNumShaderMaps() const override { return SerializedShaders.ShaderMapEntries.Num(); }
+	virtual int32 GetNumShadersForShaderMap(int32 ShaderMapIndex) const override { return SerializedShaders.ShaderMapEntries[ShaderMapIndex].NumShaders; }
+
+	virtual int32 GetShaderIndex(int32 ShaderMapIndex, int32 i) const override
 	{
-	public:
-		FMetalShaderLibraryIterator(FMetalShaderLibrary* MetalShaderLibrary) : FShaderLibraryIterator(MetalShaderLibrary), IteratorImpl(MetalShaderLibrary->Map.HashMap.CreateIterator()) {}
-		
-		virtual bool IsValid() const final override
-		{
-			return !!IteratorImpl;
-		}
-		
-		virtual FShaderLibraryEntry operator*() const final override;
-		virtual FShaderLibraryIterator& operator++() final override
-		{
-			++IteratorImpl;
-			return *this;
-		}
-		
-	private:		
-		TMap<FSHAHash, FMetalShadeEntry>::TIterator IteratorImpl;
-	};
-	
-	virtual TRefCountPtr<FShaderLibraryIterator> CreateIterator(void) final override
+		const FShaderMapEntry& ShaderMapEntry = SerializedShaders.ShaderMapEntries[ShaderMapIndex];
+		return SerializedShaders.ShaderIndices[ShaderMapEntry.ShaderIndicesOffset + i];
+	}
+
+	virtual int32 FindShaderMapIndex(const FSHAHash& Hash) override
 	{
-		return new FMetalShaderLibraryIterator(this);
+		return SerializedShaders.FindShaderMap(Hash);
 	}
-	
-	virtual bool ContainsEntry(const FSHAHash& Hash) final override;
-	virtual bool RequestEntry(const FSHAHash& Hash, FArchive* Ar) final override;
-	
-	virtual uint32 GetShaderCount(void) const final override { return Map.HashMap.Num(); }
-	
-private:
 
-	friend class FMetalDynamicRHI;
-	
-	FPixelShaderRHIRef CreatePixelShader(const FSHAHash& Hash);
-	FVertexShaderRHIRef CreateVertexShader(const FSHAHash& Hash);
-	FHullShaderRHIRef CreateHullShader(const FSHAHash& Hash);
-	FDomainShaderRHIRef CreateDomainShader(const FSHAHash& Hash);
-	FGeometryShaderRHIRef CreateGeometryShader(const FSHAHash& Hash);
-	FComputeShaderRHIRef CreateComputeShader(const FSHAHash& Hash);
+	virtual int32 FindShaderIndex(const FSHAHash& Hash) override
+	{
+		return SerializedShaders.FindShader(Hash);
+	}
+
+	/** No preload support */
+	virtual FGraphEventRef PreloadShader(int32 ShaderIndex) override { return FGraphEventRef(); }
+	virtual FGraphEventRef PreloadShaderMap(int32 ShaderMapIndex) override { return FGraphEventRef(); }
+
+	virtual TRefCountPtr<FRHIShader> CreateShader(int32 Index) override;
 	
 private:
+	FString ShaderLibraryFilename;
 	TArray<mtlpp::Library> Library;
+	FMetalShaderLibraryHeader Header;
+	FSerializedShaderArchive SerializedShaders;
+	TArray<uint8> ShaderCode;
 #if !UE_BUILD_SHIPPING
 	class FMetalShaderDebugZipFile* DebugFile;
 #endif
-	FMetalShaderMap Map;
-	FString ShaderLibraryFilename;
 };
 
 template<class T>
diff --git a/Engine/Source/Runtime/Apple/MetalRHI/Public/MetalShaderResources.h b/Engine/Source/Runtime/Apple/MetalRHI/Public/MetalShaderResources.h
--- a/Engine/Source/Runtime/Apple/MetalRHI/Public/MetalShaderResources.h
+++ b/Engine/Source/Runtime/Apple/MetalRHI/Public/MetalShaderResources.h
@@ -399,41 +399,14 @@ inline FArchive& operator<<(FArchive& Ar, FMetalCodeHeader& Header)
     return Ar;
 }
 
-struct FMetalShadeEntry
-{
-	TArray<uint8> Data;
-	uint32 Index;
-	uint8 Type;
-
-	FMetalShadeEntry()
-	: Index(0)
-	, Type(0)
-	{
-
-	}
-
-	FMetalShadeEntry(TArray<uint8> InData, uint32 InIndex, uint8 InType)
-	: Data(InData)
-	, Index(InIndex)
-	, Type(InType)
-	{
-
-	}
-	
-	friend FArchive& operator<<(FArchive& Ar, FMetalShadeEntry& Header)
-	{
-		return Ar << Header.Data << Header.Index << Header.Type;
-	}
-};
-
-struct FMetalShaderMap
+struct FMetalShaderLibraryHeader
 {
 	FString Format;
-	uint32 Count;
-	TMap<FSHAHash, FMetalShadeEntry> HashMap;
+	uint32 NumLibraries;
+	uint32 NumShadersPerLibrary;
 	
-	friend FArchive& operator<<(FArchive& Ar, FMetalShaderMap& Header)
+	friend FArchive& operator<<(FArchive& Ar, FMetalShaderLibraryHeader& Header)
 	{
-		return Ar << Header.Format << Header.Count << Header.HashMap;
+		return Ar << Header.Format << Header.NumLibraries << Header.NumShadersPerLibrary;
 	}
 };
