Enforce safe alignment for SRVs onto global dynamic read buffers.

diff --git a/Engine/Source/Runtime/Apple/MetalRHI/Private/MetalRHI.cpp b/Engine/Source/Runtime/Apple/MetalRHI/Private/MetalRHI.cpp
--- a/Engine/Source/Runtime/Apple/MetalRHI/Private/MetalRHI.cpp
+++ b/Engine/Source/Runtime/Apple/MetalRHI/Private/MetalRHI.cpp
@@ -889,6 +889,10 @@ uint64 FMetalDynamicRHI::RHICalcTextureCubePlatformSize(uint32 Size, uint8 Forma
 	}
 }
 
+uint64 FMetalDynamicRHI::RHIGetMinimumAlignmentForBufferBackedSRV(EPixelFormat Format)
+{
+	return ImmediateContext.Context->GetDevice().GetMinimumLinearTextureAlignmentForPixelFormat((mtlpp::PixelFormat)GMetalBufferFormats[Format].LinearTextureFormat);
+}
 
 void FMetalDynamicRHI::Init()
 {
diff --git a/Engine/Source/Runtime/Apple/MetalRHI/Private/MetalVertexBuffer.cpp b/Engine/Source/Runtime/Apple/MetalRHI/Private/MetalVertexBuffer.cpp
--- a/Engine/Source/Runtime/Apple/MetalRHI/Private/MetalVertexBuffer.cpp
+++ b/Engine/Source/Runtime/Apple/MetalRHI/Private/MetalVertexBuffer.cpp
@@ -300,7 +300,7 @@ FMetalTexture FMetalRHIBuffer::AllocLinearTexture(EPixelFormat Format, NSUIntege
 		{
 			Stride = 4;
 		}
-		NSUInteger NewSize = Size;
+		NSUInteger NewSize = Size - Offset;
 
 		if (FMetalCommandQueue::SupportsFeature(EMetalFeaturesTextureBuffers))
 		{
@@ -309,7 +309,7 @@ FMetalTexture FMetalRHIBuffer::AllocLinearTexture(EPixelFormat Format, NSUIntege
 		}
 		else
 		{
-			uint32 NumElements = (Buffer.GetLength() / Stride);
+			uint32 NumElements = ((Buffer.GetLength() - Offset) / Stride);
 			uint32 SizeX = NumElements;
 			uint32 SizeY = 1;
 			if (NumElements > GMaxTextureDimensions)
@@ -326,9 +326,11 @@ FMetalTexture FMetalRHIBuffer::AllocLinearTexture(EPixelFormat Format, NSUIntege
 				checkf(SizeX <= GMaxTextureDimensions, TEXT("Calculated height %u is greater than maximum permitted %d when converting buffer of size %llu with element stride %u to a 2D texture with %u elements."), SizeY, (int32)GMaxTextureDimensions, Buffer.GetLength(), Stride, NumElements);
 			}
 			
-			check(((SizeX*Stride) % 1024) == 0);
 			NewSize = SizeX*Stride;
-			
+
+			check((NewSize % 1024) == 0);
+			check((NewSize + Offset) <= Buffer.GetLength());
+
 			Desc = mtlpp::TextureDescriptor::Texture2DDescriptor(MTLFormat, SizeX, SizeY, NO);
 			Desc.SetStorageMode(Mode);
 			Desc.SetCpuCacheMode(Buffer.GetCpuCacheMode());
diff --git a/Engine/Source/Runtime/Apple/MetalRHI/Public/MetalRHI.h b/Engine/Source/Runtime/Apple/MetalRHI/Public/MetalRHI.h
--- a/Engine/Source/Runtime/Apple/MetalRHI/Public/MetalRHI.h
+++ b/Engine/Source/Runtime/Apple/MetalRHI/Public/MetalRHI.h
@@ -81,6 +81,7 @@ public:
 	virtual uint64 RHICalcTexture2DPlatformSize(uint32 SizeX, uint32 SizeY, uint8 Format, uint32 NumMips, uint32 NumSamples, uint32 Flags, const FRHIResourceCreateInfo& CreateInfo, uint32& OutAlign) final override;
 	virtual uint64 RHICalcTexture3DPlatformSize(uint32 SizeX, uint32 SizeY, uint32 SizeZ, uint8 Format, uint32 NumMips, uint32 Flags, const FRHIResourceCreateInfo& CreateInfo, uint32& OutAlign) final override;
 	virtual uint64 RHICalcTextureCubePlatformSize(uint32 Size, uint8 Format, uint32 NumMips, uint32 Flags, const FRHIResourceCreateInfo& CreateInfo, uint32& OutAlign) final override;
+	virtual uint64 RHIGetMinimumAlignmentForBufferBackedSRV(EPixelFormat Format) final override;
 	virtual void RHIGetTextureMemoryStats(FTextureMemoryStats& OutStats) final override;
 	virtual bool RHIGetTextureMemoryVisualizeData(FColor* TextureData,int32 SizeX,int32 SizeY,int32 Pitch,int32 PixelSize) final override;
 	virtual FTextureReferenceRHIRef RHICreateTextureReference(FLastRenderTimeContainer* LastRenderTime) final override;
diff --git a/Engine/Source/Runtime/RHI/Private/DynamicRHI.cpp b/Engine/Source/Runtime/RHI/Private/DynamicRHI.cpp
--- a/Engine/Source/Runtime/RHI/Private/DynamicRHI.cpp
+++ b/Engine/Source/Runtime/RHI/Private/DynamicRHI.cpp
@@ -391,6 +391,11 @@ void FDynamicRHI::RHIUpdateShaderResourceView(FRHIShaderResourceView* SRV, FRHII
 	UE_LOG(LogRHI, Fatal, TEXT("RHIUpdateShaderResourceView isn't implemented for the current RHI"));
 }
 
+uint64 FDynamicRHI::RHIGetMinimumAlignmentForBufferBackedSRV(EPixelFormat Format)
+{
+	return 16; // not-Metal requires a minimum alignment of of 16-bytes.
+}
+
 uint64 FDynamicRHI::RHICalcVMTexture2DPlatformSize(uint32 Mip0Width, uint32 Mip0Height, uint8 Format, uint32 NumMips, uint32 FirstMipIdx, uint32 NumSamples, uint32 Flags, uint32& OutAlign)
 {
 	UE_LOG(LogRHI, Fatal, TEXT("RHICalcVMTexture2DPlatformSize isn't implemented for the current RHI"));
diff --git a/Engine/Source/Runtime/RHI/Public/DynamicRHI.h b/Engine/Source/Runtime/RHI/Public/DynamicRHI.h
--- a/Engine/Source/Runtime/RHI/Public/DynamicRHI.h
+++ b/Engine/Source/Runtime/RHI/Public/DynamicRHI.h
@@ -664,6 +664,13 @@ public:
 	// FlushType: Thread safe
 	virtual uint64 RHICalcTextureCubePlatformSize(uint32 Size, uint8 Format, uint32 NumMips, uint32 Flags, const FRHIResourceCreateInfo& CreateInfo, uint32& OutAlign) = 0;
 
+	/**
+	* Gets the minimum alignment (in bytes) required for creating a shader resource view on a buffer-backed resource.
+	* @param Format - EPixelFormat texture format of the SRV.
+	*/
+	// FlushType: Thread safe
+	virtual uint64 RHIGetMinimumAlignmentForBufferBackedSRV(EPixelFormat Format);
+
 	/**
 	* Retrieves texture memory stats.
 	* safe to call on the main thread
@@ -1514,6 +1521,11 @@ FORCEINLINE uint64 RHICalcTextureCubePlatformSize(uint32 Size, uint8 Format, uin
 	return GDynamicRHI->RHICalcTextureCubePlatformSize(Size, Format, NumMips, Flags, CreateInfo, OutAlign);
 }
 
+FORCEINLINE uint64 RHIGetMinimumAlignmentForBufferBackedSRV(EPixelFormat Format)
+{
+	return GDynamicRHI->RHIGetMinimumAlignmentForBufferBackedSRV(Format);
+}
+
 FORCEINLINE void RHIGetTextureMemoryStats(FTextureMemoryStats& OutStats)
 {
 	GDynamicRHI->RHIGetTextureMemoryStats(OutStats);
diff --git a/Engine/Source/Runtime/RenderCore/Private/DynamicBufferAllocator.cpp b/Engine/Source/Runtime/RenderCore/Private/DynamicBufferAllocator.cpp
--- a/Engine/Source/Runtime/RenderCore/Private/DynamicBufferAllocator.cpp
+++ b/Engine/Source/Runtime/RenderCore/Private/DynamicBufferAllocator.cpp
@@ -98,8 +98,8 @@ FGlobalDynamicReadBuffer::FAllocation FGlobalDynamicReadBuffer::AllocateFloat(ui
 	FScopeLock ScopeLock(&FloatBufferPool->CriticalSection);
 	FAllocation Allocation;
 
-	// The codepath using FShaderResourceViewInitializer, requires the SRV to be aligned on 16 bytes on some platforms.
-	Num = Align(Num, 4);
+	// The codepath using FShaderResourceViewInitializer, requires the SRV to be aligned on some platforms.
+	Num = Align(Num, RHIGetMinimumAlignmentForBufferBackedSRV(PF_R32_FLOAT) / sizeof(float));
 
 	TotalAllocatedSinceLastCommit += Num;
 	if (IsRenderAlarmLoggingEnabled())
@@ -159,7 +159,7 @@ FGlobalDynamicReadBuffer::FAllocation FGlobalDynamicReadBuffer::AllocateInt32(ui
 	FAllocation Allocation;
 
 	// The codepath using FShaderResourceViewInitializer, requires the SRV to be aligned on 16 bytes on some platforms.
-	Num = Align(Num, 4);
+	Num = Align(Num, RHIGetMinimumAlignmentForBufferBackedSRV(PF_R32_SINT) / sizeof(int32));
 
 	TotalAllocatedSinceLastCommit += Num;
 	if (IsRenderAlarmLoggingEnabled())
@@ -264,4 +264,4 @@ void FGlobalDynamicReadBuffer::Commit()
 	}
 	Int32BufferPool->CurrentBuffer = NULL;
 	TotalAllocatedSinceLastCommit = 0;
-}
\ No newline at end of file
+}
