Fix for non-texturebuffer path to support SRV minimum alignments onto buffer-backed textures.

diff --git a/Engine/Source/Runtime/Apple/MetalRHI/Private/MetalVertexBuffer.cpp b/Engine/Source/Runtime/Apple/MetalRHI/Private/MetalVertexBuffer.cpp
--- a/Engine/Source/Runtime/Apple/MetalRHI/Private/MetalVertexBuffer.cpp
+++ b/Engine/Source/Runtime/Apple/MetalRHI/Private/MetalVertexBuffer.cpp
@@ -309,26 +309,28 @@ FMetalTexture FMetalRHIBuffer::AllocLinearTexture(EPixelFormat Format, NSUIntege
 		}
 		else
 		{
-			uint32 NumElements = ((Buffer.GetLength() - Offset) / Stride);
-			uint32 SizeX = NumElements;
+			const uint32 RequestedNumElements = (Buffer.GetLength() - Offset) / Stride;
+
+			uint32 SizeX = RequestedNumElements;
 			uint32 SizeY = 1;
-			if (NumElements > GMaxTextureDimensions)
+
+			if (RequestedNumElements > GMaxTextureDimensions)
 			{
 				uint32 Dimension = GMaxTextureDimensions;
-				while((NumElements % Dimension) != 0)
+				while ((RequestedNumElements % Dimension) != 0)
 				{
 					check(Dimension >= 1);
 					Dimension = (Dimension >> 1);
 				}
 				SizeX = Dimension;
-				SizeY = NumElements / Dimension;
-				checkf(SizeX <= GMaxTextureDimensions, TEXT("Calculated width %u is greater than maximum permitted %d when converting buffer of size %llu with element stride %u to a 2D texture with %u elements."), SizeX, (int32)GMaxTextureDimensions, Buffer.GetLength(), Stride, NumElements);
-				checkf(SizeX <= GMaxTextureDimensions, TEXT("Calculated height %u is greater than maximum permitted %d when converting buffer of size %llu with element stride %u to a 2D texture with %u elements."), SizeY, (int32)GMaxTextureDimensions, Buffer.GetLength(), Stride, NumElements);
+				SizeY = RequestedNumElements / Dimension;
+				checkf(SizeX <= GMaxTextureDimensions, TEXT("Calculated width %u is greater than maximum permitted %d when converting buffer of size %llu with element stride %u to a 2D texture with %u elements."), SizeX, (int32)GMaxTextureDimensions, Buffer.GetLength(), Stride, RequestedNumElements);
+				checkf(SizeY <= GMaxTextureDimensions, TEXT("Calculated height %u is greater than maximum permitted %d when converting buffer of size %llu with element stride %u to a 2D texture with %u elements."), SizeY, (int32)GMaxTextureDimensions, Buffer.GetLength(), Stride, RequestedNumElements);
 			}
-			
+
 			NewSize = SizeX*Stride;
 
-			check((NewSize % 1024) == 0);
+			check((NewSize % GetMetalDeviceContext().GetDevice().GetMinimumLinearTextureAlignmentForPixelFormat((mtlpp::PixelFormat)GMetalBufferFormats[Format].LinearTextureFormat)) == 0);
 			check((NewSize + Offset) <= Buffer.GetLength());
 
 			Desc = mtlpp::TextureDescriptor::Texture2DDescriptor(MTLFormat, SizeX, SizeY, NO);
@@ -337,10 +339,10 @@ FMetalTexture FMetalRHIBuffer::AllocLinearTexture(EPixelFormat Format, NSUIntege
 			Desc.SetUsage((mtlpp::TextureUsage)TexUsage);
 			Desc.SetResourceOptions((mtlpp::ResourceOptions)Options);
 		}
-		
+
 		FMetalTexture Texture = MTLPP_VALIDATE(mtlpp::Buffer, Buffer, SafeGetRuntimeDebuggingLevel() >= EMetalDebugLevelValidation, NewTexture(Desc, Offset, NewSize));
 		METAL_FATAL_ASSERT(Texture, TEXT("Failed to create linear texture, desc %s from buffer %s"), *FString([Desc description]), *FString([Buffer description]));
-		
+
 		return Texture;
 	}
 	else
