Fix compile error in MetalRenderTarget.cpp. Redefinition of default argument.

diff --git a/Engine/Source/Runtime/Apple/MetalRHI/Private/MetalRenderTarget.cpp b/Engine/Source/Runtime/Apple/MetalRHI/Private/MetalRenderTarget.cpp
--- a/Engine/Source/Runtime/Apple/MetalRHI/Private/MetalRenderTarget.cpp
+++ b/Engine/Source/Runtime/Apple/MetalRHI/Private/MetalRenderTarget.cpp
@@ -343,7 +343,7 @@ void FMetalDynamicRHI::RHIReadSurfaceData(FRHITexture* TextureRHI, FIntRect Rect
 	}
 }
 
-void FMetalDynamicRHI::RHIMapStagingSurface(FRHITexture* TextureRHI, FRHIGPUFence* FenceRHI, void*& OutData, int32& OutWidth, int32& OutHeight, uint32 GPUIndex = 0)
+void FMetalDynamicRHI::RHIMapStagingSurface(FRHITexture* TextureRHI, FRHIGPUFence* FenceRHI, void*& OutData, int32& OutWidth, int32& OutHeight, uint32 GPUIndex)
 {
 	@autoreleasepool {
     FMetalSurface* Surface = GetMetalSurfaceFromRHITexture(TextureRHI);
