Revert to using DefaultSampleRate if the iOS device query returns 0.0 as the preferred sample rate.

diff --git a/Engine/Source/Runtime/Apple/AudioMixerAudioUnit/Private/AudioMixerPlatformAudioUnit.cpp b/Engine/Source/Runtime/Apple/AudioMixerAudioUnit/Private/AudioMixerPlatformAudioUnit.cpp
--- a/Engine/Source/Runtime/Apple/AudioMixerAudioUnit/Private/AudioMixerPlatformAudioUnit.cpp
+++ b/Engine/Source/Runtime/Apple/AudioMixerAudioUnit/Private/AudioMixerPlatformAudioUnit.cpp
@@ -443,6 +443,12 @@ namespace Audio
 		{
 			NumFrames = (int32)(SampleRate * BufferSizeInSec);
 		}
+
+		if (FMath::IsNearlyZero(SampleRate))
+		{
+			SampleRate = DefaultSampleRate;
+		}
+
 		InternalPlatformSettings.SampleRate = SampleRate;
 
 		return InternalPlatformSettings;
