Fixup FMetalCodeHeader.

- CompileFlags needs to be 32-bits
- Frequency only needs to be 8-bits

Thanks to [REDACTED_PERSON] for root-causing this one!

diff --git a/Engine/Shaders/Public/Platform/Metal/MetalCommon.ush b/Engine/Shaders/Public/Platform/Metal/MetalCommon.ush
--- a/Engine/Shaders/Public/Platform/Metal/MetalCommon.ush
+++ b/Engine/Shaders/Public/Platform/Metal/MetalCommon.ush
@@ -5,7 +5,7 @@
 =============================================================================*/
 
 // Update this GUID to improve shader recompilation for Metal only shaders
-// GUID = 7E1C3027-495D-4F6A-B33D-DBDE632E480B
+// GUID = D8D6C149-3862-4960-A60F-DD7862433434
 
 #pragma once
 
diff --git a/Engine/Source/Runtime/Apple/MetalRHI/Public/MetalShaderResources.h b/Engine/Source/Runtime/Apple/MetalRHI/Public/MetalShaderResources.h
--- a/Engine/Source/Runtime/Apple/MetalRHI/Public/MetalShaderResources.h
+++ b/Engine/Source/Runtime/Apple/MetalRHI/Public/MetalShaderResources.h
@@ -327,13 +327,13 @@ struct FMetalCodeHeader
 
 	uint64 CompilerBuild;
 	uint32 CompilerVersion;
-	uint32 Frequency;
 	uint32 SourceLen;
 	uint32 SourceCRC;
 	uint32 NumThreadsX;
 	uint32 NumThreadsY;
 	uint32 NumThreadsZ;
-	uint16 CompileFlags;
+	uint32 CompileFlags;
+	uint8 Frequency;
 	uint8 Version;
 	int8 SideTable;
 	bool bDeviceFunctionConstants;
@@ -341,18 +341,17 @@ struct FMetalCodeHeader
 	FMetalCodeHeader()
 	: CompilerBuild(0)
 	, CompilerVersion(0)
-	, Frequency(0)
 	, SourceLen(0)
 	, SourceCRC(0)
 	, NumThreadsX(0)
 	, NumThreadsY(0)
 	, NumThreadsZ(0)
 	, CompileFlags(0)
+	, Frequency(0)
 	, Version(0)
 	, SideTable(-1)
 	, bDeviceFunctionConstants(false)
 	{
-		
 	}
 };
 
@@ -385,13 +384,13 @@ inline FArchive& operator<<(FArchive& Ar, FMetalCodeHeader& Header)
 
 	Ar << Header.CompilerBuild;
 	Ar << Header.CompilerVersion;
-	Ar << Header.Frequency;
 	Ar << Header.SourceLen;
 	Ar << Header.SourceCRC;
 	Ar << Header.NumThreadsX;
 	Ar << Header.NumThreadsY;
 	Ar << Header.NumThreadsZ;
 	Ar << Header.CompileFlags;
+	Ar << Header.Frequency;
 	Ar << Header.Version;
 	Ar << Header.SideTable;
 	Ar << Header.bDeviceFunctionConstants;
