Disable the use of Parallel RHI execute on macOS+AMD.  The MetalRHI is just not set up to use seperate metal command buffers across multiple RHI task threads - this piece of work will be handled outside of 4.25.

diff --git a/Engine/Source/Runtime/Apple/MetalRHI/Private/MetalRHI.cpp b/Engine/Source/Runtime/Apple/MetalRHI/Private/MetalRHI.cpp
--- a/Engine/Source/Runtime/Apple/MetalRHI/Private/MetalRHI.cpp
+++ b/Engine/Source/Runtime/Apple/MetalRHI/Private/MetalRHI.cpp
@@ -410,7 +410,11 @@ FMetalDynamicRHI::FMetalDynamicRHI(ERHIFeatureLevel::Type RequestedFeatureLevel)
 #else
 		GRHISupportsRHIThread = bSupportsRHIThread;
 #endif
-		GRHISupportsParallelRHIExecute = GRHISupportsRHIThread && ((!IsRHIDeviceIntel() && !IsRHIDeviceNVIDIA()) || FParse::Param(FCommandLine::Get(),TEXT("metalparallel")));
+		// TODO: Reenable GRHISupportsParallelRHIExecute once fixed
+		// At the moment there are too many places where we call GetDeviceContext() whoich means we end up accessing the same Metal command buffer
+		// in the immediate commandlist and all the RHI task threads which causes crashes.
+		// This should have this logic for Parallel Execute: GRHISupportsRHIThread && ((!IsRHIDeviceIntel() && !IsRHIDeviceNVIDIA()) || FParse::Param(FCommandLine::Get(),TEXT("metalparallel")));
+		GRHISupportsParallelRHIExecute = GRHISupportsRHIThread && FParse::Param(FCommandLine::Get(),TEXT("metalparallel"));
 #endif
 		GSupportsEfficientAsyncCompute = GRHISupportsParallelRHIExecute && (IsRHIDeviceAMD() || PLATFORM_IOS || FParse::Param(FCommandLine::Get(),TEXT("metalasynccompute"))); // Only AMD currently support async. compute and it requires parallel execution to be useful.
 		GSupportsParallelOcclusionQueries = GRHISupportsRHIThread;
