Fix shader library validation check.

#rnx

diff --git a/Engine/Source/Runtime/Apple/MetalRHI/Private/MetalShaders.cpp b/Engine/Source/Runtime/Apple/MetalRHI/Private/MetalShaders.cpp
--- a/Engine/Source/Runtime/Apple/MetalRHI/Private/MetalShaders.cpp
+++ b/Engine/Source/Runtime/Apple/MetalRHI/Private/MetalShaders.cpp
@@ -1448,7 +1448,7 @@ FRHIShaderLibraryRef FMetalDynamicRHI::RHICreateShaderLibrary(EShaderPlatform Pl
 		// Would be good to check the language version of the library with the archive format here.
 		if (Header.Format == PlatformName.GetPlainNameString())
 		{
-			check(SerializedShaders.GetNumShaders() / Header.NumShadersPerLibrary == Header.NumLibraries);
+			check(((SerializedShaders.GetNumShaders() + Header.NumShadersPerLibrary - 1) / Header.NumShadersPerLibrary) == Header.NumLibraries);
 
 			TArray<mtlpp::Library> Libraries;
 			Libraries.Empty(Header.NumLibraries);
