Manually merge CL 11643384

#rnx

---

Fix ShaderResourceView linear texture allocation for shared buffer backing and reuse.

Instead of just keying off an offset, have the CreateLinearTexture and AllocateLinearTexture leverage a linear texture descriptor, and make the Metal RHI agree with the higher level RHI in terms of units.

diff --git a/Engine/Source/Runtime/Apple/MetalRHI/Private/MetalIndexBuffer.cpp b/Engine/Source/Runtime/Apple/MetalRHI/Private/MetalIndexBuffer.cpp
--- a/Engine/Source/Runtime/Apple/MetalRHI/Private/MetalIndexBuffer.cpp
+++ b/Engine/Source/Runtime/Apple/MetalRHI/Private/MetalIndexBuffer.cpp
@@ -32,7 +32,7 @@ FMetalIndexBuffer::FMetalIndexBuffer(uint32 InStride, uint32 InSize, uint32 InUs
 	if (RHISupportsTessellation(GMaxRHIShaderPlatform))
 	{
 		EPixelFormat Format = IndexType == mtlpp::IndexType::UInt16 ? PF_R16_UINT : PF_R32_UINT;
-		CreateLinearTexture(Format, this, 0);
+		CreateLinearTexture(Format, this);
 	}
 }
 
diff --git a/Engine/Source/Runtime/Apple/MetalRHI/Private/MetalRHI.cpp b/Engine/Source/Runtime/Apple/MetalRHI/Private/MetalRHI.cpp
--- a/Engine/Source/Runtime/Apple/MetalRHI/Private/MetalRHI.cpp
+++ b/Engine/Source/Runtime/Apple/MetalRHI/Private/MetalRHI.cpp
@@ -498,8 +498,8 @@ FMetalDynamicRHI::FMetalDynamicRHI(ERHIFeatureLevel::Type RequestedFeatureLevel)
 	GMaxTextureDimensions = 16384;
 	GMaxCubeTextureDimensions = 16384;
 	GMaxTextureArrayLayers = 2048;
-	GMaxShadowDepthBufferSizeX = 16384;
-	GMaxShadowDepthBufferSizeY = 16384;
+	GMaxShadowDepthBufferSizeX = GMaxTextureDimensions;
+	GMaxShadowDepthBufferSizeY = GMaxTextureDimensions;
     bSupportsD16 = !FParse::Param(FCommandLine::Get(),TEXT("nometalv2")) && Device.SupportsFeatureSet(mtlpp::FeatureSet::macOS_GPUFamily1_v2);
     GRHISupportsHDROutput = FPlatformMisc::MacOSXVersionCompare(10,14,4) >= 0 && Device.SupportsFeatureSet(mtlpp::FeatureSet::macOS_GPUFamily1_v2);
 	GRHIHDRDisplayOutputFormat = (GRHISupportsHDROutput) ? PF_PLATFORM_HDR_0 : PF_B8G8R8A8;
@@ -526,11 +526,11 @@ FMetalDynamicRHI::FMetalDynamicRHI(ERHIFeatureLevel::Type RequestedFeatureLevel)
 	
 	GRHIHDRDisplayOutputFormat = (GRHISupportsHDROutput) ? PF_PLATFORM_HDR_0 : PF_B8G8R8A8;
 #endif
-	GMaxTextureDimensions = 4096;
-	GMaxCubeTextureDimensions = 4096;
+	GMaxTextureDimensions = 8192;
+	GMaxCubeTextureDimensions = 8192;
 	GMaxTextureArrayLayers = 2048;
-	GMaxShadowDepthBufferSizeX = 4096;
-	GMaxShadowDepthBufferSizeY = 4096;
+	GMaxShadowDepthBufferSizeX = GMaxTextureDimensions;
+	GMaxShadowDepthBufferSizeY = GMaxTextureDimensions;
 #endif
 	
 	GMaxTextureMipCount = FPlatformMath::CeilLogTwo( GMaxTextureDimensions ) + 1;
diff --git a/Engine/Source/Runtime/Apple/MetalRHI/Private/MetalUAV.cpp b/Engine/Source/Runtime/Apple/MetalRHI/Private/MetalUAV.cpp
--- a/Engine/Source/Runtime/Apple/MetalRHI/Private/MetalUAV.cpp
+++ b/Engine/Source/Runtime/Apple/MetalRHI/Private/MetalUAV.cpp
@@ -9,19 +9,26 @@
 static void ClearUAV(TRHICommandList_RecursiveHazardous<FMetalRHICommandContext>& RHICmdList, FMetalUnorderedAccessView* UnorderedAccessView, const void* ClearValue, bool bFloat);
 
 FMetalShaderResourceView::FMetalShaderResourceView()
-: TextureView(nullptr)
-, Offset(0)
-, MipLevel(0)
-, NumMips(0)
-, Format(0)
-, Stride(0)
+	: TextureView(nullptr)
+	, Offset(0)
+	, MipLevel(0)
+	, NumMips(0)
+	, Format(0)
+	, Stride(0)
+	, LinearTextureDesc(nullptr)
 {
-	
+	// void
 }
 
 FMetalShaderResourceView::~FMetalShaderResourceView()
 {
-	if(TextureView)
+	if (LinearTextureDesc)
+	{
+		delete LinearTextureDesc;
+		LinearTextureDesc = nullptr;
+	}
+	
+	if (TextureView)
 	{
 		FMetalSurface* Surface = GetMetalSurfaceFromRHITexture(SourceTexture);
 		if (Surface)
@@ -29,10 +36,9 @@ FMetalShaderResourceView::~FMetalShaderResourceView()
 			Surface->SRVs.Remove(this);
 		}
 		
-		if(TextureView->Texture)
+		if (TextureView->Texture)
 		{
 			TextureView->Texture = nil;
-			
 			TextureView->MSAATexture = nil;
 		}
 		delete TextureView;
@@ -43,18 +49,25 @@ FMetalShaderResourceView::~FMetalShaderResourceView()
 	SourceTexture = NULL;
 }
 
+void FMetalShaderResourceView::InitLinearTextureDescriptor(const FMetalLinearTextureDescriptor& InLinearTextureDescriptor)
+{
+	check(!LinearTextureDesc);
+	LinearTextureDesc = new FMetalLinearTextureDescriptor(InLinearTextureDescriptor);
+	check(LinearTextureDesc);
+}
+
 ns::AutoReleased<FMetalTexture> FMetalShaderResourceView::GetLinearTexture(bool const bUAV)
 {
 	ns::AutoReleased<FMetalTexture> NewLinearTexture;
 	{
 		if (IsValidRef(SourceVertexBuffer))
 		{
-			NewLinearTexture = SourceVertexBuffer->GetLinearTexture((EPixelFormat)Format, Offset);
+			NewLinearTexture = SourceVertexBuffer->GetLinearTexture((EPixelFormat)Format, LinearTextureDesc);
 			check(NewLinearTexture);
 		}
 		else if (IsValidRef(SourceIndexBuffer))
 		{
-			NewLinearTexture = SourceIndexBuffer->GetLinearTexture((EPixelFormat)Format, Offset);
+			NewLinearTexture = SourceIndexBuffer->GetLinearTexture((EPixelFormat)Format, LinearTextureDesc);
 			check(NewLinearTexture);
 		}
 	}
@@ -156,7 +169,7 @@ FUnorderedAccessViewRHIRef FMetalDynamicRHI::RHICreateUnorderedAccessView(FRHIVe
 	SRV->Format = Format;
 	{
 		check(VertexBuffer->GetUsage() & BUF_UnorderedAccess);
-		VertexBuffer->CreateLinearTexture((EPixelFormat)Format, VertexBuffer, 0);
+		VertexBuffer->CreateLinearTexture((EPixelFormat)Format, VertexBuffer);
 	}
 		
 	// create the UAV buffer to point to the structured buffer's memory
@@ -180,7 +193,7 @@ FUnorderedAccessViewRHIRef FMetalDynamicRHI::RHICreateUnorderedAccessView(FRHIIn
 		SRV->Format = Format;
 		{
 			check(IndexBuffer->GetUsage() & BUF_UnorderedAccess);
-			IndexBuffer->CreateLinearTexture((EPixelFormat)Format, IndexBuffer, 0);
+			IndexBuffer->CreateLinearTexture((EPixelFormat)Format, IndexBuffer);
 		}
 		
 		// create the UAV buffer to point to the structured buffer's memory
@@ -344,7 +357,10 @@ FShaderResourceViewRHIRef FMetalDynamicRHI::RHICreateShaderResourceView(const FS
 					SRV->Stride = Stride;
 					SRV->Offset = Desc.StartElement * Stride;
 					
-					VertexBuffer->CreateLinearTexture((EPixelFormat)Desc.Format, VertexBuffer, SRV->Offset);
+					FMetalLinearTextureDescriptor LinearTextureDesc(Desc.StartElement, Desc.NumElements, Stride);
+					SRV->InitLinearTextureDescriptor(LinearTextureDesc);
+					
+					VertexBuffer->CreateLinearTexture((EPixelFormat)Desc.Format, VertexBuffer, &LinearTextureDesc);
 				}
 				
 				return SRV;
@@ -395,8 +411,11 @@ FShaderResourceViewRHIRef FMetalDynamicRHI::RHICreateShaderResourceView(const FS
 					
 					SRV->Offset = Desc.StartElement * Stride;
 					SRV->Stride = Stride;
-						  
-					IndexBuffer->CreateLinearTexture((EPixelFormat)SRV->Format, IndexBuffer, SRV->Offset);
+					
+					FMetalLinearTextureDescriptor LinearTextureDesc(Desc.StartElement, Desc.NumElements, Stride);
+					SRV->InitLinearTextureDescriptor(LinearTextureDesc);
+					
+					IndexBuffer->CreateLinearTexture((EPixelFormat)SRV->Format, IndexBuffer, &LinearTextureDesc);
 				}
 					  
 				return SRV;
diff --git a/Engine/Source/Runtime/Apple/MetalRHI/Private/MetalVertexBuffer.cpp b/Engine/Source/Runtime/Apple/MetalRHI/Private/MetalVertexBuffer.cpp
--- a/Engine/Source/Runtime/Apple/MetalRHI/Private/MetalVertexBuffer.cpp
+++ b/Engine/Source/Runtime/Apple/MetalRHI/Private/MetalVertexBuffer.cpp
@@ -186,7 +186,7 @@ FMetalRHIBuffer::FMetalRHIBuffer(uint32 InSize, uint32 InUsage, ERHIResourceType
 
 FMetalRHIBuffer::~FMetalRHIBuffer()
 {
-	for (TPair<TTuple<EPixelFormat, uint32>, FMetalTexture>& Pair : LinearTextures)
+	for (auto& Pair : LinearTextures)
 	{
 		SafeReleaseMetalTexture(Pair.Value);
 		Pair.Value = nil;
@@ -249,7 +249,7 @@ void FMetalRHIBuffer::Alloc(uint32 InSize, EResourceLockMode LockMode, bool bIsU
         
 		if ((Usage & (BUF_UnorderedAccess|BUF_ShaderResource)))
 		{
-			for (TPair<TTuple<EPixelFormat, uint32>, FMetalTexture>& Pair : LinearTextures)
+			for (auto& Pair : LinearTextures)
 			{
 				SafeReleaseMetalTexture(Pair.Value);
 				Pair.Value = nil;
@@ -276,11 +276,11 @@ void FMetalRHIBuffer::AllocTransferBuffer(bool bOnRHIThread, uint32 InSize, ERes
     }
 }
 
-FMetalTexture FMetalRHIBuffer::AllocLinearTexture(EPixelFormat Format, NSUInteger Offset)
+FMetalTexture FMetalRHIBuffer::AllocLinearTexture(EPixelFormat InFormat, const FMetalLinearTextureDescriptor& LinearTextureDesc)
 {
 	if ((Usage & (BUF_UnorderedAccess|BUF_ShaderResource)))
 	{
-		mtlpp::PixelFormat MTLFormat = (mtlpp::PixelFormat)GMetalBufferFormats[Format].LinearTextureFormat;
+		mtlpp::PixelFormat MTLFormat = (mtlpp::PixelFormat)GMetalBufferFormats[InFormat].LinearTextureFormat;
 		
 		mtlpp::TextureDescriptor Desc;
 		NSUInteger Options = ((NSUInteger)Mode << mtlpp::ResourceStorageModeShift) | ((NSUInteger)Buffer.GetCpuCacheMode() << mtlpp::ResourceCpuCacheModeShift);
@@ -295,52 +295,62 @@ FMetalTexture FMetalRHIBuffer::AllocLinearTexture(EPixelFormat Format, NSUIntege
 			TexUsage |= mtlpp::TextureUsage::ShaderWrite;
 		}
 		
-		uint32 Stride = GPixelFormats[Format].BlockBytes;
-		if (MTLFormat == mtlpp::PixelFormat::RG11B10Float && MTLFormat != (mtlpp::PixelFormat)GPixelFormats[Format].PlatformFormat)
+		uint32 BytesPerElement = (0 == LinearTextureDesc.BytesPerElement) ? GPixelFormats[InFormat].BlockBytes : LinearTextureDesc.BytesPerElement;
+		if (MTLFormat == mtlpp::PixelFormat::RG11B10Float && MTLFormat != (mtlpp::PixelFormat)GPixelFormats[InFormat].PlatformFormat)
 		{
-			Stride = 4;
+			BytesPerElement = 4;
 		}
-		NSUInteger NewSize = Size - Offset;
+
+		const uint32 MinimumByteAlignment = GetMetalDeviceContext().GetDevice().GetMinimumLinearTextureAlignmentForPixelFormat((mtlpp::PixelFormat)GMetalBufferFormats[InFormat].LinearTextureFormat);
+		const uint32 MinimumElementAlignment = MinimumByteAlignment / BytesPerElement;
+
+		uint32 Offset = LinearTextureDesc.StartElement * BytesPerElement;
+		check(Offset % MinimumByteAlignment == 0);
+
+		uint32 NumElements = (UINT_MAX == LinearTextureDesc.NumElements) ? ((Size - Offset) / BytesPerElement) : LinearTextureDesc.NumElements;
+		NumElements = Align(NumElements, MinimumElementAlignment);
+
+		uint32 RowBytes = NumElements * BytesPerElement;
 
 		if (FMetalCommandQueue::SupportsFeature(EMetalFeaturesTextureBuffers))
 		{
-			Desc = mtlpp::TextureDescriptor::TextureBufferDescriptor(MTLFormat, NewSize / Stride, mtlpp::ResourceOptions(Options), mtlpp::TextureUsage(TexUsage));
+			Desc = mtlpp::TextureDescriptor::TextureBufferDescriptor(MTLFormat, NumElements, mtlpp::ResourceOptions(Options), mtlpp::TextureUsage(TexUsage));
 			Desc.SetAllowGPUOptimisedContents(false);
 		}
 		else
 		{
-			const uint32 RequestedNumElements = (Buffer.GetLength() - Offset) / Stride;
-
-			uint32 SizeX = RequestedNumElements;
-			uint32 SizeY = 1;
+			uint32 Width = NumElements;
+			uint32 Height = 1;
 
-			if (RequestedNumElements > GMaxTextureDimensions)
+			if (NumElements > GMaxTextureDimensions)
 			{
 				uint32 Dimension = GMaxTextureDimensions;
-				while ((RequestedNumElements % Dimension) != 0)
+				while ((NumElements % Dimension) != 0)
 				{
 					check(Dimension >= 1);
 					Dimension = (Dimension >> 1);
 				}
-				SizeX = Dimension;
-				SizeY = RequestedNumElements / Dimension;
-				checkf(SizeX <= GMaxTextureDimensions, TEXT("Calculated width %u is greater than maximum permitted %d when converting buffer of size %llu with element stride %u to a 2D texture with %u elements."), SizeX, (int32)GMaxTextureDimensions, Buffer.GetLength(), Stride, RequestedNumElements);
-				checkf(SizeY <= GMaxTextureDimensions, TEXT("Calculated height %u is greater than maximum permitted %d when converting buffer of size %llu with element stride %u to a 2D texture with %u elements."), SizeY, (int32)GMaxTextureDimensions, Buffer.GetLength(), Stride, RequestedNumElements);
+
+				Width = Dimension;
+				Height = NumElements / Dimension;
+
+				checkf(Width <= GMaxTextureDimensions, TEXT("Calculated width %u is greater than maximum permitted %d when converting buffer of size %llu with element stride %u to a 2D texture with %u elements."), Width, (int32)GMaxTextureDimensions, Buffer.GetLength(), BytesPerElement, NumElements);
+				checkf(Height <= GMaxTextureDimensions, TEXT("Calculated height %u is greater than maximum permitted %d when converting buffer of size %llu with element stride %u to a 2D texture with %u elements."), Height, (int32)GMaxTextureDimensions, Buffer.GetLength(), BytesPerElement, NumElements);
 			}
 
-			NewSize = SizeX*Stride;
+			RowBytes = Width * BytesPerElement;
 
-			check((NewSize % GetMetalDeviceContext().GetDevice().GetMinimumLinearTextureAlignmentForPixelFormat((mtlpp::PixelFormat)GMetalBufferFormats[Format].LinearTextureFormat)) == 0);
-			check((NewSize + Offset) <= Buffer.GetLength());
+			check(RowBytes % MinimumByteAlignment == 0);
+			check((RowBytes * Height) + Offset <= Buffer.GetLength());
 
-			Desc = mtlpp::TextureDescriptor::Texture2DDescriptor(MTLFormat, SizeX, SizeY, NO);
+			Desc = mtlpp::TextureDescriptor::Texture2DDescriptor(MTLFormat, Width, Height, NO);
 			Desc.SetStorageMode(Mode);
 			Desc.SetCpuCacheMode(Buffer.GetCpuCacheMode());
 			Desc.SetUsage((mtlpp::TextureUsage)TexUsage);
 			Desc.SetResourceOptions((mtlpp::ResourceOptions)Options);
 		}
 
-		FMetalTexture Texture = MTLPP_VALIDATE(mtlpp::Buffer, Buffer, SafeGetRuntimeDebuggingLevel() >= EMetalDebugLevelValidation, NewTexture(Desc, Offset, NewSize));
+		FMetalTexture Texture = MTLPP_VALIDATE(mtlpp::Buffer, Buffer, SafeGetRuntimeDebuggingLevel() >= EMetalDebugLevelValidation, NewTexture(Desc, Offset, RowBytes));
 		METAL_FATAL_ASSERT(Texture, TEXT("Failed to create linear texture, desc %s from buffer %s"), *FString([Desc description]), *FString([Buffer description]));
 
 		return Texture;
@@ -356,14 +366,18 @@ struct FMetalRHICommandCreateLinearTexture : public FRHICommand<FMetalRHICommand
 	FMetalRHIBuffer* Buffer;
 	TRefCountPtr<FRHIResource> Parent;
 	EPixelFormat Format;
-	uint32 Offset;
+	FMetalLinearTextureDescriptor LinearTextureDesc;
 	
-	FORCEINLINE_DEBUGGABLE FMetalRHICommandCreateLinearTexture(FMetalRHIBuffer* InBuffer, FRHIResource* InParent, EPixelFormat InFormat, uint32 InOffset)
-	: Buffer(InBuffer)
-	, Parent(InParent)
-	, Format(InFormat)
-	, Offset(InOffset)
+	FORCEINLINE_DEBUGGABLE FMetalRHICommandCreateLinearTexture(FMetalRHIBuffer* InBuffer, FRHIResource* InParent, EPixelFormat InFormat, const FMetalLinearTextureDescriptor* InLinearTextureDescriptor)
+		: Buffer(InBuffer)
+		, Parent(InParent)
+		, Format(InFormat)
+		, LinearTextureDesc()
 	{
+		if (InLinearTextureDescriptor)
+		{
+			LinearTextureDesc = *InLinearTextureDescriptor;
+		}
 	}
 	
 	virtual ~FMetalRHICommandCreateLinearTexture()
@@ -372,32 +386,34 @@ struct FMetalRHICommandCreateLinearTexture : public FRHICommand<FMetalRHICommand
 	
 	void Execute(FRHICommandListBase& CmdList)
 	{
-		Buffer->CreateLinearTexture(Format, Parent.GetReference(), Offset);
+		Buffer->CreateLinearTexture(Format, Parent.GetReference(), &LinearTextureDesc);
 	}
 };
 
-ns::AutoReleased<FMetalTexture> FMetalRHIBuffer::CreateLinearTexture(EPixelFormat Format, FRHIResource* InParent, uint32 Offset)
+ns::AutoReleased<FMetalTexture> FMetalRHIBuffer::CreateLinearTexture(EPixelFormat InFormat, FRHIResource* InParent, const FMetalLinearTextureDescriptor* InLinearTextureDescriptor)
 {
 	ns::AutoReleased<FMetalTexture> Texture;
-	if ((Usage & (BUF_UnorderedAccess|BUF_ShaderResource)) && GMetalBufferFormats[Format].LinearTextureFormat != mtlpp::PixelFormat::Invalid)
+	if ((Usage & (BUF_UnorderedAccess|BUF_ShaderResource)) && GMetalBufferFormats[InFormat].LinearTextureFormat != mtlpp::PixelFormat::Invalid)
 	{
 		if (IsRunningRHIInSeparateThread() && !IsInRHIThread() && !FRHICommandListExecutor::GetImmediateCommandList().Bypass())
 		{
-			new (FRHICommandListExecutor::GetImmediateCommandList().AllocCommand<FMetalRHICommandCreateLinearTexture>()) FMetalRHICommandCreateLinearTexture(this, InParent, Format, Offset);
+			new (FRHICommandListExecutor::GetImmediateCommandList().AllocCommand<FMetalRHICommandCreateLinearTexture>()) FMetalRHICommandCreateLinearTexture(this, InParent, InFormat, InLinearTextureDescriptor);
 		}
 		else
 		{
-			FMetalTexture* ExistingTexture = LinearTextures.Find(TTuple<EPixelFormat, uint32>(Format, Offset));
+			LinearTextureMapKey MapKey = (InLinearTextureDescriptor != nullptr) ? LinearTextureMapKey(InFormat, *InLinearTextureDescriptor) : LinearTextureMapKey(InFormat, FMetalLinearTextureDescriptor());
+
+			FMetalTexture* ExistingTexture = LinearTextures.Find(MapKey);
 			if (ExistingTexture)
 			{
 				Texture = *ExistingTexture;
 			}
 			else
 			{
-				FMetalTexture NewTexture = AllocLinearTexture(Format, Offset);
+				FMetalTexture NewTexture = AllocLinearTexture(InFormat, MapKey.Value);
 				check(NewTexture);
-				check(GMetalBufferFormats[Format].LinearTextureFormat == mtlpp::PixelFormat::RG11B10Float || GMetalBufferFormats[Format].LinearTextureFormat == (mtlpp::PixelFormat)NewTexture.GetPixelFormat());
-				LinearTextures.Add(TTuple<EPixelFormat, uint32>(Format, Offset), NewTexture);
+				check(GMetalBufferFormats[InFormat].LinearTextureFormat == mtlpp::PixelFormat::RG11B10Float || GMetalBufferFormats[InFormat].LinearTextureFormat == (mtlpp::PixelFormat)NewTexture.GetPixelFormat());
+				LinearTextures.Add(MapKey, NewTexture);
 				Texture = NewTexture;
 			}
 		}
@@ -405,12 +421,14 @@ ns::AutoReleased<FMetalTexture> FMetalRHIBuffer::CreateLinearTexture(EPixelForma
 	return Texture;
 }
 
-ns::AutoReleased<FMetalTexture> FMetalRHIBuffer::GetLinearTexture(EPixelFormat Format, NSUInteger Offset)
+ns::AutoReleased<FMetalTexture> FMetalRHIBuffer::GetLinearTexture(EPixelFormat InFormat, const FMetalLinearTextureDescriptor* InLinearTextureDescriptor)
 {
 	ns::AutoReleased<FMetalTexture> Texture;
-	if ((Usage & (BUF_UnorderedAccess|BUF_ShaderResource)) && GMetalBufferFormats[Format].LinearTextureFormat != mtlpp::PixelFormat::Invalid)
+	if ((Usage & (BUF_UnorderedAccess|BUF_ShaderResource)) && GMetalBufferFormats[InFormat].LinearTextureFormat != mtlpp::PixelFormat::Invalid)
 	{
-		FMetalTexture* ExistingTexture = LinearTextures.Find(TTuple<EPixelFormat, uint32>(Format, Offset));
+		LinearTextureMapKey MapKey = (InLinearTextureDescriptor != nullptr) ? LinearTextureMapKey(InFormat, *InLinearTextureDescriptor) : LinearTextureMapKey(InFormat, FMetalLinearTextureDescriptor());
+
+		FMetalTexture* ExistingTexture = LinearTextures.Find(MapKey);
 		if (ExistingTexture)
 		{
 			Texture = *ExistingTexture;
diff --git a/Engine/Source/Runtime/Apple/MetalRHI/Public/MetalResources.h b/Engine/Source/Runtime/Apple/MetalRHI/Public/MetalResources.h
--- a/Engine/Source/Runtime/Apple/MetalRHI/Public/MetalResources.h
+++ b/Engine/Source/Runtime/Apple/MetalRHI/Public/MetalResources.h
@@ -767,9 +767,64 @@ enum EMetalBufferUsage
 	EMetalBufferUsage_LinearTex = 0x40000000,
 };
 
+class FMetalLinearTextureDescriptor
+{
+public:
+	FMetalLinearTextureDescriptor()
+		: StartElement(0)
+		, NumElements(UINT_MAX)
+		, BytesPerElement(0)
+	{
+		// void
+	}
+
+	FMetalLinearTextureDescriptor(uint32 InStartElement, uint32 InNumElements, uint32 InBytesPerElement)
+		: StartElement(InStartElement)
+		, NumElements(InNumElements)
+		, BytesPerElement(InBytesPerElement)
+	{
+		// void
+	}
+	
+	FMetalLinearTextureDescriptor(const FMetalLinearTextureDescriptor& Other)
+		: StartElement(Other.StartElement)
+		, NumElements(Other.NumElements)
+		, BytesPerElement(Other.BytesPerElement)
+	{
+		// void
+	}
+	
+	~FMetalLinearTextureDescriptor()
+	{
+		// void
+	}
+
+	friend uint32 GetTypeHash(FMetalLinearTextureDescriptor const& Key)
+	{
+		uint32 Hash = GetTypeHash((uint64)Key.StartElement);
+		Hash = HashCombine(Hash, GetTypeHash((uint64)Key.NumElements));
+		Hash = HashCombine(Hash, GetTypeHash((uint64)Key.BytesPerElement));
+		return Hash;
+	}
+
+	bool operator==(FMetalLinearTextureDescriptor const& Other) const
+	{
+		return    StartElement    == Other.StartElement
+		       && NumElements     == Other.NumElements
+		       && BytesPerElement == Other.BytesPerElement;
+	}
+
+	uint32 StartElement;
+	uint32 NumElements;
+	uint32 BytesPerElement;
+};
+
 class FMetalRHIBuffer
 {
 public:
+	using LinearTextureMapKey = TTuple<EPixelFormat, FMetalLinearTextureDescriptor>;
+	using LinearTextureMap = TMap<LinearTextureMapKey, FMetalTexture>;
+	
 	FMetalRHIBuffer(uint32 InSize, uint32 InUsage, ERHIResourceType InType);
 	virtual ~FMetalRHIBuffer();
 	
@@ -801,17 +856,17 @@ public:
 	/**
 	 * Allocate a linear texture for given format.
 	 */
-	FMetalTexture AllocLinearTexture(EPixelFormat Format, NSUInteger Offset);
+	FMetalTexture AllocLinearTexture(EPixelFormat InFormat, const FMetalLinearTextureDescriptor& InLinearTextureDescriptor);
 	
 	/**
 	 * Get a linear texture for given format.
 	 */
-	ns::AutoReleased<FMetalTexture> CreateLinearTexture(EPixelFormat Format, FRHIResource* InParent, uint32 Offset);
+	ns::AutoReleased<FMetalTexture> CreateLinearTexture(EPixelFormat InFormat, FRHIResource* InParent, const FMetalLinearTextureDescriptor* InLinearTextureDescriptor = nullptr);
 	
 	/**
 	 * Get a linear texture for given format.
 	 */
-	ns::AutoReleased<FMetalTexture> GetLinearTexture(EPixelFormat Format, NSUInteger Offset);
+	ns::AutoReleased<FMetalTexture> GetLinearTexture(EPixelFormat InFormat, const FMetalLinearTextureDescriptor* InLinearTextureDescriptor = nullptr);
 	
 	/**
 	 * Prepare a CPU accessible buffer for uploading to GPU memory
@@ -837,7 +892,7 @@ public:
 	FMetalBuffer CPUBuffer;
 	
 	// The map of linear textures for this vertex buffer - may be more than one due to type conversion.
-	TMap<TTuple<EPixelFormat, uint32>, FMetalTexture> LinearTextures;
+	LinearTextureMap LinearTextures;
 	
 	/** Buffer for small buffers < 4Kb to avoid heap fragmentation. */
 	FMetalBufferData* Data;
@@ -1092,10 +1147,15 @@ public:
 	uint8 Format;
 	uint8 Stride;
 
+	void InitLinearTextureDescriptor(const FMetalLinearTextureDescriptor& InLinearTextureDescriptor);
+
 	FMetalShaderResourceView();
 	~FMetalShaderResourceView();
 	
 	ns::AutoReleased<FMetalTexture> GetLinearTexture(bool const bUAV);
+
+private:
+	FMetalLinearTextureDescriptor* LinearTextureDesc;
 };
 
 
