Fix compile IOS

diff --git a/Engine/Source/Runtime/Apple/MetalRHI/Public/MetalRHIContext.h b/Engine/Source/Runtime/Apple/MetalRHI/Public/MetalRHIContext.h
--- a/Engine/Source/Runtime/Apple/MetalRHI/Public/MetalRHIContext.h
+++ b/Engine/Source/Runtime/Apple/MetalRHI/Public/MetalRHIContext.h
@@ -223,7 +223,10 @@ public:
 	{
 		IRHICommandContext::RHITransitionResources(DepthStencilMode, DepthTexture);
 	}
-
+	virtual void RHITransitionResources(EResourceTransitionAccess TransitionType, EResourceTransitionPipeline TransitionPipeline, FRHITexture** InTextures, int32 NumTextures) final override
+	{
+		RHITransitionResources(TransitionType, InTextures, NumTextures);
+	}
 	/**
 	 * Compute queue will wait for the fence to be written before continuing.
 	 */
