Removing internal argument buffer emulation.
Saves a chunk of memory ~1.7MB standing in front of the agency.

diff --git a/Engine/Source/Runtime/Apple/MetalRHI/Private/MetalCommands.cpp b/Engine/Source/Runtime/Apple/MetalRHI/Private/MetalCommands.cpp
--- a/Engine/Source/Runtime/Apple/MetalRHI/Private/MetalCommands.cpp
+++ b/Engine/Source/Runtime/Apple/MetalRHI/Private/MetalCommands.cpp
@@ -402,94 +402,70 @@ void FMetalRHICommandContext::RHISetShaderParameter(FRHIComputeShader* ComputeSh
 void FMetalRHICommandContext::RHISetShaderUniformBuffer(FRHIGraphicsShader* ShaderRHI, uint32 BufferIndex, FRHIUniformBuffer* BufferRHI)
 {
 	@autoreleasepool {
-	EMetalShaderStages Stage = EMetalShaderStages::Num;
-	TMap<uint32, TBitArray<>>* ArgumentBitmasks = nullptr;
-	TMap<uint32, mtlpp::ArgumentEncoder>* ArgumentEncoders = nullptr;
-	FMetalShaderBindings* Bindings = nullptr;
-	switch (ShaderRHI->GetFrequency())
-	{
-	case SF_Vertex:
-	{
-		FMetalVertexShader* VertexShader = ResourceCast(static_cast<FRHIVertexShader*>(ShaderRHI));
-		ArgumentBitmasks = &VertexShader->ArgumentBitmasks;
-		ArgumentEncoders = &VertexShader->ArgumentEncoders;
-		Bindings = &VertexShader->Bindings;
-		Stage = EMetalShaderStages::Vertex;
-	}
-		break;
-#if PLATFORM_SUPPORTS_TESSELLATION_SHADERS
-	case SF_Hull:
-	{
-		Stage = EMetalShaderStages::Hull;
-		FMetalHullShader* HullShader = ResourceCast(static_cast<FRHIHullShader*>(ShaderRHI));
-		ArgumentBitmasks = &HullShader->ArgumentBitmasks;
-		ArgumentEncoders = &HullShader->ArgumentEncoders;
-		Bindings = &HullShader->Bindings;
-	}
-		break;
-	case SF_Domain:
-	{
-		Stage = EMetalShaderStages::Domain;
-		FMetalDomainShader* DomainShader = ResourceCast(static_cast<FRHIDomainShader*>(ShaderRHI));
-		ArgumentBitmasks = &DomainShader->ArgumentBitmasks;
-		ArgumentEncoders = &DomainShader->ArgumentEncoders;
-		Bindings = &DomainShader->Bindings;
-	}
-		break;
-#endif
-	case SF_Pixel:
-	{
-		Stage = EMetalShaderStages::Pixel;
-		FMetalPixelShader* PixelShader = ResourceCast(static_cast<FRHIPixelShader*>(ShaderRHI));
-		ArgumentBitmasks = &PixelShader->ArgumentBitmasks;
-		ArgumentEncoders = &PixelShader->ArgumentEncoders;
-		Bindings = &PixelShader->Bindings;
-	}
-		break;
-	default:
-		checkf(0, TEXT("FRHIShader Type %d is invalid or unsupported!"), (int32)ShaderRHI->GetFrequency());
-		NOT_SUPPORTED("RHIShaderStage");
-		break;
-	}
+		EMetalShaderStages Stage = EMetalShaderStages::Num;
+		FMetalShaderBindings* Bindings = nullptr;
+		switch (ShaderRHI->GetFrequency())
+		{
+		case SF_Vertex:
+		{
+			FMetalVertexShader* VertexShader = ResourceCast(static_cast<FRHIVertexShader*>(ShaderRHI));
+			Bindings = &VertexShader->Bindings;
+			Stage = EMetalShaderStages::Vertex;
+		}
+			break;
+	#if PLATFORM_SUPPORTS_TESSELLATION_SHADERS
+		case SF_Hull:
+		{
+			Stage = EMetalShaderStages::Hull;
+			FMetalHullShader* HullShader = ResourceCast(static_cast<FRHIHullShader*>(ShaderRHI));
+			Bindings = &HullShader->Bindings;
+		}
+			break;
+		case SF_Domain:
+		{
+			Stage = EMetalShaderStages::Domain;
+			FMetalDomainShader* DomainShader = ResourceCast(static_cast<FRHIDomainShader*>(ShaderRHI));
+			Bindings = &DomainShader->Bindings;
+		}
+			break;
+	#endif
+		case SF_Pixel:
+		{
+			Stage = EMetalShaderStages::Pixel;
+			FMetalPixelShader* PixelShader = ResourceCast(static_cast<FRHIPixelShader*>(ShaderRHI));
+			Bindings = &PixelShader->Bindings;
+		}
+			break;
+		default:
+			checkf(0, TEXT("FRHIShader Type %d is invalid or unsupported!"), (int32)ShaderRHI->GetFrequency());
+			NOT_SUPPORTED("RHIShaderStage");
+			break;
+		}
 
-	Context->GetCurrentState().BindUniformBuffer(Stage, BufferIndex, BufferRHI);
+		Context->GetCurrentState().BindUniformBuffer(Stage, BufferIndex, BufferRHI);
 
-	check(BufferIndex < Bindings->NumUniformBuffers);
-	if ((Bindings->ArgumentBuffers) & 1 << BufferIndex)
-	{
-		FMetalUniformBuffer* UB = (FMetalUniformBuffer*)BufferRHI;
-		FMetalIAB* IAB = UB->UploadIAB(Context, (*ArgumentBitmasks)[BufferIndex], (*ArgumentEncoders)[BufferIndex]);
-		// check(IAB.IndirectArgumentBuffer->IndirectArgumentEncoder.GetEncodedLength() == (*ArgumentEncoders)[BufferIndex].GetEncodedLength());
-		Context->GetCurrentState().SetShaderBuffer(Stage, IAB->IndirectArgumentBuffer, nil, 0, IAB->IndirectArgumentEncoder.GetEncodedLength(), BufferIndex, mtlpp::ResourceUsage::Read);
-	}
-	else if ((Bindings->ConstantBuffers) & 1 << BufferIndex)
-	{
-		auto* UB = (FMetalUniformBuffer*)BufferRHI;
-		Context->GetCurrentState().SetShaderBuffer(Stage, UB->Buffer, UB->Data, 0, UB->GetSize(), BufferIndex, mtlpp::ResourceUsage::Read);
-	}
+		check(BufferIndex < Bindings->NumUniformBuffers);
+		if ((Bindings->ConstantBuffers) & 1 << BufferIndex)
+		{
+			auto* UB = (FMetalUniformBuffer*)BufferRHI;
+			Context->GetCurrentState().SetShaderBuffer(Stage, UB->Buffer, UB->Data, 0, UB->GetSize(), BufferIndex, mtlpp::ResourceUsage::Read);
+		}
 	}
 }
 
 void FMetalRHICommandContext::RHISetShaderUniformBuffer(FRHIComputeShader* ComputeShaderRHI, uint32 BufferIndex, FRHIUniformBuffer* BufferRHI)
 {
 	@autoreleasepool {
-	FMetalComputeShader* ComputeShader = ResourceCast(ComputeShaderRHI);
-	Context->GetCurrentState().BindUniformBuffer(EMetalShaderStages::Compute, BufferIndex, BufferRHI);
+		FMetalComputeShader* ComputeShader = ResourceCast(ComputeShaderRHI);
+		Context->GetCurrentState().BindUniformBuffer(EMetalShaderStages::Compute, BufferIndex, BufferRHI);
 
-	auto& Bindings = ComputeShader->Bindings;
-	check(BufferIndex < Bindings.NumUniformBuffers);
-	if ((Bindings.ArgumentBuffers) & 1 << BufferIndex)
-	{
-		auto* UB = (FMetalUniformBuffer*)BufferRHI;
-		auto* IAB = UB->UploadIAB(Context, ComputeShader->ArgumentBitmasks[BufferIndex], ComputeShader->ArgumentEncoders[BufferIndex]);
-		// check(IAB.IndirectArgumentBuffer->IndirectArgumentEncoder.GetEncodedLength() == ComputeShader->ArgumentEncoders[BufferIndex].GetEncodedLength());
-		Context->GetCurrentState().SetShaderBuffer(EMetalShaderStages::Compute, IAB->IndirectArgumentBuffer, nil, 0, IAB->IndirectArgumentEncoder.GetEncodedLength(), BufferIndex, mtlpp::ResourceUsage::Read);
-	}
-	else if ((Bindings.ConstantBuffers) & 1 << BufferIndex)
-	{
-		auto* UB = (FMetalUniformBuffer*)BufferRHI;
-		Context->GetCurrentState().SetShaderBuffer(EMetalShaderStages::Compute, UB->Buffer, UB->Data, 0, UB->GetSize(), BufferIndex, mtlpp::ResourceUsage::Read);
-	}
+		auto& Bindings = ComputeShader->Bindings;
+		check(BufferIndex < Bindings.NumUniformBuffers);
+		if ((Bindings.ConstantBuffers) & 1 << BufferIndex)
+		{
+			auto* UB = (FMetalUniformBuffer*)BufferRHI;
+			Context->GetCurrentState().SetShaderBuffer(EMetalShaderStages::Compute, UB->Buffer, UB->Data, 0, UB->GetSize(), BufferIndex, mtlpp::ResourceUsage::Read);
+		}
 	}
 }
 
diff --git a/Engine/Source/Runtime/Apple/MetalRHI/Private/MetalContext.cpp b/Engine/Source/Runtime/Apple/MetalRHI/Private/MetalContext.cpp
--- a/Engine/Source/Runtime/Apple/MetalRHI/Private/MetalContext.cpp
+++ b/Engine/Source/Runtime/Apple/MetalRHI/Private/MetalContext.cpp
@@ -734,36 +734,6 @@ void FMetalDeviceContext::ReleaseFence(FMetalFence* Fence)
 	}
 }
 
-void FMetalDeviceContext::RegisterUB(FMetalUniformBuffer* UB)
-{
-	FScopeLock Lock(&FreeListMutex);
-	UniformBuffers.Add(UB);
-}
-
-void FMetalDeviceContext::UpdateIABs(FRHITextureReference* ModifiedRef)
-{
-	if(GIsMetalInitialized)
-	{
-		FScopeLock Lock(&FreeListMutex);
-		for (FMetalUniformBuffer* UB : UniformBuffers)
-		{
-			if (UB && UB->IAB && UB->TextureReferences.Contains(ModifiedRef))
-			{
-				UB->UpdateTextureReference(ModifiedRef);
-			}
-		}
-	}
-}
-
-void FMetalDeviceContext::UnregisterUB(FMetalUniformBuffer* UB)
-{
-	if(GIsMetalInitialized)
-	{
-		FScopeLock Lock(&FreeListMutex);
-		UniformBuffers.Remove(UB);
-	}
-}
-
 FMetalTexture FMetalDeviceContext::CreateTexture(FMetalSurface* Surface, mtlpp::TextureDescriptor Descriptor)
 {
 	FMetalTexture Tex = Heap.CreateTexture(Descriptor, Surface);
diff --git a/Engine/Source/Runtime/Apple/MetalRHI/Private/MetalContext.h b/Engine/Source/Runtime/Apple/MetalRHI/Private/MetalContext.h
--- a/Engine/Source/Runtime/Apple/MetalRHI/Private/MetalContext.h
+++ b/Engine/Source/Runtime/Apple/MetalRHI/Private/MetalContext.h
@@ -190,9 +190,6 @@ public:
 	void ReleaseTexture(FMetalSurface* Surface, FMetalTexture& Texture);
 	void ReleaseTexture(FMetalTexture& Texture);
 	void ReleaseFence(FMetalFence* Fence);
-	void RegisterUB(FMetalUniformBuffer* UB);
-	void UpdateIABs(FRHITextureReference* ModifiedRef);
-	void UnregisterUB(FMetalUniformBuffer* UB);
 	
 	void BeginFrame();
 	void FlushFreeList(bool const bFlushFences = true);
diff --git a/Engine/Source/Runtime/Apple/MetalRHI/Private/MetalStateCache.cpp b/Engine/Source/Runtime/Apple/MetalRHI/Private/MetalStateCache.cpp
--- a/Engine/Source/Runtime/Apple/MetalRHI/Private/MetalStateCache.cpp
+++ b/Engine/Source/Runtime/Apple/MetalRHI/Private/MetalStateCache.cpp
@@ -2730,71 +2730,6 @@ void FMetalStateCache::CommitResourceTable(EMetalShaderStages const Frequency, m
 			CommandEncoder.SetShaderSamplerState(Type, SamplerBindings.Samplers[Index], Index);
 		}
 	}
-	
-	uint32 ArgumentBuffers = 0;
-	TMap<uint8, TArray<uint8>> ArgumentBufferMasks;
-	switch(Frequency)
-	{
-		case EMetalShaderStages::Vertex:
-			ArgumentBuffers = GraphicsPSO->VertexShader->Bindings.ArgumentBuffers;
-			ArgumentBufferMasks = GraphicsPSO->VertexShader->Bindings.ArgumentBufferMasks;
-            break;
-#if PLATFORM_SUPPORTS_TESSELLATION_SHADERS
-		case EMetalShaderStages::Hull:
-			ArgumentBuffers = GraphicsPSO->HullShader->Bindings.ArgumentBuffers;
-			ArgumentBufferMasks = GraphicsPSO->HullShader->Bindings.ArgumentBufferMasks;
-			break;
-		case EMetalShaderStages::Domain:
-			ArgumentBuffers = GraphicsPSO->DomainShader->Bindings.ArgumentBuffers;
-			ArgumentBufferMasks = GraphicsPSO->DomainShader->Bindings.ArgumentBufferMasks;
-			break;
-#endif
-		case EMetalShaderStages::Pixel:
-			ArgumentBuffers = GraphicsPSO->PixelShader->Bindings.ArgumentBuffers;
-			ArgumentBufferMasks = GraphicsPSO->PixelShader->Bindings.ArgumentBufferMasks;
-			break;
-		case EMetalShaderStages::Compute:
-			ArgumentBuffers = ComputeShader->Bindings.ArgumentBuffers;
-			ArgumentBufferMasks = ComputeShader->Bindings.ArgumentBufferMasks;
-			break;
-		default:
-			break;
-	}
-	
-	while(ArgumentBuffers)
-	{
-		uint32 Index = __builtin_ctz(ArgumentBuffers);
-		ArgumentBuffers &= ~(1 << Index);
-		
-		FMetalUniformBuffer* Buffer = (FMetalUniformBuffer*)BoundUniformBuffers[Frequency][Index];
-		check(Buffer);
-		
-		FMetalUniformBuffer::FMetalIndirectArgumentBuffer& IAB = Buffer->GetIAB();
-		
-		TArray<uint8>& Mask = ArgumentBufferMasks.FindChecked((uint8)Index);
-		for (uint8 Entry : Mask)
-		{
-			FMetalUniformBuffer::Argument const& Resource = IAB.IndirectArgumentResources[Entry];
-			if (Resource.Buffer)
-			{
-				CommandEncoder.UseIndirectArgumentResource(Resource.Buffer, Resource.Usage);
-			}
-			else if (Resource.Texture)
-			{
-				CommandEncoder.UseIndirectArgumentResource(Resource.Texture, Resource.Usage);
-			}
-		}
-		
-		if (IAB.IndirectArgumentBuffer->IndirectArgumentBufferSideTable)
-		{
-			CommandEncoder.UseIndirectArgumentResource(IAB.IndirectArgumentBuffer->IndirectArgumentBufferSideTable, mtlpp::ResourceUsage::Read);
-		}
-		
-		if (Buffer->Buffer)
-		{
-			CommandEncoder.UseIndirectArgumentResource(Buffer->Buffer, mtlpp::ResourceUsage::Read);
-		}
-	}
 }
 
 FMetalBuffer& FMetalStateCache::GetDebugBuffer()
diff --git a/Engine/Source/Runtime/Apple/MetalRHI/Private/MetalTexture.cpp b/Engine/Source/Runtime/Apple/MetalRHI/Private/MetalTexture.cpp
--- a/Engine/Source/Runtime/Apple/MetalRHI/Private/MetalTexture.cpp
+++ b/Engine/Source/Runtime/Apple/MetalRHI/Private/MetalTexture.cpp
@@ -2397,10 +2397,6 @@ void FMetalRHICommandContext::RHIUpdateTextureReference(FRHITextureReference* Te
 		if (TextureRef)
 		{
 			TextureRef->SetReferencedTexture(NewTextureRHI);
-			if (NewTextureRHI)
-			{
-				GetMetalDeviceContext().UpdateIABs(TextureRefRHI);
-			}
 		}
 	}
 }
diff --git a/Engine/Source/Runtime/Apple/MetalRHI/Private/MetalUniformBuffer.cpp b/Engine/Source/Runtime/Apple/MetalRHI/Private/MetalUniformBuffer.cpp
--- a/Engine/Source/Runtime/Apple/MetalRHI/Private/MetalUniformBuffer.cpp
+++ b/Engine/Source/Runtime/Apple/MetalRHI/Private/MetalUniformBuffer.cpp
@@ -30,103 +30,11 @@ APPLE_PLATFORM_OBJECT_ALLOC_OVERRIDES(FMetalIAB)
 }
 @end
 
-struct FMetalRHICommandInitialiseUniformBufferIAB : public FRHICommand<FMetalRHICommandInitialiseUniformBufferIAB>
-{
-	TRefCountPtr<FMetalUniformBuffer> Buffer;
-	
-	FORCEINLINE_DEBUGGABLE FMetalRHICommandInitialiseUniformBufferIAB(FMetalUniformBuffer* InBuffer)
-	: Buffer(InBuffer)
-	{
-	}
-	
-	virtual ~FMetalRHICommandInitialiseUniformBufferIAB()
-	{
-	}
-	
-	void Execute(FRHICommandListBase& CmdList)
-	{
-		FPlatformAtomics::InterlockedIncrement(&Buffer->UpdateNum);
-		Buffer->UpdateIAB();
-	}
-};
-
-struct FMetalArgumentEncoderMapFuncs : TDefaultMapKeyFuncs<TArray<FMetalArgumentDesc>, mtlpp::ArgumentEncoder, /*bInAllowDuplicateKeys*/false>
-{
-	static FORCEINLINE uint32 GetKeyHash(KeyInitType Key)
-	{
-		uint32 Hash = 0;
-		
-		for (FMetalArgumentDesc const& Desc : Key)
-		{
-			HashCombine(Hash, GetTypeHash(Desc));
-		}
-		
-		return Hash;
-	}
-};
-
-class FMetalArgumentEncoderCache
-{
-public:
-	FMetalArgumentEncoderCache()
-	{
-	}
-	
-	~FMetalArgumentEncoderCache()
-	{
-	}
-	
-	static FMetalArgumentEncoderCache& Get()
-	{
-		static FMetalArgumentEncoderCache sSelf;
-		return sSelf;
-	}
-	
-	mtlpp::ArgumentEncoder CreateEncoder(TArray<FMetalArgumentDesc> const& Desc)
-	{
-		mtlpp::ArgumentEncoder Encoder;
-		
-		FRWScopeLock Lock(Mutex, SLT_ReadOnly);
-		Encoder = Encoders.FindRef(Desc);
-		
-		if (!Encoder)
-		{
-			NSMutableArray<mtlpp::ArgumentDescriptor::Type>* Arguments = [[NSMutableArray new] autorelease];
-			for (FMetalArgumentDesc const& Args : Desc)
-			{
-				mtlpp::ArgumentDescriptor Arg;
-				Args.FillDescriptor(Arg);
-				[Arguments addObject:Arg.GetPtr()];
-			}
-			
-			Encoder = GetMetalDeviceContext().GetDevice().NewArgumentEncoderWithArguments(Arguments);
-			
-			// Now we are a writer as we want to create & add the new pipeline
-			Lock.ReleaseReadOnlyLockAndAcquireWriteLock_USE_WITH_CAUTION();
-
-			if (!Encoders.Find(Desc))
-			{
-				Encoders.Add(Desc, Encoder);
-			}
-		}
-		
-		return Encoder;
-	}
-private:
-	FRWLock Mutex;
-	TMap<TArray<FMetalArgumentDesc>, mtlpp::ArgumentEncoder, FDefaultSetAllocator, FMetalArgumentEncoderMapFuncs> Encoders;
-};
-
 FMetalUniformBuffer::FMetalUniformBuffer(const void* Contents, const FRHIUniformBufferLayout& Layout, EUniformBufferUsage InUsage, EUniformBufferValidation Validation)
 	: FRHIUniformBuffer(Layout)
     , FMetalRHIBuffer(Layout.ConstantBufferSize, (FMetalCommandQueue::SupportsFeature(EMetalFeaturesIABs) && Layout.Resources.Num() ? (EMetalBufferUsage_GPUOnly|BUF_Volatile) : BUF_Volatile), RRT_UniformBuffer)
-	, UniformUsage(InUsage)
-	, IAB(nullptr)
-	, UpdateNum(0)
-	, NumResources(0)
-	, ConstantSize(Layout.ConstantBufferSize)
 {
-	NumResources = Layout.Resources.Num();
+	uint32 NumResources = Layout.Resources.Num();
 	if (NumResources)
 	{
 		ResourceTable.Empty(NumResources);
@@ -145,16 +53,8 @@ FMetalUniformBuffer::FMetalUniformBuffer(const void* Contents, const FRHIUniform
 			ResourceTable[i] = Resource;
 		}
 	}
-	
-	if (FMetalCommandQueue::SupportsFeature(EMetalFeaturesIABs))
-	{
-		ResourceTypes.AddZeroed(NumResources);
-		for (int32 i = 0; i < NumResources; ++i)
-		{
-			ResourceTypes[i] = Layout.Resources[i].MemberType;
-		}
-	}
 
+    uint32 ConstantSize = Layout.ConstantBufferSize;
 	if (ConstantSize > 0)
 	{
 		UE_CLOG(ConstantSize > 65536, LogMetal, Fatal, TEXT("Trying to allocated a uniform layout of size %d that is greater than the maximum permitted 64k."), ConstantSize);
@@ -175,207 +75,8 @@ FMetalUniformBuffer::FMetalUniformBuffer(const void* Contents, const FRHIUniform
 			FMemory::Memcpy(Data->Data, Contents, ConstantSize);
 		}
 	}
-	
-	InitIAB();
-    UpdateResourceTable(ResourceTable, Validation);
-	
-    if (NumResources && FMetalCommandQueue::SupportsFeature(EMetalFeaturesIABs))
-    {
-		check(IAB);
-        GetMetalDeviceContext().RegisterUB(this);
-    }
-}
-
-FMetalUniformBuffer::~FMetalUniformBuffer()
-{
-	if (ResourceTable.Num() && FMetalCommandQueue::SupportsFeature(EMetalFeaturesIABs))
-	{
-		GetMetalDeviceContext().UnregisterUB(this);
-	}
-	
-	if (IAB)
-	{
-		delete IAB;
-		IAB = nullptr;
-	}
-}
-
-FMetalUniformBuffer::FMetalIndirectArgumentBuffer::FMetalIndirectArgumentBuffer()
-{
-	UpdateNum = 0;
-	UpdateIAB = 0;
-	UpdateEnc = 0;
-	IndirectArgumentBuffer = nil;
-}
-
-FMetalUniformBuffer::FMetalIndirectArgumentBuffer::~FMetalIndirectArgumentBuffer()
-{
-	if (IndirectArgumentBuffer)
-		SafeReleaseMetalObject(IndirectArgumentBuffer);
-	
-	for (auto& Pair : Tier1IABs)
-	{
-		SafeReleaseMetalObject(Pair.Value);
-	}
-}
 
-FMetalIAB* FMetalUniformBuffer::UploadIAB(FMetalContext* Ctx, TBitArray<> const& Bitmask, mtlpp::ArgumentEncoder const& IABEncoder)
-{
-	check(IAB);
-	
-	int64 UpdateIAB = FPlatformAtomics::InterlockedCompareExchange(&IAB->UpdateIAB, 0, IAB->UpdateIAB);
-	if (!FMetalCommandQueue::SupportsFeature(EMetalFeaturesTier2IABs) || UpdateIAB)
-	{
-		FRWScopeLock Lock(IAB->Mutex, SLT_ReadOnly);
-		
-		FMetalIAB* NewIAB = nullptr;
-		mtlpp::ArgumentEncoder Encoder;
-		FMetalBuffer EncodingBuffer;
-		if (FMetalCommandQueue::SupportsFeature(EMetalFeaturesTier2IABs))
-		{
-			Lock.ReleaseReadOnlyLockAndAcquireWriteLock_USE_WITH_CAUTION();
-			
-			NewIAB = IAB->IndirectArgumentBuffer;
-			Encoder = FMetalArgumentEncoderCache::Get().CreateEncoder(IAB->IndirectArgumentsDecl);
-			
-			if (!NewIAB->IndirectArgumentBuffer)
-				NewIAB->IndirectArgumentBuffer = GetMetalDeviceContext().GetResourceHeap().CreateBuffer(Encoder.GetEncodedLength(), 16, BUF_Dynamic, mtlpp::ResourceOptions(BUFFER_CACHE_MODE | ((NSUInteger)mtlpp::StorageMode::Private << mtlpp::ResourceStorageModeShift)), true);
-			
-			EncodingBuffer = Ctx->GetCurrentRenderPass().AllocateTemporyBufferForCopy(NewIAB->IndirectArgumentBuffer, Encoder.GetEncodedLength(), 16);
-			
-			if (!NewIAB->IndirectArgumentBufferSideTable)
-				NewIAB->IndirectArgumentBufferSideTable = GetMetalDeviceContext().GetResourceHeap().CreateBuffer(IAB->IndirectBufferSizes.Num() * sizeof(uint32), 16, BUF_Dynamic, mtlpp::ResourceOptions(BUFFER_CACHE_MODE | ((NSUInteger)mtlpp::StorageMode::Private << mtlpp::ResourceStorageModeShift)), true);
-			
-			FMetalBuffer Temp = Ctx->GetCurrentRenderPass().AllocateTemporyBufferForCopy(NewIAB->IndirectArgumentBufferSideTable, NewIAB->IndirectArgumentBufferSideTable.GetLength(), 16);
-			
-			FMemory::Memcpy(Temp.GetContents(), IAB->IndirectBufferSizes.GetData(), IAB->IndirectBufferSizes.Num() * sizeof(uint32));
-			
-#if PLATFORM_MAC
-			if(Temp.GetStorageMode() == mtlpp::StorageMode::Managed)
-			{
-				MTLPP_VALIDATE(mtlpp::Buffer, Temp, SafeGetRuntimeDebuggingLevel() >= EMetalDebugLevelValidation, DidModify(ns::Range(0, IAB->IndirectBufferSizes.Num() * sizeof(uint32))));
-			}
-#endif
-			Ctx->AsyncCopyFromBufferToBuffer(Temp, 0, NewIAB->IndirectArgumentBufferSideTable, 0, NewIAB->IndirectArgumentBufferSideTable.GetLength());
-		}
-		else
-		{
-			NewIAB = IAB->Tier1IABs.FindRef(Bitmask);
-			if (NewIAB && (NewIAB->UpdateIAB == UpdateIAB))
-			{
-				return NewIAB;
-			}
-			
-			Lock.ReleaseReadOnlyLockAndAcquireWriteLock_USE_WITH_CAUTION();
-			NewIAB = [FMetalIAB new];
-			NewIAB->UpdateIAB = UpdateIAB;
-			
-			uint32 SetBits = 0;
-			for (uint32 i = 0; i < Bitmask.Num(); i++)
-			{
-				if (Bitmask[i])
-				{
-					SetBits++;
-				}
-			}
-			
-			Encoder = IABEncoder;
-			
-			NewIAB->IndirectArgumentBuffer = GetMetalDeviceContext().GetResourceHeap().CreateBuffer(Encoder.GetEncodedLength(), 16, BUF_Dynamic, mtlpp::ResourceOptions(BUFFER_CACHE_MODE | ((NSUInteger)BUFFER_STORAGE_MODE << mtlpp::ResourceStorageModeShift)), true);
-			
-			NewIAB->IndirectArgumentBufferSideTable = GetMetalDeviceContext().GetResourceHeap().CreateBuffer(SetBits * sizeof(uint32) * 2, 16, BUF_Dynamic, mtlpp::ResourceOptions(BUFFER_CACHE_MODE | ((NSUInteger)BUFFER_STORAGE_MODE << mtlpp::ResourceStorageModeShift)), true);
-			
-			uint32 j = 0;
-			uint32* Ptr = (uint32*)NewIAB->IndirectArgumentBufferSideTable.GetContents();
-			for (uint32 i = 0; i < Bitmask.Num(); i++)
-			{
-				if (Bitmask[i])
-				{
-					Ptr[j * 2] = IAB->IndirectBufferSizes[i * 2];
-					Ptr[(j * 2) + 1] = IAB->IndirectBufferSizes[(i * 2) + 1];
-					j++;
-				}
-			}
-			
-			EncodingBuffer = NewIAB->IndirectArgumentBuffer;
-			
-			if (IAB->Tier1IABs.Contains(Bitmask))
-			{
-				SafeReleaseMetalObject(IAB->Tier1IABs[Bitmask]);
-			}
-			
-			IAB->Tier1IABs.Add(Bitmask, NewIAB);
-		}
-		NewIAB->IndirectArgumentEncoder = Encoder;
-		
-		Encoder.SetArgumentBuffer(EncodingBuffer, 0);
-		
-		NSUInteger MTLIndex = 0;
-		for (FMetalArgumentDesc const& Arg : IAB->IndirectArgumentsDecl)
-		{
-			NSUInteger NewIndex = Arg.Index;
-			if (FMetalCommandQueue::SupportsFeature(EMetalFeaturesTier2IABs) || (NewIndex < Bitmask.Num() && Bitmask[NewIndex]))
-			{
-				NSUInteger ResourceIndex = FMetalCommandQueue::SupportsFeature(EMetalFeaturesTier2IABs) ? NewIndex : MTLIndex++;
-				switch(Arg.DataType)
-				{
-					case mtlpp::DataType::Pointer:
-						if (IAB->IndirectArgumentResources[NewIndex].Buffer)
-							Encoder.SetBuffer(IAB->IndirectArgumentResources[NewIndex].Buffer, 0, ResourceIndex);
-						break;
-					case mtlpp::DataType::Texture:
-						if (IAB->IndirectArgumentResources[NewIndex].Texture)
-							Encoder.SetTexture(IAB->IndirectArgumentResources[NewIndex].Texture, ResourceIndex);
-						break;
-					case mtlpp::DataType::Sampler:
-						if (IAB->IndirectArgumentResources[NewIndex].Sampler)
-							Encoder.SetSamplerState(IAB->IndirectArgumentResources[NewIndex].Sampler, ResourceIndex);
-						break;
-					default:
-						break;
-				}
-			}
-		}
-		
-#if PLATFORM_MAC
-		if(EncodingBuffer.GetStorageMode() == mtlpp::StorageMode::Managed)
-		{
-			MTLPP_VALIDATE(mtlpp::Buffer, EncodingBuffer, SafeGetRuntimeDebuggingLevel() >= EMetalDebugLevelValidation, DidModify(ns::Range(0, Encoder.GetEncodedLength())));
-		}
-#endif
-		if (FMetalCommandQueue::SupportsFeature(EMetalFeaturesTier2IABs))
-		{
-			Ctx->AsyncCopyFromBufferToBuffer(EncodingBuffer, 0, NewIAB->IndirectArgumentBuffer, 0, Encoder.GetEncodedLength());
-		}
-	}
-	if (FMetalCommandQueue::SupportsFeature(EMetalFeaturesTier2IABs))
-	{
-		return IAB->IndirectArgumentBuffer;
-	}
-	else
-	{
-		return IAB->Tier1IABs[Bitmask];
-	}
-}
-
-FMetalUniformBuffer::FMetalIndirectArgumentBuffer& FMetalUniformBuffer::GetIAB()
-{
-	check(ResourceTable.Num() && FMetalCommandQueue::SupportsFeature(EMetalFeaturesIABs));
-
-	check(IAB);
-	
-	return *IAB;
-}
-
-void FMetalUniformBuffer::InitIAB()
-{
-	if (NumResources && FMetalCommandQueue::SupportsFeature(EMetalFeaturesIABs) && !IAB)
-	{
-		FMetalIndirectArgumentBuffer* NewIAB = new FMetalIndirectArgumentBuffer;
-		NewIAB->IndirectArgumentBuffer = [FMetalIAB new];
-		NewIAB->UpdateNum = ~0u;
-		IAB = NewIAB;
-	}
+    UpdateResourceTable(ResourceTable, Validation);
 }
 
 void const* FMetalUniformBuffer::GetData()
@@ -394,363 +95,14 @@ void const* FMetalUniformBuffer::GetData()
 	}
 }
 
-void FMetalUniformBuffer::UpdateTextureReference(FRHITextureReference* ModifiedRef)
-{
-	if (NumResources && FMetalCommandQueue::SupportsFeature(EMetalFeaturesIABs) && IAB && TextureReferences.Contains(ModifiedRef))
-	{
-		bool bModified = false;
-		bool bUpdatedDesc = false;
-		FMetalSurface* Surface = GetMetalSurfaceFromRHITexture(ModifiedRef);
-		if (Surface)
-		{
-			TBitArray<>& Ref = TextureReferences.FindChecked(ModifiedRef);
-			for (uint32 i = 0; i < Ref.Num(); i++)
-			{
-				if (Ref[i])
-				{
-					bModified |= (IAB->IndirectArgumentResources[i].Texture != Surface->Texture);
-					IAB->IndirectArgumentResources[i].Texture = Surface->Texture;
-					
-					FMetalArgumentDesc& Desc = IAB->IndirectArgumentsDecl[i];
-					bUpdatedDesc |= Desc.TextureType != Surface->Texture.GetTextureType();
-					Desc.SetTextureType(Surface->Texture.GetTextureType());
-					
-					bModified |= bUpdatedDesc;
-				}
-			}
-		}
-		if (bUpdatedDesc)
-		{
-			FPlatformAtomics::InterlockedIncrement(&IAB->UpdateEnc);
-		}
-		if (bModified)
-		{
-			FPlatformAtomics::InterlockedIncrement(&IAB->UpdateIAB);
-		}
-	}
-}
-
-void FMetalUniformBuffer::UpdateIAB()
-{
-	if (NumResources && FMetalCommandQueue::SupportsFeature(EMetalFeaturesIABs) && IAB && (UpdateNum != IAB->UpdateNum))
-	{
-		IAB->IndirectArgumentResources.Empty();
-		
-		TArray<uint32>& BufferSizes = IAB->IndirectBufferSizes;
-		TArray<FMetalArgumentDesc>& Arguments = IAB->IndirectArgumentsDecl;
-		Arguments.Empty();
-		BufferSizes.Empty();
-		
-		int32 Index = 0;
-		// Buffer Size table
-		{
-			FMetalArgumentDesc& Desc = Arguments.Emplace_GetRef();
-			Desc.SetIndex(Index++);
-			Desc.SetAccess(mtlpp::ArgumentAccess::ReadOnly);
-			Desc.SetDataType(mtlpp::DataType::Pointer);
-			
-			BufferSizes.AddZeroed(2);
-			
-			IAB->IndirectArgumentResources.Add(Argument(FMetalBuffer(), mtlpp::ResourceUsage::Read));
-		}
-		
-		// set up an SRT-style uniform buffer
-		for (uint32 i = 0; i < NumResources; ++i)
-		{
-			FRHIResource* Resource = ResourceTable[i];
-			if (Resource)
-			{
-				switch(ResourceTypes[i])
-				{
-					case UBMT_RDG_TEXTURE_SRV:
-					case UBMT_RDG_BUFFER_SRV:
-					case UBMT_SRV:
-					{
-						FMetalArgumentDesc& Desc = Arguments.Emplace_GetRef();
-						Desc.SetIndex(Index++);
-						Desc.SetAccess(mtlpp::ArgumentAccess::ReadOnly);
-						
-						FMetalShaderResourceView* SRV = (FMetalShaderResourceView*)Resource;
-						FRHITexture* Texture = SRV->SourceTexture.GetReference();
-						FMetalVertexBuffer* VB = SRV->SourceVertexBuffer.GetReference();
-						FMetalIndexBuffer* IB = SRV->SourceIndexBuffer.GetReference();
-						FMetalStructuredBuffer* SB = SRV->SourceStructuredBuffer.GetReference();
-						if (Texture)
-						{
-							Desc.SetDataType(mtlpp::DataType::Texture);
-							
-							FRHITextureReference* Reference = Texture->GetTextureReference();
-							if (Reference)
-							{
-								TBitArray<>& Ref = TextureReferences.FindOrAdd(Reference);
-								if (Ref.Num() <= IAB->IndirectArgumentResources.Num())
-								{
-									Ref.Add(false, (IAB->IndirectArgumentResources.Num() + 1) - Ref.Num());
-								}
-								Ref[IAB->IndirectArgumentResources.Num()] = true;
-							}
-							
-							union {
-								uint8 Components[4];
-								uint32 Packed;
-							} Swizzle;
-							Swizzle.Packed = 0;
-							assert(sizeof(Swizzle) == sizeof(uint32));
-							
-							FMetalSurface* Surface = SRV->TextureView;
-							if (Surface)
-							{
-								check(!Surface->Texture.IsAliasable());
-								Desc.SetTextureType(Surface->Texture.GetTextureType());
-								IAB->IndirectArgumentResources.Add(Argument(Surface->Texture, (mtlpp::ResourceUsage)(mtlpp::ResourceUsage::Read|mtlpp::ResourceUsage::Sample)));
-								if (Surface->Texture.GetPixelFormat() == mtlpp::PixelFormat::X32_Stencil8
-#if PLATFORM_MAC
-									||	Surface->Texture.GetPixelFormat() == mtlpp::PixelFormat::X24_Stencil8
-#endif
-									)
-								{
-									Swizzle.Components[0] = Swizzle.Components[1] = Swizzle.Components[2] = Swizzle.Components[3] = 1;
-								}
-							}
-							else
-							{
-								IAB->IndirectArgumentResources.Add(Argument(FMetalTexture(), (mtlpp::ResourceUsage)(mtlpp::ResourceUsage::Read|mtlpp::ResourceUsage::Sample)));
-							}
-							
-							BufferSizes.Add(Swizzle.Packed);
-							BufferSizes.Add(GMetalBufferFormats[Texture->GetFormat()].DataFormat);
-						}
-						else if (SB)
-						{
-							Desc.SetDataType(mtlpp::DataType::Pointer);
-							
-							IAB->IndirectArgumentResources.Add(Argument(SB->Buffer, (mtlpp::ResourceUsage)(mtlpp::ResourceUsage::Read)));
-							
-							BufferSizes.Add(SB->GetSize());
-							BufferSizes.Add(GMetalBufferFormats[SRV->Format].DataFormat);
-						}
-						else
-						{
-							check(VB || IB);
-							ns::AutoReleased<FMetalTexture> Tex = SRV->GetLinearTexture(false);
-							IAB->IndirectArgumentResources.Add(Argument(Tex, (mtlpp::ResourceUsage)(mtlpp::ResourceUsage::Read|mtlpp::ResourceUsage::Sample)));
-							
-							Desc.SetDataType(mtlpp::DataType::Texture);
-							Desc.SetTextureType(Tex.GetTextureType());
-							
-							BufferSizes.Add(VB ? VB->GetSize() : IB->GetSize());
-							BufferSizes.Add(GMetalBufferFormats[SRV->Format].DataFormat);
-						}
-						break;
-					}
-					case UBMT_RDG_TEXTURE_UAV:
-					case UBMT_RDG_BUFFER_UAV:
-					{
-						FMetalArgumentDesc& Desc = Arguments.Emplace_GetRef();
-						Desc.SetIndex(Index++);
-						Desc.SetAccess(mtlpp::ArgumentAccess::ReadWrite);
-						
-						FMetalUnorderedAccessView* UAV = (FMetalUnorderedAccessView*)Resource;
-						FMetalShaderResourceView* SRV = (FMetalShaderResourceView*)UAV->SourceView;
-						FMetalStructuredBuffer* SB = UAV->SourceView->SourceStructuredBuffer.GetReference();
-						FMetalVertexBuffer* VB = UAV->SourceView->SourceVertexBuffer.GetReference();
-						FMetalIndexBuffer* IB = UAV->SourceView->SourceIndexBuffer.GetReference();
-						FRHITexture* Texture = UAV->SourceView->SourceTexture.GetReference();
-						FMetalSurface* Surface = UAV->SourceView->TextureView;
-						if (Texture)
-						{
-							Desc.SetDataType(mtlpp::DataType::Texture);
-							
-							FRHITextureReference* Reference = Texture->GetTextureReference();
-							if (Reference)
-							{
-								TBitArray<>& Ref = TextureReferences.FindOrAdd(Reference);
-								if (Ref.Num() <= IAB->IndirectArgumentResources.Num())
-							   	{
-								   	Ref.Add(false, (IAB->IndirectArgumentResources.Num() + 1) - Ref.Num());
-							   	}
-							   	Ref[IAB->IndirectArgumentResources.Num()] = true;
-							}
-							
-							union {
-								uint8 Components[4];
-								uint32 Packed;
-							} Swizzle;
-							Swizzle.Packed = 0;
-							assert(sizeof(Swizzle) == sizeof(uint32));
-							
-							if (!Surface)
-							{
-								Surface = GetMetalSurfaceFromRHITexture(Texture);
-							}
-							if (Surface)
-							{
-								check(!Surface->Texture.IsAliasable());
-								if (Surface->Texture.GetPixelFormat() == mtlpp::PixelFormat::X32_Stencil8
-#if PLATFORM_MAC
-									||	Surface->Texture.GetPixelFormat() == mtlpp::PixelFormat::X24_Stencil8
-#endif
-									)
-								{
-									Swizzle.Components[0] = Swizzle.Components[1] = Swizzle.Components[2] = Swizzle.Components[3] = 1;
-								}
-								Desc.SetTextureType(Surface->Texture.GetTextureType());
-								IAB->IndirectArgumentResources.Add(Argument(Surface->Texture, (mtlpp::ResourceUsage)(mtlpp::ResourceUsage::Read|mtlpp::ResourceUsage::Write)));
-							}
-							else
-							{
-								IAB->IndirectArgumentResources.Add(Argument(FMetalTexture(), (mtlpp::ResourceUsage)(mtlpp::ResourceUsage::Read|mtlpp::ResourceUsage::Write)));
-							}
-							BufferSizes.Add(Swizzle.Packed);
-							BufferSizes.Add(GMetalBufferFormats[Texture->GetFormat()].DataFormat);
-						}
-						else if (SB)
-						{
-							Desc.SetDataType(mtlpp::DataType::Pointer);
-							IAB->IndirectArgumentResources.Add(Argument(SB->Buffer, (mtlpp::ResourceUsage)(mtlpp::ResourceUsage::Read|mtlpp::ResourceUsage::Write)));
-							BufferSizes.Add(SB->GetSize());
-							BufferSizes.Add(PF_Unknown);
-						}
-						else
-						{
-							check(VB || IB);
-							ns::AutoReleased<FMetalTexture> Tex = SRV->GetLinearTexture(false);
-							IAB->IndirectArgumentResources.Add(Argument(Tex, (mtlpp::ResourceUsage)(mtlpp::ResourceUsage::Read|mtlpp::ResourceUsage::Write)));
-							Desc.SetDataType(mtlpp::DataType::Texture);
-							Desc.SetTextureType(Tex.GetTextureType());
-							BufferSizes.Add(0);
-							BufferSizes.Add(0);
-							
-							FMetalArgumentDesc& BufferDesc = Arguments.Emplace_GetRef();
-							BufferDesc.SetIndex(Index++);
-							BufferDesc.SetAccess(mtlpp::ArgumentAccess::ReadWrite);
-							BufferDesc.SetDataType(mtlpp::DataType::Pointer);
-							
-							if (VB)
-							{
-								IAB->IndirectArgumentResources.Add(Argument(VB->Buffer, (mtlpp::ResourceUsage)(mtlpp::ResourceUsage::Read|mtlpp::ResourceUsage::Write)));
-								BufferSizes.Add(VB->GetSize());
-								BufferSizes.Add(GMetalBufferFormats[SRV->Format].DataFormat);
-							}
-							else
-							{
-								IAB->IndirectArgumentResources.Add(Argument(IB->Buffer, (mtlpp::ResourceUsage)(mtlpp::ResourceUsage::Read|mtlpp::ResourceUsage::Write)));
-								BufferSizes.Add(IB->GetSize());
-								BufferSizes.Add(GMetalBufferFormats[SRV->Format].DataFormat);
-							}
-						}
-						break;
-					}
-					case UBMT_SAMPLER:
-					{
-						FMetalArgumentDesc& Desc = Arguments.Emplace_GetRef();
-						Desc.SetIndex(Index++);
-						Desc.SetAccess(mtlpp::ArgumentAccess::ReadOnly);
-						Desc.SetDataType(mtlpp::DataType::Sampler);
-						
-						FMetalSamplerState* Sampler = (FMetalSamplerState*)Resource;
-						IAB->IndirectArgumentResources.Add(Argument(Sampler->State));
-						
-						BufferSizes.Add(0);
-						BufferSizes.Add(0);
-						break;
-					}
-					case UBMT_RDG_TEXTURE:
-					case UBMT_TEXTURE:
-					{
-						FMetalArgumentDesc& Desc = Arguments.Emplace_GetRef();
-						Desc.SetIndex(Index++);
-						Desc.SetAccess(mtlpp::ArgumentAccess::ReadOnly);
-						Desc.SetDataType(mtlpp::DataType::Texture);
-						
-						FRHITexture* Texture = (FRHITexture*)Resource;
-						FMetalSurface* Surface = GetMetalSurfaceFromRHITexture(Texture);
-						FRHITextureReference* Reference = Texture->GetTextureReference();
-						if (Reference)
-						{
-							TBitArray<>& Ref = TextureReferences.FindOrAdd(Reference);
-							if (Ref.Num() <= IAB->IndirectArgumentResources.Num())
-							{
-								Ref.Add(false, (IAB->IndirectArgumentResources.Num() + 1) - Ref.Num());
-							}
-							Ref[IAB->IndirectArgumentResources.Num()] = true;
-						}
-						union {
-							uint8 Components[4];
-							uint32 Packed;
-						} Swizzle;
-						Swizzle.Packed = 0;
-						assert(sizeof(Swizzle) == sizeof(uint32));
-						if (Surface)
-						{
-							Desc.SetTextureType(Surface->Texture.GetTextureType());
-							check(!Surface->Texture.IsAliasable());
-							IAB->IndirectArgumentResources.Add(Argument(Surface->Texture, (mtlpp::ResourceUsage)(mtlpp::ResourceUsage::Read|mtlpp::ResourceUsage::Sample)));
-							if (Surface->Texture.GetPixelFormat() == mtlpp::PixelFormat::X32_Stencil8
-#if PLATFORM_MAC
-								||	Surface->Texture.GetPixelFormat() == mtlpp::PixelFormat::X24_Stencil8
-#endif
-								)
-							{
-								Swizzle.Components[0] = Swizzle.Components[1] = Swizzle.Components[2] = Swizzle.Components[3] = 1;
-							}
-						}
-						else
-						{
-							IAB->IndirectArgumentResources.Add(Argument(FMetalTexture(), (mtlpp::ResourceUsage)(mtlpp::ResourceUsage::Read|mtlpp::ResourceUsage::Sample)));
-						}
-						BufferSizes.Add(Swizzle.Packed);
-						BufferSizes.Add(GMetalBufferFormats[Texture->GetFormat()].DataFormat);
-						break;
-					}
-					default:
-						break;
-				}
-			}
-		}
-		
-		{
-			FMetalArgumentDesc& Desc = Arguments.Emplace_GetRef();
-			Desc.SetIndex(Index++);
-			Desc.SetAccess(mtlpp::ArgumentAccess::ReadOnly);
-			Desc.SetDataType(mtlpp::DataType::Pointer);
-			
-			BufferSizes.AddZeroed(2);
-			
-			if (ConstantSize > 0)
-				IAB->IndirectArgumentResources.Add(Argument(Buffer, mtlpp::ResourceUsage::Read));
-			else
-				IAB->IndirectArgumentResources.Add(Argument(FMetalBuffer(), mtlpp::ResourceUsage::Read));
-		}
-		
-		IAB->UpdateNum = UpdateNum;
-		FPlatformAtomics::InterlockedIncrement(&IAB->UpdateIAB);
-		FPlatformAtomics::InterlockedIncrement(&IAB->UpdateEnc);
-	}
-}
-
 void FMetalUniformBuffer::UpdateResourceTable(TArray<TRefCountPtr<FRHIResource>>& Resources, EUniformBufferValidation Validation)
 {
 	ResourceTable = Resources;
-	
-	if (NumResources && FMetalCommandQueue::SupportsFeature(EMetalFeaturesIABs))
-	{
-		FRHICommandListImmediate& RHICmdList = FRHICommandListExecutor::GetImmediateCommandList();
-		if (!(UniformUsage & UniformBuffer_SingleDraw) && IsRunningRHIInSeparateThread() && !RHICmdList.Bypass() && IsInRenderingThread())
-		{
-			new (RHICmdList.AllocCommand<FMetalRHICommandInitialiseUniformBufferIAB>()) FMetalRHICommandInitialiseUniformBufferIAB(this);
-			RHICmdList.RHIThreadFence(true);
-		}
-		else
-		{
-			FPlatformAtomics::InterlockedIncrement(&UpdateNum);
-			UpdateIAB();
-		}
-	}
 }
 
 void FMetalUniformBuffer::Update(const void* Contents, TArray<TRefCountPtr<FRHIResource>>& Resources, EUniformBufferValidation Validation)
 {
+    uint32 ConstantSize = FRHIUniformBuffer::GetSize();
     if (ConstantSize > 0)
     {
         UE_CLOG(ConstantSize > 65536, LogMetal, Fatal, TEXT("Trying to allocated a uniform layout of size %d that is greater than the maximum permitted 64k."), ConstantSize);
@@ -760,12 +112,7 @@ void FMetalUniformBuffer::Update(const void* Contents, TArray<TRefCountPtr<FRHIR
 		void* Data = Lock(true, RLM_WriteOnly, 0, 0, true);
         FMemory::Memcpy(Data, Contents, ConstantSize);
         Unlock();
-		
-		if (Buf != Buffer && FMetalCommandQueue::SupportsFeature(EMetalFeaturesIABs))
-		{
-			FPlatformAtomics::InterlockedIncrement(&UpdateNum);
-		}
-		
+
 		ConditionalSetUniformBufferFrameIndex();
 	}
 	
@@ -791,7 +138,7 @@ struct FMetalRHICommandUpateUniformBuffer : public FRHICommand<FMetalRHICommandU
 	, ResourceTable(Resources)
 	, Contents(nullptr)
 	{
-		uint32 MaxLayoutSize = InBuffer->ConstantSize;
+		uint32 MaxLayoutSize = InBuffer->GetSize();
 		Contents = new char[MaxLayoutSize];
 		FMemory::Memcpy(Contents, Data, MaxLayoutSize);
 	}
@@ -816,11 +163,11 @@ void FMetalDynamicRHI::RHIUpdateUniformBuffer(FRHIUniformBuffer* UniformBufferRH
 	FRHICommandListImmediate& RHICmdList = FRHICommandListExecutor::GetImmediateCommandList();
 		
 	TArray<TRefCountPtr<FRHIResource> > ResourceTable;
-	ResourceTable.AddZeroed(UniformBuffer->NumResources);
+	ResourceTable.AddZeroed(UniformBuffer->GetLayout().Resources.Num());
 		
 	const FRHIUniformBufferLayout& Layout = UniformBuffer->GetLayout();
 		
-	for (uint32 i = 0; i < UniformBuffer->NumResources; ++i)
+	for (uint32 i = 0; i < UniformBuffer->GetLayout().Resources.Num(); ++i)
 	{
 		FRHIResource* Resource = *(FRHIResource**)((uint8*)Contents + Layout.Resources[i].MemberOffset);
 		ResourceTable[i] = Resource;
diff --git a/Engine/Source/Runtime/Apple/MetalRHI/Public/MetalResources.h b/Engine/Source/Runtime/Apple/MetalRHI/Public/MetalResources.h
--- a/Engine/Source/Runtime/Apple/MetalRHI/Public/MetalResources.h
+++ b/Engine/Source/Runtime/Apple/MetalRHI/Public/MetalResources.h
@@ -1055,61 +1055,19 @@ struct FMetalArgumentDesc
 class FMetalUniformBuffer : public FRHIUniformBuffer, public FMetalRHIBuffer
 {
 public:
-
-	struct Argument
-	{
-		Argument() {}
-		Argument(FMetalBuffer const& InBuffer, mtlpp::ResourceUsage const InUsage) : Buffer(InBuffer), Usage(InUsage) {}
-		Argument(FMetalTexture const& InTexture, mtlpp::ResourceUsage const InUsage) : Texture(InTexture), Usage(InUsage) {}
-		Argument(FMetalSampler const& InSampler) : Sampler(InSampler), Usage(mtlpp::ResourceUsage::Read) {}
-		
-		FMetalBuffer Buffer;
-		FMetalTexture Texture;
-		FMetalSampler Sampler;
-		mtlpp::ResourceUsage Usage;
-	};
-	
-	struct FMetalIndirectArgumentBuffer
-	{
-		FMetalIndirectArgumentBuffer();
-		~FMetalIndirectArgumentBuffer();
-		
-		int64 UpdateNum; // The resource & declarations have been updated
-		int64 UpdateEnc; // The IAB encoder needs to be updated.
-		int64 UpdateIAB; // The IAB needs to be updated.
-		TArray<FMetalArgumentDesc> IndirectArgumentsDecl;
-		TArray<Argument> IndirectArgumentResources;
-		TArray<uint32> IndirectBufferSizes;
-		FMetalIAB* IndirectArgumentBuffer;
-		TMap<TBitArray<>, FMetalIAB*> Tier1IABs;
-		FRWLock Mutex;
-	};
-	
 	// Constructor
 	FMetalUniformBuffer(const void* Contents, const FRHIUniformBufferLayout& Layout, EUniformBufferUsage Usage, EUniformBufferValidation Validation);
 
 	// Destructor 
-	virtual ~FMetalUniformBuffer();
+    virtual ~FMetalUniformBuffer() {}
 	
 	void const* GetData();
 
-	void InitIAB();
-	void UpdateIAB();
-	void UpdateTextureReference(FRHITextureReference* ModifiedRef);
 	void UpdateResourceTable(TArray<TRefCountPtr<FRHIResource>>& Resources, EUniformBufferValidation Validation);
 	void Update(const void* Contents, TArray<TRefCountPtr<FRHIResource>>& Resources, EUniformBufferValidation Validation);
-	FMetalIAB* UploadIAB(class FMetalContext* Ctx, TBitArray<> const& Bitmask, mtlpp::ArgumentEncoder const& Encoder);
-	FMetalIndirectArgumentBuffer& GetIAB();
 	
 	/** Resource table containing RHI references. */
 	TArray<TRefCountPtr<FRHIResource> > ResourceTable;
-	TMap<FRHITextureReference*, TBitArray<>> TextureReferences;
-	EUniformBufferUsage UniformUsage;
-	FMetalIndirectArgumentBuffer* IAB;
-	TArray<EUniformBufferBaseType> ResourceTypes;
-	int64 UpdateNum;
-	uint32 NumResources;
-	uint32 ConstantSize;
 };
 
 
