Fix up metal heap suballoc owner tracking

diff --git a/Engine/Source/Runtime/Apple/MetalRHI/Private/MetalVertexBuffer.cpp b/Engine/Source/Runtime/Apple/MetalRHI/Private/MetalVertexBuffer.cpp
--- a/Engine/Source/Runtime/Apple/MetalRHI/Private/MetalVertexBuffer.cpp
+++ b/Engine/Source/Runtime/Apple/MetalRHI/Private/MetalVertexBuffer.cpp
@@ -95,6 +95,8 @@ void FMetalVertexBuffer::Swap(FMetalVertexBuffer& Other)
 void FMetalRHIBuffer::Swap(FMetalRHIBuffer& Other)
 {
 	::Swap(*this, Other);
+	Other.Buffer.SetOwner(&Other);
+	Buffer.SetOwner(this);
 }
 
 static bool CanUsePrivateMemory()
