Attempting to bind a memoryless texture is now a hard failure.

diff --git a/Engine/Source/Runtime/Apple/MetalRHI/Private/MetalStateCache.cpp b/Engine/Source/Runtime/Apple/MetalRHI/Private/MetalStateCache.cpp
--- a/Engine/Source/Runtime/Apple/MetalRHI/Private/MetalStateCache.cpp
+++ b/Engine/Source/Runtime/Apple/MetalRHI/Private/MetalStateCache.cpp
@@ -1400,19 +1400,18 @@ void FMetalStateCache::SetShaderTexture(EMetalShaderStages const Frequency, FMet
 {
 	check(Frequency < EMetalShaderStages::Num);
 	check(Index < ML_MaxTextures);
+
+#if (PLATFORM_IOS || PLATFORM_TVOS)
+    UE_CLOG([Texture.GetPtr() storageMode] != MTLStorageModeMemoryless, LogMetal, Fatal, TEXT("FATAL: Attempting to bind a memoryless texture. Stage %u Index %u Texture %@"), Frequency, Index, Texture.GetPtr());
+#endif
 	
 	if (ShaderTextures[Frequency].Textures[Index] != Texture
 		|| ShaderTextures[Frequency].Usage[Index] != Usage)
 	{
-		bool bMemoryLess = false;
-#if PLATFORM_IOS
-		bMemoryLess = (Texture.GetStorageMode() == mtlpp::StorageMode::Memoryless);
-#endif
-		
-		ShaderTextures[Frequency].Textures[Index] = !bMemoryLess ? Texture : FMetalTexture();
+		ShaderTextures[Frequency].Textures[Index] = Texture;
 		ShaderTextures[Frequency].Usage[Index] = Usage;
 		
-		if (Texture && !bMemoryLess)
+		if (Texture)
 		{
 			ShaderTextures[Frequency].Bound |= (FMetalTextureMask(1) << FMetalTextureMask(Index));
 		}
