The fingers, they are so fat.
(need to fail on the failure case).

diff --git a/Engine/Source/Runtime/Apple/MetalRHI/Private/MetalStateCache.cpp b/Engine/Source/Runtime/Apple/MetalRHI/Private/MetalStateCache.cpp
--- a/Engine/Source/Runtime/Apple/MetalRHI/Private/MetalStateCache.cpp
+++ b/Engine/Source/Runtime/Apple/MetalRHI/Private/MetalStateCache.cpp
@@ -1402,7 +1402,7 @@ void FMetalStateCache::SetShaderTexture(EMetalShaderStages const Frequency, FMet
 	check(Index < ML_MaxTextures);
 
 #if (PLATFORM_IOS || PLATFORM_TVOS)
-    UE_CLOG([Texture.GetPtr() storageMode] != MTLStorageModeMemoryless, LogMetal, Fatal, TEXT("FATAL: Attempting to bind a memoryless texture. Stage %u Index %u Texture %@"), Frequency, Index, Texture.GetPtr());
+    UE_CLOG([Texture.GetPtr() storageMode] == MTLStorageModeMemoryless, LogMetal, Fatal, TEXT("FATAL: Attempting to bind a memoryless texture. Stage %u Index %u Texture %@"), Frequency, Index, Texture.GetPtr());
 #endif
 	
 	if (ShaderTextures[Frequency].Textures[Index] != Texture
