Patch linear buffer fitting.

Attempt to use the widest possible texture size when the number of elements exceeds the maximum texture width, the first attempt to determine a reasonably sized 2D texture description fails, and the caller is attempting to use the entire buffer.  This is not intended to be a maximum rectangle fitting algorithm or anything like that, it's only trying to use as much of the buffer as possible without spending too much time trying to run some kind of maximum fit algorithm.

diff --git a/Engine/Source/Runtime/Apple/MetalRHI/Private/MetalVertexBuffer.cpp b/Engine/Source/Runtime/Apple/MetalRHI/Private/MetalVertexBuffer.cpp
--- a/Engine/Source/Runtime/Apple/MetalRHI/Private/MetalVertexBuffer.cpp
+++ b/Engine/Source/Runtime/Apple/MetalRHI/Private/MetalVertexBuffer.cpp
@@ -336,6 +336,26 @@ FMetalTexture FMetalRHIBuffer::AllocLinearTexture(EPixelFormat InFormat, const F
 				Width = Dimension;
 				Height = NumElements / Dimension;
 
+				// If we're just trying to fit as many elements as we can into
+				// the available buffer space, we can trim some padding at the
+				// end of the buffer in order to create widest possible linear
+				// texture that will fit.
+				if ((UINT_MAX == LinearTextureDesc.NumElements) && (Height > GMaxTextureDimensions))
+				{
+					Width = GMaxTextureDimensions;
+					Height = 1;
+
+					while ((Width * Height) < NumElements)
+					{
+						Height <<= 1;
+					}
+
+					while ((Width * Height) > NumElements)
+					{
+						Height -= 1;
+					}
+				}
+
 				checkf(Width <= GMaxTextureDimensions, TEXT("Calculated width %u is greater than maximum permitted %d when converting buffer of size %llu with element stride %u to a 2D texture with %u elements."), Width, (int32)GMaxTextureDimensions, Buffer.GetLength(), BytesPerElement, NumElements);
 				checkf(Height <= GMaxTextureDimensions, TEXT("Calculated height %u is greater than maximum permitted %d when converting buffer of size %llu with element stride %u to a 2D texture with %u elements."), Height, (int32)GMaxTextureDimensions, Buffer.GetLength(), BytesPerElement, NumElements);
 			}
