RHIThread aware version of FrameAllocator + Uniform buffer refactor.
Also removes some extra IAB stuff.

Adds FMetalFrameAllocator for allocations that live on the RHI thread (or render thread if an rhi thread does not exist) timeline).
This can be tuned by default allocation size and target allocation limit (it will go over the limit if it must)
Also provides stats to stat MetalRHI detailing memory usage

Currently there is one instance which is meant to only allocate Uniform buffers.
This can be controlled with CVars and also reports specific stats for uniform memory usage (see stat MetalRHI)

diff --git a/Engine/Source/Runtime/Apple/MetalRHI/Private/MetalCommands.cpp b/Engine/Source/Runtime/Apple/MetalRHI/Private/MetalCommands.cpp
--- a/Engine/Source/Runtime/Apple/MetalRHI/Private/MetalCommands.cpp
+++ b/Engine/Source/Runtime/Apple/MetalRHI/Private/MetalCommands.cpp
@@ -399,77 +399,6 @@ void FMetalRHICommandContext::RHISetShaderParameter(FRHIComputeShader* ComputeSh
 	}
 }
 
-void FMetalRHICommandContext::RHISetShaderUniformBuffer(FRHIGraphicsShader* ShaderRHI, uint32 BufferIndex, FRHIUniformBuffer* BufferRHI)
-{
-	@autoreleasepool {
-		EMetalShaderStages Stage = EMetalShaderStages::Num;
-		FMetalShaderBindings* Bindings = nullptr;
-		switch (ShaderRHI->GetFrequency())
-		{
-		case SF_Vertex:
-		{
-			FMetalVertexShader* VertexShader = ResourceCast(static_cast<FRHIVertexShader*>(ShaderRHI));
-			Bindings = &VertexShader->Bindings;
-			Stage = EMetalShaderStages::Vertex;
-		}
-			break;
-	#if PLATFORM_SUPPORTS_TESSELLATION_SHADERS
-		case SF_Hull:
-		{
-			Stage = EMetalShaderStages::Hull;
-			FMetalHullShader* HullShader = ResourceCast(static_cast<FRHIHullShader*>(ShaderRHI));
-			Bindings = &HullShader->Bindings;
-		}
-			break;
-		case SF_Domain:
-		{
-			Stage = EMetalShaderStages::Domain;
-			FMetalDomainShader* DomainShader = ResourceCast(static_cast<FRHIDomainShader*>(ShaderRHI));
-			Bindings = &DomainShader->Bindings;
-		}
-			break;
-	#endif
-		case SF_Pixel:
-		{
-			Stage = EMetalShaderStages::Pixel;
-			FMetalPixelShader* PixelShader = ResourceCast(static_cast<FRHIPixelShader*>(ShaderRHI));
-			Bindings = &PixelShader->Bindings;
-		}
-			break;
-		default:
-			checkf(0, TEXT("FRHIShader Type %d is invalid or unsupported!"), (int32)ShaderRHI->GetFrequency());
-			NOT_SUPPORTED("RHIShaderStage");
-			break;
-		}
-
-		Context->GetCurrentState().BindUniformBuffer(Stage, BufferIndex, BufferRHI);
-
-		check(BufferIndex < Bindings->NumUniformBuffers);
-		if ((Bindings->ConstantBuffers) & 1 << BufferIndex)
-		{
-			auto* UB = (FMetalUniformBuffer*)BufferRHI;
-			Context->GetCurrentState().SetShaderBuffer(Stage, UB->Buffer, UB->Data, 0, UB->GetSize(), BufferIndex, mtlpp::ResourceUsage::Read);
-		}
-	}
-}
-
-void FMetalRHICommandContext::RHISetShaderUniformBuffer(FRHIComputeShader* ComputeShaderRHI, uint32 BufferIndex, FRHIUniformBuffer* BufferRHI)
-{
-	@autoreleasepool {
-		FMetalComputeShader* ComputeShader = ResourceCast(ComputeShaderRHI);
-		Context->GetCurrentState().BindUniformBuffer(EMetalShaderStages::Compute, BufferIndex, BufferRHI);
-
-		auto& Bindings = ComputeShader->Bindings;
-		check(BufferIndex < Bindings.NumUniformBuffers);
-		if ((Bindings.ConstantBuffers) & 1 << BufferIndex)
-		{
-			auto* UB = (FMetalUniformBuffer*)BufferRHI;
-			Context->GetCurrentState().SetShaderBuffer(EMetalShaderStages::Compute, UB->Buffer, UB->Data, 0, UB->GetSize(), BufferIndex, mtlpp::ResourceUsage::Read);
-		}
-	}
-}
-
-
 void FMetalRHICommandContext::RHISetStencilRef(uint32 StencilRef)
 {
 	Context->GetCurrentState().SetStencilRef(StencilRef);
diff --git a/Engine/Source/Runtime/Apple/MetalRHI/Private/MetalContext.cpp b/Engine/Source/Runtime/Apple/MetalRHI/Private/MetalContext.cpp
--- a/Engine/Source/Runtime/Apple/MetalRHI/Private/MetalContext.cpp
+++ b/Engine/Source/Runtime/Apple/MetalRHI/Private/MetalContext.cpp
@@ -13,6 +13,8 @@
 #include "MetalProfiler.h"
 #include "MetalCommandBuffer.h"
 
+#include "MetalFrameAllocator.h"
+
 int32 GMetalSupportsIntermediateBackBuffer = 0;
 static FAutoConsoleVariableRef CVarMetalSupportsIntermediateBackBuffer(
 	TEXT("rhi.Metal.SupportsIntermediateBackBuffer"),
@@ -100,6 +102,26 @@ static FAutoConsoleVariableRef CVarMetalPresentFramePacing(
 	TEXT("Specify the desired frame rate for presentation (iOS 10.3+ only, default: 0.0f, off"));
 #endif
 
+#if PLATFORM_MAC
+static int32 GMetalDefaultUniformBufferAllocation = 1024*1024;
+#else
+static int32 GMetalDefaultUniformBufferAllocation = 1024*32;
+#endif
+static FAutoConsoleVariableRef CVarMetalDefaultUniformBufferAllocation(
+    TEXT("rhi.Metal.DefaultUniformBufferAllocation"),
+    GMetalDefaultUniformBufferAllocation,
+    TEXT("Default size of a uniform buffer allocation."));
+
+#if PLATFORM_MAC
+static int32 GMetalTargetUniformAllocationLimit = 1024 * 1024 * 50;
+#else
+static int32 GMetalTargetUniformAllocationLimit = 1024 * 1024 * 5;
+#endif
+static FAutoConsoleVariableRef CVarMetalTargetUniformAllocationLimit(
+     TEXT("rhi.Metal.TargetUniformAllocationLimit"),
+     GMetalTargetUniformAllocationLimit,
+     TEXT("Target Allocation limit for the uniform buffer pool."));
+
 #if PLATFORM_MAC
 static ns::AutoReleased<ns::Object<id <NSObject>>> GMetalDeviceObserver;
 static mtlpp::Device GetMTLDevice(uint32& DeviceIndex)
@@ -343,7 +365,7 @@ FMetalDeviceContext::FMetalDeviceContext(mtlpp::Device MetalDevice, uint32 InDev
 , ActiveContexts(1)
 , ActiveParallelContexts(0)
 , PSOManager(0)
-, DeviceFrameIndex(0)
+, FrameNumberRHIThread(0)
 {
 	CommandQueue.SetRuntimeDebuggingLevel(GMetalRuntimeDebugLevel);
 	
@@ -377,6 +399,12 @@ FMetalDeviceContext::FMetalDeviceContext(mtlpp::Device MetalDevice, uint32 InDev
 	{
 		GMetalSupportsIntermediateBackBuffer = 1;
 	}
+    
+    // initialize uniform allocator
+    UniformBufferAllocator = new FMetalFrameAllocator(MetalDevice.GetPtr());
+    UniformBufferAllocator->SetTargetAllocationLimitInBytes(GMetalTargetUniformAllocationLimit);
+    UniformBufferAllocator->SetDefaultAllocationSizeInBytes(GMetalDefaultUniformBufferAllocation);
+    UniformBufferAllocator->SetStatIds(GET_STATID(STAT_MetalUniformAllocatedMemory), GET_STATID(STAT_MetalUniformMemoryInFlight), GET_STATID(STAT_MetalUniformBytesPerFrame));
 	
 	PSOManager = new FMetalPipelineStateCacheManager();
 	
@@ -391,6 +419,8 @@ FMetalDeviceContext::~FMetalDeviceContext()
 	delete &(GetCommandQueue());
 	
 	delete PSOManager;
+    
+    delete UniformBufferAllocator;
 	
 #if PLATFORM_MAC
 	if (FPlatformMisc::MacOSXVersionCompare(10, 13, 4) >= 0)
@@ -413,9 +443,6 @@ void FMetalDeviceContext::BeginFrame()
 	
 	// Wait for the frame semaphore on the immediate context.
 	dispatch_semaphore_wait(CommandBufferSemaphore, DISPATCH_TIME_FOREVER);
-	
-	// Bump the frame counter.
-	DeviceFrameIndex++;
 }
 
 #if METAL_DEBUG_OPTIONS
@@ -532,7 +559,11 @@ void FMetalDeviceContext::EndFrame()
 		SubmitFlags |= EMetalSubmitFlagsWaitOnCommandBuffer;
 	}
 #endif
+    
 	SubmitCommandsHint((uint32)SubmitFlags);
+    
+    // increment the internal frame counter
+    FrameNumberRHIThread++;
 	
     FlushFreeList();
     
diff --git a/Engine/Source/Runtime/Apple/MetalRHI/Private/MetalContext.h b/Engine/Source/Runtime/Apple/MetalRHI/Private/MetalContext.h
--- a/Engine/Source/Runtime/Apple/MetalRHI/Private/MetalContext.h
+++ b/Engine/Source/Runtime/Apple/MetalRHI/Private/MetalContext.h
@@ -23,6 +23,7 @@
 
 class FMetalRHICommandContext;
 class FMetalPipelineStateCacheManager;
+class FMetalFrameAllocator;
 
 class FMetalContext
 {
@@ -222,11 +223,15 @@ public:
 	/** Get the index of the bound Metal device in the global list of rendering devices. */
 	uint32 GetDeviceIndex(void) const;
     
-	/** Device frame index accessor. */
-	uint64 GetDeviceFrameIndex() const
-	{
-		return DeviceFrameIndex;
-	}
+    FMetalFrameAllocator* GetUniformAllocator()
+    {
+        return UniformBufferAllocator;
+    }
+    
+    uint32 GetFrameNumberRHIThread()
+    {
+        return FrameNumberRHIThread;
+    }
 	
 #if METAL_DEBUG_OPTIONS
     void AddActiveBuffer(FMetalBuffer const& Buffer);
@@ -268,7 +273,8 @@ private:
 	};
 	TArray<FMetalDelayedFreeList*> DelayedFreeLists;
 	
-	TSet<FMetalUniformBuffer*> UniformBuffers;
+//	TSet<FMetalUniformBuffer*> UniformBuffers;
+    FMetalFrameAllocator* UniformBufferAllocator;
 	
 #if METAL_DEBUG_OPTIONS
 	/** The list of fences for the current frame */
@@ -313,6 +319,7 @@ private:
 	/** PSO cache manager */
 	FMetalPipelineStateCacheManager* PSOManager;
 
-	/** Device frame index, glorified frame counter in the device namespace. */
-	uint64 DeviceFrameIndex;
+    /** Thread index owned by the RHI Thread. Monotonically increases every call to EndFrame() */
+    uint32 FrameNumberRHIThread;
 };
+
diff --git a/Engine/Source/Runtime/Apple/MetalRHI/Private/MetalFrameAllocator.cpp b/Engine/Source/Runtime/Apple/MetalRHI/Private/MetalFrameAllocator.cpp
new file mode 100644
--- /dev/null
+++ b/Engine/Source/Runtime/Apple/MetalRHI/Private/MetalFrameAllocator.cpp
@@ -0,0 +1,205 @@
+// Copyright Epic Games, Inc. All Rights Reserved.
+
+#include "MetalRHIPrivate.h"
+#include "MetalFrameAllocator.h"
+#include "MetalLLM.h"
+#include "MetalProfiler.h"
+
+#pragma mark Constructor/Destructor
+
+FMetalFrameAllocator::FMetalFrameAllocator(id <MTLDevice> TargetDevice)
+    : Device(TargetDevice)
+#if METAL_FRAME_ALLOCATOR_VALIDATION
+    , AllAllocatedBuffers(nil)
+#endif
+    , AvailableBuffers(nil)
+    , BuffersInFlight(nil)
+    , CurrentBuffer(nil)
+    , CurrentCursor(0)
+    , CurrentFrame(0)
+    , TargetAllocationLimit(0)
+    , DefaultAllocationSize(0)
+    , TotalBytesAllocated(0)
+    , BytesAcquiredInCurrentFrame(0)
+    , BytesInFlight(0)
+{
+    check(Device);
+#if METAL_FRAME_ALLOCATOR_VALIDATION
+    AllAllocatedBuffers = [NSMutableArray new];
+#endif
+    AvailableBuffers = [NSMutableArray new];
+    BuffersInFlight = [NSMutableArray new];
+    UnfairLock = OS_UNFAIR_LOCK_INIT;
+    
+    TotalAllocationStat = GET_STATID(STAT_MetalFrameAllocatorAllocatedMemory);
+    MemoryInFlightStat = GET_STATID(STAT_MetalFrameAllocatorMemoryInFlight);
+    BytesPerFrameStat = GET_STATID(STAT_MetalFrameAllocatorBytesPerFrame);
+    
+    [Device retain];
+}
+
+FMetalFrameAllocator::~FMetalFrameAllocator()
+{
+    // Owner of this object is responsible for finishing all rendering before destroying this object.
+    // This does not check and will crash.
+    os_unfair_lock_lock(&UnfairLock);
+    
+    [BuffersInFlight release];
+    [AvailableBuffers release];
+#if METAL_FRAME_ALLOCATOR_VALIDATION
+    [AllAllocatedBuffers release];
+#endif
+    
+    [Device release];
+    
+    os_unfair_lock_unlock(&UnfairLock);
+}
+
+#pragma mark -
+
+#pragma mark API Entrypoints
+
+void FMetalFrameAllocator::SetTargetAllocationLimitInBytes(uint32 LimitInBytes)
+{
+    TargetAllocationLimit = LimitInBytes;
+}
+
+void FMetalFrameAllocator::SetDefaultAllocationSizeInBytes(uint32 DefaultInBytes)
+{
+    DefaultAllocationSize = DefaultInBytes;
+}
+
+void FMetalFrameAllocator::SetStatIds(const TStatId& InTotalAllocationStat, const TStatId& InMemoryInFlightStat, const TStatId& InBytesPerFrameStat)
+{
+    TotalAllocationStat = InTotalAllocationStat;
+    MemoryInFlightStat = InMemoryInFlightStat;
+    BytesPerFrameStat = InBytesPerFrameStat;
+}
+
+FMetalFrameAllocator::AllocationEntry FMetalFrameAllocator::AcquireSpace(uint32 SizeInBytes)
+{
+    uint32 AlignedSize = Align(SizeInBytes, BufferOffsetAlignment);
+    
+    os_unfair_lock_lock(&UnfairLock);
+    
+    NSUInteger SpaceRequired = this->CurrentCursor + AlignedSize;
+    
+    if(this->CurrentBuffer == nil || (this->CurrentBuffer && SpaceRequired > [this->CurrentBuffer length]))
+    {
+        this->CurrentBuffer = this->FindOrAllocateBufferUnsafe(AlignedSize);
+        this->CurrentCursor = 0;
+        [this->BuffersInFlight addObject:this->CurrentBuffer];
+        this->BytesInFlight += [this->CurrentBuffer length];
+    }
+    
+    check(this->CurrentBuffer != nil);
+    
+    AllocationEntry Entry;
+    Entry.Backing = this->CurrentBuffer;
+    Entry.Offset = this->CurrentCursor;
+    this->CurrentCursor += AlignedSize;
+    
+    this->BytesAcquiredInCurrentFrame += AlignedSize;
+    
+    os_unfair_lock_unlock(&UnfairLock);
+    
+    return Entry;
+}
+
+void FMetalFrameAllocator::MarkEndOfFrame(uint32 FrameNumberThatEnded, id <MTLCommandBuffer> LastCommandBufferInFrame)
+{
+    __block FMetalFrameAllocator* Allocator = this;
+    __block NSMutableArray<id <MTLBuffer>>* BuffersToRecycle = this->BuffersInFlight;
+    
+    // Pass ownership of the BuffersInFlight list to the completion handler.
+    // It will be responsible for destroying the array.
+    [LastCommandBufferInFrame addCompletedHandler:
+    ^(id <MTLCommandBuffer> CompletedCommandBuffer)
+    {
+        Allocator->RecycleBuffers(BuffersToRecycle);
+        [BuffersToRecycle release];
+    }];
+    
+    this->BuffersInFlight = [NSMutableArray new];
+    this->CurrentBuffer = nil;
+    this->CurrentCursor = 0;
+    
+    SET_MEMORY_STAT_FName(BytesPerFrameStat.GetName(), this->BytesAcquiredInCurrentFrame);
+    SET_MEMORY_STAT_FName(MemoryInFlightStat.GetName(), this->BytesInFlight);
+    this->BytesAcquiredInCurrentFrame = 0;
+}
+
+#pragma mark Private
+
+id <MTLBuffer> FMetalFrameAllocator::FindOrAllocateBufferUnsafe(uint32 SizeInBytes)
+{
+    check(SizeInBytes % BufferOffsetAlignment == 0);
+    
+    id <MTLBuffer> Found = nil;
+    for(id <MTLBuffer> Buf : this->AvailableBuffers)
+    {
+        if(Buf.length > SizeInBytes)
+        {
+            Found = Buf;
+            break;
+        }
+    }
+    
+    if(Found != nil)
+    {
+        [this->AvailableBuffers removeObject:Found];
+        [Found setPurgeableState:MTLPurgeableStateNonVolatile];
+    }
+    else
+    {
+        uint32 AllocationSize = FMath::Max(DefaultAllocationSize, SizeInBytes);
+        Found = [Device newBufferWithLength:AllocationSize options:MTLResourceStorageModeShared | MTLResourceCPUCacheModeWriteCombined];
+        [Found setLabel: [NSString stringWithFormat:@"UniformBacking %u", AllocationSize]];
+        
+        check(Found);
+        
+        this->TotalBytesAllocated += AllocationSize;
+        INC_MEMORY_STAT_BY_FName(TotalAllocationStat.GetName(), [Found length]);
+#if STATS || ENABLE_LOW_LEVEL_MEM_TRACKER
+        MetalLLM::LogAllocBuffer(GetMetalDeviceContext().GetDevice(), Found);
+#endif
+
+#if METAL_FRAME_ALLOCATOR_VALIDATION
+        [this->AllAllocatedBuffers addObject:Found];
+#endif
+    }
+    
+    check(Found);
+    check([Found setPurgeableState:MTLPurgeableStateKeepCurrent] == MTLPurgeableStateNonVolatile);
+    
+    return Found;
+}
+
+void FMetalFrameAllocator::RecycleBuffers(NSArray <id <MTLBuffer>>* BuffersToRecycle)
+{
+    // This will run on the completion handler.
+    os_unfair_lock_lock(&UnfairLock);
+    
+    for(id <MTLBuffer> Buffer : BuffersToRecycle)
+    {
+        this->BytesInFlight -= [Buffer length];
+        
+        // Very simple logic: if we are over the target we will destroy the buffer instead of returning it to the pool.
+        if(this->TotalBytesAllocated < this->TargetAllocationLimit)
+        {
+            [Buffer setPurgeableState:MTLPurgeableStateVolatile];
+            [this->AvailableBuffers addObject:Buffer];
+        }
+        else
+        {
+            this->TotalBytesAllocated -= [Buffer length];
+            DEC_MEMORY_STAT_BY_FName(TotalAllocationStat.GetName(), [Buffer length]);
+#if METAL_FRAME_ALLOCATOR_VALIDATION
+            [this->AllAllocatedBuffers removeObject:Buffer];
+#endif
+            [Buffer release];
+        }
+    }
+
+    os_unfair_lock_unlock(&UnfairLock);
+}
diff --git a/Engine/Source/Runtime/Apple/MetalRHI/Private/MetalProfiler.cpp b/Engine/Source/Runtime/Apple/MetalRHI/Private/MetalProfiler.cpp
--- a/Engine/Source/Runtime/Apple/MetalRHI/Private/MetalProfiler.cpp
+++ b/Engine/Source/Runtime/Apple/MetalRHI/Private/MetalProfiler.cpp
@@ -55,6 +55,14 @@ DEFINE_STAT(STAT_MetalTextureCount);
 DEFINE_STAT(STAT_MetalHeapCount);
 DEFINE_STAT(STAT_MetalFenceCount);
 
+DEFINE_STAT(STAT_MetalUniformMemoryInFlight);
+DEFINE_STAT(STAT_MetalUniformAllocatedMemory);
+DEFINE_STAT(STAT_MetalUniformBytesPerFrame);
+
+DEFINE_STAT(STAT_MetalFrameAllocatorMemoryInFlight);
+DEFINE_STAT(STAT_MetalFrameAllocatorAllocatedMemory);
+DEFINE_STAT(STAT_MetalFrameAllocatorBytesPerFrame);
+
 int64 volatile GMetalTexturePageOnTime = 0;
 int64 volatile GMetalGPUWorkTime = 0;
 int64 volatile GMetalGPUIdleTime = 0;
diff --git a/Engine/Source/Runtime/Apple/MetalRHI/Private/MetalProfiler.h b/Engine/Source/Runtime/Apple/MetalRHI/Private/MetalProfiler.h
--- a/Engine/Source/Runtime/Apple/MetalRHI/Private/MetalProfiler.h
+++ b/Engine/Source/Runtime/Apple/MetalRHI/Private/MetalProfiler.h
@@ -42,6 +42,14 @@ DECLARE_MEMORY_STAT_EXTERN(TEXT("Heap Memory"), STAT_MetalHeapMemory, STATGROUP_
 DECLARE_MEMORY_STAT_EXTERN(TEXT("Unused Buffer Memory"), STAT_MetalBufferUnusedMemory, STATGROUP_MetalRHI, );
 DECLARE_MEMORY_STAT_EXTERN(TEXT("Unused Texture Memory"), STAT_MetalTextureUnusedMemory, STATGROUP_MetalRHI, );
 
+DECLARE_MEMORY_STAT_EXTERN(TEXT("Uniform Memory In Flight"), STAT_MetalUniformMemoryInFlight, STATGROUP_MetalRHI, );
+DECLARE_MEMORY_STAT_EXTERN(TEXT("Allocated Uniform Pool Memory"), STAT_MetalUniformAllocatedMemory, STATGROUP_MetalRHI, );
+DECLARE_MEMORY_STAT_EXTERN(TEXT("Uniform Memory Per Frame"), STAT_MetalUniformBytesPerFrame, STATGROUP_MetalRHI, );
+
+DECLARE_MEMORY_STAT_EXTERN(TEXT("General Frame Allocator Memory In Flight"), STAT_MetalFrameAllocatorMemoryInFlight, STATGROUP_MetalRHI, );
+DECLARE_MEMORY_STAT_EXTERN(TEXT("Allocated Frame Allocator Memory"), STAT_MetalFrameAllocatorAllocatedMemory, STATGROUP_MetalRHI, );
+DECLARE_MEMORY_STAT_EXTERN(TEXT("Frame Allocator Memory Per Frame"), STAT_MetalFrameAllocatorBytesPerFrame, STATGROUP_MetalRHI, );
+
 DECLARE_DWORD_ACCUMULATOR_STAT_EXTERN(TEXT("Buffer Count"), STAT_MetalBufferCount, STATGROUP_MetalRHI, );
 DECLARE_DWORD_ACCUMULATOR_STAT_EXTERN(TEXT("Texture Count"), STAT_MetalTextureCount, STATGROUP_MetalRHI, );
 DECLARE_DWORD_ACCUMULATOR_STAT_EXTERN(TEXT("Heap Count"), STAT_MetalHeapCount, STATGROUP_MetalRHI, );
diff --git a/Engine/Source/Runtime/Apple/MetalRHI/Private/MetalRHIPrivate.h b/Engine/Source/Runtime/Apple/MetalRHI/Private/MetalRHIPrivate.h
--- a/Engine/Source/Runtime/Apple/MetalRHI/Private/MetalRHIPrivate.h
+++ b/Engine/Source/Runtime/Apple/MetalRHI/Private/MetalRHIPrivate.h
@@ -185,6 +185,12 @@ FMetalSurface* GetMetalSurfaceFromRHITexture(FRHITexture* Texture);
 
 #define NOT_SUPPORTED(Func) UE_LOG(LogMetal, Fatal, TEXT("'%s' is not supported"), L##Func);
 
+// Verifies we are on the correct thread to mutate internal MetalRHI resources.
+FORCEINLINE void CheckMetalThread()
+{
+    check((IsInRenderingThread() && (!IsRunningRHIInSeparateThread() || !FRHICommandListExecutor::IsRHIThreadActive())) || IsInRHIThread());
+}
+
 FORCEINLINE mtlpp::IndexType GetMetalIndexType(EMetalIndexType IndexType)
 {
 	switch (IndexType)
@@ -327,15 +333,5 @@ FORCEINLINE EMetalShaderStages GetMetalShaderFrequency(EShaderFrequency Stage)
 	}
 }
 
-@interface FMetalIAB : FApplePlatformObject<NSObject>
-{
-@public
-	FMetalBuffer IndirectArgumentBuffer;
-	FMetalBuffer IndirectArgumentBufferSideTable;
-	mtlpp::ArgumentEncoder IndirectArgumentEncoder;
-	int64 UpdateIAB;
-}
-@end
-
 #include "MetalStateCache.h"
 #include "MetalContext.h"
diff --git a/Engine/Source/Runtime/Apple/MetalRHI/Private/MetalRenderPass.cpp b/Engine/Source/Runtime/Apple/MetalRHI/Private/MetalRenderPass.cpp
--- a/Engine/Source/Runtime/Apple/MetalRHI/Private/MetalRenderPass.cpp
+++ b/Engine/Source/Runtime/Apple/MetalRHI/Private/MetalRenderPass.cpp
@@ -10,6 +10,7 @@
 #include "MetalRenderPass.h"
 #include "MetalCommandBuffer.h"
 #include "MetalProfiler.h"
+#include "MetalFrameAllocator.h"
 
 #pragma mark - Private Console Variables -
 
@@ -162,6 +163,21 @@ TRefCountPtr<FMetalFence> const& FMetalRenderPass::Submit(EMetalSubmitFlags Flag
 			PrologueEncoder.CommitCommandBuffer((Flags & EMetalSubmitFlagsAsyncCommandBuffer) ? Flags : EMetalSubmitFlagsNone);
         }
     }
+
+    // Must be on the render thread if there's no RHI thread
+    // Must be on the RHI thread otherwise
+    CheckMetalThread();
+    
+    if (Flags & EMetalSubmitFlagsLastCommandBuffer)
+    {
+        check(CurrentEncoder.GetCommandBuffer());
+        id <MTLCommandBuffer> CommandBuffer = CurrentEncoder.GetCommandBuffer().GetPtr();
+        
+        FMetalDeviceContext& DeviceContext = (FMetalDeviceContext&)GetMetalDeviceContext();
+        FMetalFrameAllocator* UniformAllocator = DeviceContext.GetUniformAllocator();
+        
+        UniformAllocator->MarkEndOfFrame(DeviceContext.GetFrameNumberRHIThread(), CommandBuffer);
+    }
     
     if (CurrentEncoder.GetCommandBuffer() && !(Flags & EMetalSubmitFlagsAsyncCommandBuffer))
     {
diff --git a/Engine/Source/Runtime/Apple/MetalRHI/Private/MetalUniformBuffer.cpp b/Engine/Source/Runtime/Apple/MetalRHI/Private/MetalUniformBuffer.cpp
--- a/Engine/Source/Runtime/Apple/MetalRHI/Private/MetalUniformBuffer.cpp
+++ b/Engine/Source/Runtime/Apple/MetalRHI/Private/MetalUniformBuffer.cpp
@@ -4,183 +4,107 @@
 	MetalConstantBuffer.cpp: Metal Constant buffer implementation.
 =============================================================================*/
 
+#include "MetalRHI.h"
+#include "MetalResources.h"
+#include "MetalFrameAllocator.h"
 #include "MetalRHIPrivate.h"
-#include "MetalProfiler.h"
-#include "MetalBuffer.h"
-#include "MetalCommandBuffer.h"
-#include "HAL/LowLevelMemTracker.h"
-#include "Misc/ScopeRWLock.h"
 
-@implementation FMetalIAB
-
-APPLE_PLATFORM_OBJECT_ALLOC_OVERRIDES(FMetalIAB)
-
--(instancetype)init
-{
-	id Self = [super init];
-	if (Self)
-		((FMetalIAB*)Self)->UpdateIAB = 0;
-	return Self;
-}
--(void)dealloc
-{
-	SafeReleaseMetalBuffer(IndirectArgumentBuffer);
-	SafeReleaseMetalBuffer(IndirectArgumentBufferSideTable);
-	[super dealloc];
-}
-@end
-
-FMetalUniformBuffer::FMetalUniformBuffer(const void* Contents, const FRHIUniformBufferLayout& Layout, EUniformBufferUsage InUsage, EUniformBufferValidation Validation)
-	: FRHIUniformBuffer(Layout)
-    , FMetalRHIBuffer(Layout.ConstantBufferSize, (FMetalCommandQueue::SupportsFeature(EMetalFeaturesIABs) && Layout.Resources.Num() ? (EMetalBufferUsage_GPUOnly|BUF_Volatile) : BUF_Volatile), RRT_UniformBuffer)
-{
-	uint32 NumResources = Layout.Resources.Num();
-	if (NumResources)
-	{
-		ResourceTable.Empty(NumResources);
-		ResourceTable.AddZeroed(NumResources);
-		
-		for (uint32 i = 0; i < NumResources; ++i)
-		{
-			FRHIResource* Resource = *(FRHIResource**)((uint8*)Contents + Layout.Resources[i].MemberOffset);
-			
-			// Allow null SRV's in uniform buffers for feature levels that don't support SRV's in shaders
-			if (Validation == EUniformBufferValidation::ValidateResources && !(GMaxRHIFeatureLevel <= ERHIFeatureLevel::ES3_1 && Layout.Resources[i].MemberType == UBMT_SRV))
-			{
-				check(Resource);
-			}
-			
-			ResourceTable[i] = Resource;
-		}
-	}
-
-    uint32 ConstantSize = Layout.ConstantBufferSize;
-	if (ConstantSize > 0)
-	{
-		UE_CLOG(ConstantSize > 65536, LogMetal, Fatal, TEXT("Trying to allocated a uniform layout of size %d that is greater than the maximum permitted 64k."), ConstantSize);
-		
-		if (Buffer)
-		{
-			FMemory::Memcpy(Buffer.GetContents(), Contents, ConstantSize);
-#if PLATFORM_MAC
-			if(Mode == mtlpp::StorageMode::Managed)
-			{
-				MTLPP_VALIDATE(mtlpp::Buffer, Buffer, SafeGetRuntimeDebuggingLevel() >= EMetalDebugLevelValidation, DidModify(ns::Range(0, ConstantSize)));
-			}
+#pragma mark Suballocated Uniform Buffer Implementation
+
+FMetalSuballocatedUniformBuffer::FMetalSuballocatedUniformBuffer(const FRHIUniformBufferLayout& Layout, EUniformBufferUsage Usage, EUniformBufferValidation InValidation)
+    : FRHIUniformBuffer(Layout)
+    , LastFrameUpdated(0)
+    , Offset(0)
+    , Backing(nil)
+    , Shadow(nullptr)
+    , ResourceTable()
+#if !UE_BUILD_SHIPPING
+    , Validation(InValidation)
 #endif
-		}
-		else
-		{
-			check(Data && Data->Data);
-			FMemory::Memcpy(Data->Data, Contents, ConstantSize);
-		}
-	}
-
-    UpdateResourceTable(ResourceTable, Validation);
+{
+    // Slate can create SingleDraw uniform buffers and use them several frames later. So it must be included.
+    if(Usage == UniformBuffer_SingleDraw || Usage == UniformBuffer_MultiFrame)
+    {
+        Shadow = FMemory::Malloc(this->GetSize());
+    }
 }
 
-void const* FMetalUniformBuffer::GetData()
+FMetalSuballocatedUniformBuffer::~FMetalSuballocatedUniformBuffer()
 {
-	if (Data)
-	{
-		return Data->Data;
-	}
-	else if (Buffer)
-	{
-		return MTLPP_VALIDATE(mtlpp::Buffer, Buffer, SafeGetRuntimeDebuggingLevel() >= EMetalDebugLevelValidation, GetContents());
-	}
-	else
-	{
-		return nullptr;
-	}
+    if(this->HasShadow())
+    {
+        FMemory::Free(Shadow);
+    }
+    
+    // Note: this object does NOT own a reference
+    // to the uniform buffer backing store
 }
 
-void FMetalUniformBuffer::UpdateResourceTable(TArray<TRefCountPtr<FRHIResource>>& Resources, EUniformBufferValidation Validation)
+bool FMetalSuballocatedUniformBuffer::HasShadow()
 {
-	ResourceTable = Resources;
+    return this->Shadow != nullptr;
 }
 
-void FMetalUniformBuffer::Update(const void* Contents, TArray<TRefCountPtr<FRHIResource>>& Resources, EUniformBufferValidation Validation)
+void FMetalSuballocatedUniformBuffer::Update(const void* Contents)
 {
-    uint32 ConstantSize = FRHIUniformBuffer::GetSize();
-    if (ConstantSize > 0)
+    if(this->HasShadow())
     {
-        UE_CLOG(ConstantSize > 65536, LogMetal, Fatal, TEXT("Trying to allocated a uniform layout of size %d that is greater than the maximum permitted 64k."), ConstantSize);
+        FMemory::Memcpy(this->Shadow, Contents, this->GetSize());
+    }
+    
+    const FRHIUniformBufferLayout& Layout = this->GetLayout();
+    uint32 NumResources = Layout.Resources.Num();
+    if (NumResources > 0)
+    {
+        ResourceTable.Empty(NumResources);
+        ResourceTable.AddZeroed(NumResources);
         
-		ns::AutoReleased<FMetalBuffer> Buf(Buffer);
-		
-		void* Data = Lock(true, RLM_WriteOnly, 0, 0, true);
-        FMemory::Memcpy(Data, Contents, ConstantSize);
-        Unlock();
+        for (uint32 i = 0; i < NumResources; ++i)
+        {
+            FRHIResource* Resource = *(FRHIResource**)((uint8*)Contents + Layout.Resources[i].MemberOffset);
+            
+            // Allow null SRV's in uniform buffers for feature levels that don't support SRV's in shaders
+#if METAL_UNIFORM_BUFFER_VALIDATION
+            if (Validation == EUniformBufferValidation::ValidateResources && !(GMaxRHIFeatureLevel <= ERHIFeatureLevel::ES3_1 && GetLayout().Resources[i].MemberType == UBMT_SRV))
+            {
+                check(Resource);
+            }
+#endif
+            
+            ResourceTable[i] = Resource;
+        }
+    }
 
-		ConditionalSetUniformBufferFrameIndex();
-	}
-	
-	UpdateResourceTable(Resources, Validation);
+    this->PushToGPUBacking(Contents);
 }
 
-FUniformBufferRHIRef FMetalDynamicRHI::RHICreateUniformBuffer(const void* Contents, const FRHIUniformBufferLayout& Layout, EUniformBufferUsage Usage, EUniformBufferValidation Validation)
+// Acquires a region in the current frame's uniform buffer and
+// pushes the data in Contents into that GPU backing store
+// The amount of data read from Contents is given by the Layout
+void FMetalSuballocatedUniformBuffer::PushToGPUBacking(const void* Contents)
 {
-	@autoreleasepool {
-	check(IsInRenderingThread() || IsInParallelRenderingThread() || IsInRHIThread());
-		return new FMetalUniformBuffer(Contents, Layout, Usage, Validation);
-	}
+    check(IsInRenderingThread() ^ IsRunningRHIInSeparateThread());
+    
+    FMetalDeviceContext& DeviceContext = GetMetalDeviceContext();
+    
+    FMetalFrameAllocator* Allocator = DeviceContext.GetUniformAllocator();
+    FMetalFrameAllocator::AllocationEntry Entry = Allocator->AcquireSpace(this->GetSize());
+    // copy contents into backing
+    this->Backing = Entry.Backing;
+    this->Offset = Entry.Offset;
+    uint8* ConstantSpace = reinterpret_cast<uint8*>([this->Backing contents]) + Entry.Offset;
+    FMemory::Memcpy(ConstantSpace, Contents, this->GetSize());
+    LastFrameUpdated = DeviceContext.GetFrameNumberRHIThread();
 }
 
-struct FMetalRHICommandUpateUniformBuffer : public FRHICommand<FMetalRHICommandUpateUniformBuffer>
+// Because we can create a uniform buffer on frame N and may not bind it until frame N+10
+// we need to keep a copy of the most recent data. Then when it's time to bind this
+// uniform buffer we can push the data into the GPU backing.
+void FMetalSuballocatedUniformBuffer::PrepareToBind()
 {
-	TRefCountPtr<FMetalUniformBuffer> Buffer;
-	TArray<TRefCountPtr<FRHIResource> > ResourceTable;
-	char* Contents;
-	
-	FORCEINLINE_DEBUGGABLE FMetalRHICommandUpateUniformBuffer(FMetalUniformBuffer* InBuffer, void const* Data, TArray<TRefCountPtr<FRHIResource>>& Resources)
-	: Buffer(InBuffer)
-	, ResourceTable(Resources)
-	, Contents(nullptr)
-	{
-		uint32 MaxLayoutSize = InBuffer->GetSize();
-		Contents = new char[MaxLayoutSize];
-		FMemory::Memcpy(Contents, Data, MaxLayoutSize);
-	}
-	
-	virtual ~FMetalRHICommandUpateUniformBuffer()
-	{
-		delete [] Contents;
-	}
-	
-	void Execute(FRHICommandListBase& CmdList)
-	{
-		Buffer->Update(Contents, ResourceTable, EUniformBufferValidation::None);
-	}
-};
-
-void FMetalDynamicRHI::RHIUpdateUniformBuffer(FRHIUniformBuffer* UniformBufferRHI, const void* Contents)
-{
-	@autoreleasepool {
-	// check((IsInRenderingThread() || IsInRHIThread()) && !IsInParallelRenderingThread());
-
-	FMetalUniformBuffer* UniformBuffer = ResourceCast(UniformBufferRHI);
-	FRHICommandListImmediate& RHICmdList = FRHICommandListExecutor::GetImmediateCommandList();
-		
-	TArray<TRefCountPtr<FRHIResource> > ResourceTable;
-	ResourceTable.AddZeroed(UniformBuffer->GetLayout().Resources.Num());
-		
-	const FRHIUniformBufferLayout& Layout = UniformBuffer->GetLayout();
-		
-	for (uint32 i = 0; i < UniformBuffer->GetLayout().Resources.Num(); ++i)
-	{
-		FRHIResource* Resource = *(FRHIResource**)((uint8*)Contents + Layout.Resources[i].MemberOffset);
-		ResourceTable[i] = Resource;
-	}	
-		
-	if (RHICmdList.Bypass() || !IsRunningRHIInSeparateThread())
-	{
-		UniformBuffer->Update(Contents, ResourceTable, EUniformBufferValidation::None);
-	}
-	else
-	{
-		new (RHICmdList.AllocCommand<FMetalRHICommandUpateUniformBuffer>()) FMetalRHICommandUpateUniformBuffer(UniformBuffer, Contents, ResourceTable);
-		RHICmdList.RHIThreadFence(true);
-	}
-	}
+    FMetalDeviceContext& DeviceContext = GetMetalDeviceContext();
+    if(Shadow && LastFrameUpdated < DeviceContext.GetFrameNumberRHIThread())
+    {
+        this->PushToGPUBacking(this->Shadow);
+    }
 }
diff --git a/Engine/Source/Runtime/Apple/MetalRHI/Private/MetalUniformBufferCommands.cpp b/Engine/Source/Runtime/Apple/MetalRHI/Private/MetalUniformBufferCommands.cpp
new file mode 100644
--- /dev/null
+++ b/Engine/Source/Runtime/Apple/MetalRHI/Private/MetalUniformBufferCommands.cpp
@@ -0,0 +1,109 @@
+// Copyright Epic Games, Inc. All Rights Reserved.
+
+#include "RHI.h"
+#include "MetalRHI.h"
+#include "MetalRHIPrivate.h"
+#include "MetalResources.h"
+#include "MetalFrameAllocator.h"
+
+static void DoUpdateUniformBuffer(FMetalUniformBuffer* UB, const void* Contents)
+{
+    check(IsInRenderingThread() || IsInRHIThread());
+    
+    FRHICommandListImmediate& RHICmdList = FRHICommandListExecutor::GetImmediateCommandList();
+    // The only way we can be on the RHI thread here is if we're in the process of creating a FLocalUniformBuffer.
+    bool bUpdateImmediately = RHICmdList.Bypass() || IsInRHIThread();
+    
+    if(bUpdateImmediately)
+    {
+        UB->Update(Contents);
+    }
+    else
+    {
+        const uint32 NumBytes = UB->GetLayout().ConstantBufferSize;
+        void* Data = RHICmdList.Alloc(NumBytes, 16);
+        FMemory::Memcpy(Data, Contents, NumBytes);
+        
+        RHICmdList.EnqueueLambda([Data, UB](FRHICommandListImmediate& RHICmdList)
+        {
+            UB->Update(Data);
+        });
+        
+        RHICmdList.RHIThreadFence(true);
+    }
+}
+
+FUniformBufferRHIRef FMetalDynamicRHI::RHICreateUniformBuffer(const void* Contents, const FRHIUniformBufferLayout& Layout, EUniformBufferUsage Usage, EUniformBufferValidation Validation)
+{
+    FMetalDeviceContext& DeviceContext = (FMetalDeviceContext&)GetMetalDeviceContext();
+    FMetalFrameAllocator* UniformAllocator = DeviceContext.GetUniformAllocator();
+    
+    FMetalUniformBuffer* UB = new FMetalUniformBuffer(Layout, Usage, Validation);
+    
+    DoUpdateUniformBuffer(UB, Contents);
+    
+    return UB;
+}
+
+void FMetalDynamicRHI::RHIUpdateUniformBuffer(FRHIUniformBuffer* UniformBufferRHI, const void* Contents)
+{
+    check(IsInRenderingThread());
+    
+    FMetalUniformBuffer* UB = ResourceCast(UniformBufferRHI);
+    DoUpdateUniformBuffer(UB, Contents);
+}
+
+template<EMetalShaderStages Stage, typename RHIShaderType>
+static void SetUniformBufferInternal(FMetalContext* Context, RHIShaderType* ShaderRHI, uint32 BufferIndex, FRHIUniformBuffer* UBRHI)
+{
+    @autoreleasepool
+    {
+        auto Shader = ResourceCast(ShaderRHI);
+        Context->GetCurrentState().BindUniformBuffer(Stage, BufferIndex, UBRHI);
+        
+        auto& Bindings = Shader->Bindings;
+        if((Bindings.ConstantBuffers) & (1 << BufferIndex))
+        {
+            FMetalUniformBuffer* UB = ResourceCast(UBRHI);
+            UB->PrepareToBind();
+            
+            FMetalBuffer Buf(UB->Backing, ns::Ownership::AutoRelease);
+            Context->GetCurrentState().SetShaderBuffer(Stage, Buf, nil, UB->Offset, UB->GetSize(), BufferIndex, mtlpp::ResourceUsage::Read);
+        }
+    }
+}
+
+void FMetalRHICommandContext::RHISetShaderUniformBuffer(FRHIVertexShader* VertexShaderRHI, uint32 BufferIndex, FRHIUniformBuffer* BufferRHI)
+{
+    SetUniformBufferInternal<EMetalShaderStages::Vertex, FRHIVertexShader>(Context, VertexShaderRHI, BufferIndex, BufferRHI);
+}
+
+void FMetalRHICommandContext::RHISetShaderUniformBuffer(FRHIHullShader* HullShaderRHI, uint32 BufferIndex, FRHIUniformBuffer* BufferRHI)
+{
+#if PLATFORM_SUPPORTS_TESSELLATION_SHADERS
+    SetUniformBufferInternal<EMetalShaderStages::Hull, FRHIHullShader>(Context, HullShaderRHI, BufferIndex, BufferRHI);
+#endif
+}
+
+void FMetalRHICommandContext::RHISetShaderUniformBuffer(FRHIDomainShader* DomainShaderRHI, uint32 BufferIndex, FRHIUniformBuffer* BufferRHI)
+{
+#if PLATFORM_SUPPORTS_TESSELLATION_SHADERS
+    SetUniformBufferInternal<EMetalShaderStages::Domain, FRHIDomainShader>(Context, DomainShaderRHI, BufferIndex, BufferRHI);
+#endif
+}
+
+void FMetalRHICommandContext::RHISetShaderUniformBuffer(FRHIGeometryShader* GeometryShader, uint32 BufferIndex, FRHIUniformBuffer* BufferRHI)
+{
+    NOT_SUPPORTED("RHISetShaderUniformBuffer-Geometry");
+}
+
+void FMetalRHICommandContext::RHISetShaderUniformBuffer(FRHIPixelShader* PixelShaderRHI, uint32 BufferIndex, FRHIUniformBuffer* BufferRHI)
+{
+    SetUniformBufferInternal<EMetalShaderStages::Pixel, FRHIPixelShader>(Context, PixelShaderRHI, BufferIndex, BufferRHI);
+}
+
+void FMetalRHICommandContext::RHISetShaderUniformBuffer(FRHIComputeShader* ComputeShaderRHI, uint32 BufferIndex, FRHIUniformBuffer* BufferRHI)
+{
+    SetUniformBufferInternal<EMetalShaderStages::Compute, FRHIComputeShader>(Context, ComputeShaderRHI, BufferIndex, BufferRHI);
+}
+
diff --git a/Engine/Source/Runtime/Apple/MetalRHI/Private/MetalVertexBuffer.cpp b/Engine/Source/Runtime/Apple/MetalRHI/Private/MetalVertexBuffer.cpp
--- a/Engine/Source/Runtime/Apple/MetalRHI/Private/MetalVertexBuffer.cpp
+++ b/Engine/Source/Runtime/Apple/MetalRHI/Private/MetalVertexBuffer.cpp
@@ -236,7 +236,7 @@ void FMetalRHIBuffer::Unalias()
 	}
 }
 
-void FMetalRHIBuffer::Alloc(uint32 InSize, EResourceLockMode LockMode, bool bIsUniformBuffer)
+void FMetalRHIBuffer::Alloc(uint32 InSize, EResourceLockMode LockMode)
 {
 	if (!Buffer)
 	{
@@ -260,8 +260,6 @@ void FMetalRHIBuffer::Alloc(uint32 InSize, EResourceLockMode LockMode, bool bIsU
 				check(Pair.Value);
 			}
 		}
-
-		bIsUniformBufferBacking = bIsUniformBuffer;
 	}
 }
 
@@ -459,22 +457,7 @@ ns::AutoReleased<FMetalTexture> FMetalRHIBuffer::GetLinearTexture(EPixelFormat I
 	return Texture;
 }
 
-bool FMetalRHIBuffer::CanUseBufferAsBackingForAsyncCopy() const
-{
-	return (!bIsUniformBufferBacking ||
-			UniformBufferFrameIndex != GetMetalDeviceContext().GetDeviceFrameIndex() ||
-			UniformBufferPreviousOffset != Buffer.GetOffset());
-}
-
-void FMetalRHIBuffer::ConditionalSetUniformBufferFrameIndex()
-{
-	if (bIsUniformBufferBacking)
-	{
-		UniformBufferFrameIndex = GetMetalDeviceContext().GetDeviceFrameIndex();
-	}
-}
-
-void* FMetalRHIBuffer::Lock(bool bIsOnRHIThread, EResourceLockMode LockMode, uint32 Offset, uint32 InSize, bool bIsUniformBuffer /*= false*/)
+void* FMetalRHIBuffer::Lock(bool bIsOnRHIThread, EResourceLockMode LockMode, uint32 Offset, uint32 InSize)
 {
 	check(LockSize == 0 && LockOffset == 0);
 	
@@ -519,16 +502,14 @@ void* FMetalRHIBuffer::Lock(bool bIsOnRHIThread, EResourceLockMode LockMode, uin
 	// When we can't we have to reallocate the backing store
 	if (LockMode != RLM_ReadOnly &&
 		Mode == mtlpp::StorageMode::Private &&
-		Buffer &&
-		(!GetMetalDeviceContext().CanAsyncCopyToBuffer(Buffer) ||
-		 !CanUseBufferAsBackingForAsyncCopy()))
+		Buffer)
 	{
 		METAL_INC_DWORD_STAT_BY(Type, MemFreed, Len);
 		SafeReleaseMetalBuffer(Buffer);
 		Buffer = nil;
 	}
 	
-    Alloc(Len, LockMode, bIsUniformBuffer);
+    Alloc(Len, LockMode);
 	AllocTransferBuffer(bIsOnRHIThread, Len, LockMode);
 	
 	FMetalBuffer& theBufferToUse = CPUBuffer ? CPUBuffer : Buffer;
@@ -579,8 +560,6 @@ void FMetalRHIBuffer::Unlock()
 			// Synchronise the buffer with the GPU
 			GetMetalDeviceContext().AsyncCopyFromBufferToBuffer(CPUBuffer, 0, Buffer, 0, FMath::Min(CPUBuffer.GetLength(), Buffer.GetLength()));
 			
-			ConditionalSetUniformBufferPreviousOffset();
-			
 			if (CPUBuffer)
             {
 				SafeReleaseMetalBuffer(CPUBuffer);
diff --git a/Engine/Source/Runtime/Apple/MetalRHI/Public/MetalFrameAllocator.h b/Engine/Source/Runtime/Apple/MetalRHI/Public/MetalFrameAllocator.h
new file mode 100644
--- /dev/null
+++ b/Engine/Source/Runtime/Apple/MetalRHI/Public/MetalFrameAllocator.h
@@ -0,0 +1,86 @@
+// Copyright Epic Games, Inc. All Rights Reserved.
+
+#import <Metal/Metal.h>
+#import <os/lock.h>
+
+#include "MetalRHI.h"
+
+#define METAL_FRAME_ALLOCATOR_VALIDATION !(UE_BUILD_SHIPPING || UE_BUILD_TEST)
+
+class FMetalFrameAllocator
+{
+public:
+
+    struct AllocationEntry
+    {
+        id <MTLBuffer> Backing;
+        uint32 Offset;
+    };
+    
+private:
+    // The device. This object will take a reference.
+    // You must make sure your rendering is on this same device or things will explode horribly.
+    id <MTLDevice> Device;
+#if METAL_FRAME_ALLOCATOR_VALIDATION
+    // Every buffer that is currently allocated.
+    NSMutableArray<id <MTLBuffer>>* AllAllocatedBuffers;
+#endif
+    // Buffers that are available. These buffers are not currently being read from nor written to.
+    NSMutableArray<id <MTLBuffer>>* AvailableBuffers;
+    // The current set of buffers in use for this frame.
+    // Must ONLY be mutated on RHI thread if present.
+    NSMutableArray<id <MTLBuffer>>* BuffersInFlight;
+    // The current buffer we are suballocating from and the current offset.
+    id <MTLBuffer> CurrentBuffer;
+    uint32 CurrentCursor;
+    uint32 CurrentFrame;
+    // Lock around AllAllocatedBuffers and AvailableBuffers since they are mutated from a completion handler.
+    os_unfair_lock UnfairLock;
+    // Allocation info and target limits.
+    uint32 TargetAllocationLimit;
+    uint32 DefaultAllocationSize;
+    // Stats!
+    uint64 TotalBytesAllocated;
+    uint64 BytesAcquiredInCurrentFrame;
+    uint64 BytesInFlight;
+    
+    TStatId TotalAllocationStat;
+    TStatId MemoryInFlightStat;
+    TStatId BytesPerFrameStat;
+public:
+    // API
+    
+    // This object will hold a reference to TargetDevice for its lifetime.
+    FMetalFrameAllocator(id <MTLDevice> TargetDevice);
+    ~FMetalFrameAllocator();
+    
+    // Sets a target allocation limit.
+    // This is not a hard high-water mark.
+    // If a frame needs to use more memory it will do so.
+    // When the total bytes allocated goes over this amount we will start removing buffers from the pool.
+    void SetTargetAllocationLimitInBytes(uint32 LimitInBytes);
+    
+    // The size of each buffer.
+    // The actual size is computed by max(SizeNeeded, DefaultSize)
+    void SetDefaultAllocationSizeInBytes(uint32 DefaultInBytes);
+    
+    void SetStatIds(const TStatId& TotalAllocationStat, const TStatId& MemoryInFlightStat, const TStatId& BytesPerFrameStat);
+    
+    // Acquires SizeInBytes bytes for uniform data.
+    // Returns the currently active buffer and offset.
+    // WARNING: This is NOT completely thread safe.
+    // WARNING: You MUST only call from the RHI thread if present. You MAY call it from the parallel translate threads.
+    AllocationEntry AcquireSpace(uint32 SizeInBytes);
+    // Marks the end of a frame.
+    // When LastCommandBufferInFrame is retired, the BuffersInFlight
+    // will be returned to the AvailableBuffers pool.
+    // When this call returns CurrentBuffer and CurrentCursor will be reset to 0.
+    // This is NOT thread-safe. It must be called from the RHI thread if present.
+    void MarkEndOfFrame(uint32 FrameNumberThatEnded, id <MTLCommandBuffer> LastCommandBufferInFrame);
+    
+private:
+    // This is NOT thread safe. You must lock around it.
+    id <MTLBuffer> FindOrAllocateBufferUnsafe(uint32 SizeInBytes);
+    void RecycleBuffers(NSArray <id <MTLBuffer>>* BuffersToRecycle);
+};
+
diff --git a/Engine/Source/Runtime/Apple/MetalRHI/Public/MetalResources.h b/Engine/Source/Runtime/Apple/MetalRHI/Public/MetalResources.h
--- a/Engine/Source/Runtime/Apple/MetalRHI/Public/MetalResources.h
+++ b/Engine/Source/Runtime/Apple/MetalRHI/Public/MetalResources.h
@@ -20,6 +20,8 @@ THIRD_PARTY_INCLUDES_END
 class FMetalContext;
 @class FMetalShaderPipeline;
 
+#include "MetalUniformBuffer.h"
+
 /** The MTLVertexDescriptor and a pre-calculated hash value used to simplify comparisons (as vendor MTLVertexDescriptor implementations aren't all comparable) */
 struct FMetalHashedVertexDescriptor
 {
@@ -846,7 +848,7 @@ public:
 	/**
 	 * Allocate the buffer backing store.
 	 */
-	void Alloc(uint32 InSize, EResourceLockMode LockMode, bool bIsUniformBuffer=false);
+	void Alloc(uint32 InSize, EResourceLockMode LockMode);
 	
 	/**
 	 * Allocate the CPU accessible buffer for data transfer.
@@ -871,7 +873,7 @@ public:
 	/**
 	 * Prepare a CPU accessible buffer for uploading to GPU memory
 	 */
-	void* Lock(bool bIsOnRHIThread, EResourceLockMode LockMode, uint32 Offset, uint32 Size=0, bool bIsUniformBuffer=false);
+	void* Lock(bool bIsOnRHIThread, EResourceLockMode LockMode, uint32 Offset, uint32 Size=0);
 	
 	/**
 	 * Prepare a CPU accessible buffer for uploading to GPU memory
@@ -917,33 +919,6 @@ public:
 	
 	// Resource type
 	ERHIResourceType Type;
-
-	/**
-	 * For resources backing uniform buffers, set the device frame index.
-	 *
-	 * Calling this function on resources that do not back uniform buffers has no effect.
-	 */
-	void ConditionalSetUniformBufferFrameIndex();
-
-private:
-	bool CanUseBufferAsBackingForAsyncCopy() const;
-
-	/**
-	 * For resources backing uniform buffers, updates the last used sub-allocation offset; this is used for
-	 * internal tracking and deciding when to recycle backing stores during uniform buffer updates.
-	 */
-	void ConditionalSetUniformBufferPreviousOffset()
-	{
-		if (bIsUniformBufferBacking)
-		{
-			check(Buffer);
-			UniformBufferPreviousOffset = Buffer.GetOffset();
-		}
-	}
-
-	bool bIsUniformBufferBacking = false;
-	uint64 UniformBufferFrameIndex = 0;
-	uint64 UniformBufferPreviousOffset = UINT64_MAX;
 };
 
 /** Index buffer resource class that stores stride information. */
@@ -1050,27 +1025,6 @@ struct FMetalArgumentDesc
 	}
 };
 
-@class FMetalIAB;
-
-class FMetalUniformBuffer : public FRHIUniformBuffer, public FMetalRHIBuffer
-{
-public:
-	// Constructor
-	FMetalUniformBuffer(const void* Contents, const FRHIUniformBufferLayout& Layout, EUniformBufferUsage Usage, EUniformBufferValidation Validation);
-
-	// Destructor 
-    virtual ~FMetalUniformBuffer() {}
-	
-	void const* GetData();
-
-	void UpdateResourceTable(TArray<TRefCountPtr<FRHIResource>>& Resources, EUniformBufferValidation Validation);
-	void Update(const void* Contents, TArray<TRefCountPtr<FRHIResource>>& Resources, EUniformBufferValidation Validation);
-	
-	/** Resource table containing RHI references. */
-	TArray<TRefCountPtr<FRHIResource> > ResourceTable;
-};
-
-
 class FMetalStructuredBuffer : public FRHIStructuredBuffer, public FMetalRHIBuffer
 {
 public:
@@ -1363,7 +1317,7 @@ struct TMetalResourceTraits<FRHIRenderQuery>
 template<>
 struct TMetalResourceTraits<FRHIUniformBuffer>
 {
-	typedef FMetalUniformBuffer TConcreteType;
+	typedef FMetalSuballocatedUniformBuffer TConcreteType;
 };
 template<>
 struct TMetalResourceTraits<FRHIIndexBuffer>
diff --git a/Engine/Source/Runtime/Apple/MetalRHI/Public/MetalUniformBuffer.h b/Engine/Source/Runtime/Apple/MetalRHI/Public/MetalUniformBuffer.h
new file mode 100644
--- /dev/null
+++ b/Engine/Source/Runtime/Apple/MetalRHI/Public/MetalUniformBuffer.h
@@ -0,0 +1,63 @@
+// Copyright Epic Games, Inc. All Rights Reserved.
+
+//
+//  Implements handles to linearly allocated per-frame constant buffers for shared memory systems.
+//
+
+#import <Metal/Metal.h>
+
+#include "RHIResources.h"
+
+#define METAL_UNIFORM_BUFFER_VALIDATION !UE_BUILD_SHIPPING
+
+class FMetalStateCache;
+
+class FMetalSuballocatedUniformBuffer : public FRHIUniformBuffer
+{
+    friend class FMetalStateCache;
+public:
+    // The last render thread frame this uniform buffer updated or pushed contents to the GPU backing
+    int32 LastFrameUpdated;
+    // Offset within the GPU backing this uniform buffer owns
+    uint32 Offset;
+    // The GPU backing buffer for this uniform buffer. Many FMetalMobileUniformBuffers can own regions.
+    // This UB does not own a reference to the backing buffer.
+    // This Backing is recycled at the end of every frame so you MUST update it if LastFrameUpdated != this frame or contents are undefined.
+    id <MTLBuffer> Backing;
+    // CPU side shadow memory to hold updates for single-draw or multi-frame buffers.
+    // This allows you to upload on a frame but actually use this UB later on
+    void* Shadow;
+private:
+    TArray<TRefCountPtr<FRHIResource> > ResourceTable;
+#if METAL_UNIFORM_BUFFER_VALIDATION
+    EUniformBufferValidation Validation;
+#endif
+    
+public:
+    // Creates a uniform buffer.
+    // If Usage is SingleDraw or MultiFrame we will keep a copy of the data
+    FMetalSuballocatedUniformBuffer(const FRHIUniformBufferLayout& Layout, EUniformBufferUsage Usage, EUniformBufferValidation Validation);
+    ~FMetalSuballocatedUniformBuffer();
+    
+    void Update(const void* Contents);
+    
+    // Prepares this uniform buffer for binding.
+    // You MUST call this before binding it.
+    // If this is a UB that was created before the current frame the data on the GPU will
+    // not be correct until this function returns.
+    void PrepareToBind();
+    
+private:
+    
+    // Pushes the data in Contents to the gpu.
+    // Updates the frame counter to FrameNumber.
+    // (this is to support the case where we create buffer and reuse it many frames later).
+    // This acquires a region in the current frame's transient uniform buffer
+    // and copies Contents into that backing.
+    // The amount of data is determined by the Layout
+    void PushToGPUBacking(const void* Contents);
+    
+    bool HasShadow();
+};
+
+typedef FMetalSuballocatedUniformBuffer FMetalUniformBuffer;
