#UE4 Add #pragma once

diff --git a/Engine/Source/Runtime/Apple/MetalRHI/Public/MetalFrameAllocator.h b/Engine/Source/Runtime/Apple/MetalRHI/Public/MetalFrameAllocator.h
--- a/Engine/Source/Runtime/Apple/MetalRHI/Public/MetalFrameAllocator.h
+++ b/Engine/Source/Runtime/Apple/MetalRHI/Public/MetalFrameAllocator.h
@@ -1,5 +1,7 @@
 // Copyright Epic Games, Inc. All Rights Reserved.
 
+#pragma once
+
 #import <Metal/Metal.h>
 #import <os/lock.h>
 
diff --git a/Engine/Source/Runtime/Apple/MetalRHI/Public/MetalUniformBuffer.h b/Engine/Source/Runtime/Apple/MetalRHI/Public/MetalUniformBuffer.h
--- a/Engine/Source/Runtime/Apple/MetalRHI/Public/MetalUniformBuffer.h
+++ b/Engine/Source/Runtime/Apple/MetalRHI/Public/MetalUniformBuffer.h
@@ -1,5 +1,7 @@
 // Copyright Epic Games, Inc. All Rights Reserved.
 
+#pragma once
+
 //
 //  Implements handles to linearly allocated per-frame constant buffers for shared memory systems.
 //
