Fix up code problems with a file split-up vs change to original file, with MetalRHI::SetrUniformBuffer

diff --git a/Engine/Source/Runtime/Apple/MetalRHI/Private/MetalUniformBufferCommands.cpp b/Engine/Source/Runtime/Apple/MetalRHI/Private/MetalUniformBufferCommands.cpp
--- a/Engine/Source/Runtime/Apple/MetalRHI/Private/MetalUniformBufferCommands.cpp
+++ b/Engine/Source/Runtime/Apple/MetalRHI/Private/MetalUniformBufferCommands.cpp
@@ -73,37 +73,72 @@ static void SetUniformBufferInternal(FMetalContext* Context, RHIShaderType* Shad
     }
 }
 
-void FMetalRHICommandContext::RHISetShaderUniformBuffer(FRHIVertexShader* VertexShaderRHI, uint32 BufferIndex, FRHIUniformBuffer* BufferRHI)
+void FMetalRHICommandContext::RHISetShaderUniformBuffer(FRHIGraphicsShader* ShaderRHI, uint32 BufferIndex, FRHIUniformBuffer* BufferRHI)
 {
-    SetUniformBufferInternal<EMetalShaderStages::Vertex, FRHIVertexShader>(Context, VertexShaderRHI, BufferIndex, BufferRHI);
-}
-
-void FMetalRHICommandContext::RHISetShaderUniformBuffer(FRHIHullShader* HullShaderRHI, uint32 BufferIndex, FRHIUniformBuffer* BufferRHI)
-{
-#if PLATFORM_SUPPORTS_TESSELLATION_SHADERS
-    SetUniformBufferInternal<EMetalShaderStages::Hull, FRHIHullShader>(Context, HullShaderRHI, BufferIndex, BufferRHI);
-#endif
-}
-
-void FMetalRHICommandContext::RHISetShaderUniformBuffer(FRHIDomainShader* DomainShaderRHI, uint32 BufferIndex, FRHIUniformBuffer* BufferRHI)
-{
-#if PLATFORM_SUPPORTS_TESSELLATION_SHADERS
-    SetUniformBufferInternal<EMetalShaderStages::Domain, FRHIDomainShader>(Context, DomainShaderRHI, BufferIndex, BufferRHI);
-#endif
-}
+	@autoreleasepool{
+		EMetalShaderStages Stage = EMetalShaderStages::Num;
+		FMetalShaderBindings* Bindings = nullptr;
+		switch (ShaderRHI->GetFrequency())
+		{
+		case SF_Vertex:
+		{
+			FMetalVertexShader* VertexShader = ResourceCast(static_cast<FRHIVertexShader*>(ShaderRHI));
+			Bindings = &VertexShader->Bindings;
+			Stage = EMetalShaderStages::Vertex;
+		}
+			break;
+	#if PLATFORM_SUPPORTS_TESSELLATION_SHADERS
+		case SF_Hull:
+		{
+			Stage = EMetalShaderStages::Hull;
+			FMetalHullShader* HullShader = ResourceCast(static_cast<FRHIHullShader*>(ShaderRHI));
+			Bindings = &HullShader->Bindings;
+		}
+			break;
+		case SF_Domain:
+		{
+			Stage = EMetalShaderStages::Domain;
+			FMetalDomainShader* DomainShader = ResourceCast(static_cast<FRHIDomainShader*>(ShaderRHI));
+			Bindings = &DomainShader->Bindings;
+		}
+			break;
+	#endif
+		case SF_Pixel:
+		{
+			Stage = EMetalShaderStages::Pixel;
+			FMetalPixelShader* PixelShader = ResourceCast(static_cast<FRHIPixelShader*>(ShaderRHI));
+			Bindings = &PixelShader->Bindings;
+		}
+			break;
+		default:
+			checkf(0, TEXT("FRHIShader Type %d is invalid or unsupported!"), (int32)ShaderRHI->GetFrequency());
+			NOT_SUPPORTED("RHIShaderStage");
+			break;
+		}
 
-void FMetalRHICommandContext::RHISetShaderUniformBuffer(FRHIGeometryShader* GeometryShader, uint32 BufferIndex, FRHIUniformBuffer* BufferRHI)
-{
-    NOT_SUPPORTED("RHISetShaderUniformBuffer-Geometry");
-}
+		Context->GetCurrentState().BindUniformBuffer(Stage, BufferIndex, BufferRHI);
 
-void FMetalRHICommandContext::RHISetShaderUniformBuffer(FRHIPixelShader* PixelShaderRHI, uint32 BufferIndex, FRHIUniformBuffer* BufferRHI)
-{
-    SetUniformBufferInternal<EMetalShaderStages::Pixel, FRHIPixelShader>(Context, PixelShaderRHI, BufferIndex, BufferRHI);
+		check(BufferIndex < Bindings->NumUniformBuffers);
+		if ((Bindings->ConstantBuffers) & 1 << BufferIndex)
+		{
+			auto* UB = (FMetalUniformBuffer*)BufferRHI;
+			Context->GetCurrentState().SetShaderBuffer(Stage, UB->Buffer, UB->Data, 0, UB->GetSize(), BufferIndex, mtlpp::ResourceUsage::Read);
+		}
+	}
 }
 
 void FMetalRHICommandContext::RHISetShaderUniformBuffer(FRHIComputeShader* ComputeShaderRHI, uint32 BufferIndex, FRHIUniformBuffer* BufferRHI)
 {
-    SetUniformBufferInternal<EMetalShaderStages::Compute, FRHIComputeShader>(Context, ComputeShaderRHI, BufferIndex, BufferRHI);
-}
+	@autoreleasepool{
+		FMetalComputeShader * ComputeShader = ResourceCast(ComputeShaderRHI);
+		Context->GetCurrentState().BindUniformBuffer(EMetalShaderStages::Compute, BufferIndex, BufferRHI);
 
+		auto& Bindings = ComputeShader->Bindings;
+		check(BufferIndex < Bindings.NumUniformBuffers);
+		if ((Bindings.ConstantBuffers) & 1 << BufferIndex)
+		{
+			auto* UB = (FMetalUniformBuffer*)BufferRHI;
+			Context->GetCurrentState().SetShaderBuffer(EMetalShaderStages::Compute, UB->Buffer, UB->Data, 0, UB->GetSize(), BufferIndex, mtlpp::ResourceUsage::Read);
+		}
+	}
+}
