Fix MTLSamplerDescriptor compareFunction assertion on startup on older iOS devices

diff --git a/Engine/Source/Runtime/Apple/MetalRHI/Private/MetalState.cpp b/Engine/Source/Runtime/Apple/MetalRHI/Private/MetalState.cpp
--- a/Engine/Source/Runtime/Apple/MetalRHI/Private/MetalState.cpp
+++ b/Engine/Source/Runtime/Apple/MetalRHI/Private/MetalState.cpp
@@ -349,7 +349,13 @@ static FMetalSampler FindOrCreateSamplerState(mtlpp::Device Device, const FSampl
 		Desc.SetRAddressMode(TranslateWrapMode(Initializer.AddressW));
 		Desc.SetLodMinClamp(Initializer.MinMipLevel);
 		Desc.SetLodMaxClamp(Initializer.MaxMipLevel);
+#if PLATFORM_TVOS
+		Desc.SetCompareFunction(mtlpp::CompareFunction::Never);	
+#elif PLATFORM_IOS
+		Desc.SetCompareFunction(Device.SupportsFeatureSet(mtlpp::FeatureSet::iOS_GPUFamily3_v1) ? TranslateSamplerCompareFunction(Initializer.SamplerComparisonFunction) : mtlpp::CompareFunction::Never);
+#else
 		Desc.SetCompareFunction(TranslateSamplerCompareFunction(Initializer.SamplerComparisonFunction));
+#endif
 #if PLATFORM_MAC
 		Desc.SetBorderColor(Initializer.BorderColor == 0 ? mtlpp::SamplerBorderColor::TransparentBlack : mtlpp::SamplerBorderColor::OpaqueWhite);
 #endif
