- More compile fixes
#rnx

diff --git a/Engine/Source/Runtime/Apple/MetalRHI/Private/MetalUniformBufferCommands.cpp b/Engine/Source/Runtime/Apple/MetalRHI/Private/MetalUniformBufferCommands.cpp
--- a/Engine/Source/Runtime/Apple/MetalRHI/Private/MetalUniformBufferCommands.cpp
+++ b/Engine/Source/Runtime/Apple/MetalRHI/Private/MetalUniformBufferCommands.cpp
@@ -58,7 +58,7 @@ static void SetUniformBufferInternal(FMetalContext* Context, RHIShaderType* Shad
 {
     @autoreleasepool
     {
-        auto Shader = ResourceCast(ShaderRHI);
+        auto Shader = ResourceCast(static_cast<RHIShaderType*>(ShaderRHI));
         Context->GetCurrentState().BindUniformBuffer(Stage, BufferIndex, UBRHI);
         
         auto& Bindings = Shader->Bindings;
@@ -76,69 +76,34 @@ static void SetUniformBufferInternal(FMetalContext* Context, RHIShaderType* Shad
 void FMetalRHICommandContext::RHISetShaderUniformBuffer(FRHIGraphicsShader* ShaderRHI, uint32 BufferIndex, FRHIUniformBuffer* BufferRHI)
 {
 	@autoreleasepool{
+
 		EMetalShaderStages Stage = EMetalShaderStages::Num;
 		FMetalShaderBindings* Bindings = nullptr;
 		switch (ShaderRHI->GetFrequency())
 		{
 		case SF_Vertex:
-		{
-			FMetalVertexShader* VertexShader = ResourceCast(static_cast<FRHIVertexShader*>(ShaderRHI));
-			Bindings = &VertexShader->Bindings;
-			Stage = EMetalShaderStages::Vertex;
-		}
+			SetUniformBufferInternal<EMetalShaderStages::Vertex, FRHIVertexShader>(Context, static_cast<FRHIVertexShader*>(ShaderRHI), BufferIndex, BufferRHI);
 			break;
 	#if PLATFORM_SUPPORTS_TESSELLATION_SHADERS
 		case SF_Hull:
-		{
-			Stage = EMetalShaderStages::Hull;
-			FMetalHullShader* HullShader = ResourceCast(static_cast<FRHIHullShader*>(ShaderRHI));
-			Bindings = &HullShader->Bindings;
-		}
+			SetUniformBufferInternal<EMetalShaderStages::Hull, FRHIHullShader>(Context, static_cast<FRHIHullShader*>(ShaderRHI), BufferIndex, BufferRHI);
 			break;
 		case SF_Domain:
-		{
-			Stage = EMetalShaderStages::Domain;
-			FMetalDomainShader* DomainShader = ResourceCast(static_cast<FRHIDomainShader*>(ShaderRHI));
-			Bindings = &DomainShader->Bindings;
-		}
+			SetUniformBufferInternal<EMetalShaderStages::Domain, FRHIDomainShader>(Context, static_cast<FRHIDomainShader*>(ShaderRHI), BufferIndex, BufferRHI);
 			break;
 	#endif
 		case SF_Pixel:
-		{
-			Stage = EMetalShaderStages::Pixel;
-			FMetalPixelShader* PixelShader = ResourceCast(static_cast<FRHIPixelShader*>(ShaderRHI));
-			Bindings = &PixelShader->Bindings;
-		}
+			SetUniformBufferInternal<EMetalShaderStages::Pixel, FRHIDomainShader>(Context, static_cast<FRHIDomainShader*>(ShaderRHI), BufferIndex, BufferRHI);
 			break;
 		default:
 			checkf(0, TEXT("FRHIShader Type %d is invalid or unsupported!"), (int32)ShaderRHI->GetFrequency());
 			NOT_SUPPORTED("RHIShaderStage");
 			break;
 		}
-
-		Context->GetCurrentState().BindUniformBuffer(Stage, BufferIndex, BufferRHI);
-
-		check(BufferIndex < Bindings->NumUniformBuffers);
-		if ((Bindings->ConstantBuffers) & 1 << BufferIndex)
-		{
-			auto* UB = (FMetalUniformBuffer*)BufferRHI;
-			Context->GetCurrentState().SetShaderBuffer(Stage, UB->Buffer, UB->Data, 0, UB->GetSize(), BufferIndex, mtlpp::ResourceUsage::Read);
-		}
 	}
 }
 
 void FMetalRHICommandContext::RHISetShaderUniformBuffer(FRHIComputeShader* ComputeShaderRHI, uint32 BufferIndex, FRHIUniformBuffer* BufferRHI)
 {
-	@autoreleasepool{
-		FMetalComputeShader * ComputeShader = ResourceCast(ComputeShaderRHI);
-		Context->GetCurrentState().BindUniformBuffer(EMetalShaderStages::Compute, BufferIndex, BufferRHI);
-
-		auto& Bindings = ComputeShader->Bindings;
-		check(BufferIndex < Bindings.NumUniformBuffers);
-		if ((Bindings.ConstantBuffers) & 1 << BufferIndex)
-		{
-			auto* UB = (FMetalUniformBuffer*)BufferRHI;
-			Context->GetCurrentState().SetShaderBuffer(EMetalShaderStages::Compute, UB->Buffer, UB->Data, 0, UB->GetSize(), BufferIndex, mtlpp::ResourceUsage::Read);
-		}
-	}
+	SetUniformBufferInternal<EMetalShaderStages::Compute, FRHIComputeShader>(Context, ComputeShaderRHI, BufferIndex, BufferRHI);
 }
