Fix for crash using when native Metal libs and packing with r.Shaders.Optimize=0.  Compressed text source doesn't exist when using native libs, however if compiled for debug using metal libs then source should be available for debugging anyway.

diff --git a/Engine/Source/Runtime/Apple/MetalRHI/Private/MetalShaders.cpp b/Engine/Source/Runtime/Apple/MetalRHI/Private/MetalShaders.cpp
--- a/Engine/Source/Runtime/Apple/MetalRHI/Private/MetalShaders.cpp
+++ b/Engine/Source/Runtime/Apple/MetalRHI/Private/MetalShaders.cpp
@@ -552,7 +552,7 @@ void TMetalBaseShader<BaseResourceType, ShaderType>::Init(TArrayView<const uint8
 			GlslCodeNSString = [FMetalShaderDebugCache::Get().GetShaderCode(SourceLen, SourceCRC).GetPtr() retain];
 		}
 #endif
-        if (bForceTextShaders)
+        if (bForceTextShaders && CodeSize && CompressedSource.Num())
         {
             bHasShaderSource = (GetSourceCode() != nil);
         }
