Merge CL 11258451 from //UE4/Dev-Rendering/...

---

Enable Metal RHI runtime virtual texture support.

diff --git a/Engine/Source/Runtime/Apple/MetalRHI/Private/MetalTexture.cpp b/Engine/Source/Runtime/Apple/MetalRHI/Private/MetalTexture.cpp
--- a/Engine/Source/Runtime/Apple/MetalRHI/Private/MetalTexture.cpp
+++ b/Engine/Source/Runtime/Apple/MetalRHI/Private/MetalTexture.cpp
@@ -2688,7 +2688,18 @@ void FMetalDynamicRHI::RHICopySubTextureRegion(FRHITexture2D* SourceTexture, FRH
 	}
 }
 
-
+inline bool MetalRHICopyTexutre_IsTextureFormatCompatible(EPixelFormat SrcFmt, EPixelFormat DstFmt)
+{
+	//
+	// For now, we only support copies between textures of mismatching
+	// formats if they are of size-compatible internal formats.  This allows us
+	// to copy from uncompressed to compressed textures, specifically in support
+	// of the runtime virtual texture system.  Note that copies of compatible
+	// formats incur the cost of an extra copy, as we must copy from the source
+	// texture to a temporary buffer and finally to the destination texture.
+	//
+	return ((SrcFmt == DstFmt) || (GPixelFormats[SrcFmt].BlockBytes == GPixelFormats[DstFmt].BlockBytes));
+}
 
 void FMetalRHICommandContext::RHICopyTexture(FRHITexture* SourceTextureRHI, FRHITexture* DestTextureRHI, const FRHICopyTextureInfo& CopyInfo)
 {
@@ -2707,26 +2718,33 @@ void FMetalRHICommandContext::RHICopyTexture(FRHITexture* SourceTextureRHI, FRHI
 		FMetalSurface* MetalSrcTexture = GetMetalSurfaceFromRHITexture(SourceTextureRHI);
 		FMetalSurface* MetalDestTexture = GetMetalSurfaceFromRHITexture(DestTextureRHI);
 		
-		if(SourceTextureRHI->GetFormat() == DestTextureRHI->GetFormat())
+		const bool TextureFormatExactMatch = (SourceTextureRHI->GetFormat() == DestTextureRHI->GetFormat());
+		const bool TextureFormatCompatible = MetalRHICopyTexutre_IsTextureFormatCompatible(SourceTextureRHI->GetFormat(), DestTextureRHI->GetFormat());
+		
+		if (TextureFormatExactMatch || TextureFormatCompatible)
 		{
 			FIntVector Size = (CopyInfo.Size != FIntVector::ZeroValue) ? CopyInfo.Size : FIntVector(MetalSrcTexture->SizeX, MetalSrcTexture->SizeY, MetalSrcTexture->SizeZ);
 			
 			mtlpp::Origin SourceOrigin(CopyInfo.SourcePosition.X, CopyInfo.SourcePosition.Y, CopyInfo.SourcePosition.Z);
 			mtlpp::Origin DestinationOrigin(CopyInfo.DestPosition.X, CopyInfo.DestPosition.Y, CopyInfo.DestPosition.Z);
-
+			
 			FMetalTexture SrcTexture;
-			mtlpp::TextureUsage Usage = MetalSrcTexture->Texture.GetUsage();
-			if(Usage & mtlpp::TextureUsage::PixelFormatView)
+
+			if (TextureFormatExactMatch)
 			{
-				ns::Range Slices(0, MetalSrcTexture->Texture.GetArrayLength() * (MetalSrcTexture->bIsCubemap ? 6 : 1));
-				if(MetalSrcTexture->Texture.GetPixelFormat() != MetalDestTexture->Texture.GetPixelFormat())
+				mtlpp::TextureUsage Usage = MetalSrcTexture->Texture.GetUsage();
+				if (Usage & mtlpp::TextureUsage::PixelFormatView)
 				{
-					SrcTexture = MetalSrcTexture->Texture.NewTextureView(MetalDestTexture->Texture.GetPixelFormat(), MetalSrcTexture->Texture.GetTextureType(), ns::Range(0, MetalSrcTexture->Texture.GetMipmapLevelCount()), Slices);
+					ns::Range Slices(0, MetalSrcTexture->Texture.GetArrayLength() * (MetalSrcTexture->bIsCubemap ? 6 : 1));
+					if (MetalSrcTexture->Texture.GetPixelFormat() != MetalDestTexture->Texture.GetPixelFormat())
+					{
+						SrcTexture = MetalSrcTexture->Texture.NewTextureView(MetalDestTexture->Texture.GetPixelFormat(), MetalSrcTexture->Texture.GetTextureType(), ns::Range(0, MetalSrcTexture->Texture.GetMipmapLevelCount()), Slices);
+					}
+				}
+				if (!SrcTexture)
+				{
+					SrcTexture = MetalSrcTexture->Texture;
 				}
-			}
-			if (!SrcTexture)
-			{
-				SrcTexture = MetalSrcTexture->Texture;
 			}
 			
 			for (uint32 SliceIndex = 0; SliceIndex < CopyInfo.NumSlices; ++SliceIndex)
@@ -2739,19 +2757,33 @@ void FMetalRHICommandContext::RHICopyTexture(FRHITexture* SourceTextureRHI, FRHI
 					uint32 SourceMipIndex = CopyInfo.SourceMipIndex + MipIndex;
 					uint32 DestMipIndex = CopyInfo.DestMipIndex + MipIndex;
 					mtlpp::Size SourceSize(FMath::Max(Size.X >> MipIndex, 1), FMath::Max(Size.Y >> MipIndex, 1), FMath::Max(Size.Z >> MipIndex, 1));
-					
+					mtlpp::Size DestSize = SourceSize;
+
+					if (TextureFormatCompatible)
+					{
+						DestSize.width  *= GPixelFormats[MetalDestTexture->PixelFormat].BlockSizeX;
+						DestSize.height *= GPixelFormats[MetalDestTexture->PixelFormat].BlockSizeY;
+					}
+
 					// Account for create with TexCreate_SRGB flag which could make these different
-					if(SrcTexture.GetPixelFormat() == MetalDestTexture->Texture.GetPixelFormat())
+					if (TextureFormatExactMatch && (SrcTexture.GetPixelFormat() == MetalDestTexture->Texture.GetPixelFormat()))
 					{
 						GetInternalContext().CopyFromTextureToTexture(SrcTexture, SourceSliceIndex, SourceMipIndex, SourceOrigin,SourceSize,MetalDestTexture->Texture, DestSliceIndex, DestMipIndex, DestinationOrigin);
 					}
 					else
 					{
-						// Linear and sRGB mismatch then try to go via metal buffer
-						// Modified clone of logic from MetalRenderTarget.cpp
-						uint32 BytesPerPixel = (MetalSrcTexture->PixelFormat != PF_DepthStencil) ? GPixelFormats[MetalSrcTexture->PixelFormat].BlockBytes : 1;
+						//
+						// In the case of compatible texture formats or pixel
+						// format mismatch (like linear vs. sRGB), then we must
+						// achieve the copy by going through a buffer object.
+						//
+						const uint32 BytesPerPixel = (MetalSrcTexture->PixelFormat != PF_DepthStencil) ? GPixelFormats[MetalSrcTexture->PixelFormat].BlockBytes : 1;
 						const uint32 Stride = BytesPerPixel * SourceSize.width;
-						const uint32 Alignment = PLATFORM_MAC ? 1u : 64u;
+#if PLATFORM_MAC
+						const uint32 Alignment = 1u;
+#else
+						const uint32 Alignment = 64u;
+#endif
 						const uint32 AlignedStride = ((Stride - 1) & ~(Alignment - 1)) + Alignment;
 						const uint32 BytesPerImage = AlignedStride *  SourceSize.height;
 						const uint32 DataSize = BytesPerImage * SourceSize.depth;
@@ -2768,14 +2800,14 @@ void FMetalRHICommandContext::RHICopyTexture(FRHITexture* SourceTextureRHI, FRHI
 						}
 #endif
 						GetInternalContext().CopyFromTextureToBuffer(MetalSrcTexture->Texture, SourceSliceIndex, SourceMipIndex, SourceOrigin, SourceSize, Buffer, 0, AlignedStride, BytesPerImage, Options);
-						GetInternalContext().CopyFromBufferToTexture(Buffer, 0, Stride, BytesPerImage, SourceSize, MetalDestTexture->Texture, DestSliceIndex, DestMipIndex, DestinationOrigin, Options);
+						GetInternalContext().CopyFromBufferToTexture(Buffer, 0, Stride, BytesPerImage, DestSize, MetalDestTexture->Texture, DestSliceIndex, DestMipIndex, DestinationOrigin, Options);
 						
 						GetMetalDeviceContext().ReleaseBuffer(Buffer);
 					}
 				}
 			}
 			
-			if (SrcTexture != MetalSrcTexture->Texture)
+			if (SrcTexture && (SrcTexture != MetalSrcTexture->Texture))
 			{
 				SafeReleaseMetalTexture(SrcTexture);
 			}
diff --git a/Engine/Source/Runtime/Core/Public/Mac/MacPlatform.h b/Engine/Source/Runtime/Core/Public/Mac/MacPlatform.h
--- a/Engine/Source/Runtime/Core/Public/Mac/MacPlatform.h
+++ b/Engine/Source/Runtime/Core/Public/Mac/MacPlatform.h
@@ -67,6 +67,7 @@ typedef FMacPlatformTypes FPlatformTypes;
 #define PLATFORM_SUPPORTS_STACK_SYMBOLS				1
 #define PLATFORM_HAS_BSD_SOCKET_FEATURE_MSG_DONTWAIT	1
 #define PLATFORM_IS_ANSI_MALLOC_THREADSAFE			1
+#define PLATFORM_SUPPORTS_VIRTUAL_TEXTURE_STREAMING		1
 
 #define PLATFORM_RHITHREAD_DEFAULT_BYPASS			WITH_EDITOR
 
diff --git a/Engine/Source/Runtime/Core/Public/Mac/MacPlatformProperties.h b/Engine/Source/Runtime/Core/Public/Mac/MacPlatformProperties.h
--- a/Engine/Source/Runtime/Core/Public/Mac/MacPlatformProperties.h
+++ b/Engine/Source/Runtime/Core/Public/Mac/MacPlatformProperties.h
@@ -132,6 +132,11 @@ struct FMacPlatformProperties
 	{
 		return !IsServerOnly();
 	}
+
+	static FORCEINLINE bool SupportsVirtualTextureStreaming()
+	{
+		return true;
+	}
 };
 
 #ifdef PROPERTY_HEADER_SHOULD_DEFINE_TYPE
