Enable GRHIRequiresRenderTargetForPixelShaderUAVs for Metal RHI.

This fixes a Metal API Validation error that was triggering when opening a new material for edit.

#rnx

diff --git a/Engine/Source/Runtime/Apple/MetalRHI/Private/MetalRHI.cpp b/Engine/Source/Runtime/Apple/MetalRHI/Private/MetalRHI.cpp
--- a/Engine/Source/Runtime/Apple/MetalRHI/Private/MetalRHI.cpp
+++ b/Engine/Source/Runtime/Apple/MetalRHI/Private/MetalRHI.cpp
@@ -115,7 +115,7 @@ FMetalDynamicRHI::FMetalDynamicRHI(ERHIFeatureLevel::Type RequestedFeatureLevel)
 	GSupportsVolumeTextureRendering = false;
 	
 	// Metal always needs a render target to render with fragment shaders!
-	// GRHIRequiresRenderTargetForPixelShaderUAVs = true;
+	GRHIRequiresRenderTargetForPixelShaderUAVs = true;
 
 	//@todo-rco: Query name from API
 	GRHIAdapterName = TEXT("Metal");
