macOS: Make MetalRHI use the FRHIGPUFence when mapping the staging buffer.  Use that and texture type to inform if we need to apply a texture sync for CPU access.  If managed and we have requested a GPU readback then mark the Metal Surface for auto Texture Sync on copying to the resolve target.  This stops us from needing to apply a blit encoder sync and command buffer wait at the point of the requested readback.

diff --git a/Engine/Source/Runtime/Apple/MetalRHI/Private/MetalRenderTarget.cpp b/Engine/Source/Runtime/Apple/MetalRHI/Private/MetalRenderTarget.cpp
--- a/Engine/Source/Runtime/Apple/MetalRHI/Private/MetalRenderTarget.cpp
+++ b/Engine/Source/Runtime/Apple/MetalRHI/Private/MetalRenderTarget.cpp
@@ -116,6 +116,14 @@ void FMetalRHICommandContext::RHICopyToResolveTarget(FRHITexture* SourceTextureR
 		{
 			Context->CopyFromTextureToTexture(Source->Texture, SrcIndex, ResolveParams.MipIndex, Origin, Size, Destination->Texture, DestIndex, ResolveParams.MipIndex, Origin);
 		}
+
+#if PLATFORM_MAC
+		if((Destination->GPUReadback & FMetalSurface::EMetalGPUReadbackFlags::ReadbackRequested) != 0)
+		{
+			Context->GetCurrentRenderPass().SynchronizeTexture(Destination->Texture, DestIndex, ResolveParams.MipIndex);
+		}
+#endif
+		
 	}
 	}
 }
@@ -349,10 +357,19 @@ void FMetalDynamicRHI::RHIMapStagingSurface(FRHITexture* TextureRHI, FRHIGPUFenc
     FMetalSurface* Surface = GetMetalSurfaceFromRHITexture(TextureRHI);
     FMetalTexture2D* Texture = (FMetalTexture2D*)TextureRHI->GetTexture2D();
     
+#if PLATFORM_MAC
+	uint16 FencePoll = FenceRHI && FenceRHI->Poll() ? 1 : 0;
+    Surface->GPUReadback |= FencePoll << FMetalSurface::EMetalGPUReadbackFlags::ReadbackFenceCompleteShift;
+#endif
+    
     uint32 Stride = 0;
     OutWidth = Texture->GetSizeX();
     OutHeight = Texture->GetSizeY();
     OutData = Surface->Lock(0, 0, RLM_ReadOnly, Stride);
+    
+#if PLATFORM_MAC
+	Surface->GPUReadback = (Surface->Texture.GetPtr() && Surface->Texture.GetStorageMode() == mtlpp::StorageMode::Managed) ? FMetalSurface::EMetalGPUReadbackFlags::ReadbackRequested : 0;
+#endif
 	}
 }
 
diff --git a/Engine/Source/Runtime/Apple/MetalRHI/Private/MetalTexture.cpp b/Engine/Source/Runtime/Apple/MetalRHI/Private/MetalTexture.cpp
--- a/Engine/Source/Runtime/Apple/MetalRHI/Private/MetalTexture.cpp
+++ b/Engine/Source/Runtime/Apple/MetalRHI/Private/MetalTexture.cpp
@@ -1361,10 +1361,12 @@ void* FMetalSurface::Lock(uint32 MipIndex, uint32 ArrayIndex, EResourceLockMode
 			else
 			{
 #if PLATFORM_MAC
-				GetMetalDeviceContext().SynchronizeTexture(Texture, ArrayIndex, MipIndex);
-				
-				//kick the current command buffer.
-				GetMetalDeviceContext().SubmitCommandBufferAndWait();
+				if((GPUReadback & EMetalGPUReadbackFlags::ReadbackRequestedAndComplete) != EMetalGPUReadbackFlags::ReadbackRequestedAndComplete)
+				{
+					// A previous texture sync has not been done, need the data now, request texture sync and kick the current command buffer.
+					GetMetalDeviceContext().SynchronizeTexture(Texture, ArrayIndex, MipIndex);
+					GetMetalDeviceContext().SubmitCommandBufferAndWait();
+				}
 #endif
 				
 				// This block breaks the texture atlas system in Ocean, which depends on nonzero strides coming back from compressed textures. Turning off.
diff --git a/Engine/Source/Runtime/Apple/MetalRHI/Public/MetalResources.h b/Engine/Source/Runtime/Apple/MetalRHI/Public/MetalResources.h
--- a/Engine/Source/Runtime/Apple/MetalRHI/Public/MetalResources.h
+++ b/Engine/Source/Runtime/Apple/MetalRHI/Public/MetalResources.h
@@ -553,7 +553,16 @@ public:
 	FMetalTexture MSAAResolveTexture;
 	uint32 SizeX, SizeY, SizeZ;
 	bool bIsCubemap;
-	int32 volatile Written;
+	int16 volatile Written;
+	int16 GPUReadback = 0;
+	enum EMetalGPUReadbackFlags : int16
+	{
+		ReadbackRequestedShift 			= 0,
+		ReadbackFenceCompleteShift  	= 1,
+		ReadbackRequested 				= 1 << ReadbackRequestedShift,
+		ReadbackFenceComplete 			= 1 << ReadbackFenceCompleteShift,
+		ReadbackRequestedAndComplete 	= ReadbackRequested | ReadbackFenceComplete
+	};
 	
 	uint32 Flags;
 	// one per mip
