Fixing bug where Metal LLM tags where registering over AppleLLM tags due to bad initialization code.
#tests iOS Client

diff --git a/Engine/Source/Runtime/Apple/MetalRHI/Private/MetalLLM.cpp b/Engine/Source/Runtime/Apple/MetalRHI/Private/MetalLLM.cpp
--- a/Engine/Source/Runtime/Apple/MetalRHI/Private/MetalLLM.cpp
+++ b/Engine/Source/Runtime/Apple/MetalRHI/Private/MetalLLM.cpp
@@ -40,7 +40,7 @@ void MetalLLM::Initialise()
 
 	for (int32 Index = 0; Index < TagCount; ++Index)
 	{
-		int32 Tag = (int32)ELLMTag::PlatformTagStart + Index;
+		int32 Tag = (int32)ELLMTagApple::AppleMetalTagsStart + Index;
 		const FLLMTagInfoMetal& TagInfo = ELLMTagNamesMetal[Index];
 
 		FLowLevelMemTracker::Get().RegisterPlatformTag(Tag, TagInfo.Name, TagInfo.StatName, TagInfo.SummaryStatName);
