Move MetalStatistics header into Restricted/NotForLicensees.

diff --git a/Engine/Source/Runtime/Apple/MetalRHI/MetalRHI.Build.cs b/Engine/Source/Runtime/Apple/MetalRHI/MetalRHI.Build.cs
--- a/Engine/Source/Runtime/Apple/MetalRHI/MetalRHI.Build.cs
+++ b/Engine/Source/Runtime/Apple/MetalRHI/MetalRHI.Build.cs
@@ -1,5 +1,6 @@
 // Copyright Epic Games, Inc. All Rights Reserved.
 
+using System.IO;
 using UnrealBuildTool;
 
 public class MetalRHI : ModuleRules
@@ -27,12 +28,11 @@ public class MetalRHI : ModuleRules
 			PublicFrameworks.Add("QuartzCore");
 		}
 
-		var StatsModule = "../Plugins/NotForLicensees/MetalStatistics/MetalStatistics.uplugin";
-		bool bMetalStats = System.IO.File.Exists(StatsModule);
+		string StatsPlugin = Path.Combine(EngineDirectory, "Restricted/NotForLicensees/Plugins/MetalStatistics/MetalStatistics.uplugin");
+		bool bMetalStats = File.Exists(StatsPlugin);
 		if ( bMetalStats && Target.Configuration != UnrealTargetConfiguration.Shipping )
 		{
 			PublicDefinitions.Add("METAL_STATISTICS=1");
-			WhitelistRestrictedFolders.Add("Public/NotForLicensees");
 		}
 	}
 }
diff --git a/Engine/Source/Runtime/Apple/MetalRHI/Private/MetalProfiler.h b/Engine/Source/Runtime/Apple/MetalRHI/Private/MetalProfiler.h
--- a/Engine/Source/Runtime/Apple/MetalRHI/Private/MetalProfiler.h
+++ b/Engine/Source/Runtime/Apple/MetalRHI/Private/MetalProfiler.h
@@ -7,7 +7,7 @@
 #include "GPUProfiler.h"
 
 #if METAL_STATISTICS
-#include "NotForLicensees/MetalStatistics.h"
+#include "../../../../../Restricted/NotForLicensees/Plugins/MetalStatistics/Source/MetalStatistics/Public/MetalStatistics.h"
 #endif
 
 // Stats
