Fixes failed check() during landscape streaming

diff --git a/Engine/Source/Runtime/Apple/MetalRHI/Private/MetalBuffer.cpp b/Engine/Source/Runtime/Apple/MetalRHI/Private/MetalBuffer.cpp
--- a/Engine/Source/Runtime/Apple/MetalRHI/Private/MetalBuffer.cpp
+++ b/Engine/Source/Runtime/Apple/MetalRHI/Private/MetalBuffer.cpp
@@ -156,11 +156,11 @@ void FMetalBuffer::Release()
 	}
 }
 
-void FMetalBuffer::SetOwner(class FMetalRHIBuffer* Owner)
+void FMetalBuffer::SetOwner(class FMetalRHIBuffer* Owner, bool bIsSwap)
 {
 	if (Heap)
 	{
-		Heap->SetOwner(ns::Range(GetOffset(), GetLength()), Owner);
+		Heap->SetOwner(ns::Range(GetOffset(), GetLength()), Owner, bIsSwap);
 	}
 }
 
@@ -219,7 +219,7 @@ FMetalSubBufferHeap::~FMetalSubBufferHeap()
 	}
 }
 
-void FMetalSubBufferHeap::SetOwner(ns::Range const& Range, FMetalRHIBuffer* Owner)
+void FMetalSubBufferHeap::SetOwner(ns::Range const& Range, FMetalRHIBuffer* Owner, bool bIsSwap)
 {
 	FScopeLock Lock(&PoolMutex);
 	for (uint32 i = 0; i < AllocRanges.Num(); i++)
@@ -227,7 +227,7 @@ void FMetalSubBufferHeap::SetOwner(ns::Range const& Range, FMetalRHIBuffer* Owne
 		if (AllocRanges[i].Range.Location == Range.Location)
 		{
 			check(AllocRanges[i].Range.Length == Range.Length);
-			check(AllocRanges[i].Owner == nullptr || Owner == nullptr);
+			check(AllocRanges[i].Owner == nullptr || Owner == nullptr || bIsSwap);
 			AllocRanges[i].Owner = Owner;
 			break;
 		}
diff --git a/Engine/Source/Runtime/Apple/MetalRHI/Private/MetalBuffer.h b/Engine/Source/Runtime/Apple/MetalRHI/Private/MetalBuffer.h
--- a/Engine/Source/Runtime/Apple/MetalRHI/Private/MetalBuffer.h
+++ b/Engine/Source/Runtime/Apple/MetalRHI/Private/MetalBuffer.h
@@ -53,7 +53,7 @@ public:
     mtlpp::PurgeableState SetPurgeableState(mtlpp::PurgeableState state);
 	void FreeRange(ns::Range const& Range);
 
-    void SetOwner(ns::Range const& Range, FMetalRHIBuffer* Owner);
+    void SetOwner(ns::Range const& Range, FMetalRHIBuffer* Owner, bool bIsSwap);
 
 private:
     struct Allocation
diff --git a/Engine/Source/Runtime/Apple/MetalRHI/Private/MetalRHIContext.cpp b/Engine/Source/Runtime/Apple/MetalRHI/Private/MetalRHIContext.cpp
--- a/Engine/Source/Runtime/Apple/MetalRHI/Private/MetalRHIContext.cpp
+++ b/Engine/Source/Runtime/Apple/MetalRHI/Private/MetalRHIContext.cpp
@@ -42,7 +42,7 @@ void SafeReleaseMetalBuffer(FMetalBuffer& Buffer)
 {
 	if(GIsMetalInitialized && GDynamicRHI && Buffer)
 	{
-		Buffer.SetOwner(nullptr);
+		Buffer.SetOwner(nullptr, false);
 		FMetalRHICommandContext* Context = static_cast<FMetalRHICommandContext*>(RHIGetDefaultContext());
 		if(Context)
 		{
diff --git a/Engine/Source/Runtime/Apple/MetalRHI/Private/MetalVertexBuffer.cpp b/Engine/Source/Runtime/Apple/MetalRHI/Private/MetalVertexBuffer.cpp
--- a/Engine/Source/Runtime/Apple/MetalRHI/Private/MetalVertexBuffer.cpp
+++ b/Engine/Source/Runtime/Apple/MetalRHI/Private/MetalVertexBuffer.cpp
@@ -95,8 +95,8 @@ void FMetalVertexBuffer::Swap(FMetalVertexBuffer& Other)
 void FMetalRHIBuffer::Swap(FMetalRHIBuffer& Other)
 {
 	::Swap(*this, Other);
-	Other.Buffer.SetOwner(&Other);
-	Buffer.SetOwner(this);
+	Other.Buffer.SetOwner(&Other, true);
+	Buffer.SetOwner(this, true);
 }
 
 static bool CanUsePrivateMemory()
@@ -245,7 +245,7 @@ void FMetalRHIBuffer::Alloc(uint32 InSize, EResourceLockMode LockMode)
 		Buffer = GetMetalDeviceContext().CreatePooledBuffer(Args);
 		METAL_FATAL_ASSERT(Buffer, TEXT("Failed to create buffer of size %u and storage mode %u"), InSize, (uint32)Mode);
 
-        Buffer.SetOwner(this);
+        Buffer.SetOwner(this, false);
 
         METAL_INC_DWORD_STAT_BY(Type, MemAlloc, InSize);
         
@@ -269,7 +269,7 @@ void FMetalRHIBuffer::AllocTransferBuffer(bool bOnRHIThread, uint32 InSize, ERes
 	{
 		FMetalPooledBufferArgs ArgsCPU(GetMetalDeviceContext().GetDevice(), InSize, BUF_Dynamic, mtlpp::StorageMode::Shared);
         CPUBuffer = GetMetalDeviceContext().CreatePooledBuffer(ArgsCPU);
-		CPUBuffer.SetOwner(this);
+		CPUBuffer.SetOwner(this, false);
         check(CPUBuffer && CPUBuffer.GetPtr());
         METAL_INC_DWORD_STAT_BY(Type, MemAlloc, InSize);
         METAL_FATAL_ASSERT(CPUBuffer, TEXT("Failed to create buffer of size %u and storage mode %u"), InSize, (uint32)mtlpp::StorageMode::Shared);
diff --git a/Engine/Source/Runtime/Apple/MetalRHI/Public/MetalResources.h b/Engine/Source/Runtime/Apple/MetalRHI/Public/MetalResources.h
--- a/Engine/Source/Runtime/Apple/MetalRHI/Public/MetalResources.h
+++ b/Engine/Source/Runtime/Apple/MetalRHI/Public/MetalResources.h
@@ -393,7 +393,7 @@ public:
 	inline bool IsPooled() const { return bPooled; }
 	inline bool IsSingleUse() const { return bSingleUse; }
 	inline void MarkSingleUse() { bSingleUse = true; }
-    void SetOwner(class FMetalRHIBuffer* Owner);
+    void SetOwner(class FMetalRHIBuffer* Owner, bool bIsSwap);
 	void Release();
 	
 	friend uint32 GetTypeHash(FMetalBuffer const& Hash)
