Fixes flickering decals in mobile preview.
Implements NextSubpass on desktop (which just inserts a texture barrier)

diff --git a/Engine/Source/Runtime/Apple/MetalRHI/Private/MetalRHIContext.cpp b/Engine/Source/Runtime/Apple/MetalRHI/Private/MetalRHIContext.cpp
--- a/Engine/Source/Runtime/Apple/MetalRHI/Private/MetalRHIContext.cpp
+++ b/Engine/Source/Runtime/Apple/MetalRHI/Private/MetalRHIContext.cpp
@@ -223,6 +223,17 @@ void FMetalRHICommandContext::RHIEndRenderPass()
 	}
 }
 
+void FMetalRHICommandContext::RHINextSubpass()
+{
+#if PLATFORM_MAC
+	if (RenderPassInfo.SubpassHint == ESubpassHint::DepthReadSubpass)
+	{
+		FMetalRenderPass& RP = Context->GetCurrentRenderPass();
+		RP.InsertTextureBarrier();
+	}
+#endif
+}
+
 void FMetalRHICommandContext::RHIBeginComputePass(const TCHAR* InName)
 {
 	RHISetRenderTargets(0, nullptr, nullptr);
diff --git a/Engine/Source/Runtime/Apple/MetalRHI/Private/MetalRenderPass.cpp b/Engine/Source/Runtime/Apple/MetalRHI/Private/MetalRenderPass.cpp
--- a/Engine/Source/Runtime/Apple/MetalRHI/Private/MetalRenderPass.cpp
+++ b/Engine/Source/Runtime/Apple/MetalRHI/Private/MetalRenderPass.cpp
@@ -1282,6 +1282,17 @@ TRefCountPtr<FMetalFence> const& FMetalRenderPass::EndRenderPass(void)
 	return CurrentEncoderFence;
 }
 
+void FMetalRenderPass::InsertTextureBarrier()
+{
+#if PLATFORM_MAC
+	check(CurrentEncoder.IsRenderCommandEncoderActive());
+	
+	id <MTLRenderCommandEncoder> RenderEncoder = CurrentEncoder.GetRenderCommandEncoder().GetPtr();
+	check(RenderEncoder);
+	[RenderEncoder textureBarrier];
+#endif
+}
+
 void FMetalRenderPass::CopyFromTextureToBuffer(FMetalTexture const& Texture, uint32 sourceSlice, uint32 sourceLevel, mtlpp::Origin sourceOrigin, mtlpp::Size sourceSize, FMetalBuffer const& toBuffer, uint32 destinationOffset, uint32 destinationBytesPerRow, uint32 destinationBytesPerImage, mtlpp::BlitOption options)
 {
 	ConditionalSwitchToBlit();
diff --git a/Engine/Source/Runtime/Apple/MetalRHI/Private/MetalRenderPass.h b/Engine/Source/Runtime/Apple/MetalRHI/Private/MetalRenderPass.h
--- a/Engine/Source/Runtime/Apple/MetalRHI/Private/MetalRenderPass.h
+++ b/Engine/Source/Runtime/Apple/MetalRHI/Private/MetalRenderPass.h
@@ -158,6 +158,8 @@ public:
 	 */
     mtlpp::RenderCommandEncoder GetParallelRenderCommandEncoder(uint32 Index, mtlpp::ParallelRenderCommandEncoder& ParallelEncoder);
 	
+	void InsertTextureBarrier();
+	
 private:
 #pragma mark -
     void ConditionalSwitchToRender(void);
diff --git a/Engine/Source/Runtime/Apple/MetalRHI/Public/MetalRHIContext.h b/Engine/Source/Runtime/Apple/MetalRHI/Public/MetalRHIContext.h
--- a/Engine/Source/Runtime/Apple/MetalRHI/Public/MetalRHIContext.h
+++ b/Engine/Source/Runtime/Apple/MetalRHI/Public/MetalRHIContext.h
@@ -241,6 +241,8 @@ public:
 	virtual void RHIBeginComputePass(const TCHAR* InName) final override;
 	
 	virtual void RHIEndComputePass() final override;
+	
+	virtual void RHINextSubpass() final override;
 
 protected:
 	static TGlobalResource<TBoundShaderStateHistory<10000>> BoundShaderStateHistory;
