Add Hull shader hash to FMetalGraphicsPipelineKey to differentiate between Adaptive Tessellation on/off changes the hull function signature.  Put Hull and Domain function hashes in FMetalGraphicsPipelineKey under PLATFORM_SUPPORTS_TESSELLATION_SHADERS like everything else.

diff --git a/Engine/Source/Runtime/Apple/MetalRHI/Private/MetalPipeline.cpp b/Engine/Source/Runtime/Apple/MetalRHI/Private/MetalPipeline.cpp
--- a/Engine/Source/Runtime/Apple/MetalRHI/Private/MetalPipeline.cpp
+++ b/Engine/Source/Runtime/Apple/MetalRHI/Private/MetalPipeline.cpp
@@ -272,7 +272,10 @@ struct FMetalGraphicsPipelineKey
 	FMetalRenderPipelineHash RenderPipelineHash;
 	FMetalHashedVertexDescriptor VertexDescriptorHash;
 	FSHAHash VertexFunction;
+#if PLATFORM_SUPPORTS_TESSELLATION_SHADERS
+	FSHAHash HullFunction;
 	FSHAHash DomainFunction;
+#endif
 	FSHAHash PixelFunction;
 
 	template<typename Type>
@@ -296,7 +299,10 @@ struct FMetalGraphicsPipelineKey
 		return (RenderPipelineHash == Other.RenderPipelineHash
 		&& VertexDescriptorHash == Other.VertexDescriptorHash
 		&& VertexFunction == Other.VertexFunction
+#if PLATFORM_SUPPORTS_TESSELLATION_SHADERS
+		&& HullFunction == Other.HullFunction
 		&& DomainFunction == Other.DomainFunction
+#endif
 		&& PixelFunction == Other.PixelFunction);
 	}
 	
@@ -304,7 +310,10 @@ struct FMetalGraphicsPipelineKey
 	{
 		uint32 H = FCrc::MemCrc32(&Key.RenderPipelineHash, sizeof(Key.RenderPipelineHash), GetTypeHash(Key.VertexDescriptorHash));
 		H = FCrc::MemCrc32(Key.VertexFunction.Hash, sizeof(Key.VertexFunction.Hash), H);
+#if PLATFORM_SUPPORTS_TESSELLATION_SHADERS
+		H = FCrc::MemCrc32(Key.HullFunction.Hash, sizeof(Key.HullFunction.Hash), H);
 		H = FCrc::MemCrc32(Key.DomainFunction.Hash, sizeof(Key.DomainFunction.Hash), H);
+#endif
 		H = FCrc::MemCrc32(Key.PixelFunction.Hash, sizeof(Key.PixelFunction.Hash), H);
 		return H;
 	}
@@ -400,8 +409,10 @@ struct FMetalGraphicsPipelineKey
 
 #if PLATFORM_SUPPORTS_TESSELLATION_SHADERS
 		FMetalDomainShader* DomainShader = (FMetalDomainShader*)Init.BoundShaderState.DomainShaderRHI;
-		if (DomainShader)
+		FMetalHullShader* HullShader = (FMetalHullShader*)Init.BoundShaderState.HullShaderRHI;
+		if (DomainShader && HullShader)
 		{
+			Key.HullFunction = HullShader->GetHash();
 			Key.DomainFunction = DomainShader->GetHash();
 			Key.SetHashValue(Offset_IndexType, NumBits_IndexType, IndexType);
 		}
