Fixes assertion firing in EngineTests
MapStagingSurface has a path that blits and stalls. This was occurring immediately regardless if the RHI thread was active or not.
Now we detect if we're running commandlists or the RHI thread and enqueue a lambda then stall. This is a quick fix to unblock automation while a fuller fix is developed.

diff --git a/Engine/Source/Runtime/Apple/MetalRHI/Private/MetalTexture.cpp b/Engine/Source/Runtime/Apple/MetalRHI/Private/MetalTexture.cpp
--- a/Engine/Source/Runtime/Apple/MetalRHI/Private/MetalTexture.cpp
+++ b/Engine/Source/Runtime/Apple/MetalRHI/Private/MetalTexture.cpp
@@ -1339,6 +1339,9 @@ void* FMetalSurface::Lock(uint32 MipIndex, uint32 ArrayIndex, EResourceLockMode
 		{
 			SCOPE_CYCLE_COUNTER(STAT_MetalTexturePageOffTime);
 			
+			FRHICommandListImmediate& RHICmdList = FRHICommandListExecutor::GetImmediateCommandList();
+			const bool bIssueImmediateCommands = RHICmdList.Bypass() || IsInRHIThread();
+			
 			mtlpp::Region Region;
 			if (SizeZ <= 1 || bIsCubemap)
 			{
@@ -1353,10 +1356,24 @@ void* FMetalSurface::Lock(uint32 MipIndex, uint32 ArrayIndex, EResourceLockMode
 			
 			if (Texture.GetStorageMode() == mtlpp::StorageMode::Private)
 			{
-				GetMetalDeviceContext().CopyFromTextureToBuffer(Texture, ArrayIndex, MipIndex, Region.origin, Region.size, SourceData, 0, DestStride, MipBytes, mtlpp::BlitOption::None);
+				// If we are running with command lists or the RHI thread is enabled we have to execute GFX commands in that context.
+				auto CopyTexToBuf =
+				[this, &ArrayIndex, &MipIndex, &Region, &SourceData, &DestStride, &MipBytes](FRHICommandListImmediate& RHICmdList)
+				{
+					GetMetalDeviceContext().CopyFromTextureToBuffer(this->Texture, ArrayIndex, MipIndex, Region.origin, Region.size, SourceData, 0, DestStride, MipBytes, mtlpp::BlitOption::None);
+					//kick the current command buffer.
+					GetMetalDeviceContext().SubmitCommandBufferAndWait();
+				};
 				
-				//kick the current command buffer.
-				GetMetalDeviceContext().SubmitCommandBufferAndWait();
+				if (bIssueImmediateCommands)
+				{
+					CopyTexToBuf(RHICmdList);
+				}
+				else
+				{
+					RHICmdList.EnqueueLambda(MoveTemp(CopyTexToBuf));
+					RHICmdList.ImmediateFlush(EImmediateFlushType::FlushRHIThread);
+				}
 			}
 			else
 			{
@@ -1364,8 +1381,24 @@ void* FMetalSurface::Lock(uint32 MipIndex, uint32 ArrayIndex, EResourceLockMode
 				if((GPUReadback & EMetalGPUReadbackFlags::ReadbackRequestedAndComplete) != EMetalGPUReadbackFlags::ReadbackRequestedAndComplete)
 				{
 					// A previous texture sync has not been done, need the data now, request texture sync and kick the current command buffer.
-					GetMetalDeviceContext().SynchronizeTexture(Texture, ArrayIndex, MipIndex);
-					GetMetalDeviceContext().SubmitCommandBufferAndWait();
+					auto SyncReadbackToCPU =
+					[this, &ArrayIndex, &MipIndex](FRHICommandListImmediate& RHICmdList)
+					{
+						GetMetalDeviceContext().SynchronizeTexture(this->Texture, ArrayIndex, MipIndex);
+						GetMetalDeviceContext().SubmitCommandBufferAndWait();
+					};
+					
+					// Similar to above. If we are in a context where we have command lists or the RHI thread we must execute
+					// commands there. Otherwise we can just do this directly.
+					if (bIssueImmediateCommands)
+					{
+						SyncReadbackToCPU(RHICmdList);
+					}
+					else
+					{
+						RHICmdList.EnqueueLambda(MoveTemp(SyncReadbackToCPU));
+						RHICmdList.ImmediateFlush(EImmediateFlushType::FlushRHIThread);
+					}
 				}
 #endif
 				
