Buffer refactor merging Vertex, Index and Structured Buffer types into a single FRHIBuffer. No changes to APIs in this changelist other than those that naturally folded together when types merged. A few typedefs provide backward compatibility for the rendering code so that no changes outside of the RHI is needed at this point. There are plenty of API improvements and cleanups that can be imagined once buffer types are merged, but that's deferred to a later changelist. Currently buffers are still created with the old functions, but a unified CreateBuffer() could be imagined in the future.

diff --git a/Engine/Plugins/Experimental/HairStrands/Source/HairStrandsCore/Public/GroomTextureBuilder.h b/Engine/Plugins/Experimental/HairStrands/Source/HairStrandsCore/Public/GroomTextureBuilder.h
--- a/Engine/Plugins/Experimental/HairStrands/Source/HairStrandsCore/Public/GroomTextureBuilder.h
+++ b/Engine/Plugins/Experimental/HairStrands/Source/HairStrandsCore/Public/GroomTextureBuilder.h
@@ -12,7 +12,7 @@ class UTexture2D;
 class UGroomAsset;
 class USkeletalMesh;
 class UStaticMesh;
-class FRHIIndexBuffer;
+class FRHIBuffer;
 class FRHIShaderResourceView;
 
 ////////////////////////////////////////////////////////////////////////////////////////////////////////////////////////
@@ -73,4 +73,4 @@ struct HAIRSTRANDSCORE_API FGroomTextureBuilder
 #endif
 };
 
-////////////////////////////////////////////////////////////////////////////////////////////////////////////////////////
\ No newline at end of file
+////////////////////////////////////////////////////////////////////////////////////////////////////////////////////////
diff --git a/Engine/Plugins/Tests/RHITests/Source/RHITests/Private/RHIBufferTests.cpp b/Engine/Plugins/Tests/RHITests/Source/RHITests/Private/RHIBufferTests.cpp
--- a/Engine/Plugins/Tests/RHITests/Source/RHITests/Private/RHIBufferTests.cpp
+++ b/Engine/Plugins/Tests/RHITests/Source/RHITests/Private/RHIBufferTests.cpp
@@ -3,7 +3,7 @@
 #include "RHIBufferTests.h"
 
 // Copies data in the specified vertex buffer back to the CPU, and passes a pointer to that data to the provided verification lambda.
-bool FRHIBufferTests::VerifyBufferContents(const TCHAR* TestName, FRHICommandListImmediate& RHICmdList, FRHIVertexBuffer* Buffer, TFunctionRef<bool(void* Ptr, uint32 NumBytes)> VerifyCallback)
+bool FRHIBufferTests::VerifyBufferContents(const TCHAR* TestName, FRHICommandListImmediate& RHICmdList, FRHIBuffer* Buffer, TFunctionRef<bool(void* Ptr, uint32 NumBytes)> VerifyCallback)
 {
 	bool Result;
 	{
@@ -35,31 +35,3 @@ bool FRHIBufferTests::VerifyBufferContents(const TCHAR* TestName, FRHICommandLis
 
 	return Result;
 }
-
-// Copies data in the specified vertex buffer back to the CPU, and passes a pointer to that data to the provided verification lambda.
-bool FRHIBufferTests::VerifyBufferContents(const TCHAR* TestName, FRHICommandListImmediate& RHICmdList, FRHIStructuredBuffer* Buffer, TFunctionRef<bool(void* Ptr, uint32 NumBytes)> VerifyCallback)
-{
-	bool Result;
-	{
-		uint32 NumBytes = Buffer->GetSize();
-
-		RHICmdList.ImmediateFlush(EImmediateFlushType::FlushRHIThreadFlushResourcesFlushDeferredDeletes);
-		void* Memory = RHICmdList.LockStructuredBuffer(Buffer, 0, NumBytes, EResourceLockMode::RLM_ReadOnly);
-		Result = VerifyCallback(Memory, NumBytes);
-		RHICmdList.UnlockStructuredBuffer(Buffer);
-	}
-
-	// Immediate flush to clean up the staging buffer / other resources
-	RHICmdList.ImmediateFlush(EImmediateFlushType::FlushRHIThreadFlushResourcesFlushDeferredDeletes);
-
-	if (!Result)
-	{
-		UE_LOG(LogRHIUnitTestCommandlet, Error, TEXT("Test failed. \"%s\""), TestName);
-	}
-	else
-	{
-		UE_LOG(LogRHIUnitTestCommandlet, Display, TEXT("Test passed. \"%s\""), TestName);
-	}
-
-	return Result;
-}
\ No newline at end of file
diff --git a/Engine/Plugins/Tests/RHITests/Source/RHITests/Public/RHIBufferTests.h b/Engine/Plugins/Tests/RHITests/Source/RHITests/Public/RHIBufferTests.h
--- a/Engine/Plugins/Tests/RHITests/Source/RHITests/Public/RHIBufferTests.h
+++ b/Engine/Plugins/Tests/RHITests/Source/RHITests/Public/RHIBufferTests.h
@@ -7,10 +7,7 @@
 class FRHIBufferTests
 {
 	// Copies data in the specified vertex buffer back to the CPU, and passes a pointer to that data to the provided verification lambda.
-	static bool VerifyBufferContents(const TCHAR* TestName, FRHICommandListImmediate& RHICmdList, FRHIVertexBuffer* Buffer, TFunctionRef<bool(void* Ptr, uint32 NumBytes)> VerifyCallback);
-
-	// Copies data in the specified vertex buffer back to the CPU, and passes a pointer to that data to the provided verification lambda.
-	static bool VerifyBufferContents(const TCHAR* TestName, FRHICommandListImmediate& RHICmdList, FRHIStructuredBuffer* Buffer, TFunctionRef<bool(void* Ptr, uint32 NumBytes)> VerifyCallback);
+	static bool VerifyBufferContents(const TCHAR* TestName, FRHICommandListImmediate& RHICmdList, FRHIBuffer* Buffer, TFunctionRef<bool(void* Ptr, uint32 NumBytes)> VerifyCallback);
 
 	template <typename BufferType, typename ValueType, uint32 NumTestBytes>
 	static bool RunTest_UAVClear_Buffer(FRHICommandListImmediate& RHICmdList, const FString& TestName, BufferType* BufferRHI, FRHIUnorderedAccessView* UAV, uint32 BufferSize, const ValueType& ClearValue, void(FRHIComputeCommandList::* ClearPtr)(FRHIUnorderedAccessView*, ValueType const&), const uint8(&TestValue)[NumTestBytes])
diff --git a/Engine/Source/Runtime/Apple/MetalRHI/Private/MetalDynamicRHI.h b/Engine/Source/Runtime/Apple/MetalRHI/Private/MetalDynamicRHI.h
--- a/Engine/Source/Runtime/Apple/MetalRHI/Private/MetalDynamicRHI.h
+++ b/Engine/Source/Runtime/Apple/MetalRHI/Private/MetalDynamicRHI.h
@@ -58,24 +58,21 @@ public:
 	virtual FIndexBufferRHIRef RHICreateIndexBuffer(uint32 Stride, uint32 Size, uint32 InUsage, ERHIAccess InResourceState, FRHIResourceCreateInfo& CreateInfo) final override;
 	virtual void* LockIndexBuffer_BottomOfPipe(FRHICommandListImmediate& RHICmdList, FRHIIndexBuffer* IndexBuffer, uint32 Offset, uint32 Size, EResourceLockMode LockMode) final override;
 	virtual void UnlockIndexBuffer_BottomOfPipe(FRHICommandListImmediate& RHICmdList, FRHIIndexBuffer* IndexBuffer) final override;
-	virtual void RHITransferIndexBufferUnderlyingResource(FRHIIndexBuffer* DestIndexBuffer, FRHIIndexBuffer* SrcIndexBuffer) final override;
+	virtual void RHITransferBufferUnderlyingResource(FRHIBuffer* DestBuffer, FRHIBuffer* SrcBuffer) final override;
 	virtual FVertexBufferRHIRef RHICreateVertexBuffer(uint32 Size, uint32 InUsage, ERHIAccess InResourceState, FRHIResourceCreateInfo& CreateInfo) final override;
 	virtual void* LockVertexBuffer_BottomOfPipe(FRHICommandListImmediate& RHICmdList, FRHIVertexBuffer* VertexBuffer, uint32 Offset, uint32 SizeRHI, EResourceLockMode LockMode) final override;
 	virtual void UnlockVertexBuffer_BottomOfPipe(FRHICommandListImmediate& RHICmdList, FRHIVertexBuffer* VertexBuffer) final override;
 	virtual void RHICopyVertexBuffer(FRHIVertexBuffer* SourceBuffer, FRHIVertexBuffer* DestBuffer) final override;
-	virtual void RHITransferVertexBufferUnderlyingResource(FRHIVertexBuffer* DestVertexBuffer, FRHIVertexBuffer* SrcVertexBuffer) final override;
 	virtual FStructuredBufferRHIRef RHICreateStructuredBuffer(uint32 Stride, uint32 Size, uint32 InUsage, ERHIAccess InResourceState, FRHIResourceCreateInfo& CreateInfo) final override;
 	virtual void* LockStructuredBuffer_BottomOfPipe(FRHICommandListImmediate& RHICmdList, FRHIStructuredBuffer* StructuredBuffer, uint32 Offset, uint32 SizeRHI, EResourceLockMode LockMode) final override;
 	virtual void UnlockStructuredBuffer_BottomOfPipe(FRHICommandListImmediate& RHICmdList, FRHIStructuredBuffer* StructuredBuffer) final override;
 	virtual FUnorderedAccessViewRHIRef RHICreateUnorderedAccessView(FRHIStructuredBuffer* StructuredBuffer, bool bUseUAVCounter, bool bAppendBuffer) final override;
 	virtual FUnorderedAccessViewRHIRef RHICreateUnorderedAccessView(FRHITexture* Texture, uint32 MipLevel) final override;
-	virtual FUnorderedAccessViewRHIRef RHICreateUnorderedAccessView(FRHIVertexBuffer* VertexBuffer, uint8 Format) final override;
-	virtual FUnorderedAccessViewRHIRef RHICreateUnorderedAccessView(FRHIIndexBuffer* IndexBuffer, uint8 Format) final override;
-	virtual FShaderResourceViewRHIRef RHICreateShaderResourceView(FRHIStructuredBuffer* StructuredBuffer) final override;
+	virtual FUnorderedAccessViewRHIRef RHICreateUnorderedAccessView(FRHIBuffer* Buffer, uint8 Format) final override;
+	virtual FShaderResourceViewRHIRef RHICreateShaderResourceView(FRHIBuffer* Buffer) final override;
 	virtual FShaderResourceViewRHIRef RHICreateShaderResourceView(FRHITexture* Texture2DRHI, const FRHITextureSRVCreateInfo& CreateInfo) final override;
-	virtual FShaderResourceViewRHIRef RHICreateShaderResourceView(FRHIVertexBuffer* VertexBuffer, uint32 Stride, uint8 Format) final override;
+	virtual FShaderResourceViewRHIRef RHICreateShaderResourceView(FRHIBuffer* Buffer, uint32 Stride, uint8 Format) final override;
 	virtual FShaderResourceViewRHIRef RHICreateShaderResourceView(const FShaderResourceViewInitializer& Initializer) final override;
-	virtual FShaderResourceViewRHIRef RHICreateShaderResourceView(FRHIIndexBuffer* Buffer) final override;
 	virtual void RHIUpdateShaderResourceView(FRHIShaderResourceView* SRV, FRHIVertexBuffer* VertexBuffer, uint32 Stride, uint8 Format) final override;
 	virtual void RHIUpdateShaderResourceView(FRHIShaderResourceView* SRV, FRHIIndexBuffer* IndexBuffer) final override;
 	virtual uint64 RHICalcTexture2DPlatformSize(uint32 SizeX, uint32 SizeY, uint8 Format, uint32 NumMips, uint32 NumSamples, ETextureCreateFlags Flags, const FRHIResourceCreateInfo& CreateInfo, uint32& OutAlign) final override;
@@ -162,18 +159,16 @@ public:
 	
 	virtual FRenderQueryRHIRef RHICreateRenderQuery_RenderThread(class FRHICommandListImmediate& RHICmdList, ERenderQueryType QueryType) final override;
 	
-	virtual FUnorderedAccessViewRHIRef RHICreateUnorderedAccessView_RenderThread(class FRHICommandListImmediate& RHICmdList, FRHIStructuredBuffer* StructuredBuffer, bool bUseUAVCounter, bool bAppendBuffer) final override;
+	virtual FUnorderedAccessViewRHIRef RHICreateUnorderedAccessView_RenderThread(class FRHICommandListImmediate& RHICmdList, FRHIBuffer* Buffer, bool bUseUAVCounter, bool bAppendBuffer) final override;
 	virtual FUnorderedAccessViewRHIRef RHICreateUnorderedAccessView_RenderThread(class FRHICommandListImmediate& RHICmdList, FRHITexture* Texture, uint32 MipLevel) final override;
-	virtual FUnorderedAccessViewRHIRef RHICreateUnorderedAccessView_RenderThread(class FRHICommandListImmediate& RHICmdList, FRHIVertexBuffer* VertexBuffer, uint8 Format) final override;
-	virtual FUnorderedAccessViewRHIRef RHICreateUnorderedAccessView_RenderThread(class FRHICommandListImmediate& RHICmdList, FRHIIndexBuffer* IndexBuffer, uint8 Format) final override;
-	
-	virtual FShaderResourceViewRHIRef RHICreateShaderResourceView_RenderThread(class FRHICommandListImmediate& RHICmdList, FRHIVertexBuffer* VertexBuffer, uint32 Stride, uint8 Format) final override;
+	virtual FUnorderedAccessViewRHIRef RHICreateUnorderedAccessView_RenderThread(class FRHICommandListImmediate& RHICmdList, FRHIBuffer* Buffer, uint8 Format) final override;
+
+	virtual FShaderResourceViewRHIRef RHICreateShaderResourceView_RenderThread(class FRHICommandListImmediate& RHICmdList, FRHIBuffer* Buffer, uint32 Stride, uint8 Format) final override;
 	virtual FShaderResourceViewRHIRef RHICreateShaderResourceView_RenderThread(class FRHICommandListImmediate& RHICmdList, const FShaderResourceViewInitializer& Initializer) final override;
 	virtual FShaderResourceViewRHIRef CreateShaderResourceView_RenderThread(class FRHICommandListImmediate& RHICmdList, const FShaderResourceViewInitializer& Initializer) final override;
-	virtual FShaderResourceViewRHIRef CreateShaderResourceView_RenderThread(class FRHICommandListImmediate& RHICmdList, FRHIIndexBuffer* Buffer) final override;
-	virtual FShaderResourceViewRHIRef CreateShaderResourceView_RenderThread(class FRHICommandListImmediate& RHICmdList, FRHIVertexBuffer* VertexBuffer, uint32 Stride, uint8 Format) final override;
-	virtual FShaderResourceViewRHIRef RHICreateShaderResourceView_RenderThread(class FRHICommandListImmediate& RHICmdList, FRHIIndexBuffer* Buffer) final override;
-	virtual FShaderResourceViewRHIRef RHICreateShaderResourceView_RenderThread(class FRHICommandListImmediate& RHICmdList, FRHIStructuredBuffer* StructuredBuffer) final override;
+	virtual FShaderResourceViewRHIRef CreateShaderResourceView_RenderThread(class FRHICommandListImmediate& RHICmdList, FRHIBuffer* Buffer) final override;
+	virtual FShaderResourceViewRHIRef CreateShaderResourceView_RenderThread(class FRHICommandListImmediate& RHICmdList, FRHIBuffer* Buffer, uint32 Stride, uint8 Format) final override;
+	virtual FShaderResourceViewRHIRef RHICreateShaderResourceView_RenderThread(class FRHICommandListImmediate& RHICmdList, FRHIBuffer* Buffer) final override;
 	virtual FShaderResourceViewRHIRef RHICreateShaderResourceView_RenderThread(class FRHICommandListImmediate& RHICmdList, FRHITexture* Texture2DRHI, const FRHITextureSRVCreateInfo& CreateInfo) final override;
 
 	virtual void RHICreateTransition(FRHITransition* Transition, ERHIPipeline SrcPipelines, ERHIPipeline DstPipelines, ERHICreateTransitionFlags CreateFlags, TArrayView<const FRHITransitionInfo> Infos) final override;
@@ -187,11 +182,9 @@ public:
 	
 	virtual void RHIAcquireTransientResource_RenderThread(FRHITexture* Texture) final override;
 	virtual void RHIDiscardTransientResource_RenderThread(FRHITexture* Texture) final override;
-	virtual void RHIAcquireTransientResource_RenderThread(FRHIVertexBuffer* Buffer) final override;
-	virtual void RHIDiscardTransientResource_RenderThread(FRHIVertexBuffer* Buffer) final override;
-	virtual void RHIAcquireTransientResource_RenderThread(FRHIStructuredBuffer* Buffer) final override;
-	virtual void RHIDiscardTransientResource_RenderThread(FRHIStructuredBuffer* Buffer) final override;
-	
+	virtual void RHIAcquireTransientResource_RenderThread(FRHIBuffer* Buffer) final override;
+	virtual void RHIDiscardTransientResource_RenderThread(FRHIBuffer* Buffer) final override;
+
 	/* Copy the source box pixels in the destination box texture, return true if implemented for the current platform*/
 	virtual void RHICopySubTextureRegion_RenderThread(class FRHICommandListImmediate& RHICmdList, FRHITexture2D* SourceTexture, FRHITexture2D* DestinationTexture, FBox2D SourceBox, FBox2D DestinationBox) final override;
 	virtual void RHICopySubTextureRegion(FRHITexture2D* SourceTexture, FRHITexture2D* DestinationTexture, FBox2D SourceBox, FBox2D DestinationBox) final override;
diff --git a/Engine/Source/Runtime/Apple/MetalRHI/Private/MetalIndexBuffer.cpp b/Engine/Source/Runtime/Apple/MetalRHI/Private/MetalIndexBuffer.cpp
--- a/Engine/Source/Runtime/Apple/MetalRHI/Private/MetalIndexBuffer.cpp
+++ b/Engine/Source/Runtime/Apple/MetalRHI/Private/MetalIndexBuffer.cpp
@@ -14,54 +14,25 @@
 
 static uint32 MetalIndexBufferUsage(uint32 InUsage)
 {
-	uint32 Usage = InUsage;
+	uint32 Usage = InUsage | BUF_IndexBuffer;
 	if (RHISupportsTessellation(GMaxRHIShaderPlatform))
 	{
 		Usage |= BUF_ShaderResource;
 	}
-	Usage |= EMetalBufferUsage_GPUOnly|EMetalBufferUsage_LinearTex;
+	Usage |= (EMetalBufferUsage_GPUOnly | EMetalBufferUsage_LinearTex);
 	return Usage;
 }
 
-/** Constructor */
-FMetalIndexBuffer::FMetalIndexBuffer(uint32 InStride, uint32 InSize, uint32 InUsage)
-	: FRHIIndexBuffer(InStride, InSize, InUsage)
-	, FMetalRHIBuffer(InSize, MetalIndexBufferUsage(InUsage), RRT_IndexBuffer)
-	, IndexType((InStride == 2) ? mtlpp::IndexType::UInt16 : mtlpp::IndexType::UInt32)
-{
-	if (InSize)
-	{
-		if (RHISupportsTessellation(GMaxRHIShaderPlatform))
-		{
-			EPixelFormat Format = IndexType == mtlpp::IndexType::UInt16 ? PF_R16_UINT : PF_R32_UINT;
-			CreateLinearTexture(Format, this);
-		}
-	}
-}
-
-FMetalIndexBuffer::~FMetalIndexBuffer()
-{
-}
-
-void FMetalIndexBuffer::Swap(FMetalIndexBuffer& Other)
-{
-	@autoreleasepool {
-	FRHIIndexBuffer::Swap(Other);
-	FMetalRHIBuffer::Swap(Other);
-	::Swap(IndexType, Other.IndexType);
-	}
-}
-
 FIndexBufferRHIRef FMetalDynamicRHI::RHICreateIndexBuffer(uint32 Stride, uint32 Size, uint32 InUsage, ERHIAccess InResourceState, FRHIResourceCreateInfo& CreateInfo)
 {
 	@autoreleasepool {
 	if (CreateInfo.bWithoutNativeResource)
 	{
-		return new FMetalIndexBuffer(2, 0, 0);
+		return new FMetalResourceMultiBuffer(0, MetalIndexBufferUsage(0), 2, nullptr, RRT_IndexBuffer);
 	}
 		
 	// make the RHI object, which will allocate memory
-	FMetalIndexBuffer* IndexBuffer = new FMetalIndexBuffer(Stride, Size, InUsage);
+	FMetalResourceMultiBuffer* IndexBuffer = new FMetalResourceMultiBuffer(Size, MetalIndexBufferUsage(InUsage), Stride, nullptr, RRT_IndexBuffer);
 	
 	if (CreateInfo.ResourceArray)
 	{
@@ -83,20 +54,19 @@ FIndexBufferRHIRef FMetalDynamicRHI::RHICreateIndexBuffer(uint32 Stride, uint32
 	}
 }
 
-void FMetalDynamicRHI::RHITransferIndexBufferUnderlyingResource(FRHIIndexBuffer* DestIndexBuffer, FRHIIndexBuffer* SrcIndexBuffer)
+void FMetalDynamicRHI::RHITransferBufferUnderlyingResource(FRHIBuffer* DestBuffer, FRHIBuffer* SrcBuffer)
 {
 	@autoreleasepool {
-	check(DestIndexBuffer);
-	FMetalIndexBuffer* Dest = ResourceCast(DestIndexBuffer);
-	if (!SrcIndexBuffer)
+	check(DestBuffer);
+	FMetalResourceMultiBuffer* Dest = ResourceCast(DestBuffer);
+	if (!SrcBuffer)
 	{
-		FRHIResourceCreateInfo CreateInfo;
-		TRefCountPtr<FMetalIndexBuffer> DeletionProxy = new FMetalIndexBuffer(2, 0, 0);
+		TRefCountPtr<FMetalResourceMultiBuffer> DeletionProxy = new FMetalResourceMultiBuffer(0, Dest->GetUsage(), Dest->GetStride(), nullptr, Dest->Type);
 		Dest->Swap(*DeletionProxy);
 	}
 	else
 	{
-		FMetalIndexBuffer* Src = ResourceCast(SrcIndexBuffer);
+		FMetalResourceMultiBuffer* Src = ResourceCast(SrcBuffer);
 		Dest->Swap(*Src);
 	}
 	}
@@ -125,11 +95,11 @@ FIndexBufferRHIRef FMetalDynamicRHI::CreateIndexBuffer_RenderThread(class FRHICo
 	@autoreleasepool {
 		if (CreateInfo.bWithoutNativeResource)
 		{
-			return new FMetalIndexBuffer(2, 0, 0);
+			return new FMetalResourceMultiBuffer(0, MetalIndexBufferUsage(0), 2, nullptr, RRT_IndexBuffer);
 		}
 		
 		// make the RHI object, which will allocate memory
-		TRefCountPtr<FMetalIndexBuffer> IndexBuffer = new FMetalIndexBuffer(Stride, Size, InUsage);
+		TRefCountPtr<FMetalResourceMultiBuffer> IndexBuffer = new FMetalResourceMultiBuffer(Size, MetalIndexBufferUsage(InUsage), Stride, nullptr, RRT_IndexBuffer);
 		
 		IndexBuffer->Init_RenderThread(RHICmdList, Size, InUsage, CreateInfo, IndexBuffer);
 		
diff --git a/Engine/Source/Runtime/Apple/MetalRHI/Private/MetalStructuredBuffer.cpp b/Engine/Source/Runtime/Apple/MetalRHI/Private/MetalStructuredBuffer.cpp
--- a/Engine/Source/Runtime/Apple/MetalRHI/Private/MetalStructuredBuffer.cpp
+++ b/Engine/Source/Runtime/Apple/MetalRHI/Private/MetalStructuredBuffer.cpp
@@ -6,40 +6,15 @@
 #include "MetalCommandBuffer.h"
 #include "Containers/ResourceArray.h"
 
-FMetalStructuredBuffer::FMetalStructuredBuffer(uint32 Stride, uint32 InSize, FResourceArrayInterface* ResourceArray, uint32 InUsage)
-	: FRHIStructuredBuffer(Stride, InSize, InUsage)
-	, FMetalRHIBuffer(InSize, InUsage|EMetalBufferUsage_GPUOnly, RRT_StructuredBuffer)
-{
-	check((InSize % Stride) == 0);
-	
-	if (ResourceArray)
-	{
-		// copy any resources to the CPU address
-		void* LockedMemory = RHILockStructuredBuffer(this, 0, InSize, RLM_WriteOnly);
- 		FMemory::Memcpy(LockedMemory, ResourceArray->GetResourceData(), InSize);
-		ResourceArray->Discard();
-		RHIUnlockStructuredBuffer(this);
-	}
-}
-
-FMetalStructuredBuffer::~FMetalStructuredBuffer()
+static uint32 MetalStructuredBufferUsage(uint32 InUsage)
 {
+	return (InUsage | EMetalBufferUsage_GPUOnly);
 }
 
-
 FStructuredBufferRHIRef FMetalDynamicRHI::RHICreateStructuredBuffer(uint32 Stride, uint32 Size, uint32 InUsage, ERHIAccess InResourceState, FRHIResourceCreateInfo& CreateInfo)
 {
 	@autoreleasepool {
-		FMetalStructuredBuffer* Buffer = new FMetalStructuredBuffer(Stride, Size, CreateInfo.ResourceArray, InUsage);
-//		if (!CreateInfo.ResourceArray && Buffer->Mode == mtlpp::StorageMode::Private)
-//		{
-//			if (Buffer->TransferBuffer)
-//			{
-//				SafeReleaseMetalBuffer(Buffer->TransferBuffer);
-//				Buffer->TransferBuffer = nil;
-//			}
-//		}
-		return Buffer;
+		return new FMetalResourceMultiBuffer(Size, MetalStructuredBufferUsage(InUsage), Stride, CreateInfo.ResourceArray, RRT_StructuredBuffer);
 	}
 }
 
@@ -65,7 +40,7 @@ FStructuredBufferRHIRef FMetalDynamicRHI::CreateStructuredBuffer_RenderThread(cl
 {
 	@autoreleasepool {
 		// make the RHI object, which will allocate memory
-		TRefCountPtr<FMetalStructuredBuffer> VertexBuffer = new FMetalStructuredBuffer(Stride, Size, nullptr, InUsage);
+		TRefCountPtr<FMetalResourceMultiBuffer> VertexBuffer = new FMetalResourceMultiBuffer(Size, MetalStructuredBufferUsage(InUsage), Stride, nullptr, RRT_StructuredBuffer);
 		
 		VertexBuffer->Init_RenderThread(RHICmdList, Size, InUsage, CreateInfo, VertexBuffer);
 		
diff --git a/Engine/Source/Runtime/Apple/MetalRHI/Private/MetalTexture.cpp b/Engine/Source/Runtime/Apple/MetalRHI/Private/MetalTexture.cpp
--- a/Engine/Source/Runtime/Apple/MetalRHI/Private/MetalTexture.cpp
+++ b/Engine/Source/Runtime/Apple/MetalRHI/Private/MetalTexture.cpp
@@ -2541,22 +2541,12 @@ void FMetalDynamicRHI::RHIDiscardTransientResource_RenderThread(FRHITexture* Tex
 	}
 }
 
-void FMetalDynamicRHI::RHIAcquireTransientResource_RenderThread(FRHIVertexBuffer* Buffer)
+void FMetalDynamicRHI::RHIAcquireTransientResource_RenderThread(FRHIBuffer* Buffer)
 {
 	// This is a no-op for the Metal RHI.
 }
 
-void FMetalDynamicRHI::RHIDiscardTransientResource_RenderThread(FRHIVertexBuffer* Buffer)
-{
-	// This is a no-op for the Metal RHI.
-}
-
-void FMetalDynamicRHI::RHIAcquireTransientResource_RenderThread(FRHIStructuredBuffer* Buffer)
-{
-	// This is a no-op for the Metal RHI.
-}
-
-void FMetalDynamicRHI::RHIDiscardTransientResource_RenderThread(FRHIStructuredBuffer* Buffer)
+void FMetalDynamicRHI::RHIDiscardTransientResource_RenderThread(FRHIBuffer* Buffer)
 {
 	// This is a no-op for the Metal RHI.
 }
diff --git a/Engine/Source/Runtime/Apple/MetalRHI/Private/MetalTransitionData.cpp b/Engine/Source/Runtime/Apple/MetalRHI/Private/MetalTransitionData.cpp
--- a/Engine/Source/Runtime/Apple/MetalRHI/Private/MetalTransitionData.cpp
+++ b/Engine/Source/Runtime/Apple/MetalRHI/Private/MetalTransitionData.cpp
@@ -94,33 +94,15 @@ void FMetalTransitionData::EndResourceTransitions() const
 		switch (Info.Type)
 		{
 			case FRHITransitionInfo::EType::UAV:
-			{
 				GetMetalDeviceContext().TransitionResource(Info.UAV);
-			}
 				break;
 
-			case FRHITransitionInfo::EType::VertexBuffer:
-			{
-				GetMetalDeviceContext().TransitionRHIResource(Info.VertexBuffer);
-			}
-				break;
-
-			case FRHITransitionInfo::EType::IndexBuffer:
-			{
-				GetMetalDeviceContext().TransitionRHIResource(Info.IndexBuffer);
-			}
-				break;
-
-			case FRHITransitionInfo::EType::StructuredBuffer:
-			{
-				GetMetalDeviceContext().TransitionRHIResource(Info.StructuredBuffer);
-			}
+			case FRHITransitionInfo::EType::Buffer:
+				GetMetalDeviceContext().TransitionRHIResource(Info.Buffer);
 				break;
 
 			case FRHITransitionInfo::EType::Texture:
-			{
 				GetMetalDeviceContext().TransitionResource(Info.Texture);
-			}
 				break;
 
 			default:
diff --git a/Engine/Source/Runtime/Apple/MetalRHI/Private/MetalUAV.cpp b/Engine/Source/Runtime/Apple/MetalRHI/Private/MetalUAV.cpp
--- a/Engine/Source/Runtime/Apple/MetalRHI/Private/MetalUAV.cpp
+++ b/Engine/Source/Runtime/Apple/MetalRHI/Private/MetalUAV.cpp
@@ -69,9 +69,9 @@ ns::AutoReleased<FMetalTexture> FMetalShaderResourceView::GetLinearTexture(bool
 	return NewLinearTexture;
 }
 
-FUnorderedAccessViewRHIRef FMetalDynamicRHI::RHICreateUnorderedAccessView_RenderThread(class FRHICommandListImmediate& RHICmdList, FRHIStructuredBuffer* StructuredBuffer, bool bUseUAVCounter, bool bAppendBuffer)
+FUnorderedAccessViewRHIRef FMetalDynamicRHI::RHICreateUnorderedAccessView_RenderThread(class FRHICommandListImmediate& RHICmdList, FRHIBuffer* Buffer, bool bUseUAVCounter, bool bAppendBuffer)
 {
-	return GDynamicRHI->RHICreateUnorderedAccessView(StructuredBuffer, bUseUAVCounter, bAppendBuffer);
+	return GDynamicRHI->RHICreateUnorderedAccessView(Buffer, bUseUAVCounter, bAppendBuffer);
 }
 
 FUnorderedAccessViewRHIRef FMetalDynamicRHI::RHICreateUnorderedAccessView_RenderThread(class FRHICommandListImmediate& RHICmdList, FRHITexture* Texture, uint32 MipLevel)
@@ -89,24 +89,17 @@ FUnorderedAccessViewRHIRef FMetalDynamicRHI::RHICreateUnorderedAccessView_Render
 	}
 }
 
-FUnorderedAccessViewRHIRef FMetalDynamicRHI::RHICreateUnorderedAccessView_RenderThread(class FRHICommandListImmediate& RHICmdList, FRHIVertexBuffer* VertexBuffer, uint8 Format)
+FUnorderedAccessViewRHIRef FMetalDynamicRHI::RHICreateUnorderedAccessView_RenderThread(class FRHICommandListImmediate& RHICmdList, FRHIBuffer* Buffer, uint8 Format)
 {
-	FUnorderedAccessViewRHIRef Result = GDynamicRHI->RHICreateUnorderedAccessView(VertexBuffer, Format);
+	FUnorderedAccessViewRHIRef Result = GDynamicRHI->RHICreateUnorderedAccessView(Buffer, Format);
 	if (IsRunningRHIInSeparateThread() && !RHICmdList.Bypass()) 	{ 		RHICmdList.RHIThreadFence(true); 	}
 	return Result;
 }
 
-FUnorderedAccessViewRHIRef FMetalDynamicRHI::RHICreateUnorderedAccessView_RenderThread(class FRHICommandListImmediate& RHICmdList, FRHIIndexBuffer* IndexBuffer, uint8 Format)
-{
-	FUnorderedAccessViewRHIRef Result = GDynamicRHI->RHICreateUnorderedAccessView(IndexBuffer, Format);
-	if (IsRunningRHIInSeparateThread() && !RHICmdList.Bypass()) 	{ 		RHICmdList.RHIThreadFence(true); 	}
-	return Result;
-}
-
-FUnorderedAccessViewRHIRef FMetalDynamicRHI::RHICreateUnorderedAccessView(FRHIStructuredBuffer* StructuredBufferRHI, bool bUseUAVCounter, bool bAppendBuffer)
+FUnorderedAccessViewRHIRef FMetalDynamicRHI::RHICreateUnorderedAccessView(FRHIBuffer* BufferRHI, bool bUseUAVCounter, bool bAppendBuffer)
 {
 	@autoreleasepool {
-	FMetalStructuredBuffer* StructuredBuffer = ResourceCast(StructuredBufferRHI);
+	FMetalStructuredBuffer* StructuredBuffer = ResourceCast(BufferRHI);
 	
 	FMetalShaderResourceView* SRV = new FMetalShaderResourceView;
 	SRV->SourceVertexBuffer = nullptr;
@@ -151,20 +144,20 @@ FUnorderedAccessViewRHIRef FMetalDynamicRHI::RHICreateUnorderedAccessView(FRHITe
 	}
 }
 
-FUnorderedAccessViewRHIRef FMetalDynamicRHI::RHICreateUnorderedAccessView(FRHIVertexBuffer* VertexBufferRHI, uint8 Format)
+FUnorderedAccessViewRHIRef FMetalDynamicRHI::RHICreateUnorderedAccessView(FRHIBuffer* BufferRHI, uint8 Format)
 {
 	@autoreleasepool {
-	FMetalVertexBuffer* VertexBuffer = ResourceCast(VertexBufferRHI);
+	FMetalResourceMultiBuffer* Buffer = ResourceCast(BufferRHI);
 	
 	FMetalShaderResourceView* SRV = new FMetalShaderResourceView;
-	SRV->SourceVertexBuffer = VertexBuffer;
+	SRV->SourceVertexBuffer = (BufferRHI->GetUsage() & BUF_VertexBuffer) ? Buffer : nullptr;
 	SRV->TextureView = nullptr;
-	SRV->SourceIndexBuffer = nullptr;
+	SRV->SourceIndexBuffer = (BufferRHI->GetUsage() & BUF_IndexBuffer) ? Buffer : nullptr;
 	SRV->SourceStructuredBuffer = nullptr;
 	SRV->Format = Format;
 	{
-		check(VertexBuffer->GetUsage() & BUF_UnorderedAccess);
-		VertexBuffer->CreateLinearTexture((EPixelFormat)Format, VertexBuffer);
+		check(Buffer->GetUsage() & BUF_UnorderedAccess);
+		Buffer->CreateLinearTexture((EPixelFormat)Format, Buffer);
 	}
 		
 	// create the UAV buffer to point to the structured buffer's memory
@@ -175,33 +168,9 @@ FUnorderedAccessViewRHIRef FMetalDynamicRHI::RHICreateUnorderedAccessView(FRHIVe
 	}
 }
 
-FUnorderedAccessViewRHIRef FMetalDynamicRHI::RHICreateUnorderedAccessView(FRHIIndexBuffer* IndexBufferRHI, uint8 Format)
-{
-	@autoreleasepool {
-		FMetalIndexBuffer* IndexBuffer = ResourceCast(IndexBufferRHI);
-		
-		FMetalShaderResourceView* SRV = new FMetalShaderResourceView;
-		SRV->SourceVertexBuffer = nullptr;
-		SRV->TextureView = nullptr;
-		SRV->SourceIndexBuffer = IndexBuffer;
-		SRV->SourceStructuredBuffer = nullptr;
-		SRV->Format = Format;
-		{
-			check(IndexBuffer->GetUsage() & BUF_UnorderedAccess);
-			IndexBuffer->CreateLinearTexture((EPixelFormat)Format, IndexBuffer);
-		}
-		
-		// create the UAV buffer to point to the structured buffer's memory
-		FMetalUnorderedAccessView* UAV = new FMetalUnorderedAccessView;
-		UAV->SourceView = SRV;
-		
-		return UAV;
-	}
-}
-
-FShaderResourceViewRHIRef FMetalDynamicRHI::CreateShaderResourceView_RenderThread(class FRHICommandListImmediate& RHICmdList, FRHIVertexBuffer* VertexBuffer, uint32 Stride, uint8 Format)
+FShaderResourceViewRHIRef FMetalDynamicRHI::CreateShaderResourceView_RenderThread(class FRHICommandListImmediate& RHICmdList, FRHIBuffer* Buffer, uint32 Stride, uint8 Format)
 {
-	FShaderResourceViewRHIRef Result = GDynamicRHI->RHICreateShaderResourceView(VertexBuffer, Stride, Format);
+	FShaderResourceViewRHIRef Result = GDynamicRHI->RHICreateShaderResourceView(Buffer, Stride, Format);
 	if (IsRunningRHIInSeparateThread() && !RHICmdList.Bypass()) 	{ 		RHICmdList.RHIThreadFence(true); 	}
 	return Result;
 }
@@ -213,16 +182,16 @@ FShaderResourceViewRHIRef FMetalDynamicRHI::CreateShaderResourceView_RenderThrea
 	return Result;
 }
 
-FShaderResourceViewRHIRef FMetalDynamicRHI::CreateShaderResourceView_RenderThread(class FRHICommandListImmediate& RHICmdList, FRHIIndexBuffer* Buffer)
+FShaderResourceViewRHIRef FMetalDynamicRHI::CreateShaderResourceView_RenderThread(class FRHICommandListImmediate& RHICmdList, FRHIBuffer* Buffer)
 {
 	FShaderResourceViewRHIRef Result = GDynamicRHI->RHICreateShaderResourceView(Buffer);
 	if (IsRunningRHIInSeparateThread() && !RHICmdList.Bypass()) 	{ 		RHICmdList.RHIThreadFence(true); 	}
 	return Result;
 }
 
-FShaderResourceViewRHIRef FMetalDynamicRHI::RHICreateShaderResourceView_RenderThread(class FRHICommandListImmediate& RHICmdList, FRHIVertexBuffer* VertexBuffer, uint32 Stride, uint8 Format)
+FShaderResourceViewRHIRef FMetalDynamicRHI::RHICreateShaderResourceView_RenderThread(class FRHICommandListImmediate& RHICmdList, FRHIBuffer* Buffer, uint32 Stride, uint8 Format)
 {
-	FShaderResourceViewRHIRef Result = GDynamicRHI->RHICreateShaderResourceView(VertexBuffer, Stride, Format);
+	FShaderResourceViewRHIRef Result = GDynamicRHI->RHICreateShaderResourceView(Buffer, Stride, Format);
 	if (IsRunningRHIInSeparateThread() && !RHICmdList.Bypass()) 	{ 		RHICmdList.RHIThreadFence(true); 	}
 	return Result;
 }
@@ -230,22 +199,17 @@ FShaderResourceViewRHIRef FMetalDynamicRHI::RHICreateShaderResourceView_RenderTh
 FShaderResourceViewRHIRef FMetalDynamicRHI::RHICreateShaderResourceView_RenderThread(class FRHICommandListImmediate& RHICmdList, const FShaderResourceViewInitializer& Initializer)
 {
 	FShaderResourceViewRHIRef Result = GDynamicRHI->RHICreateShaderResourceView(Initializer);
-	if (IsRunningRHIInSeparateThread() && !RHICmdList.Bypass()) { RHICmdList.RHIThreadFence(true); }
+	if (IsRunningRHIInSeparateThread() && !RHICmdList.Bypass()) 	{ 		RHICmdList.RHIThreadFence(true); 	}
 	return Result;
 }
 
-FShaderResourceViewRHIRef FMetalDynamicRHI::RHICreateShaderResourceView_RenderThread(class FRHICommandListImmediate& RHICmdList, FRHIIndexBuffer* Buffer)
+FShaderResourceViewRHIRef FMetalDynamicRHI::RHICreateShaderResourceView_RenderThread(class FRHICommandListImmediate& RHICmdList, FRHIBuffer* Buffer)
 {
 	FShaderResourceViewRHIRef Result = GDynamicRHI->RHICreateShaderResourceView(Buffer);
 	if (IsRunningRHIInSeparateThread() && !RHICmdList.Bypass()) 	{ 		RHICmdList.RHIThreadFence(true); 	}
 	return Result;
 }
 
-FShaderResourceViewRHIRef FMetalDynamicRHI::RHICreateShaderResourceView_RenderThread(class FRHICommandListImmediate& RHICmdList, FRHIStructuredBuffer* StructuredBuffer)
-{
-	return GDynamicRHI->RHICreateShaderResourceView(StructuredBuffer);
-}
-
 FShaderResourceViewRHIRef FMetalDynamicRHI::RHICreateShaderResourceView_RenderThread(class FRHICommandListImmediate& RHICmdList, FRHITexture* Texture2DRHI, const FRHITextureSRVCreateInfo& CreateInfo)
 {
 	FMetalSurface* Surface = (FMetalSurface*)Texture2DRHI->GetTextureBaseRHI();
@@ -299,26 +263,20 @@ FShaderResourceViewRHIRef FMetalDynamicRHI::RHICreateShaderResourceView(FRHIText
 	}
 }
 
-FShaderResourceViewRHIRef FMetalDynamicRHI::RHICreateShaderResourceView(FRHIStructuredBuffer* StructuredBufferRHI)
+FShaderResourceViewRHIRef FMetalDynamicRHI::RHICreateShaderResourceView(FRHIBuffer* BufferRHI)
 {
-	return RHICreateShaderResourceView(FShaderResourceViewInitializer(StructuredBufferRHI));
+	@autoreleasepool {
+	return RHICreateShaderResourceView(FShaderResourceViewInitializer(BufferRHI));
+	}
 }
 
-FShaderResourceViewRHIRef FMetalDynamicRHI::RHICreateShaderResourceView(FRHIVertexBuffer* VertexBufferRHI, uint32 Stride, uint8 Format)
+FShaderResourceViewRHIRef FMetalDynamicRHI::RHICreateShaderResourceView(FRHIBuffer* BufferRHI, uint32 Stride, uint8 Format)
 {
 	check(GPixelFormats[Format].BlockBytes == Stride);
 
 	@autoreleasepool
 	{
-		return RHICreateShaderResourceView(FShaderResourceViewInitializer(VertexBufferRHI, EPixelFormat(Format)));
-	}
-}
-
-FShaderResourceViewRHIRef FMetalDynamicRHI::RHICreateShaderResourceView(FRHIIndexBuffer* BufferRHI)
-{
-	@autoreleasepool
-	{
-		return RHICreateShaderResourceView(FShaderResourceViewInitializer(BufferRHI));
+		return RHICreateShaderResourceView(FShaderResourceViewInitializer(BufferRHI, EPixelFormat(Format)));
 	}
 }
 
@@ -332,7 +290,7 @@ FShaderResourceViewRHIRef FMetalDynamicRHI::RHICreateShaderResourceView(const FS
 			{
 				FShaderResourceViewInitializer::FVertexBufferShaderResourceViewInitializer Desc = Initializer.AsVertexBufferSRV();
 				
-				FMetalVertexBuffer* VertexBuffer = ResourceCast(Desc.VertexBuffer);
+				FMetalVertexBuffer* VertexBuffer = ResourceCast(Desc.Buffer);
 				
 				FMetalShaderResourceView* SRV = new FMetalShaderResourceView;
 				SRV->SourceVertexBuffer = VertexBuffer;
@@ -368,7 +326,7 @@ FShaderResourceViewRHIRef FMetalDynamicRHI::RHICreateShaderResourceView(const FS
 			{
 				FShaderResourceViewInitializer::FStructuredBufferShaderResourceViewInitializer Desc = Initializer.AsStructuredBufferSRV();
 				
-				FMetalStructuredBuffer* StructuredBuffer = ResourceCast(Desc.StructuredBuffer);
+				FMetalStructuredBuffer* StructuredBuffer = ResourceCast(Desc.Buffer);
 
 				FMetalShaderResourceView* SRV = new FMetalShaderResourceView;
 				SRV->SourceVertexBuffer = nullptr;
@@ -387,7 +345,7 @@ FShaderResourceViewRHIRef FMetalDynamicRHI::RHICreateShaderResourceView(const FS
 			{
 				FShaderResourceViewInitializer::FIndexBufferShaderResourceViewInitializer Desc = Initializer.AsIndexBufferSRV();
 				
-				FMetalIndexBuffer* IndexBuffer = ResourceCast(Desc.IndexBuffer);
+				FMetalIndexBuffer* IndexBuffer = ResourceCast(Desc.Buffer);
 				
 				FMetalShaderResourceView* SRV = new FMetalShaderResourceView;
 				SRV->SourceVertexBuffer = nullptr;
@@ -404,7 +362,7 @@ FShaderResourceViewRHIRef FMetalDynamicRHI::RHICreateShaderResourceView(const FS
 				{
 					SRV->Format = (IndexBuffer->IndexType == mtlpp::IndexType::UInt16) ? PF_R16_UINT : PF_R32_UINT;
 					
-					const uint32 Stride = Desc.IndexBuffer->GetStride();
+					const uint32 Stride = Desc.Buffer->GetStride();
 					check(Stride == ((SRV->Format == PF_R16_UINT) ? 2 : 4));
 					
 					SRV->Offset = Desc.StartOffsetBytes;
diff --git a/Engine/Source/Runtime/Apple/MetalRHI/Private/MetalVertexBuffer.cpp b/Engine/Source/Runtime/Apple/MetalRHI/Private/MetalVertexBuffer.cpp
--- a/Engine/Source/Runtime/Apple/MetalRHI/Private/MetalVertexBuffer.cpp
+++ b/Engine/Source/Runtime/Apple/MetalRHI/Private/MetalVertexBuffer.cpp
@@ -81,21 +81,9 @@ APPLE_PLATFORM_OBJECT_ALLOC_OVERRIDES(FMetalBufferData)
 }
 @end
 
-
-FMetalVertexBuffer::FMetalVertexBuffer(uint32 InSize, uint32 InUsage)
-	: FRHIVertexBuffer(InSize, InUsage)
-	, FMetalRHIBuffer(InSize, InUsage|EMetalBufferUsage_LinearTex, RRT_VertexBuffer)
-{
-}
-
-FMetalVertexBuffer::~FMetalVertexBuffer()
-{
-}
-
-void FMetalVertexBuffer::Swap(FMetalVertexBuffer& Other)
+static uint32 MetalVertexBufferUsage(uint32 InUsage)
 {
-	FRHIVertexBuffer::Swap(Other);
-	FMetalRHIBuffer::Swap(Other);
+	return (InUsage | BUF_VertexBuffer | EMetalBufferUsage_LinearTex);
 }
 
 void FMetalRHIBuffer::Swap(FMetalRHIBuffer& Other)
@@ -688,17 +676,61 @@ void FMetalRHIBuffer::Unlock()
 	LastLockFrame = GetMetalDeviceContext().GetFrameNumberRHIThread();
 }
 
+FMetalResourceMultiBuffer::FMetalResourceMultiBuffer(uint32 InSize, uint32 InUsage, uint32 InStride, FResourceArrayInterface* ResourceArray, ERHIResourceType Type)
+	: FRHIBuffer(InSize, InUsage & ~EMetalBufferUsageFlags, InStride)
+	, FMetalRHIBuffer(InSize, InUsage, Type)
+	, IndexType((InStride == 2) ? mtlpp::IndexType::UInt16 : mtlpp::IndexType::UInt32)
+{
+	if (InUsage & BUF_IndexBuffer)
+	{
+		if (InSize)
+		{
+			if (RHISupportsTessellation(GMaxRHIShaderPlatform))
+			{
+				EPixelFormat Format = IndexType == mtlpp::IndexType::UInt16 ? PF_R16_UINT : PF_R32_UINT;
+				CreateLinearTexture(Format, this);
+			}
+		}
+	}
+	else if (InUsage & BUF_StructuredBuffer)
+	{
+		check((InSize % InStride) == 0);
+
+		if (ResourceArray)
+		{
+			// copy any resources to the CPU address
+			void* LockedMemory = RHILockStructuredBuffer(this, 0, InSize, RLM_WriteOnly);
+			FMemory::Memcpy(LockedMemory, ResourceArray->GetResourceData(), InSize);
+			ResourceArray->Discard();
+			RHIUnlockStructuredBuffer(this);
+		}
+	}
+}
+
+FMetalResourceMultiBuffer::~FMetalResourceMultiBuffer()
+{
+}
+
+void FMetalResourceMultiBuffer::Swap(FMetalResourceMultiBuffer& Other)
+{
+	@autoreleasepool {
+		FRHIBuffer::Swap(Other);
+		FMetalRHIBuffer::Swap(Other);
+		::Swap(IndexType, Other.IndexType);
+	}
+}
+
 FVertexBufferRHIRef FMetalDynamicRHI::RHICreateVertexBuffer(uint32 Size, uint32 InUsage, ERHIAccess InResourceState, FRHIResourceCreateInfo& CreateInfo)
 {
 	check(0);
 	@autoreleasepool {
 	if (CreateInfo.bWithoutNativeResource)
 	{
-		return new FMetalVertexBuffer(0, 0);
+		return new FMetalResourceMultiBuffer(0, MetalVertexBufferUsage(0), 0, nullptr, RRT_VertexBuffer);
 	}
 	
 	// make the RHI object, which will allocate memory
-	FMetalVertexBuffer* VertexBuffer = new FMetalVertexBuffer(Size, InUsage);
+	FMetalResourceMultiBuffer* VertexBuffer = new FMetalResourceMultiBuffer(Size, MetalVertexBufferUsage(InUsage), 0, nullptr, RRT_VertexBuffer);
 
 	if (CreateInfo.ResourceArray)
 	{
@@ -853,30 +885,14 @@ FVertexBufferRHIRef FMetalDynamicRHI::CreateVertexBuffer_RenderThread(class FRHI
 	@autoreleasepool {
 		if (CreateInfo.bWithoutNativeResource)
 		{
-			return new FMetalVertexBuffer(0, 0);
+			return new FMetalResourceMultiBuffer(0, MetalVertexBufferUsage(0), 0, nullptr, RRT_VertexBuffer);
 		}
 		
 		// make the RHI object, which will allocate memory
-		TRefCountPtr<FMetalVertexBuffer> VertexBuffer = new FMetalVertexBuffer(Size, InUsage);
+		TRefCountPtr<FMetalResourceMultiBuffer> VertexBuffer = new FMetalResourceMultiBuffer(Size, MetalVertexBufferUsage(InUsage), 0, nullptr, RRT_VertexBuffer);
 		
 		VertexBuffer->Init_RenderThread(RHICmdList, Size, InUsage, CreateInfo, VertexBuffer);
 		
 		return VertexBuffer.GetReference();
 	}
 }
-
-void FMetalDynamicRHI::RHITransferVertexBufferUnderlyingResource(FRHIVertexBuffer* DestVertexBuffer, FRHIVertexBuffer* SrcVertexBuffer)
-{
-	check(DestVertexBuffer);
-	FMetalVertexBuffer* Dest = ResourceCast(DestVertexBuffer);
-	if (!SrcVertexBuffer)
-	{
-		TRefCountPtr<FMetalVertexBuffer> DeletionProxy = new FMetalVertexBuffer(0, 0);
-		Dest->Swap(*DeletionProxy);
-	}
-	else
-	{
-		FMetalVertexBuffer* Src = ResourceCast(SrcVertexBuffer);
-		Dest->Swap(*Src);
-	}
-}
diff --git a/Engine/Source/Runtime/Apple/MetalRHI/Public/MetalResources.h b/Engine/Source/Runtime/Apple/MetalRHI/Public/MetalResources.h
--- a/Engine/Source/Runtime/Apple/MetalRHI/Public/MetalResources.h
+++ b/Engine/Source/Runtime/Apple/MetalRHI/Public/MetalResources.h
@@ -392,6 +392,10 @@ enum EMetalBufferUsage
 {
 	EMetalBufferUsage_GPUOnly = 0x80000000,
 	EMetalBufferUsage_LinearTex = 0x40000000,
+
+	EMetalBufferUsageMask	= 0xF,
+	EMetalBufferUsageShift	= 0x1C,
+	EMetalBufferUsageFlags	= (EMetalBufferUsageMask << EMetalBufferUsageShift)
 };
 
 class FMetalLinearTextureDescriptor
@@ -582,43 +586,21 @@ private:
 	void AllocLinearTextures(const LinearTextureMapKey& InLinearTextureMapKey);
 };
 
-/** Index buffer resource class that stores stride information. */
-class FMetalIndexBuffer : public FRHIIndexBuffer, public FMetalRHIBuffer
+class FMetalResourceMultiBuffer : public FRHIBuffer, public FMetalRHIBuffer
 {
 public:
-	
-	/** Constructor */
-	FMetalIndexBuffer(uint32 InStride, uint32 InSize, uint32 InUsage);
-	virtual ~FMetalIndexBuffer();
-	
-	void Swap(FMetalIndexBuffer& Other);
-	
-	// 16- or 32-bit
-	mtlpp::IndexType IndexType;
-};
+	FMetalResourceMultiBuffer(uint32 InSize, uint32 InUsage, uint32 InStride, FResourceArrayInterface* ResourceArray, ERHIResourceType ResourceType);
+	virtual ~FMetalResourceMultiBuffer();
 
-/** Vertex buffer resource class that stores usage type. */
-class FMetalVertexBuffer : public FRHIVertexBuffer, public FMetalRHIBuffer
-{
-public:
-
-	/** Constructor */
-	FMetalVertexBuffer(uint32 InSize, uint32 InUsage);
-	virtual ~FMetalVertexBuffer();
-
-	void Swap(FMetalVertexBuffer& Other);
-};
+	void Swap(FMetalResourceMultiBuffer& Other);
 
-class FMetalStructuredBuffer : public FRHIStructuredBuffer, public FMetalRHIBuffer
-{
-public:
-	// Constructor
-	FMetalStructuredBuffer(uint32 Stride, uint32 Size, FResourceArrayInterface* ResourceArray, uint32 InUsage);
-
-	// Destructor
-	~FMetalStructuredBuffer();
+	// 16- or 32-bit; used for index buffers only.
+	mtlpp::IndexType IndexType;
 };
 
+typedef FMetalResourceMultiBuffer FMetalIndexBuffer;
+typedef FMetalResourceMultiBuffer FMetalVertexBuffer;
+typedef FMetalResourceMultiBuffer FMetalStructuredBuffer;
 
 class FMetalShaderResourceView : public FRHIShaderResourceView
 {
@@ -777,19 +759,9 @@ struct TMetalResourceTraits<FRHIUniformBuffer>
 	typedef FMetalSuballocatedUniformBuffer TConcreteType;
 };
 template<>
-struct TMetalResourceTraits<FRHIIndexBuffer>
-{
-	typedef FMetalIndexBuffer TConcreteType;
-};
-template<>
-struct TMetalResourceTraits<FRHIStructuredBuffer>
-{
-	typedef FMetalStructuredBuffer TConcreteType;
-};
-template<>
-struct TMetalResourceTraits<FRHIVertexBuffer>
+struct TMetalResourceTraits<FRHIBuffer>
 {
-	typedef FMetalVertexBuffer TConcreteType;
+	typedef FMetalResourceMultiBuffer TConcreteType;
 };
 template<>
 struct TMetalResourceTraits<FRHIShaderResourceView>
diff --git a/Engine/Source/Runtime/D3D12RHI/Private/D3D12Adapter.h b/Engine/Source/Runtime/D3D12RHI/Private/D3D12Adapter.h
--- a/Engine/Source/Runtime/D3D12RHI/Private/D3D12Adapter.h
+++ b/Engine/Source/Runtime/D3D12RHI/Private/D3D12Adapter.h
@@ -258,8 +258,7 @@ public:
 		const TCHAR* Name,
 		D3D12_RESOURCE_FLAGS Flags = D3D12_RESOURCE_FLAG_NONE);
 
-	template <typename BufferType> 
-	BufferType* CreateRHIBuffer(FRHICommandListImmediate* RHICmdList,
+	FD3D12Buffer* CreateRHIBuffer(FRHICommandListImmediate* RHICmdList,
 		const D3D12_RESOURCE_DESC& Desc,
 		uint32 Alignment, uint32 Stride, uint32 Size, uint32 InUsage,
 		ED3D12ResourceStateMode InResourceStateMode,
diff --git a/Engine/Source/Runtime/D3D12RHI/Private/D3D12Buffer.cpp b/Engine/Source/Runtime/D3D12RHI/Private/D3D12Buffer.cpp
--- a/Engine/Source/Runtime/D3D12RHI/Private/D3D12Buffer.cpp
+++ b/Engine/Source/Runtime/D3D12RHI/Private/D3D12Buffer.cpp
@@ -6,25 +6,6 @@ D3D12Buffer.cpp: D3D Common code for buffers.
 
 #include "D3D12RHIPrivate.h"
 
-// Forward declarations are required for the template functions
-template FD3D12VertexBuffer* FD3D12Adapter::CreateRHIBuffer<FD3D12VertexBuffer>(FRHICommandListImmediate* RHICmdList,
-	const D3D12_RESOURCE_DESC& Desc,
-	uint32 Alignment, uint32 Stride, uint32 Size, uint32 InUsage,
-	ED3D12ResourceStateMode InResourceStateMode,
-	FRHIResourceCreateInfo& CreateInfo);
-
-template FD3D12IndexBuffer* FD3D12Adapter::CreateRHIBuffer<FD3D12IndexBuffer>(FRHICommandListImmediate* RHICmdList,
-	const D3D12_RESOURCE_DESC& Desc,
-	uint32 Alignment, uint32 Stride, uint32 Size, uint32 InUsage,
-	ED3D12ResourceStateMode InResourceStateMode,
-	FRHIResourceCreateInfo& CreateInfo);
-
-template FD3D12StructuredBuffer* FD3D12Adapter::CreateRHIBuffer<FD3D12StructuredBuffer>(FRHICommandListImmediate* RHICmdList,
-	const D3D12_RESOURCE_DESC& Desc,
-	uint32 Alignment, uint32 Stride, uint32 Size, uint32 InUsage,
-	ED3D12ResourceStateMode InResourceStateMode,
-	FRHIResourceCreateInfo& CreateInfo);
-
 struct FRHICommandUpdateBufferString
 {
 	static const TCHAR* TStr() { return TEXT("FRHICommandUpdateBuffer"); }
@@ -166,9 +147,7 @@ void FD3D12Adapter::AllocateBuffer(FD3D12Device* Device,
 	}
 }
 
-// This is a templated function used to create FD3D12VertexBuffers, FD3D12IndexBuffers and FD3D12StructuredBuffers
-template<class BufferType>
-BufferType* FD3D12Adapter::CreateRHIBuffer(FRHICommandListImmediate* RHICmdList,
+FD3D12Buffer* FD3D12Adapter::CreateRHIBuffer(FRHICommandListImmediate* RHICmdList,
 	const D3D12_RESOURCE_DESC& InDesc,
 	uint32 Alignment,
 	uint32 Stride,
@@ -182,13 +161,13 @@ BufferType* FD3D12Adapter::CreateRHIBuffer(FRHICommandListImmediate* RHICmdList,
 	const bool bIsDynamic = (InUsage & BUF_AnyDynamic) ? true : false;
 	const uint32 FirstGPUIndex = CreateInfo.GPUMask.GetFirstIndex();
 
-	BufferType* BufferOut = nullptr;
+	FD3D12Buffer* BufferOut = nullptr;
 	if (bIsDynamic)
 	{
-		BufferType* NewBuffer0 = nullptr;
-		BufferOut = CreateLinkedObject<BufferType>(CreateInfo.GPUMask, [&](FD3D12Device* Device)
+		FD3D12Buffer* NewBuffer0 = nullptr;
+		BufferOut = CreateLinkedObject<FD3D12Buffer>(CreateInfo.GPUMask, [&](FD3D12Device* Device)
 		{
-			BufferType* NewBuffer = new BufferType(Device, Stride, Size, InUsage);
+			FD3D12Buffer* NewBuffer = new FD3D12Buffer(Device, Size, InUsage, Stride);
 			NewBuffer->BufferAlignment = Alignment;
 
 			if (Device->GetGPUIndex() == FirstGPUIndex)
@@ -207,9 +186,9 @@ BufferType* FD3D12Adapter::CreateRHIBuffer(FRHICommandListImmediate* RHICmdList,
 	}
 	else
 	{
-		BufferOut = CreateLinkedObject<BufferType>(CreateInfo.GPUMask, [&](FD3D12Device* Device)
+		BufferOut = CreateLinkedObject<FD3D12Buffer>(CreateInfo.GPUMask, [&](FD3D12Device* Device)
 		{
-			BufferType* NewBuffer = new BufferType(Device, Stride, Size, InUsage);
+			FD3D12Buffer* NewBuffer = new FD3D12Buffer(Device, Size, InUsage, Stride);
 			NewBuffer->BufferAlignment = Alignment;
 
 			AllocateBuffer(Device, InDesc, Size, InUsage, InResourceStateMode, CreateInfo, Alignment, *NewBuffer, NewBuffer->ResourceLocation);
@@ -286,7 +265,7 @@ BufferType* FD3D12Adapter::CreateRHIBuffer(FRHICommandListImmediate* RHICmdList,
 		CreateInfo.ResourceArray->Discard();
 	}
 
-	UpdateBufferStats<BufferType>(&BufferOut->ResourceLocation, true);
+	UpdateBufferStats(GetBufferStats(InUsage), BufferOut->ResourceLocation.GetSize());
 
 	return BufferOut;
 }
@@ -330,6 +309,15 @@ void FD3D12Buffer::RenameLDAChain(FD3D12ResourceLocation& NewLocation)
 	}
 }
 
+void FD3D12Buffer::Swap(FD3D12Buffer& Other)
+{
+	check(!LockedData.bLocked && !Other.LockedData.bLocked);
+	FRHIBuffer::Swap(Other);
+	FD3D12BaseShaderResource::Swap(Other);
+	FD3D12TransientResource::Swap(Other);
+	FD3D12LinkedAdapterObject<FD3D12Buffer>::Swap(Other);
+}
+
 void FD3D12Buffer::ReleaseUnderlyingResource()
 {
 	check(IsHeadLink());
diff --git a/Engine/Source/Runtime/D3D12RHI/Private/D3D12CommandContext.h b/Engine/Source/Runtime/D3D12RHI/Private/D3D12CommandContext.h
--- a/Engine/Source/Runtime/D3D12RHI/Private/D3D12CommandContext.h
+++ b/Engine/Source/Runtime/D3D12RHI/Private/D3D12CommandContext.h
@@ -248,7 +248,7 @@ public:
 	virtual void RHISetComputeShader(FRHIComputeShader* ComputeShader) final override;
 	virtual void RHISetComputePipelineState(FRHIComputePipelineState* ComputePipelineState) final override;
 	virtual void RHIDispatchComputeShader(uint32 ThreadGroupCountX, uint32 ThreadGroupCountY, uint32 ThreadGroupCountZ) final override;
-	virtual void RHIDispatchIndirectComputeShader(FRHIVertexBuffer* ArgumentBuffer, uint32 ArgumentOffset) final override;
+	virtual void RHIDispatchIndirectComputeShader(FRHIBuffer* ArgumentBuffer, uint32 ArgumentOffset) final override;
 	virtual void RHISetGlobalUniformBuffers(const FUniformBufferStaticBindings& InUniformBuffers) final override;
 	virtual void RHISetShaderTexture(FRHIComputeShader* PixelShader, uint32 TextureIndex, FRHITexture* NewTexture) final override;
 	virtual void RHISetShaderSampler(FRHIComputeShader* ComputeShader, uint32 SamplerIndex, FRHISamplerState* NewState) final override;
@@ -291,10 +291,10 @@ public:
 	virtual void SetRenderTargets(uint32 NumSimultaneousRenderTargets, const FRHIRenderTargetView* NewRenderTargets, const FRHIDepthRenderTargetView* NewDepthStencilTarget) final override;
 	void SetRenderTargetsAndClear(const FRHISetRenderTargetsInfo& RenderTargetsInfo);
 	virtual void RHIDrawPrimitive(uint32 BaseVertexIndex, uint32 NumPrimitives, uint32 NumInstances) final override;
-	virtual void RHIDrawPrimitiveIndirect(FRHIVertexBuffer* ArgumentBuffer, uint32 ArgumentOffset) final override;
-	virtual void RHIDrawIndexedIndirect(FRHIIndexBuffer* IndexBufferRHI, FRHIStructuredBuffer* ArgumentsBufferRHI, int32 DrawArgumentsIndex, uint32 NumInstances) final override;
-	virtual void RHIDrawIndexedPrimitive(FRHIIndexBuffer* IndexBuffer, int32 BaseVertexIndex, uint32 FirstInstance, uint32 NumVertices, uint32 StartIndex, uint32 NumPrimitives, uint32 NumInstances) final override;
-	virtual void RHIDrawIndexedPrimitiveIndirect(FRHIIndexBuffer* IndexBuffer, FRHIVertexBuffer* ArgumentBuffer, uint32 ArgumentOffset) final override;
+	virtual void RHIDrawPrimitiveIndirect(FRHIBuffer* ArgumentBuffer, uint32 ArgumentOffset) final override;
+	virtual void RHIDrawIndexedIndirect(FRHIBuffer* IndexBufferRHI, FRHIBuffer* ArgumentsBufferRHI, int32 DrawArgumentsIndex, uint32 NumInstances) final override;
+	virtual void RHIDrawIndexedPrimitive(FRHIBuffer* IndexBuffer, int32 BaseVertexIndex, uint32 FirstInstance, uint32 NumVertices, uint32 StartIndex, uint32 NumPrimitives, uint32 NumInstances) final override;
+	virtual void RHIDrawIndexedPrimitiveIndirect(FRHIBuffer* IndexBuffer, FRHIBuffer* ArgumentBuffer, uint32 ArgumentOffset) final override;
 	virtual void RHISetDepthBounds(float MinDepth, float MaxDepth) final override;
     virtual void RHISetShadingRate(EVRSShadingRate ShadingRate, EVRSRateCombiner Combiner) final override;
     virtual void RHISetShadingRateImage(FRHITexture* RateImageTexture, EVRSRateCombiner Combiner) final override;
@@ -480,7 +480,7 @@ public:
 	{
 		ContextRedirect(RHIDispatchComputeShader(ThreadGroupCountX, ThreadGroupCountY, ThreadGroupCountZ));
 	}
-	FORCEINLINE virtual void RHIDispatchIndirectComputeShader(FRHIVertexBuffer* ArgumentBuffer, uint32 ArgumentOffset) final override
+	FORCEINLINE virtual void RHIDispatchIndirectComputeShader(FRHIBuffer* ArgumentBuffer, uint32 ArgumentOffset) final override
 	{
 		ContextRedirect(RHIDispatchIndirectComputeShader(ArgumentBuffer, ArgumentOffset));
 	}
@@ -643,19 +643,19 @@ public:
 	{
 		ContextRedirect(RHIDrawPrimitive(BaseVertexIndex, NumPrimitives, NumInstances));
 	}
-	FORCEINLINE virtual void RHIDrawPrimitiveIndirect(FRHIVertexBuffer* ArgumentBuffer, uint32 ArgumentOffset) final override
+	FORCEINLINE virtual void RHIDrawPrimitiveIndirect(FRHIBuffer* ArgumentBuffer, uint32 ArgumentOffset) final override
 	{
 		ContextRedirect(RHIDrawPrimitiveIndirect(ArgumentBuffer, ArgumentOffset));
 	}
-	FORCEINLINE virtual void RHIDrawIndexedIndirect(FRHIIndexBuffer* IndexBufferRHI, FRHIStructuredBuffer* ArgumentsBufferRHI, int32 DrawArgumentsIndex, uint32 NumInstances) final override
+	FORCEINLINE virtual void RHIDrawIndexedIndirect(FRHIBuffer* IndexBufferRHI, FRHIBuffer* ArgumentsBufferRHI, int32 DrawArgumentsIndex, uint32 NumInstances) final override
 	{
 		ContextRedirect(RHIDrawIndexedIndirect(IndexBufferRHI, ArgumentsBufferRHI, DrawArgumentsIndex, NumInstances));
 	}
-	FORCEINLINE virtual void RHIDrawIndexedPrimitive(FRHIIndexBuffer* IndexBuffer, int32 BaseVertexIndex, uint32 FirstInstance, uint32 NumVertices, uint32 StartIndex, uint32 NumPrimitives, uint32 NumInstances) final override
+	FORCEINLINE virtual void RHIDrawIndexedPrimitive(FRHIBuffer* IndexBuffer, int32 BaseVertexIndex, uint32 FirstInstance, uint32 NumVertices, uint32 StartIndex, uint32 NumPrimitives, uint32 NumInstances) final override
 	{
 		ContextRedirect(RHIDrawIndexedPrimitive(IndexBuffer, BaseVertexIndex, FirstInstance, NumVertices, StartIndex, NumPrimitives, NumInstances));
 	}
-	FORCEINLINE virtual void RHIDrawIndexedPrimitiveIndirect(FRHIIndexBuffer* IndexBuffer, FRHIVertexBuffer* ArgumentBuffer, uint32 ArgumentOffset) final override
+	FORCEINLINE virtual void RHIDrawIndexedPrimitiveIndirect(FRHIBuffer* IndexBuffer, FRHIBuffer* ArgumentBuffer, uint32 ArgumentOffset) final override
 	{
 		ContextRedirect(RHIDrawIndexedPrimitiveIndirect(IndexBuffer, ArgumentBuffer, ArgumentOffset));
 	}
diff --git a/Engine/Source/Runtime/D3D12RHI/Private/D3D12Commands.cpp b/Engine/Source/Runtime/D3D12RHI/Private/D3D12Commands.cpp
--- a/Engine/Source/Runtime/D3D12RHI/Private/D3D12Commands.cpp
+++ b/Engine/Source/Runtime/D3D12RHI/Private/D3D12Commands.cpp
@@ -153,9 +153,9 @@ void FD3D12CommandContext::RHIDispatchComputeShader(uint32 ThreadGroupCountX, ui
 	DEBUG_EXECUTE_COMMAND_LIST(this);
 }
 
-void FD3D12CommandContext::RHIDispatchIndirectComputeShader(FRHIVertexBuffer* ArgumentBufferRHI, uint32 ArgumentOffset)
+void FD3D12CommandContext::RHIDispatchIndirectComputeShader(FRHIBuffer* ArgumentBufferRHI, uint32 ArgumentOffset)
 {
-	FD3D12VertexBuffer* ArgumentBuffer = RetrieveObject<FD3D12VertexBuffer>(ArgumentBufferRHI);
+	FD3D12Buffer* ArgumentBuffer = RetrieveObject<FD3D12Buffer>(ArgumentBufferRHI);
 
 	if (IsDefaultContext())
 	{
@@ -260,36 +260,22 @@ void ProcessResource(FD3D12CommandContext& Context, const FRHITransitionInfo& In
 		FRHITransitionInfo LocalInfo = Info;
 		LocalInfo.MipIndex = UAV->GetViewSubresourceSubset().MostDetailedMip();
 		Function(LocalInfo, UAV->GetResource());
+		break;
 	}
-	break;
-	case FRHITransitionInfo::EType::VertexBuffer:
-	{
-		FD3D12VertexBuffer* VertexBuffer = Context.RetrieveObject<FD3D12VertexBuffer>(Info.VertexBuffer);
-		check(VertexBuffer);
-		Function(Info, VertexBuffer->GetResource());
-	}
-	break;
-	case FRHITransitionInfo::EType::IndexBuffer:
-	{
-		FD3D12IndexBuffer* IndexBuffer = Context.RetrieveObject<FD3D12IndexBuffer>(Info.IndexBuffer);
-		check(IndexBuffer);
-		Function(Info, IndexBuffer->GetResource());
-	}
-	break;
-	case FRHITransitionInfo::EType::StructuredBuffer:
+	case FRHITransitionInfo::EType::Buffer:
 	{
-		FD3D12StructuredBuffer* StructuredBuffer = Context.RetrieveObject<FD3D12StructuredBuffer>(Info.StructuredBuffer);
-		check(StructuredBuffer);
-		Function(Info, StructuredBuffer->GetResource());
+		FD3D12Buffer* Buffer = Context.RetrieveObject<FD3D12Buffer>(Info.Buffer);
+		check(Buffer);
+		Function(Info, Buffer->GetResource());
+		break;
 	}
-	break;
 	case FRHITransitionInfo::EType::Texture:
 	{
 		FD3D12TextureBase* Texture = Context.RetrieveTextureBase(Info.Texture);
 		check(Texture);
 		Function(Info, Texture->GetResource());
+		break;
 	}
-	break;
 	default:
 		checkNoEntry();
 		break;
@@ -522,7 +508,7 @@ void FD3D12CommandContext::RHICopyToStagingBuffer(FRHIVertexBuffer* SourceBuffer
 	check(StagingBuffer);
 	ensureMsgf(!StagingBuffer->bIsLocked, TEXT("Attempting to Copy to a locked staging buffer. This may have undefined behavior"));
 
-	FD3D12VertexBuffer* VertexBuffer = RetrieveObject<FD3D12VertexBuffer>(SourceBufferRHI);
+	FD3D12Buffer* VertexBuffer = RetrieveObject<FD3D12Buffer>(SourceBufferRHI);
 	check(VertexBuffer);
 
 	ensureMsgf((SourceBufferRHI->GetUsage() & BUF_SourceCopy) != 0, TEXT("Buffers used as copy source need to be created with BUF_SourceCopy"));
@@ -1608,7 +1594,7 @@ void FD3D12CommandContext::RHIDrawPrimitive(uint32 BaseVertexIndex, uint32 NumPr
 	DEBUG_EXECUTE_COMMAND_LIST(this);
 }
 
-void FD3D12CommandContext::RHIDrawPrimitiveIndirect(FRHIVertexBuffer* ArgumentBufferRHI, uint32 ArgumentOffset)
+void FD3D12CommandContext::RHIDrawPrimitiveIndirect(FRHIBuffer* ArgumentBufferRHI, uint32 ArgumentOffset)
 {
 	FD3D12Buffer* ArgumentBuffer = RetrieveObject<FD3D12Buffer>(ArgumentBufferRHI);
 
@@ -1648,7 +1634,7 @@ void FD3D12CommandContext::RHIDrawPrimitiveIndirect(FRHIVertexBuffer* ArgumentBu
 	DEBUG_EXECUTE_COMMAND_LIST(this);
 }
 
-void FD3D12CommandContext::RHIDrawIndexedIndirect(FRHIIndexBuffer* IndexBufferRHI, FRHIStructuredBuffer* ArgumentsBufferRHI, int32 DrawArgumentsIndex, uint32 NumInstances)
+void FD3D12CommandContext::RHIDrawIndexedIndirect(FRHIBuffer* IndexBufferRHI, FRHIBuffer* ArgumentsBufferRHI, int32 DrawArgumentsIndex, uint32 NumInstances)
 {
 	const uint32 IndexBufferStride = FD3D12DynamicRHI::ResourceCast(IndexBufferRHI)->GetStride();
 	const uint32 ArgumentsBufferStride = FD3D12DynamicRHI::ResourceCast(ArgumentsBufferRHI)->GetStride();
@@ -1697,7 +1683,7 @@ void FD3D12CommandContext::RHIDrawIndexedIndirect(FRHIIndexBuffer* IndexBufferRH
 	DEBUG_EXECUTE_COMMAND_LIST(this);
 }
 
-void FD3D12CommandContext::RHIDrawIndexedPrimitive(FRHIIndexBuffer* IndexBufferRHI, int32 BaseVertexIndex, uint32 FirstInstance, uint32 NumVertices, uint32 StartIndex, uint32 NumPrimitives, uint32 NumInstances)
+void FD3D12CommandContext::RHIDrawIndexedPrimitive(FRHIBuffer* IndexBufferRHI, int32 BaseVertexIndex, uint32 FirstInstance, uint32 NumVertices, uint32 StartIndex, uint32 NumPrimitives, uint32 NumInstances)
 {
 	// called should make sure the input is valid, this avoid hidden bugs
 	ensure(NumPrimitives > 0);
@@ -1715,16 +1701,15 @@ void FD3D12CommandContext::RHIDrawIndexedPrimitive(FRHIIndexBuffer* IndexBufferR
 
 	uint32 IndexCount = StateCache.GetVertexCountAndIncrementStat(NumPrimitives);
 
-	FD3D12IndexBuffer* RHIIndexBuffer = FD3D12DynamicRHI::ResourceCast(IndexBufferRHI);
 	FD3D12Buffer* IndexBuffer = RetrieveObject<FD3D12Buffer>(IndexBufferRHI);
 
 	// Verify that we are not trying to read outside the index buffer range
 	// test is an optimized version of: StartIndex + IndexCount <= IndexBuffer->GetSize() / IndexBuffer->GetStride() 
-	checkf((StartIndex + IndexCount) * RHIIndexBuffer->GetStride() <= RHIIndexBuffer->GetSize(),
-		TEXT("Start %u, Count %u, Type %u, Buffer Size %u, Buffer stride %u"), StartIndex, IndexCount, StateCache.GetGraphicsPipelinePrimitiveType(), RHIIndexBuffer->GetSize(), RHIIndexBuffer->GetStride());
+	checkf((StartIndex + IndexCount) * IndexBuffer->GetStride() <= IndexBuffer->GetSize(),
+		TEXT("Start %u, Count %u, Type %u, Buffer Size %u, Buffer stride %u"), StartIndex, IndexCount, StateCache.GetGraphicsPipelinePrimitiveType(), IndexBuffer->GetSize(), IndexBuffer->GetStride());
 
 	// determine 16bit vs 32bit indices
-	const DXGI_FORMAT Format = (RHIIndexBuffer->GetStride() == sizeof(uint16) ? DXGI_FORMAT_R16_UINT : DXGI_FORMAT_R32_UINT);
+	const DXGI_FORMAT Format = (IndexBuffer->GetStride() == sizeof(uint16) ? DXGI_FORMAT_R16_UINT : DXGI_FORMAT_R32_UINT);
 	StateCache.SetIndexBuffer(IndexBuffer->ResourceLocation, Format, 0);
 	StateCache.ApplyState<D3D12PT_Graphics>();
 
@@ -1736,7 +1721,7 @@ void FD3D12CommandContext::RHIDrawIndexedPrimitive(FRHIIndexBuffer* IndexBufferR
 	DEBUG_EXECUTE_COMMAND_LIST(this);
 }
 
-void FD3D12CommandContext::RHIDrawIndexedPrimitiveIndirect(FRHIIndexBuffer* IndexBufferRHI, FRHIVertexBuffer* ArgumentBufferRHI, uint32 ArgumentOffset)
+void FD3D12CommandContext::RHIDrawIndexedPrimitiveIndirect(FRHIBuffer* IndexBufferRHI, FRHIBuffer* ArgumentBufferRHI, uint32 ArgumentOffset)
 {
 	const uint32 IndexBufferStride = FD3D12DynamicRHI::ResourceCast(IndexBufferRHI)->GetStride();
 	FD3D12Buffer* IndexBuffer = RetrieveObject<FD3D12Buffer>(IndexBufferRHI);
diff --git a/Engine/Source/Runtime/D3D12RHI/Private/D3D12IndexBuffer.cpp b/Engine/Source/Runtime/D3D12RHI/Private/D3D12IndexBuffer.cpp
--- a/Engine/Source/Runtime/D3D12RHI/Private/D3D12IndexBuffer.cpp
+++ b/Engine/Source/Runtime/D3D12RHI/Private/D3D12IndexBuffer.cpp
@@ -29,44 +29,20 @@ D3D12_RESOURCE_DESC CreateIndexBufferResourceDesc(uint32 Size, uint32 InUsage)
 	return Desc;
 }
 
-FD3D12IndexBuffer::~FD3D12IndexBuffer()
-{
-	if (ResourceLocation.IsValid())
-	{
-		UpdateBufferStats<FD3D12IndexBuffer>(&ResourceLocation, false);
-	}
-}
-
-void FD3D12IndexBuffer::Swap(FD3D12IndexBuffer& Other)
-{
-	check(!LockedData.bLocked && !Other.LockedData.bLocked);
-	FRHIIndexBuffer::Swap(Other);
-	FD3D12BaseShaderResource::Swap(Other);
-	FD3D12TransientResource::Swap(Other);
-	FD3D12LinkedAdapterObject<FD3D12Buffer>::Swap(Other);
-}
-
-void FD3D12IndexBuffer::ReleaseUnderlyingResource()
-{
-	UpdateBufferStats<FD3D12IndexBuffer>(&ResourceLocation, false);
-	FRHIIndexBuffer::ReleaseUnderlyingResource();
-	FD3D12Buffer::ReleaseUnderlyingResource();
-}
-
 FIndexBufferRHIRef FD3D12DynamicRHI::RHICreateIndexBuffer(uint32 Stride, uint32 Size, uint32 InUsage, ERHIAccess InResourceState, FRHIResourceCreateInfo& CreateInfo)
 {
 	if (CreateInfo.bWithoutNativeResource)
 	{
-		return GetAdapter().CreateLinkedObject<FD3D12IndexBuffer>(CreateInfo.GPUMask, [](FD3D12Device* Device)
+		return GetAdapter().CreateLinkedObject<FD3D12Buffer>(CreateInfo.GPUMask, [](FD3D12Device* Device)
 			{
-				return new FD3D12IndexBuffer();
+				return new FD3D12Buffer();
 			});
 	}
 
 	const D3D12_RESOURCE_DESC Desc = CreateIndexBufferResourceDesc(Size, InUsage);
 	const uint32 Alignment = 4;
 
-	FD3D12IndexBuffer* Buffer = GetAdapter().CreateRHIBuffer<FD3D12IndexBuffer>(nullptr, Desc, Alignment, Stride, Size, InUsage, ED3D12ResourceStateMode::Default, CreateInfo);
+	FD3D12Buffer* Buffer = GetAdapter().CreateRHIBuffer(nullptr, Desc, Alignment, Stride, Size, InUsage | BUF_IndexBuffer, ED3D12ResourceStateMode::Default, CreateInfo);
 	if (Buffer->ResourceLocation.IsTransient())
 	{
 		// TODO: this should ideally be set in platform-independent code, since this tracking is for the high level
@@ -78,27 +54,27 @@ FIndexBufferRHIRef FD3D12DynamicRHI::RHICreateIndexBuffer(uint32 Stride, uint32
 
 void* FD3D12DynamicRHI::RHILockIndexBuffer(FRHICommandListImmediate& RHICmdList, FRHIIndexBuffer* IndexBufferRHI, uint32 Offset, uint32 Size, EResourceLockMode LockMode)
 {
-	FD3D12IndexBuffer* Buffer = FD3D12DynamicRHI::ResourceCast(IndexBufferRHI);
+	FD3D12Buffer* Buffer = FD3D12DynamicRHI::ResourceCast(IndexBufferRHI);
 	return LockBuffer(&RHICmdList, Buffer, Buffer->GetSize(), Buffer->GetUsage(), Offset, Size, LockMode);
 }
 
 void FD3D12DynamicRHI::RHIUnlockIndexBuffer(FRHICommandListImmediate& RHICmdList, FRHIIndexBuffer* IndexBufferRHI)
 {
-	FD3D12IndexBuffer* Buffer = FD3D12DynamicRHI::ResourceCast(IndexBufferRHI);
+	FD3D12Buffer* Buffer = FD3D12DynamicRHI::ResourceCast(IndexBufferRHI);
 	UnlockBuffer(&RHICmdList, Buffer, Buffer->GetUsage());
 }
 
-void FD3D12DynamicRHI::RHITransferIndexBufferUnderlyingResource(FRHIIndexBuffer* DestIndexBuffer, FRHIIndexBuffer* SrcIndexBuffer)
+void FD3D12DynamicRHI::RHITransferBufferUnderlyingResource(FRHIBuffer* DestBuffer, FRHIBuffer* SrcBuffer)
 {
-	check(DestIndexBuffer);
-	FD3D12IndexBuffer* Dest = ResourceCast(DestIndexBuffer);
-	if (!SrcIndexBuffer)
+	check(DestBuffer);
+	FD3D12Buffer* Dest = ResourceCast(DestBuffer);
+	if (!SrcBuffer)
 	{
 		Dest->ReleaseUnderlyingResource();
 	}
 	else
 	{
-		FD3D12IndexBuffer* Src = ResourceCast(SrcIndexBuffer);
+		FD3D12Buffer* Src = ResourceCast(SrcBuffer);
 		Dest->Swap(*Src);
 	}
 }
@@ -107,16 +83,16 @@ FIndexBufferRHIRef FD3D12DynamicRHI::CreateIndexBuffer_RenderThread(class FRHICo
 {
 	if (CreateInfo.bWithoutNativeResource)
 	{
-		return GetAdapter().CreateLinkedObject<FD3D12IndexBuffer>(CreateInfo.GPUMask, [](FD3D12Device* Device)
+		return GetAdapter().CreateLinkedObject<FD3D12Buffer>(CreateInfo.GPUMask, [](FD3D12Device* Device)
 			{
-				return new FD3D12IndexBuffer();
+				return new FD3D12Buffer();
 			});
 	}
 
 	const D3D12_RESOURCE_DESC Desc = CreateIndexBufferResourceDesc(Size, InUsage);
 	const uint32 Alignment = 4;
 
-	FD3D12IndexBuffer* Buffer = GetAdapter().CreateRHIBuffer<FD3D12IndexBuffer>(&RHICmdList, Desc, Alignment, Stride, Size, InUsage, ED3D12ResourceStateMode::Default, CreateInfo);
+	FD3D12Buffer* Buffer = GetAdapter().CreateRHIBuffer(&RHICmdList, Desc, Alignment, Stride, Size, InUsage | BUF_IndexBuffer, ED3D12ResourceStateMode::Default, CreateInfo);
 	if (Buffer->ResourceLocation.IsTransient())
 	{
 		// TODO: this should ideally be set in platform-independent code, since this tracking is for the high level
@@ -131,7 +107,7 @@ FIndexBufferRHIRef FD3D12DynamicRHI::CreateAndLockIndexBuffer_RenderThread(class
 	const D3D12_RESOURCE_DESC Desc = CreateIndexBufferResourceDesc(Size, InUsage);
 	const uint32 Alignment = 4;
 
-	FD3D12IndexBuffer* Buffer = GetAdapter().CreateRHIBuffer<FD3D12IndexBuffer>(&RHICmdList, Desc, Alignment, Stride, Size, InUsage, ED3D12ResourceStateMode::Default, CreateInfo);
+	FD3D12Buffer* Buffer = GetAdapter().CreateRHIBuffer(&RHICmdList, Desc, Alignment, Stride, Size, InUsage | BUF_IndexBuffer, ED3D12ResourceStateMode::Default, CreateInfo);
 	if (Buffer->ResourceLocation.IsTransient())
 	{
 		// TODO: this should ideally be set in platform-independent code, since this tracking is for the high level
diff --git a/Engine/Source/Runtime/D3D12RHI/Private/D3D12RHIPrivate.h b/Engine/Source/Runtime/D3D12RHI/Private/D3D12RHIPrivate.h
--- a/Engine/Source/Runtime/D3D12RHI/Private/D3D12RHIPrivate.h
+++ b/Engine/Source/Runtime/D3D12RHI/Private/D3D12RHIPrivate.h
@@ -278,25 +278,22 @@ public:
 	virtual FIndexBufferRHIRef RHICreateIndexBuffer(uint32 Stride, uint32 Size, uint32 InUsage, ERHIAccess InResourceState, FRHIResourceCreateInfo& CreateInfo) final override;
 	virtual void* RHILockIndexBuffer(FRHICommandListImmediate& RHICmdList, FRHIIndexBuffer* IndexBuffer, uint32 Offset, uint32 Size, EResourceLockMode LockMode) final override;
 	virtual void RHIUnlockIndexBuffer(FRHICommandListImmediate& RHICmdList, FRHIIndexBuffer* IndexBuffer) final override;
-	virtual void RHITransferIndexBufferUnderlyingResource(FRHIIndexBuffer* DestIndexBuffer, FRHIIndexBuffer* SrcIndexBuffer) final override;
+	virtual void RHITransferBufferUnderlyingResource(FRHIBuffer* DestBuffer, FRHIBuffer* SrcBuffer) final override;
 	virtual FVertexBufferRHIRef RHICreateVertexBuffer(uint32 Size, uint32 InUsage, ERHIAccess InResourceState, FRHIResourceCreateInfo& CreateInfo) final override;
 	virtual void* RHILockVertexBuffer(FRHICommandListImmediate& RHICmdList, FRHIVertexBuffer* VertexBuffer, uint32 Offset, uint32 SizeRHI, EResourceLockMode LockMode) final override;
 	virtual void RHIUnlockVertexBuffer(FRHICommandListImmediate& RHICmdList, FRHIVertexBuffer* VertexBuffer) final override;
 	virtual void RHICopyVertexBuffer(FRHIVertexBuffer* SourceBuffer, FRHIVertexBuffer* DestBuffer) final override;
-	virtual void RHITransferVertexBufferUnderlyingResource(FRHIVertexBuffer* DestVertexBuffer, FRHIVertexBuffer* SrcVertexBuffer) final override;
 	virtual FStructuredBufferRHIRef RHICreateStructuredBuffer(uint32 Stride, uint32 Size, uint32 InUsage, ERHIAccess InResourceState, FRHIResourceCreateInfo& CreateInfo) final override;
 	virtual void* RHILockStructuredBuffer(FRHICommandListImmediate& RHICmdList, FRHIStructuredBuffer* StructuredBuffer, uint32 Offset, uint32 SizeRHI, EResourceLockMode LockMode) final override;
 	virtual void RHIUnlockStructuredBuffer(FRHICommandListImmediate& RHICmdList, FRHIStructuredBuffer* StructuredBuffer) final override;
-	virtual FUnorderedAccessViewRHIRef RHICreateUnorderedAccessView(FRHIStructuredBuffer* StructuredBuffer, bool bUseUAVCounter, bool bAppendBuffer) final override;
+	virtual FUnorderedAccessViewRHIRef RHICreateUnorderedAccessView(FRHIBuffer* StructuredBuffer, bool bUseUAVCounter, bool bAppendBuffer) final override;
 	virtual FUnorderedAccessViewRHIRef RHICreateUnorderedAccessView(FRHITexture* Texture, uint32 MipLevel) final override;
-	virtual FUnorderedAccessViewRHIRef RHICreateUnorderedAccessView(FRHIVertexBuffer* VertexBuffer, uint8 Format) final override;
-	virtual FUnorderedAccessViewRHIRef RHICreateUnorderedAccessView(FRHIIndexBuffer* IndexBuffer, uint8 Format) final override;
+	virtual FUnorderedAccessViewRHIRef RHICreateUnorderedAccessView(FRHIBuffer* Buffer, uint8 Format) final override;
 	virtual FShaderResourceViewRHIRef RHICreateShaderResourceView(FRHIStructuredBuffer* StructuredBuffer) final override;
 	virtual FShaderResourceViewRHIRef RHICreateShaderResourceView(FRHIVertexBuffer* VertexBuffer, uint32 Stride, uint8 Format) final override;
 	virtual FShaderResourceViewRHIRef RHICreateShaderResourceView(const FShaderResourceViewInitializer& Initializer) final override;
-	virtual FShaderResourceViewRHIRef RHICreateShaderResourceView(FRHIIndexBuffer* Buffer) final override;
-	virtual void RHIUpdateShaderResourceView(FRHIShaderResourceView* SRV, FRHIVertexBuffer* VertexBuffer, uint32 Stride, uint8 Format) final override;
-	virtual void RHIUpdateShaderResourceView(FRHIShaderResourceView* SRV, FRHIIndexBuffer* IndexBuffer) final override;
+	virtual void RHIUpdateShaderResourceView(FRHIShaderResourceView* SRV, FRHIBuffer* Buffer, uint32 Stride, uint8 Format) final override;
+	virtual void RHIUpdateShaderResourceView(FRHIShaderResourceView* SRV, FRHIBuffer* Buffer) final override;
 	virtual uint64 RHICalcTexture2DPlatformSize(uint32 SizeX, uint32 SizeY, uint8 Format, uint32 NumMips, uint32 NumSamples, ETextureCreateFlags InFlags, const FRHIResourceCreateInfo& CreateInfo, uint32& OutAlign) override;
 	virtual uint64 RHICalcTexture3DPlatformSize(uint32 SizeX, uint32 SizeY, uint32 SizeZ, uint8 Format, uint32 NumMips, ETextureCreateFlags InFlags, const FRHIResourceCreateInfo& CreateInfo, uint32& OutAlign) final override;
 	virtual uint64 RHICalcTextureCubePlatformSize(uint32 Size, uint8 Format, uint32 NumMips, ETextureCreateFlags InFlags, const FRHIResourceCreateInfo& CreateInfo, uint32& OutAlign) final override;
@@ -444,18 +441,18 @@ public:
 	virtual FTexture2DRHIRef RHICreateTexture2D_RenderThread(class FRHICommandListImmediate& RHICmdList, uint32 SizeX, uint32 SizeY, uint8 Format, uint32 NumMips, uint32 NumSamples, ETextureCreateFlags Flags, ERHIAccess InResourceState, FRHIResourceCreateInfo& CreateInfo) override;
 	virtual FTexture2DArrayRHIRef RHICreateTexture2DArray_RenderThread(class FRHICommandListImmediate& RHICmdList, uint32 SizeX, uint32 SizeY, uint32 SizeZ, uint8 Format, uint32 NumMips, uint32 NumSamples, ETextureCreateFlags Flags, ERHIAccess InResourceState, FRHIResourceCreateInfo& CreateInfo) override;
 	virtual FTexture3DRHIRef RHICreateTexture3D_RenderThread(class FRHICommandListImmediate& RHICmdList, uint32 SizeX, uint32 SizeY, uint32 SizeZ, uint8 Format, uint32 NumMips, ETextureCreateFlags Flags, ERHIAccess InResourceState, FRHIResourceCreateInfo& CreateInfo) override;
-	virtual FUnorderedAccessViewRHIRef RHICreateUnorderedAccessView_RenderThread(class FRHICommandListImmediate& RHICmdList, FRHIStructuredBuffer* StructuredBuffer, bool bUseUAVCounter, bool bAppendBuffer) override final;
+	virtual FUnorderedAccessViewRHIRef RHICreateUnorderedAccessView_RenderThread(class FRHICommandListImmediate& RHICmdList, FRHIBuffer* Buffer, bool bUseUAVCounter, bool bAppendBuffer) override final;
 	virtual FUnorderedAccessViewRHIRef RHICreateUnorderedAccessView_RenderThread(class FRHICommandListImmediate& RHICmdList, FRHITexture* Texture, uint32 MipLevel) override final;
-	virtual FUnorderedAccessViewRHIRef RHICreateUnorderedAccessView_RenderThread(class FRHICommandListImmediate& RHICmdList, FRHIVertexBuffer* VertexBuffer, uint8 Format) override final;
+	virtual FUnorderedAccessViewRHIRef RHICreateUnorderedAccessView_RenderThread(class FRHICommandListImmediate& RHICmdList, FRHIBuffer* Buffer, uint8 Format) override final;
 	virtual FShaderResourceViewRHIRef RHICreateShaderResourceView_RenderThread(class FRHICommandListImmediate& RHICmdList, FRHITexture* Texture, const FRHITextureSRVCreateInfo& CreateInfo) override final;
-	virtual FShaderResourceViewRHIRef RHICreateShaderResourceView_RenderThread(class FRHICommandListImmediate& RHICmdList, FRHIVertexBuffer* VertexBuffer, uint32 Stride, uint8 Format) override final;
+	virtual FShaderResourceViewRHIRef RHICreateShaderResourceView_RenderThread(class FRHICommandListImmediate& RHICmdList, FRHIBuffer* Buffer, uint32 Stride, uint8 Format) override final;
 	virtual FShaderResourceViewRHIRef RHICreateShaderResourceView_RenderThread(class FRHICommandListImmediate& RHICmdList, const FShaderResourceViewInitializer& Initializer) override final;
 	virtual FShaderResourceViewRHIRef RHICreateShaderResourceViewWriteMask_RenderThread(class FRHICommandListImmediate& RHICmdList, FRHITexture2D* Texture2D) override final;
 
-	virtual FShaderResourceViewRHIRef CreateShaderResourceView_RenderThread(class FRHICommandListImmediate& RHICmdList, FRHIVertexBuffer* VertexBuffer, uint32 Stride, uint8 Format) override final;
-	virtual FShaderResourceViewRHIRef CreateShaderResourceView_RenderThread(class FRHICommandListImmediate& RHICmdList, FRHIIndexBuffer* Buffer) override final;
+	virtual FShaderResourceViewRHIRef CreateShaderResourceView_RenderThread(class FRHICommandListImmediate& RHICmdList, FRHIBuffer* Buffer, uint32 Stride, uint8 Format) override final;
+	virtual FShaderResourceViewRHIRef CreateShaderResourceView_RenderThread(class FRHICommandListImmediate& RHICmdList, FRHIBuffer* Buffer) override final;
 	virtual FShaderResourceViewRHIRef CreateShaderResourceView_RenderThread(class FRHICommandListImmediate& RHICmdList, const FShaderResourceViewInitializer& Initializer) override final;
-	virtual FShaderResourceViewRHIRef RHICreateShaderResourceView_RenderThread(class FRHICommandListImmediate& RHICmdList, FRHIStructuredBuffer* StructuredBuffer) override final;
+	virtual FShaderResourceViewRHIRef RHICreateShaderResourceView_RenderThread(class FRHICommandListImmediate& RHICmdList, FRHIBuffer* Buffer) override final;
 	virtual FTextureCubeRHIRef RHICreateTextureCube_RenderThread(class FRHICommandListImmediate& RHICmdList, uint32 Size, uint8 Format, uint32 NumMips, ETextureCreateFlags Flags, ERHIAccess InResourceState, FRHIResourceCreateInfo& CreateInfo) override;
 	virtual FTextureCubeRHIRef RHICreateTextureCubeArray_RenderThread(class FRHICommandListImmediate& RHICmdList, uint32 Size, uint32 ArraySize, uint8 Format, uint32 NumMips, ETextureCreateFlags Flags, ERHIAccess InResourceState, FRHIResourceCreateInfo& CreateInfo) override;
 
diff --git a/Engine/Source/Runtime/D3D12RHI/Private/D3D12RayTracing.cpp b/Engine/Source/Runtime/D3D12RHI/Private/D3D12RayTracing.cpp
--- a/Engine/Source/Runtime/D3D12RHI/Private/D3D12RayTracing.cpp
+++ b/Engine/Source/Runtime/D3D12RHI/Private/D3D12RayTracing.cpp
@@ -1379,7 +1379,7 @@ public:
 		CreateInfo.ResourceArray = &Data;
 		CreateInfo.GPUMask = FRHIGPUMask::FromIndex(Device->GetGPUIndex());
 
-		Buffer = Adapter->CreateRHIBuffer<FD3D12MemBuffer>(
+		Buffer = Adapter->CreateRHIBuffer(
 			nullptr, BufferDesc, BufferDesc.Alignment,
 			0, BufferDesc.Width, BUF_Static, ED3D12ResourceStateMode::SingleState, CreateInfo);
 
@@ -1459,7 +1459,7 @@ public:
 	TResourceArray<uint8, D3D12_RAYTRACING_SHADER_TABLE_BYTE_ALIGNMENT> Data;
 
 	bool bIsDirty = true;
-	TRefCountPtr<FD3D12MemBuffer> Buffer;
+	TRefCountPtr<FD3D12Buffer> Buffer;
 
 	// SBTs have their own descriptor heaps
 	FD3D12RayTracingDescriptorCache* DescriptorCache = nullptr;
@@ -2291,7 +2291,7 @@ FD3D12RayTracingGeometry::FD3D12RayTracingGeometry(const FRayTracingGeometryInit
 
 FD3D12RayTracingGeometry::~FD3D12RayTracingGeometry()
 {
-	for (TRefCountPtr<FD3D12MemBuffer>& Buffer : AccelerationStructureBuffers)
+	for (TRefCountPtr<FD3D12Buffer>& Buffer : AccelerationStructureBuffers)
 	{
 		if (Buffer)
 		{
@@ -2300,7 +2300,7 @@ FD3D12RayTracingGeometry::~FD3D12RayTracingGeometry()
 		}
 	}
 
-	for (TRefCountPtr<FD3D12MemBuffer>& Buffer : ScratchBuffers)
+	for (TRefCountPtr<FD3D12Buffer>& Buffer : ScratchBuffers)
 	{
 		if (Buffer)
 		{
@@ -2337,8 +2337,8 @@ void FD3D12RayTracingGeometry::TransitionBuffers(FD3D12CommandContext& CommandCo
 }
 
 static void CreateAccelerationStructureBuffers(
-	TRefCountPtr<FD3D12MemBuffer>& AccelerationStructureBuffer,
-	TRefCountPtr<FD3D12MemBuffer>&  ScratchBuffer,
+	TRefCountPtr<FD3D12Buffer>& AccelerationStructureBuffer,
+	TRefCountPtr<FD3D12Buffer>&  ScratchBuffer,
 	FD3D12Adapter* Adapter, uint32 GPUIndex,
 	const D3D12_RAYTRACING_ACCELERATION_STRUCTURE_PREBUILD_INFO& PrebuildInfo,
 	const TCHAR* DebugName)
@@ -2350,7 +2350,7 @@ static void CreateAccelerationStructureBuffers(
 
 	CreateInfo.GPUMask = FRHIGPUMask::FromIndex(GPUIndex);
 	CreateInfo.DebugName = DebugName;
-	AccelerationStructureBuffer = Adapter->CreateRHIBuffer<FD3D12MemBuffer>(
+	AccelerationStructureBuffer = Adapter->CreateRHIBuffer(
 		nullptr, AccelerationStructureBufferDesc, D3D12_RAYTRACING_ACCELERATION_STRUCTURE_BYTE_ALIGNMENT,
 		0, AccelerationStructureBufferDesc.Width, BUF_AccelerationStructure, ED3D12ResourceStateMode::SingleState, CreateInfo);
 
@@ -2362,7 +2362,7 @@ static void CreateAccelerationStructureBuffers(
 
 	CreateInfo.GPUMask = FRHIGPUMask::FromIndex(GPUIndex);
 	CreateInfo.DebugName = TEXT("ScratchBuffer");
-	ScratchBuffer = Adapter->CreateRHIBuffer<FD3D12MemBuffer>(
+	ScratchBuffer = Adapter->CreateRHIBuffer(
 		nullptr, ScratchBufferDesc, D3D12_RAYTRACING_ACCELERATION_STRUCTURE_BYTE_ALIGNMENT,
 		0, ScratchBufferDesc.Width, BUF_UnorderedAccess, ED3D12ResourceStateMode::SingleState, CreateInfo);
 
@@ -2561,7 +2561,7 @@ void FD3D12RayTracingGeometry::BuildAccelerationStructure(FD3D12CommandContext&
 
 		CreateInfo.GPUMask = FRHIGPUMask::FromIndex(GPUIndex);
 		CreateInfo.DebugName = TEXT("PostBuildInfoBuffer");
-		PostBuildInfoBuffers[GPUIndex] = Adapter->CreateRHIBuffer<FD3D12MemBuffer>(
+		PostBuildInfoBuffers[GPUIndex] = Adapter->CreateRHIBuffer(
 			nullptr, PostBuildInfoBufferDesc, 8,
 			0, PostBuildInfoBufferDesc.Width, BUF_UnorderedAccess | BUF_SourceCopy, ED3D12ResourceStateMode::MultiState, CreateInfo);
 
@@ -2686,7 +2686,7 @@ void FD3D12RayTracingGeometry::ConditionalCompactAccelerationStructure(FD3D12Com
 	DEC_MEMORY_STAT_BY(STAT_D3D12RayTracingBLASMemory, AccelerationStructureBuffers[GPUIndex]->GetSize());
 
 	// Move old AS into this temporary variable which gets released when this function returns
-	TRefCountPtr<FD3D12MemBuffer> OldAccelerationStructure = MoveTemp(AccelerationStructureBuffers[GPUIndex]);
+	TRefCountPtr<FD3D12Buffer> OldAccelerationStructure = MoveTemp(AccelerationStructureBuffers[GPUIndex]);
 
 	{
 		FRHIResourceCreateInfo CreateInfo;
@@ -2695,7 +2695,7 @@ void FD3D12RayTracingGeometry::ConditionalCompactAccelerationStructure(FD3D12Com
 
 		CreateInfo.GPUMask = FRHIGPUMask::FromIndex(GPUIndex);
 		CreateInfo.DebugName = TEXT("AccelerationStructureBuffer");
-		AccelerationStructureBuffers[GPUIndex] = Adapter->CreateRHIBuffer<FD3D12MemBuffer>(
+		AccelerationStructureBuffers[GPUIndex] = Adapter->CreateRHIBuffer(
 			nullptr, AccelerationStructureBufferDesc, D3D12_RAYTRACING_ACCELERATION_STRUCTURE_BYTE_ALIGNMENT,
 			0, AccelerationStructureBufferDesc.Width, BUF_AccelerationStructure, ED3D12ResourceStateMode::SingleState, CreateInfo);
 
@@ -2790,8 +2790,8 @@ void FD3D12RayTracingScene::BuildAccelerationStructure(FD3D12CommandContext& Com
 	TRACE_CPUPROFILER_EVENT_SCOPE(BuildAccelerationStructure_TopLevel);
 	SCOPE_CYCLE_COUNTER(STAT_D3D12BuildTLAS);
 
-	TRefCountPtr<FD3D12StructuredBuffer> InstanceBuffer;
-	TRefCountPtr<FD3D12MemBuffer> ScratchBuffer;
+	TRefCountPtr<FD3D12Buffer> InstanceBuffer;
+	TRefCountPtr<FD3D12Buffer> ScratchBuffer;
 
 	const uint32 GPUIndex = CommandContext.GetGPUIndex();
 	FD3D12Adapter* Adapter = CommandContext.GetParentAdapter();
@@ -2818,7 +2818,7 @@ void FD3D12RayTracingScene::BuildAccelerationStructure(FD3D12CommandContext& Com
 
 	RayTracingDevice->GetRaytracingAccelerationStructurePrebuildInfo(&PrebuildDescInputs, &PrebuildInfo);
 
-	TRefCountPtr<FD3D12MemBuffer>& AccelerationStructureBuffer = AccelerationStructureBuffers[GPUIndex];
+	TRefCountPtr<FD3D12Buffer>& AccelerationStructureBuffer = AccelerationStructureBuffers[GPUIndex];
 
 	CreateAccelerationStructureBuffers(
 		AccelerationStructureBuffer,
@@ -2865,7 +2865,7 @@ void FD3D12RayTracingScene::BuildAccelerationStructure(FD3D12CommandContext& Com
 			D3D12_RESOURCE_FLAG_ALLOW_UNORDERED_ACCESS);
 
 		// #dxr_todo multi state only when bShouldCopyIndirectInstances is set - will still need transition to copy_src for lock behind but this can be done on the complete pool in theory (have to check cost)
-		InstanceBuffer = Adapter->CreateRHIBuffer<FD3D12StructuredBuffer>(
+		InstanceBuffer = Adapter->CreateRHIBuffer(
 			nullptr, InstanceBufferDesc, D3D12_RAYTRACING_INSTANCE_DESCS_BYTE_ALIGNMENT,
 			sizeof(D3D12_RAYTRACING_INSTANCE_DESC), InstanceBufferDesc.Width, BUF_UnorderedAccess, ED3D12ResourceStateMode::MultiState, CreateInfo);
 
diff --git a/Engine/Source/Runtime/D3D12RHI/Private/D3D12RayTracing.h b/Engine/Source/Runtime/D3D12RHI/Private/D3D12RayTracing.h
--- a/Engine/Source/Runtime/D3D12RHI/Private/D3D12RayTracing.h
+++ b/Engine/Source/Runtime/D3D12RHI/Private/D3D12RayTracing.h
@@ -11,8 +11,6 @@
 class FD3D12RayTracingPipelineState;
 class FD3D12RayTracingShaderTable;
 
-typedef FD3D12VertexBuffer FD3D12MemBuffer; // Generic GPU memory buffer
-
 // Built-in local root parameters that are always bound to all hit shaders
 struct FHitGroupSystemParameters
 {
@@ -57,11 +55,11 @@ public:
 	FIndexBufferRHIRef  RHIIndexBuffer;
 	static FVertexBufferRHIRef NullTransformBuffer; // Null transform for hidden sections
 
-	TRefCountPtr<FD3D12MemBuffer> AccelerationStructureBuffers[MAX_NUM_GPUS];
-	TRefCountPtr<FD3D12MemBuffer> ScratchBuffers[MAX_NUM_GPUS];
+	TRefCountPtr<FD3D12Buffer> AccelerationStructureBuffers[MAX_NUM_GPUS];
+	TRefCountPtr<FD3D12Buffer> ScratchBuffers[MAX_NUM_GPUS];
 
 	uint64 PostBuildInfoBufferReadbackFences[MAX_NUM_GPUS];
-	TRefCountPtr<FD3D12MemBuffer> PostBuildInfoBuffers[MAX_NUM_GPUS];
+	TRefCountPtr<FD3D12Buffer> PostBuildInfoBuffers[MAX_NUM_GPUS];
 	FStagingBufferRHIRef PostBuildInfoStagingBuffers[MAX_NUM_GPUS];
 
 	// Hit shader parameters per geometry segment
@@ -84,7 +82,7 @@ public:
 
 	void BuildAccelerationStructure(FD3D12CommandContext& CommandContext, D3D12_RAYTRACING_ACCELERATION_STRUCTURE_BUILD_FLAGS BuildFlags);
 
-	TRefCountPtr<FD3D12MemBuffer> AccelerationStructureBuffers[MAX_NUM_GPUS];
+	TRefCountPtr<FD3D12Buffer> AccelerationStructureBuffers[MAX_NUM_GPUS];
 	bool bAccelerationStructureViewInitialized[MAX_NUM_GPUS] = {};
 
 	TArray<FRayTracingGeometryInstance> Instances;
diff --git a/Engine/Source/Runtime/D3D12RHI/Private/D3D12Stats.cpp b/Engine/Source/Runtime/D3D12RHI/Private/D3D12Stats.cpp
--- a/Engine/Source/Runtime/D3D12RHI/Private/D3D12Stats.cpp
+++ b/Engine/Source/Runtime/D3D12RHI/Private/D3D12Stats.cpp
@@ -281,71 +281,22 @@ float FD3D12EventNode::GetTiming()
 	return Result;
 }
 
-void UpdateBufferStats(FD3D12ResourceLocation* ResourceLocation, bool bAllocating, uint32 BufferType)
+void UpdateBufferStats(FName Name, int64 RequestedSize)
 {
-	uint64 RequestedSize = ResourceLocation->GetSize();
-
-	const bool bUniformBuffer = BufferType == D3D12_BUFFER_TYPE_CONSTANT;
-	const bool bIndexBuffer = BufferType == D3D12_BUFFER_TYPE_INDEX;
-	const bool bVertexBuffer = BufferType == D3D12_BUFFER_TYPE_VERTEX;
-
-	if (bAllocating)
-	{
-		if (bUniformBuffer)
-		{
-			INC_MEMORY_STAT_BY(STAT_UniformBufferMemory, RequestedSize);
-		}
-		else if (bIndexBuffer)
-		{
-			INC_MEMORY_STAT_BY(STAT_IndexBufferMemory, RequestedSize);
-		}
-		else if (bVertexBuffer)
-		{
-			INC_MEMORY_STAT_BY(STAT_VertexBufferMemory, RequestedSize);
-		}
-		else
-		{
-			INC_MEMORY_STAT_BY(STAT_StructuredBufferMemory, RequestedSize);
-		}
-
-#if PLATFORM_WINDOWS
-		// this is a work-around on Windows. Due to the fact that there is no way
-		// to hook the actual d3d allocations it is very difficult to track memory
-		// in the normal way. The problem is that some buffers are allocated from
-		// the allocators and some are allocated from the device. Ideally this
-		// tracking would be moved to where the actual d3d resource is created and
-		// released and the tracking could be re-enabled in the buddy allocator.
-		// The problem is that the releasing of resources happens in a generic way
-		// (see FD3D12ResourceLocation) 
-		LLM_SCOPED_PAUSE_TRACKING_WITH_ENUM_AND_AMOUNT(ELLMTag::Meshes, RequestedSize, ELLMTracker::Default, ELLMAllocType::None);
-		LLM_SCOPED_PAUSE_TRACKING_WITH_ENUM_AND_AMOUNT(ELLMTag::GraphicsPlatform, RequestedSize, ELLMTracker::Platform, ELLMAllocType::None);
-#endif
-	}
-	else
-	{
-		if (bUniformBuffer)
-		{
-			DEC_MEMORY_STAT_BY(STAT_UniformBufferMemory, RequestedSize);
-		}
-		else if (bIndexBuffer)
-		{
-			DEC_MEMORY_STAT_BY(STAT_IndexBufferMemory, RequestedSize);
-		}
-		else if (bVertexBuffer)
-		{
-			DEC_MEMORY_STAT_BY(STAT_VertexBufferMemory, RequestedSize);
-		}
-		else
-		{
-			DEC_MEMORY_STAT_BY(STAT_StructuredBufferMemory, RequestedSize);
-		}
+	INC_MEMORY_STAT_BY_FName(Name, RequestedSize);
 
 #if PLATFORM_WINDOWS
-		// this is a work-around on Windows. See comment above.
-		LLM_SCOPED_PAUSE_TRACKING_WITH_ENUM_AND_AMOUNT(ELLMTag::Meshes, -(int64)RequestedSize, ELLMTracker::Default, ELLMAllocType::None);
-		LLM_SCOPED_PAUSE_TRACKING_WITH_ENUM_AND_AMOUNT(ELLMTag::GraphicsPlatform, -(int64)RequestedSize, ELLMTracker::Platform, ELLMAllocType::None);
+	// this is a work-around on Windows. Due to the fact that there is no way
+	// to hook the actual d3d allocations it is very difficult to track memory
+	// in the normal way. The problem is that some buffers are allocated from
+	// the allocators and some are allocated from the device. Ideally this
+	// tracking would be moved to where the actual d3d resource is created and
+	// released and the tracking could be re-enabled in the buddy allocator.
+	// The problem is that the releasing of resources happens in a generic way
+	// (see FD3D12ResourceLocation) 
+	LLM_SCOPED_PAUSE_TRACKING_WITH_ENUM_AND_AMOUNT(ELLMTag::Meshes, RequestedSize, ELLMTracker::Default, ELLMAllocType::None);
+	LLM_SCOPED_PAUSE_TRACKING_WITH_ENUM_AND_AMOUNT(ELLMTag::GraphicsPlatform, RequestedSize, ELLMTracker::Platform, ELLMAllocType::None);
 #endif
-	}
 }
 
 #if NV_AFTERMATH
diff --git a/Engine/Source/Runtime/D3D12RHI/Private/D3D12StructuredBuffer.cpp b/Engine/Source/Runtime/D3D12RHI/Private/D3D12StructuredBuffer.cpp
--- a/Engine/Source/Runtime/D3D12RHI/Private/D3D12StructuredBuffer.cpp
+++ b/Engine/Source/Runtime/D3D12RHI/Private/D3D12StructuredBuffer.cpp
@@ -38,7 +38,7 @@ FStructuredBufferRHIRef FD3D12DynamicRHI::CreateStructuredBuffer_RenderThread(FR
 	// can be addressed correctly with element based offsets.
 	const uint32 Alignment = ((InUsage & (BUF_ByteAddressBuffer | BUF_DrawIndirect)) == 0) ? Stride : 4;
 
-	FD3D12StructuredBuffer* NewBuffer = GetAdapter().CreateRHIBuffer<FD3D12StructuredBuffer>(&RHICmdList, Desc, Alignment, Stride, Size, InUsage, ED3D12ResourceStateMode::Default, CreateInfo);
+	FD3D12Buffer* NewBuffer = GetAdapter().CreateRHIBuffer(&RHICmdList, Desc, Alignment, Stride, Size, InUsage | BUF_StructuredBuffer, ED3D12ResourceStateMode::Default, CreateInfo);
 	if (NewBuffer->ResourceLocation.IsTransient())
 	{
 		// TODO: this should ideally be set in platform-independent code, since this tracking is for the high level
@@ -59,7 +59,7 @@ FStructuredBufferRHIRef FD3D12DynamicRHI::RHICreateStructuredBuffer(uint32 Strid
 	// can be addressed correctly with element based offsets.
 	const uint32 Alignment = ((InUsage & (BUF_ByteAddressBuffer | BUF_DrawIndirect)) == 0) ? Stride : 4;
 
-	FD3D12StructuredBuffer* NewBuffer = GetAdapter().CreateRHIBuffer<FD3D12StructuredBuffer>(nullptr, Desc, Alignment, Stride, Size, InUsage, ED3D12ResourceStateMode::Default, CreateInfo);
+	FD3D12Buffer* NewBuffer = GetAdapter().CreateRHIBuffer(nullptr, Desc, Alignment, Stride, Size, InUsage | BUF_StructuredBuffer, ED3D12ResourceStateMode::Default, CreateInfo);
 	if (NewBuffer->ResourceLocation.IsTransient())
 	{
 		// TODO: this should ideally be set in platform-independent code, since this tracking is for the high level
@@ -69,19 +69,14 @@ FStructuredBufferRHIRef FD3D12DynamicRHI::RHICreateStructuredBuffer(uint32 Strid
 	return NewBuffer;
 }
 
-FD3D12StructuredBuffer::~FD3D12StructuredBuffer()
-{
-	UpdateBufferStats<FD3D12StructuredBuffer>(&ResourceLocation, false);
-}
-
 void* FD3D12DynamicRHI::RHILockStructuredBuffer(FRHICommandListImmediate& RHICmdList, FRHIStructuredBuffer* StructuredBufferRHI, uint32 Offset, uint32 Size, EResourceLockMode LockMode)
 {
-	FD3D12StructuredBuffer* Buffer = FD3D12DynamicRHI::ResourceCast(StructuredBufferRHI);
+	FD3D12Buffer* Buffer = FD3D12DynamicRHI::ResourceCast(StructuredBufferRHI);
 	return LockBuffer(&RHICmdList, Buffer, Buffer->GetSize(), Buffer->GetUsage(), Offset, Size, LockMode);
 }
 
 void FD3D12DynamicRHI::RHIUnlockStructuredBuffer(FRHICommandListImmediate& RHICmdList, FRHIStructuredBuffer* StructuredBufferRHI)
 {
-	FD3D12StructuredBuffer* Buffer = FD3D12DynamicRHI::ResourceCast(StructuredBufferRHI);
+	FD3D12Buffer* Buffer = FD3D12DynamicRHI::ResourceCast(StructuredBufferRHI);
 	UnlockBuffer(&RHICmdList, Buffer, Buffer->GetUsage());
 }
diff --git a/Engine/Source/Runtime/D3D12RHI/Private/D3D12UAV.cpp b/Engine/Source/Runtime/D3D12RHI/Private/D3D12UAV.cpp
--- a/Engine/Source/Runtime/D3D12RHI/Private/D3D12UAV.cpp
+++ b/Engine/Source/Runtime/D3D12RHI/Private/D3D12UAV.cpp
@@ -28,13 +28,13 @@ inline FD3D12UnorderedAccessView* CreateUAV(D3D12_UNORDERED_ACCESS_VIEW_DESC& De
 	});
 }
 
-FUnorderedAccessViewRHIRef FD3D12DynamicRHI::RHICreateUnorderedAccessView(FRHIStructuredBuffer* StructuredBufferRHI, bool bUseUAVCounter, bool bAppendBuffer)
+FUnorderedAccessViewRHIRef FD3D12DynamicRHI::RHICreateUnorderedAccessView(FRHIBuffer* BufferRHI, bool bUseUAVCounter, bool bAppendBuffer)
 {
-	FD3D12StructuredBuffer*  StructuredBuffer = FD3D12DynamicRHI::ResourceCast(StructuredBufferRHI);
+	FD3D12Buffer* Buffer = FD3D12DynamicRHI::ResourceCast(BufferRHI);
 
-	FD3D12ResourceLocation& Location = StructuredBuffer->ResourceLocation;
+	FD3D12ResourceLocation& Location = Buffer->ResourceLocation;
 
-	const uint32 BufferUsage = StructuredBuffer->GetUsage();
+	const uint32 BufferUsage = Buffer->GetUsage();
 	const bool bByteAccessBuffer = (BufferUsage & BUF_ByteAddressBuffer) != 0;
 	const bool bStructuredBuffer = !bByteAccessBuffer;
 	check(bByteAccessBuffer != bStructuredBuffer); // You can't have a structured buffer that allows raw views
@@ -43,7 +43,7 @@ FUnorderedAccessViewRHIRef FD3D12DynamicRHI::RHICreateUnorderedAccessView(FRHISt
 	UAVDesc.ViewDimension = D3D12_UAV_DIMENSION_BUFFER;
 	UAVDesc.Format = DXGI_FORMAT_UNKNOWN;
 
-	uint32 EffectiveStride = StructuredBuffer->GetStride();
+	uint32 EffectiveStride = Buffer->GetStride();
 
 	if (bByteAccessBuffer)
 	{
@@ -62,7 +62,7 @@ FUnorderedAccessViewRHIRef FD3D12DynamicRHI::RHICreateUnorderedAccessView(FRHISt
 	UAVDesc.Buffer.StructureByteStride = bStructuredBuffer ? EffectiveStride : 0;
 
 	const bool bNeedsCounterResource = bAppendBuffer | bUseUAVCounter;
-	return CreateUAV(UAVDesc, StructuredBuffer, bNeedsCounterResource);
+	return CreateUAV(UAVDesc, Buffer, bNeedsCounterResource);
 }
 
 FUnorderedAccessViewRHIRef FD3D12DynamicRHI::RHICreateUnorderedAccessView(FRHITexture* TextureRHI, uint32 MipLevel)
@@ -117,16 +117,16 @@ FUnorderedAccessViewRHIRef FD3D12DynamicRHI::RHICreateUnorderedAccessView(FRHITe
 	}
 }
 
-FUnorderedAccessViewRHIRef FD3D12DynamicRHI::RHICreateUnorderedAccessView(FRHIVertexBuffer* VertexBufferRHI, uint8 Format)
+FUnorderedAccessViewRHIRef FD3D12DynamicRHI::RHICreateUnorderedAccessView(FRHIBuffer* BufferRHI, uint8 Format)
 {
-	FD3D12VertexBuffer*  VertexBuffer = FD3D12DynamicRHI::ResourceCast(VertexBufferRHI);
-	FD3D12ResourceLocation& Location = VertexBuffer->ResourceLocation;
+	FD3D12Buffer* Buffer = FD3D12DynamicRHI::ResourceCast(BufferRHI);
+	FD3D12ResourceLocation& Location = Buffer->ResourceLocation;
 
 	D3D12_UNORDERED_ACCESS_VIEW_DESC UAVDesc{};
 	UAVDesc.ViewDimension = D3D12_UAV_DIMENSION_BUFFER;
 
 	uint32 EffectiveStride;
-	if (VertexBuffer->GetUsage() & BUF_ByteAddressBuffer)
+	if (Buffer->GetUsage() & BUF_ByteAddressBuffer)
 	{
 		UAVDesc.Format = DXGI_FORMAT_R32_TYPELESS;
 		UAVDesc.Buffer.Flags |= D3D12_BUFFER_UAV_FLAG_RAW;
@@ -141,34 +141,7 @@ FUnorderedAccessViewRHIRef FD3D12DynamicRHI::RHICreateUnorderedAccessView(FRHIVe
 	UAVDesc.Buffer.FirstElement = Location.GetOffsetFromBaseOfResource() / EffectiveStride;
 	UAVDesc.Buffer.NumElements = Location.GetSize() / EffectiveStride;
 
-	return CreateUAV(UAVDesc, VertexBuffer, false);
-}
-
-FUnorderedAccessViewRHIRef FD3D12DynamicRHI::RHICreateUnorderedAccessView(FRHIIndexBuffer* IndexBufferRHI, uint8 Format)
-{
-	FD3D12IndexBuffer* IndexBuffer = FD3D12DynamicRHI::ResourceCast(IndexBufferRHI);
-	FD3D12ResourceLocation& Location = IndexBuffer->ResourceLocation;
-
-	D3D12_UNORDERED_ACCESS_VIEW_DESC UAVDesc{};
-	UAVDesc.ViewDimension = D3D12_UAV_DIMENSION_BUFFER;
-
-	uint32 EffectiveStride;
-	if ((IndexBuffer->GetUsage() & BUF_ByteAddressBuffer) != 0)
-	{
-		UAVDesc.Format = DXGI_FORMAT_R32_TYPELESS;
-		UAVDesc.Buffer.Flags |= D3D12_BUFFER_UAV_FLAG_RAW;
-		EffectiveStride = 4;
-	}
-	else
-	{
-		UAVDesc.Format = FindUnorderedAccessDXGIFormat((DXGI_FORMAT)GPixelFormats[Format].PlatformFormat);
-		EffectiveStride = GPixelFormats[Format].BlockBytes;
-	}
-
-	UAVDesc.Buffer.FirstElement = Location.GetOffsetFromBaseOfResource() / EffectiveStride;
-	UAVDesc.Buffer.NumElements = Location.GetSize() / EffectiveStride;
-
-	return CreateUAV(UAVDesc, IndexBuffer, false);
+	return CreateUAV(UAVDesc, Buffer, false);
 }
 
 void FD3D12CommandContext::ClearUAV(TRHICommandList_RecursiveHazardous<FD3D12CommandContext>& RHICmdList, FD3D12UnorderedAccessView* UnorderedAccessView, const void* ClearValues, bool bFloat)
@@ -323,17 +296,17 @@ void FD3D12CommandContext::RHIClearUAVUint(FRHIUnorderedAccessView* UnorderedAcc
 	ClearUAV(RHICmdList, ResourceCast(UnorderedAccessViewRHI), &Values, false);
 }
 
-FUnorderedAccessViewRHIRef FD3D12DynamicRHI::RHICreateUnorderedAccessView_RenderThread(FRHICommandListImmediate& RHICmdList, FRHIStructuredBuffer* StructuredBufferRHI, bool bUseUAVCounter, bool bAppendBuffer)
+FUnorderedAccessViewRHIRef FD3D12DynamicRHI::RHICreateUnorderedAccessView_RenderThread(FRHICommandListImmediate& RHICmdList, FRHIBuffer* BufferRHI, bool bUseUAVCounter, bool bAppendBuffer)
 {
-	FD3D12StructuredBuffer* StructuredBuffer = FD3D12DynamicRHI::ResourceCast(StructuredBufferRHI);
+	FD3D12Buffer* Buffer = FD3D12DynamicRHI::ResourceCast(BufferRHI);
 	// TODO: we have to stall the RHI thread when creating SRVs of dynamic buffers because they get renamed.
 	// perhaps we could do a deferred operation?
-	if (StructuredBuffer->GetUsage() & BUF_AnyDynamic)
+	if (Buffer->GetUsage() & BUF_AnyDynamic)
 	{
 		FScopedRHIThreadStaller StallRHIThread(RHICmdList);
-		return RHICreateUnorderedAccessView(StructuredBufferRHI, bUseUAVCounter, bAppendBuffer);
+		return RHICreateUnorderedAccessView(BufferRHI, bUseUAVCounter, bAppendBuffer);
 	}
-	return RHICreateUnorderedAccessView(StructuredBufferRHI, bUseUAVCounter, bAppendBuffer);
+	return RHICreateUnorderedAccessView(BufferRHI, bUseUAVCounter, bAppendBuffer);
 }
 
 FUnorderedAccessViewRHIRef FD3D12DynamicRHI::RHICreateUnorderedAccessView_RenderThread(FRHICommandListImmediate& RHICmdList, FRHITexture* Texture, uint32 MipLevel)
@@ -341,18 +314,18 @@ FUnorderedAccessViewRHIRef FD3D12DynamicRHI::RHICreateUnorderedAccessView_Render
 	return FD3D12DynamicRHI::RHICreateUnorderedAccessView(Texture, MipLevel);
 }
 
-FUnorderedAccessViewRHIRef FD3D12DynamicRHI::RHICreateUnorderedAccessView_RenderThread(FRHICommandListImmediate& RHICmdList, FRHIVertexBuffer* VertexBufferRHI, uint8 Format)
+FUnorderedAccessViewRHIRef FD3D12DynamicRHI::RHICreateUnorderedAccessView_RenderThread(FRHICommandListImmediate& RHICmdList, FRHIBuffer* BufferRHI, uint8 Format)
 {
-	FD3D12VertexBuffer* VertexBuffer = FD3D12DynamicRHI::ResourceCast(VertexBufferRHI);
+	FD3D12Buffer* Buffer = FD3D12DynamicRHI::ResourceCast(BufferRHI);
 
 	// TODO: we have to stall the RHI thread when creating SRVs of dynamic buffers because they get renamed.
 	// perhaps we could do a deferred operation?
-	if (VertexBuffer->GetUsage() & BUF_AnyDynamic)
+	if (Buffer->GetUsage() & BUF_AnyDynamic)
 	{
 		FScopedRHIThreadStaller StallRHIThread(RHICmdList);
-		return RHICreateUnorderedAccessView(VertexBufferRHI, Format);
+		return RHICreateUnorderedAccessView(BufferRHI, Format);
 	}
-	return RHICreateUnorderedAccessView(VertexBufferRHI, Format);
+	return RHICreateUnorderedAccessView(BufferRHI, Format);
 }
 
 FD3D12StagingBuffer::~FD3D12StagingBuffer()
diff --git a/Engine/Source/Runtime/D3D12RHI/Private/D3D12UniformBuffer.cpp b/Engine/Source/Runtime/D3D12RHI/Private/D3D12UniformBuffer.cpp
--- a/Engine/Source/Runtime/D3D12RHI/Private/D3D12UniformBuffer.cpp
+++ b/Engine/Source/Runtime/D3D12RHI/Private/D3D12UniformBuffer.cpp
@@ -88,7 +88,7 @@ FUniformBufferRHIRef FD3D12DynamicRHI::RHICreateUniformBuffer(const void* Conten
 		}
 	}
 
-	UpdateBufferStats<FD3D12UniformBuffer>(&UniformBufferOut->ResourceLocation, true);
+	UpdateBufferStats(GET_STATFNAME(STAT_UniformBufferMemory), UniformBufferOut->ResourceLocation.GetSize());
 
 	return UniformBufferOut;
 }
@@ -212,7 +212,8 @@ FD3D12UniformBuffer::~FD3D12UniformBuffer()
 {
 	check(!GRHISupportsRHIThread || IsInRenderingThread());
 
-	UpdateBufferStats<FD3D12UniformBuffer>(&ResourceLocation, false);
+	int64 BufferSize = ResourceLocation.GetSize();
+	UpdateBufferStats(GET_STATFNAME(STAT_UniformBufferMemory), -BufferSize);
 
 #if USE_STATIC_ROOT_SIGNATURE
 	delete View;
