Replace RHICopyVertexBuffer() with generic RHICopyBuffer() function

diff --git a/Engine/Source/Runtime/Apple/MetalRHI/Private/MetalDynamicRHI.h b/Engine/Source/Runtime/Apple/MetalRHI/Private/MetalDynamicRHI.h
--- a/Engine/Source/Runtime/Apple/MetalRHI/Private/MetalDynamicRHI.h
+++ b/Engine/Source/Runtime/Apple/MetalRHI/Private/MetalDynamicRHI.h
@@ -58,7 +58,7 @@ public:
 	virtual FIndexBufferRHIRef RHICreateIndexBuffer(uint32 Stride, uint32 Size, uint32 InUsage, ERHIAccess InResourceState, FRHIResourceCreateInfo& CreateInfo) final override;
 	virtual void RHITransferBufferUnderlyingResource(FRHIBuffer* DestBuffer, FRHIBuffer* SrcBuffer) final override;
 	virtual FVertexBufferRHIRef RHICreateVertexBuffer(uint32 Size, uint32 InUsage, ERHIAccess InResourceState, FRHIResourceCreateInfo& CreateInfo) final override;
-	virtual void RHICopyVertexBuffer(FRHIVertexBuffer* SourceBuffer, FRHIVertexBuffer* DestBuffer) final override;
+	virtual void RHICopyBuffer(FRHIBuffer* SourceBuffer, FRHIBuffer* DestBuffer) final override;
 	virtual FStructuredBufferRHIRef RHICreateStructuredBuffer(uint32 Stride, uint32 Size, uint32 InUsage, ERHIAccess InResourceState, FRHIResourceCreateInfo& CreateInfo) final override;
 	virtual void* LockBuffer_BottomOfPipe(FRHICommandListImmediate& RHICmdList, FRHIBuffer* Buffer, uint32 Offset, uint32 SizeRHI, EResourceLockMode LockMode) final override;
 	virtual void UnlockBuffer_BottomOfPipe(FRHICommandListImmediate& RHICmdList, FRHIBuffer* Buffer) final override;
diff --git a/Engine/Source/Runtime/Apple/MetalRHI/Private/MetalTexture.cpp b/Engine/Source/Runtime/Apple/MetalRHI/Private/MetalTexture.cpp
--- a/Engine/Source/Runtime/Apple/MetalRHI/Private/MetalTexture.cpp
+++ b/Engine/Source/Runtime/Apple/MetalRHI/Private/MetalTexture.cpp
@@ -1,7 +1,7 @@
 // Copyright Epic Games, Inc. All Rights Reserved.
 
 /*=============================================================================
- MetalVertexBuffer.cpp: Metal texture RHI implementation.
+ MetalTexture.cpp: Metal texture RHI implementation.
  =============================================================================*/
 
 #include "MetalRHIPrivate.h"
@@ -2804,7 +2804,7 @@ void FMetalRHICommandContext::RHICopyTexture(FRHITexture* SourceTextureRHI, FRHI
 	}
 }
 
-void FMetalRHICommandContext::RHICopyBufferRegion(FRHIVertexBuffer* DstBufferRHI, uint64 DstOffset, FRHIVertexBuffer* SrcBufferRHI, uint64 SrcOffset, uint64 NumBytes)
+void FMetalRHICommandContext::RHICopyBufferRegion(FRHIBuffer* DstBufferRHI, uint64 DstOffset, FRHIBuffer* SrcBufferRHI, uint64 SrcOffset, uint64 NumBytes)
 {
 	if (!DstBufferRHI || !SrcBufferRHI || DstBufferRHI == SrcBufferRHI || !NumBytes)
 	{
@@ -2812,13 +2812,13 @@ void FMetalRHICommandContext::RHICopyBufferRegion(FRHIVertexBuffer* DstBufferRHI
 	}
 
 	@autoreleasepool {
-		FMetalVertexBuffer* DstVertexBuffer = ResourceCast(DstBufferRHI);
-		FMetalVertexBuffer* SrcVertexBuffer = ResourceCast(SrcBufferRHI);
+		FMetalResourceMultiBuffer* DstBuffer = ResourceCast(DstBufferRHI);
+		FMetalResourceMultiBuffer* SrcBuffer = ResourceCast(SrcBufferRHI);
 
-		check(DstVertexBuffer && SrcVertexBuffer);
-		check(!DstVertexBuffer->Data && !SrcVertexBuffer->Data);
+		check(DstBuffer && SrcBuffer);
+		check(!DstBuffer->Data && !SrcBuffer->Data);
 		check(DstOffset + NumBytes <= DstBufferRHI->GetSize() && SrcOffset + NumBytes <= SrcBufferRHI->GetSize());
 
-		GetInternalContext().CopyFromBufferToBuffer(SrcVertexBuffer->GetCurrentBuffer(), SrcOffset, DstVertexBuffer->GetCurrentBuffer(), DstOffset, NumBytes);
+		GetInternalContext().CopyFromBufferToBuffer(SrcBuffer->GetCurrentBuffer(), SrcOffset, DstBuffer->GetCurrentBuffer(), DstOffset, NumBytes);
 	}
 }
diff --git a/Engine/Source/Runtime/Apple/MetalRHI/Private/MetalVertexBuffer.cpp b/Engine/Source/Runtime/Apple/MetalRHI/Private/MetalVertexBuffer.cpp
--- a/Engine/Source/Runtime/Apple/MetalRHI/Private/MetalVertexBuffer.cpp
+++ b/Engine/Source/Runtime/Apple/MetalRHI/Private/MetalVertexBuffer.cpp
@@ -795,34 +795,34 @@ void FMetalDynamicRHI::UnlockBuffer_BottomOfPipe(FRHICommandListImmediate& RHICm
 	}
 }
 
-void FMetalDynamicRHI::RHICopyVertexBuffer(FRHIVertexBuffer* SourceBufferRHI, FRHIVertexBuffer* DestBufferRHI)
+void FMetalDynamicRHI::RHICopyBuffer(FRHIBuffer* SourceBufferRHI, FRHIBuffer* DestBufferRHI)
 {
 	@autoreleasepool {
-		FMetalVertexBuffer* SrcVertexBuffer = ResourceCast(SourceBufferRHI);
-		FMetalVertexBuffer* DstVertexBuffer = ResourceCast(DestBufferRHI);
+		FMetalResourceMultiBuffer* SrcBuffer = ResourceCast(SourceBufferRHI);
+		FMetalResourceMultiBuffer* DstBuffer = ResourceCast(DestBufferRHI);
 		
-		const FMetalBuffer& TheSrcBuffer = SrcVertexBuffer->GetCurrentBuffer();
-		const FMetalBuffer& TheDstBuffer = DstVertexBuffer->GetCurrentBuffer();
+		const FMetalBuffer& TheSrcBuffer = SrcBuffer->GetCurrentBuffer();
+		const FMetalBuffer& TheDstBuffer = DstBuffer->GetCurrentBuffer();
 	
 		if (TheSrcBuffer && TheDstBuffer)
 		{
-			GetMetalDeviceContext().CopyFromBufferToBuffer(TheSrcBuffer, 0, TheDstBuffer, 0, FMath::Min(SrcVertexBuffer->GetSize(), DstVertexBuffer->GetSize()));
+			GetMetalDeviceContext().CopyFromBufferToBuffer(TheSrcBuffer, 0, TheDstBuffer, 0, FMath::Min(SrcBuffer->GetSize(), DstBuffer->GetSize()));
 		}
 		else if (TheDstBuffer)
 		{
-			FMetalPooledBufferArgs ArgsCPU(GetMetalDeviceContext().GetDevice(), SrcVertexBuffer->GetSize(), BUF_Dynamic, mtlpp::StorageMode::Shared);
+			FMetalPooledBufferArgs ArgsCPU(GetMetalDeviceContext().GetDevice(), SrcBuffer->GetSize(), BUF_Dynamic, mtlpp::StorageMode::Shared);
 			FMetalBuffer TempBuffer = GetMetalDeviceContext().CreatePooledBuffer(ArgsCPU);
-			FMemory::Memcpy(TempBuffer.GetContents(), SrcVertexBuffer->Data->Data, SrcVertexBuffer->GetSize());
-			GetMetalDeviceContext().CopyFromBufferToBuffer(TempBuffer, 0, TheDstBuffer, 0, FMath::Min(SrcVertexBuffer->GetSize(), DstVertexBuffer->GetSize()));
+			FMemory::Memcpy(TempBuffer.GetContents(), SrcBuffer->Data->Data, SrcBuffer->GetSize());
+			GetMetalDeviceContext().CopyFromBufferToBuffer(TempBuffer, 0, TheDstBuffer, 0, FMath::Min(SrcBuffer->GetSize(), DstBuffer->GetSize()));
 			SafeReleaseMetalBuffer(TempBuffer);
 		}
 		else
 		{
-			void const* SrcData = SrcVertexBuffer->Lock(true, RLM_ReadOnly, 0);
-			void* DstData = DstVertexBuffer->Lock(true, RLM_WriteOnly, 0);
-			FMemory::Memcpy(DstData, SrcData, FMath::Min(SrcVertexBuffer->GetSize(), DstVertexBuffer->GetSize()));
-			SrcVertexBuffer->Unlock();
-			DstVertexBuffer->Unlock();
+			void const* SrcData = SrcBuffer->Lock(true, RLM_ReadOnly, 0);
+			void* DstData = DstBuffer->Lock(true, RLM_WriteOnly, 0);
+			FMemory::Memcpy(DstData, SrcData, FMath::Min(SrcBuffer->GetSize(), DstBuffer->GetSize()));
+			SrcBuffer->Unlock();
+			DstBuffer->Unlock();
 		}
 	}
 }
diff --git a/Engine/Source/Runtime/Apple/MetalRHI/Public/MetalRHIContext.h b/Engine/Source/Runtime/Apple/MetalRHI/Public/MetalRHIContext.h
--- a/Engine/Source/Runtime/Apple/MetalRHI/Public/MetalRHIContext.h
+++ b/Engine/Source/Runtime/Apple/MetalRHI/Public/MetalRHIContext.h
@@ -58,7 +58,7 @@ public:
 	
 	virtual void RHICopyTexture(FRHITexture* SourceTextureRHI, FRHITexture* DestTextureRHI, const FRHICopyTextureInfo& CopyInfo) final override;
 	
-	virtual void RHICopyBufferRegion(FRHIVertexBuffer* DstBufferRHI, uint64 DstOffset, FRHIVertexBuffer* SrcBufferRHI, uint64 SrcOffset, uint64 NumBytes) final override;
+	virtual void RHICopyBufferRegion(FRHIBuffer* DstBufferRHI, uint64 DstOffset, FRHIBuffer* SrcBufferRHI, uint64 SrcOffset, uint64 NumBytes) final override;
 
 	/**
 	 * Resolves from one texture to another.
diff --git a/Engine/Source/Runtime/D3D12RHI/Private/D3D12CommandContext.h b/Engine/Source/Runtime/D3D12RHI/Private/D3D12CommandContext.h
--- a/Engine/Source/Runtime/D3D12RHI/Private/D3D12CommandContext.h
+++ b/Engine/Source/Runtime/D3D12RHI/Private/D3D12CommandContext.h
@@ -363,7 +363,7 @@ public:
 	virtual void RHIBroadcastTemporalEffect(const FName& InEffectName, const TArrayView<FRHITexture*> InTextures) final AFR_API_OVERRIDE;
 
 #if D3D12_RHI_RAYTRACING
-	virtual void RHICopyBufferRegion(FRHIVertexBuffer* DestBuffer, uint64 DstOffset, FRHIVertexBuffer* SourceBuffer, uint64 SrcOffset, uint64 NumBytes) final override;
+	virtual void RHICopyBufferRegion(FRHIBuffer* DestBuffer, uint64 DstOffset, FRHIBuffer* SourceBuffer, uint64 SrcOffset, uint64 NumBytes) final override;
 	virtual void RHICopyBufferRegions(const TArrayView<const FCopyBufferRegionParams> Params) final override;
 	virtual void RHIBuildAccelerationStructures(const TArrayView<const FAccelerationStructureBuildParams> Params) final override;
 	virtual void RHIBuildAccelerationStructure(FRHIRayTracingScene* Scene) final override;
diff --git a/Engine/Source/Runtime/D3D12RHI/Private/D3D12RHIPrivate.h b/Engine/Source/Runtime/D3D12RHI/Private/D3D12RHIPrivate.h
--- a/Engine/Source/Runtime/D3D12RHI/Private/D3D12RHIPrivate.h
+++ b/Engine/Source/Runtime/D3D12RHI/Private/D3D12RHIPrivate.h
@@ -281,7 +281,7 @@ public:
 	virtual FIndexBufferRHIRef RHICreateIndexBuffer(uint32 Stride, uint32 Size, uint32 InUsage, ERHIAccess InResourceState, FRHIResourceCreateInfo& CreateInfo) final override;
 	virtual void RHITransferBufferUnderlyingResource(FRHIBuffer* DestBuffer, FRHIBuffer* SrcBuffer) final override;
 	virtual FVertexBufferRHIRef RHICreateVertexBuffer(uint32 Size, uint32 InUsage, ERHIAccess InResourceState, FRHIResourceCreateInfo& CreateInfo) final override;
-	virtual void RHICopyVertexBuffer(FRHIVertexBuffer* SourceBuffer, FRHIVertexBuffer* DestBuffer) final override;
+	virtual void RHICopyBuffer(FRHIBuffer* SourceBuffer, FRHIBuffer* DestBuffer) final override;
 	virtual FStructuredBufferRHIRef RHICreateStructuredBuffer(uint32 Stride, uint32 Size, uint32 InUsage, ERHIAccess InResourceState, FRHIResourceCreateInfo& CreateInfo) final override;
 	virtual void* RHILockBuffer(FRHICommandListImmediate& RHICmdList, FRHIBuffer* Buffer, uint32 Offset, uint32 Size, EResourceLockMode LockMode) final override;
 	virtual void RHIUnlockBuffer(FRHICommandListImmediate& RHICmdList, FRHIBuffer* Buffer) final override;
diff --git a/Engine/Source/Runtime/D3D12RHI/Private/D3D12VertexBuffer.cpp b/Engine/Source/Runtime/D3D12RHI/Private/D3D12VertexBuffer.cpp
--- a/Engine/Source/Runtime/D3D12RHI/Private/D3D12VertexBuffer.cpp
+++ b/Engine/Source/Runtime/D3D12RHI/Private/D3D12VertexBuffer.cpp
@@ -82,7 +82,7 @@ FVertexBufferRHIRef FD3D12DynamicRHI::CreateVertexBuffer_RenderThread(FRHIComman
 	return Buffer;
 }
 
-void FD3D12DynamicRHI::RHICopyVertexBuffer(FRHIVertexBuffer* SourceBufferRHI, FRHIVertexBuffer* DestBufferRHI)
+void FD3D12DynamicRHI::RHICopyBuffer(FRHIBuffer* SourceBufferRHI, FRHIBuffer* DestBufferRHI)
 {
 	FD3D12Buffer* SrcBuffer = FD3D12DynamicRHI::ResourceCast(SourceBufferRHI);
 	FD3D12Buffer* DstBuffer = FD3D12DynamicRHI::ResourceCast(DestBufferRHI);
@@ -120,7 +120,7 @@ void FD3D12DynamicRHI::RHICopyVertexBuffer(FRHIVertexBuffer* SourceBufferRHI, FR
 }
 
 #if D3D12_RHI_RAYTRACING
-void FD3D12CommandContext::RHICopyBufferRegion(FRHIVertexBuffer* DestBufferRHI, uint64 DstOffset, FRHIVertexBuffer* SourceBufferRHI, uint64 SrcOffset, uint64 NumBytes)
+void FD3D12CommandContext::RHICopyBufferRegion(FRHIBuffer* DestBufferRHI, uint64 DstOffset, FRHIBuffer* SourceBufferRHI, uint64 SrcOffset, uint64 NumBytes)
 {
 	FD3D12Buffer* SourceBuffer = RetrieveObject<FD3D12Buffer>(SourceBufferRHI);
 	FD3D12Buffer* DestBuffer = RetrieveObject<FD3D12Buffer>(DestBufferRHI);
diff --git a/Engine/Source/Runtime/EmptyRHI/Private/EmptyVertexBuffer.cpp b/Engine/Source/Runtime/EmptyRHI/Private/EmptyVertexBuffer.cpp
--- a/Engine/Source/Runtime/EmptyRHI/Private/EmptyVertexBuffer.cpp
+++ b/Engine/Source/Runtime/EmptyRHI/Private/EmptyVertexBuffer.cpp
@@ -66,12 +66,12 @@ void FEmptyDynamicRHI::UnlockVertexBuffer_BottomOfPipe(FRHICommandListImmediate&
 	VertexBuffer->Unlock();
 }
 
-void FEmptyDynamicRHI::RHICopyVertexBuffer(FRHIVertexBuffer* SourceBufferRHI, FRHIVertexBuffer* DestBufferRHI)
+void FEmptyDynamicRHI::RHICopyBuffer(FRHIBuffer* SourceBufferRHI, FRHIBuffer* DestBufferRHI)
 {
 
 }
 
-void FEmptyDynamicRHI::RHITransferVertexBufferUnderlyingResource(FRHIVertexBuffer* DestVertexBuffer, FRHIVertexBuffer* SrcVertexBuffer)
+void FEmptyDynamicRHI::RHITransferBufferUnderlyingResource(FRHIBuffer* DestBuffer, FRHIBuffer* SrcBuffer)
 {
 
 }
\ No newline at end of file
diff --git a/Engine/Source/Runtime/EmptyRHI/Public/EmptyRHI.h b/Engine/Source/Runtime/EmptyRHI/Public/EmptyRHI.h
--- a/Engine/Source/Runtime/EmptyRHI/Public/EmptyRHI.h
+++ b/Engine/Source/Runtime/EmptyRHI/Public/EmptyRHI.h
@@ -56,8 +56,8 @@ public:
 	virtual FVertexBufferRHIRef RHICreateVertexBuffer(uint32 Size, uint32 InUsage, FRHIResourceCreateInfo& CreateInfo) final override;
 	virtual void* LockVertexBuffer_BottomOfPipe(FRHICommandListImmediate& RHICmdList, FRHIVertexBuffer* VertexBuffer, uint32 Offset, uint32 SizeRHI, EResourceLockMode LockMode) final override;
 	virtual void UnlockVertexBuffer_BottomOfPipe(FRHICommandListImmediate& RHICmdList, FRHIVertexBuffer* VertexBuffer) final override;
-	virtual void RHICopyVertexBuffer(FRHIVertexBuffer* SourceBuffer, FRHIVertexBuffer* DestBuffer) final override;
-	virtual void RHITransferVertexBufferUnderlyingResource(FRHIVertexBuffer* DestVertexBuffer, FRHIVertexBuffer* SrcVertexBuffer) final override;
+	virtual void RHICopyBuffer(FRHIBuffer* SourceBuffer, FRHIBuffer* DestBuffer) final override;
+	virtual void RHITransferBufferUnderlyingResource(FRHIBuffer* DestBuffer, FRHIBuffer* SrcBuffer) final override;
 	virtual FStructuredBufferRHIRef RHICreateStructuredBuffer(uint32 Stride, uint32 Size, uint32 InUsage, FRHIResourceCreateInfo& CreateInfo) final override;
 	virtual void* LockStructuredBuffer_BottomOfPipe(FRHICommandListImmediate& RHICmdList, FRHIStructuredBuffer* StructuredBuffer, uint32 Offset, uint32 SizeRHI, EResourceLockMode LockMode) final override;
 	virtual void UnlockStructuredBuffer_BottomOfPipe(FRHICommandListImmediate& RHICmdList, FRHIStructuredBuffer* StructuredBuffer) final override;
diff --git a/Engine/Source/Runtime/NullDrv/Public/NullRHI.h b/Engine/Source/Runtime/NullDrv/Public/NullRHI.h
--- a/Engine/Source/Runtime/NullDrv/Public/NullRHI.h
+++ b/Engine/Source/Runtime/NullDrv/Public/NullRHI.h
@@ -142,7 +142,7 @@ public:
 		return new FRHIVertexBuffer(Size, InUsage, 0); 
 	}
 
-	virtual void RHICopyVertexBuffer(FRHIVertexBuffer* SourceBuffer, FRHIVertexBuffer* DestBuffer) final override
+	virtual void RHICopyBuffer(FRHIBuffer* SourceBuffer, FRHIBuffer* DestBuffer) final override
 	{
 
 	}
diff --git a/Engine/Source/Runtime/OpenGLDrv/Private/OpenGLVertexBuffer.cpp b/Engine/Source/Runtime/OpenGLDrv/Private/OpenGLVertexBuffer.cpp
--- a/Engine/Source/Runtime/OpenGLDrv/Private/OpenGLVertexBuffer.cpp
+++ b/Engine/Source/Runtime/OpenGLDrv/Private/OpenGLVertexBuffer.cpp
@@ -246,13 +246,13 @@ void FOpenGLDynamicRHI::UnlockBuffer_BottomOfPipe(FRHICommandListImmediate& RHIC
 	RHITHREAD_GLCOMMAND_EPILOGUE();
 }
 
-void FOpenGLDynamicRHI::RHICopyVertexBuffer(FRHIVertexBuffer* SourceBufferRHI, FRHIVertexBuffer* DestBufferRHI)
+void FOpenGLDynamicRHI::RHICopyBuffer(FRHIBuffer* SourceBufferRHI, FRHIBuffer* DestBufferRHI)
 {
 	check(SourceBufferRHI && DestBufferRHI && SourceBufferRHI->GetSize() == DestBufferRHI->GetSize());
 	RHICopyBufferRegion(DestBufferRHI, 0, SourceBufferRHI, 0, SourceBufferRHI->GetSize());
 }
 
-void FOpenGLDynamicRHI::RHICopyBufferRegion(FRHIVertexBuffer* DestBufferRHI, uint64 DstOffset, FRHIVertexBuffer* SourceBufferRHI, uint64 SrcOffset, uint64 NumBytes)
+void FOpenGLDynamicRHI::RHICopyBufferRegion(FRHIBuffer* DestBufferRHI, uint64 DstOffset, FRHIBuffer* SourceBufferRHI, uint64 SrcOffset, uint64 NumBytes)
 {
 	VERIFY_GL_SCOPE();
 	FOpenGLBuffer* SourceBuffer = ResourceCast(SourceBufferRHI);
diff --git a/Engine/Source/Runtime/OpenGLDrv/Public/OpenGLDrv.h b/Engine/Source/Runtime/OpenGLDrv/Public/OpenGLDrv.h
--- a/Engine/Source/Runtime/OpenGLDrv/Public/OpenGLDrv.h
+++ b/Engine/Source/Runtime/OpenGLDrv/Public/OpenGLDrv.h
@@ -349,7 +349,7 @@ public:
 	virtual void RHIUpdateUniformBuffer(FRHIUniformBuffer* UniformBufferRHI, const void* Contents) final override;
 	virtual FIndexBufferRHIRef RHICreateIndexBuffer(uint32 Stride, uint32 Size, uint32 InUsage, ERHIAccess InResourceState, FRHIResourceCreateInfo& CreateInfo) final override;
 	virtual FVertexBufferRHIRef RHICreateVertexBuffer(uint32 Size, uint32 InUsage, ERHIAccess InResourceState, FRHIResourceCreateInfo& CreateInfo) final override;
-	virtual void RHICopyVertexBuffer(FRHIVertexBuffer* SourceBuffer, FRHIVertexBuffer* DestBuffer) final override;
+	virtual void RHICopyBuffer(FRHIBuffer* SourceBuffer, FRHIBuffer* DestBuffer) final override;
 	virtual void RHITransferBufferUnderlyingResource(FRHIBuffer* DestBuffer, FRHIBuffer* SrcBuffer) final override;
 	virtual FStructuredBufferRHIRef RHICreateStructuredBuffer(uint32 Stride, uint32 Size, uint32 InUsage, ERHIAccess InResourceState, FRHIResourceCreateInfo& CreateInfo) final override;
 	virtual void* LockBuffer_BottomOfPipe(FRHICommandListImmediate& RHICmdList, FRHIBuffer* Buffer, uint32 Offset, uint32 Size, EResourceLockMode LockMode) final override;
@@ -500,7 +500,7 @@ public:
 	// FIXME: Broken on Android for cubemaps
 	virtual void RHICopyTexture(FRHITexture* SourceTexture, FRHITexture* DestTexture, const FRHICopyTextureInfo& CopyInfo) final override;
 
-	virtual void RHICopyBufferRegion(FRHIVertexBuffer* DestBuffer, uint64 DstOffset, FRHIVertexBuffer* SourceBuffer, uint64 SrcOffset, uint64 NumBytes) final override;
+	virtual void RHICopyBufferRegion(FRHIBuffer* DestBuffer, uint64 DstOffset, FRHIBuffer* SourceBuffer, uint64 SrcOffset, uint64 NumBytes) final override;
 
 	// FOpenGLDynamicRHI interface.
 	virtual FTexture2DRHIRef RHICreateTexture2DFromResource(EPixelFormat Format, uint32 SizeX, uint32 SizeY, uint32 NumMips, uint32 NumSamples, uint32 NumSamplesTileMem, const FClearValueBinding& ClearValueBinding, GLuint Resource, ETextureCreateFlags Flags);
diff --git a/Engine/Source/Runtime/RHI/Public/DynamicRHI.h b/Engine/Source/Runtime/RHI/Public/DynamicRHI.h
--- a/Engine/Source/Runtime/RHI/Public/DynamicRHI.h
+++ b/Engine/Source/Runtime/RHI/Public/DynamicRHI.h
@@ -530,9 +530,9 @@ public:
 	// FlushType: Wait RHI Thread
 	virtual FVertexBufferRHIRef RHICreateVertexBuffer(uint32 Size, uint32 InUsage, ERHIAccess InResourceState, FRHIResourceCreateInfo& CreateInfo) = 0;
 
-	/** Copies the contents of one vertex buffer to another vertex buffer.  They must have identical sizes. */
+	/** Copies the contents of one buffer to another buffer. They must have identical sizes. */
 	// FlushType: Flush Immediate (seems dangerous)
-	virtual void RHICopyVertexBuffer(FRHIVertexBuffer* SourceBuffer, FRHIVertexBuffer* DestBuffer) = 0;
+	virtual void RHICopyBuffer(FRHIBuffer* SourceBuffer, FRHIBuffer* DestBuffer) = 0;
 
 	/**
 	 * Transfer metadata and underlying resource from src to dest and release any resource owned by dest.
diff --git a/Engine/Source/Runtime/RHI/Public/RHICommandList.h b/Engine/Source/Runtime/RHI/Public/RHICommandList.h
--- a/Engine/Source/Runtime/RHI/Public/RHICommandList.h
+++ b/Engine/Source/Runtime/RHI/Public/RHICommandList.h
@@ -3456,7 +3456,7 @@ public:
 	}
 #endif // #if PLATFORM_USE_BACKBUFFER_WRITE_TRANSITION_TRACKING
 	
-	FORCEINLINE_DEBUGGABLE void CopyBufferRegion(FRHIVertexBuffer* DestBuffer, uint64 DstOffset, FRHIVertexBuffer* SourceBuffer, uint64 SrcOffset, uint64 NumBytes)
+	FORCEINLINE_DEBUGGABLE void CopyBufferRegion(FRHIBuffer* DestBuffer, uint64 DstOffset, FRHIBuffer* SourceBuffer, uint64 SrcOffset, uint64 NumBytes)
 	{
 		// No copy/DMA operation inside render passes
 		check(IsOutsideRenderPass());
@@ -3965,11 +3965,17 @@ public:
 		GDynamicRHI->RHIUnlockBuffer(*this, VertexBuffer);
 	}
 	
-	FORCEINLINE void CopyVertexBuffer(FRHIVertexBuffer* SourceBuffer, FRHIVertexBuffer* DestBuffer)
+	FORCEINLINE void CopyBuffer(FRHIBuffer* SourceBuffer, FRHIBuffer* DestBuffer)
 	{
-		QUICK_SCOPE_CYCLE_COUNTER(STAT_RHIMETHOD_CopyVertexBuffer_Flush);
+		QUICK_SCOPE_CYCLE_COUNTER(STAT_RHIMETHOD_CopyBuffer_Flush);
 		ImmediateFlush(EImmediateFlushType::FlushRHIThread);  
-		GDynamicRHI->RHICopyVertexBuffer(SourceBuffer,DestBuffer);
+		GDynamicRHI->RHICopyBuffer(SourceBuffer, DestBuffer);
+	}
+
+	UE_DEPRECATED(5.0, "CopyVertexBuffer() has been replaced with a general CopyBuffer() function.")
+	FORCEINLINE void CopyVertexBuffer(FRHIVertexBuffer* SourceBuffer, FRHIVertexBuffer* DestBuffer)
+	{
+		CopyBuffer(SourceBuffer, DestBuffer);
 	}
 
 	UE_DEPRECATED(5.0, "Buffer locks have been unified. Use LockBuffer() instead.")
diff --git a/Engine/Source/Runtime/RHI/Public/RHIContext.h b/Engine/Source/Runtime/RHI/Public/RHIContext.h
--- a/Engine/Source/Runtime/RHI/Public/RHIContext.h
+++ b/Engine/Source/Runtime/RHI/Public/RHIContext.h
@@ -600,7 +600,7 @@ public:
 		}
 	}
 
-	virtual void RHICopyBufferRegion(FRHIVertexBuffer* DestBuffer, uint64 DstOffset, FRHIVertexBuffer* SourceBuffer, uint64 SrcOffset, uint64 NumBytes)
+	virtual void RHICopyBufferRegion(FRHIBuffer* DestBuffer, uint64 DstOffset, FRHIBuffer* SourceBuffer, uint64 SrcOffset, uint64 NumBytes)
 	{
 		checkNoEntry();
 	}
diff --git a/Engine/Source/Runtime/RHI/Public/RHIValidation.h b/Engine/Source/Runtime/RHI/Public/RHIValidation.h
--- a/Engine/Source/Runtime/RHI/Public/RHIValidation.h
+++ b/Engine/Source/Runtime/RHI/Public/RHIValidation.h
@@ -378,9 +378,9 @@ public:
 
 	/** Copies the contents of one vertex buffer to another vertex buffer.  They must have identical sizes. */
 	// FlushType: Flush Immediate (seems dangerous)
-	virtual void RHICopyVertexBuffer(FRHIVertexBuffer* SourceBuffer, FRHIVertexBuffer* DestBuffer) override final
+	virtual void RHICopyBuffer(FRHIBuffer* SourceBuffer, FRHIBuffer* DestBuffer) override final
 	{
-		RHI->RHICopyVertexBuffer(SourceBuffer, DestBuffer);
+		RHI->RHICopyBuffer(SourceBuffer, DestBuffer);
 	}
 
 	/**
diff --git a/Engine/Source/Runtime/RHI/Public/RHIValidationContext.h b/Engine/Source/Runtime/RHI/Public/RHIValidationContext.h
--- a/Engine/Source/Runtime/RHI/Public/RHIValidationContext.h
+++ b/Engine/Source/Runtime/RHI/Public/RHIValidationContext.h
@@ -864,7 +864,7 @@ public:
 	}
 
 #if RHI_RAYTRACING
-	virtual void RHICopyBufferRegion(FRHIVertexBuffer* DestBuffer, uint64 DstOffset, FRHIVertexBuffer* SourceBuffer, uint64 SrcOffset, uint64 NumBytes)
+	virtual void RHICopyBufferRegion(FRHIBuffer* DestBuffer, uint64 DstOffset, FRHIBuffer* SourceBuffer, uint64 SrcOffset, uint64 NumBytes)
 	{
 		Tracker->Assert(SourceBuffer->GetWholeResourceIdentity(), ERHIAccess::CopySrc);
 		Tracker->Assert(DestBuffer->GetWholeResourceIdentity(), ERHIAccess::CopyDest);
diff --git a/Engine/Source/Runtime/VulkanRHI/Private/VulkanContext.h b/Engine/Source/Runtime/VulkanRHI/Private/VulkanContext.h
--- a/Engine/Source/Runtime/VulkanRHI/Private/VulkanContext.h
+++ b/Engine/Source/Runtime/VulkanRHI/Private/VulkanContext.h
@@ -75,7 +75,7 @@ public:
 	virtual void RHIClearUAVUint(FRHIUnorderedAccessView* UnorderedAccessViewRHI, const FUintVector4& Values) final override;
 	virtual void RHICopyToResolveTarget(FRHITexture* SourceTexture, FRHITexture* DestTexture, const FResolveParams& ResolveParams) final override;
 	virtual void RHICopyTexture(FRHITexture* SourceTexture, FRHITexture* DestTexture, const FRHICopyTextureInfo& CopyInfo) final override;
-	virtual void RHICopyBufferRegion(FRHIVertexBuffer* DstBuffer, uint64 DstOffset, FRHIVertexBuffer* SrcBuffer, uint64 SrcOffset, uint64 NumBytes) final override;
+	virtual void RHICopyBufferRegion(FRHIBuffer* DstBuffer, uint64 DstOffset, FRHIBuffer* SrcBuffer, uint64 SrcOffset, uint64 NumBytes) final override;
 	virtual void RHIBeginTransitions(TArrayView<const FRHITransition*> Transitions) override final;
 	virtual void RHIEndTransitions(TArrayView<const FRHITransition*> Transitions) override final;
 	virtual void RHICopyToStagingBuffer(FRHIVertexBuffer* SourceBuffer, FRHIStagingBuffer* DestinationStagingBuffer, uint32 Offset, uint32 NumBytes) final override;
diff --git a/Engine/Source/Runtime/VulkanRHI/Private/VulkanTexture.cpp b/Engine/Source/Runtime/VulkanRHI/Private/VulkanTexture.cpp
--- a/Engine/Source/Runtime/VulkanRHI/Private/VulkanTexture.cpp
+++ b/Engine/Source/Runtime/VulkanRHI/Private/VulkanTexture.cpp
@@ -3015,7 +3015,7 @@ void FVulkanCommandListContext::RHICopyTexture(FRHITexture* SourceTexture, FRHIT
 	}
 }
 
-void FVulkanCommandListContext::RHICopyBufferRegion(FRHIVertexBuffer* DstBuffer, uint64 DstOffset, FRHIVertexBuffer* SrcBuffer, uint64 SrcOffset, uint64 NumBytes)
+void FVulkanCommandListContext::RHICopyBufferRegion(FRHIBuffer* DstBuffer, uint64 DstOffset, FRHIBuffer* SrcBuffer, uint64 SrcOffset, uint64 NumBytes)
 {
 	if (!DstBuffer || !SrcBuffer || DstBuffer == SrcBuffer || !NumBytes)
 	{
diff --git a/Engine/Source/Runtime/VulkanRHI/Private/VulkanVertexBuffer.cpp b/Engine/Source/Runtime/VulkanRHI/Private/VulkanVertexBuffer.cpp
--- a/Engine/Source/Runtime/VulkanRHI/Private/VulkanVertexBuffer.cpp
+++ b/Engine/Source/Runtime/VulkanRHI/Private/VulkanVertexBuffer.cpp
@@ -18,7 +18,7 @@ FVertexBufferRHIRef FVulkanDynamicRHI::RHICreateVertexBuffer(uint32 Size, uint32
 	return VertexBuffer;
 }
 
-void FVulkanDynamicRHI::RHICopyVertexBuffer(FRHIVertexBuffer* SourceBufferRHI, FRHIVertexBuffer* DestBufferRHI)
+void FVulkanDynamicRHI::RHICopyBuffer(FRHIBuffer* SourceBufferRHI, FRHIBuffer* DestBufferRHI)
 {
 	VULKAN_SIGNAL_UNIMPLEMENTED();
 }
diff --git a/Engine/Source/Runtime/VulkanRHI/Public/VulkanDynamicRHI.h b/Engine/Source/Runtime/VulkanRHI/Public/VulkanDynamicRHI.h
--- a/Engine/Source/Runtime/VulkanRHI/Public/VulkanDynamicRHI.h
+++ b/Engine/Source/Runtime/VulkanRHI/Public/VulkanDynamicRHI.h
@@ -82,7 +82,7 @@ public:
 	virtual void RHIUpdateUniformBuffer(FRHIUniformBuffer* UniformBufferRHI, const void* Contents) final override;
 	virtual FIndexBufferRHIRef RHICreateIndexBuffer(uint32 Stride, uint32 Size, uint32 InUsage, ERHIAccess InResourceState, FRHIResourceCreateInfo& CreateInfo) final override;
 	virtual FVertexBufferRHIRef RHICreateVertexBuffer(uint32 Size, uint32 InUsage, ERHIAccess InResourceState, FRHIResourceCreateInfo& CreateInfo) final override;
-	virtual void RHICopyVertexBuffer(FRHIVertexBuffer* SourceBuffer, FRHIVertexBuffer* DestBuffer) final override;
+	virtual void RHICopyBuffer(FRHIBuffer* SourceBuffer, FRHIBuffer* DestBuffer) final override;
 	virtual void RHITransferBufferUnderlyingResource(FRHIBuffer* DestBuffer, FRHIBuffer* SrcBuffer) final override;
 	virtual FStructuredBufferRHIRef RHICreateStructuredBuffer(uint32 Stride, uint32 Size, uint32 InUsage, ERHIAccess InResourceState, FRHIResourceCreateInfo& CreateInfo) final override;
 	virtual void* LockBuffer_BottomOfPipe(FRHICommandListImmediate& RHICmdList, FRHIBuffer* Buffer, uint32 Offset, uint32 Size, EResourceLockMode LockMode) final override;
diff --git a/Engine/Source/Runtime/Windows/D3D11RHI/Private/D3D11RHIPrivate.h b/Engine/Source/Runtime/Windows/D3D11RHI/Private/D3D11RHIPrivate.h
--- a/Engine/Source/Runtime/Windows/D3D11RHI/Private/D3D11RHIPrivate.h
+++ b/Engine/Source/Runtime/Windows/D3D11RHI/Private/D3D11RHIPrivate.h
@@ -439,7 +439,7 @@ public:
 	virtual void RHIUpdateUniformBuffer(FRHIUniformBuffer* UniformBufferRHI, const void* Contents) final override;
 	virtual FIndexBufferRHIRef RHICreateIndexBuffer(uint32 Stride, uint32 Size, uint32 InUsage, ERHIAccess InResourceState, FRHIResourceCreateInfo& CreateInfo) final override;
 	virtual FVertexBufferRHIRef RHICreateVertexBuffer(uint32 Size, uint32 InUsage, ERHIAccess InResourceState, FRHIResourceCreateInfo& CreateInfo) final override;
-	virtual void RHICopyVertexBuffer(FRHIVertexBuffer* SourceBuffer, FRHIVertexBuffer* DestBuffer) final override;
+	virtual void RHICopyBuffer(FRHIBuffer* SourceBuffer, FRHIBuffer* DestBuffer) final override;
 	virtual void RHITransferBufferUnderlyingResource(FRHIBuffer* DestBuffer, FRHIBuffer* SrcBuffer) final override;
 	virtual FStructuredBufferRHIRef RHICreateStructuredBuffer(uint32 Stride, uint32 Size, uint32 InUsage, ERHIAccess InResourceState, FRHIResourceCreateInfo& CreateInfo) final override;
 	virtual void* LockBuffer_BottomOfPipe(FRHICommandListImmediate& RHICmdList, FRHIBuffer* Buffer, uint32 Offset, uint32 Size, EResourceLockMode LockMode) final override;
@@ -533,7 +533,7 @@ public:
 	virtual void RHIClearUAVUint(FRHIUnorderedAccessView* UnorderedAccessViewRHI, const FUintVector4& Values) final override;
 	virtual void RHICopyToResolveTarget(FRHITexture* SourceTexture, FRHITexture* DestTexture, const FResolveParams& ResolveParams) final override;
 	virtual void RHICopyTexture(FRHITexture* SourceTexture, FRHITexture* DestTexture, const FRHICopyTextureInfo& CopyInfo) final override;
-	virtual void RHICopyBufferRegion(FRHIVertexBuffer* DestBuffer, uint64 DstOffset, FRHIVertexBuffer* SourceBuffer, uint64 SrcOffset, uint64 NumBytes) final override;
+	virtual void RHICopyBufferRegion(FRHIBuffer* DestBuffer, uint64 DstOffset, FRHIBuffer* SourceBuffer, uint64 SrcOffset, uint64 NumBytes) final override;
 	virtual void RHICreateTransition(FRHITransition* Transition, ERHIPipeline SrcPipelines, ERHIPipeline DstPipelines, ERHICreateTransitionFlags CreateFlags, TArrayView<const FRHITransitionInfo> Infos) final override;
 	virtual void RHIReleaseTransition(FRHITransition* Transition) final override;
 	virtual void RHIBeginTransitions(TArrayView<const FRHITransition*> Transitions) override final;
diff --git a/Engine/Source/Runtime/Windows/D3D11RHI/Private/D3D11Texture.cpp b/Engine/Source/Runtime/Windows/D3D11RHI/Private/D3D11Texture.cpp
--- a/Engine/Source/Runtime/Windows/D3D11RHI/Private/D3D11Texture.cpp
+++ b/Engine/Source/Runtime/Windows/D3D11RHI/Private/D3D11Texture.cpp
@@ -2725,7 +2725,7 @@ void FD3D11DynamicRHI::RHICopyTexture(FRHITexture* SourceTextureRHI, FRHITexture
 	}
 }
 
-void FD3D11DynamicRHI::RHICopyBufferRegion(FRHIVertexBuffer* DstBuffer, uint64 DstOffset, FRHIVertexBuffer* SrcBuffer, uint64 SrcOffset, uint64 NumBytes)
+void FD3D11DynamicRHI::RHICopyBufferRegion(FRHIBuffer* DstBuffer, uint64 DstOffset, FRHIBuffer* SrcBuffer, uint64 SrcOffset, uint64 NumBytes)
 {
 	if (!DstBuffer || !SrcBuffer || DstBuffer == SrcBuffer || !NumBytes)
 	{
diff --git a/Engine/Source/Runtime/Windows/D3D11RHI/Private/D3D11VertexBuffer.cpp b/Engine/Source/Runtime/Windows/D3D11RHI/Private/D3D11VertexBuffer.cpp
--- a/Engine/Source/Runtime/Windows/D3D11RHI/Private/D3D11VertexBuffer.cpp
+++ b/Engine/Source/Runtime/Windows/D3D11RHI/Private/D3D11VertexBuffer.cpp
@@ -126,7 +126,7 @@ FVertexBufferRHIRef FD3D11DynamicRHI::CreateVertexBuffer_RenderThread(
 	return RHICreateVertexBuffer(Size, InUsage, InResourceState, CreateInfo);
 }
 
-void FD3D11DynamicRHI::RHICopyVertexBuffer(FRHIVertexBuffer* SourceBufferRHI, FRHIVertexBuffer* DestBufferRHI)
+void FD3D11DynamicRHI::RHICopyBuffer(FRHIBuffer* SourceBufferRHI, FRHIBuffer* DestBufferRHI)
 {
 	FD3D11Buffer* SourceBuffer = ResourceCast(SourceBufferRHI);
 	FD3D11Buffer* DestBuffer = ResourceCast(DestBufferRHI);
