Don't unset mtlpp::TextureUsage::PixelFormatView or we end up creating a new blank texture with which to reference via a requested SRV later on.  We've broken the link for RDG.

[Depth Stencil creation with TexCreate_RenderTargetable  during: FRDGBuilder::CollectPassResources() -> EnumerateTextureParameters() -> if (FRDGTextureRef Texture = RenderTargets.DepthStencil.GetTexture())-> RHICreateTargetableShaderResource2D() ].

diff --git a/Engine/Source/Runtime/Apple/MetalRHI/Private/MetalTexture.cpp b/Engine/Source/Runtime/Apple/MetalRHI/Private/MetalTexture.cpp
--- a/Engine/Source/Runtime/Apple/MetalRHI/Private/MetalTexture.cpp
+++ b/Engine/Source/Runtime/Apple/MetalRHI/Private/MetalTexture.cpp
@@ -112,8 +112,6 @@ static mtlpp::TextureUsage ConvertFlagsToUsage(ETextureCreateFlags Flags)
 	{
 		Usage |= mtlpp::TextureUsage::RenderTarget;
 		Usage |= mtlpp::TextureUsage::ShaderRead;
-		// The cost of PixelFormatView on macOS/iOS is exorbitant, we need to reallocate on demand to avoid it
-		Usage &= ~(mtlpp::TextureUsage::PixelFormatView);
 	}
 	return (mtlpp::TextureUsage)Usage;
 }
