Build fixes

diff --git a/Engine/Source/Runtime/Apple/MetalRHI/Private/MetalLLM.cpp b/Engine/Source/Runtime/Apple/MetalRHI/Private/MetalLLM.cpp
--- a/Engine/Source/Runtime/Apple/MetalRHI/Private/MetalLLM.cpp
+++ b/Engine/Source/Runtime/Apple/MetalRHI/Private/MetalLLM.cpp
@@ -150,7 +150,7 @@ static mtlpp::SizeAndAlign TextureSizeAndAlign(mtlpp::TextureType TextureType, u
 	SizeAlign.Align = 0;
 	
 	uint32 Align = 0;
-	FRHIResourceCreateInfo CreateInfo;
+	FRHIResourceCreateInfo CreateInfo(TEXT(""));
 	switch (TextureType)
 	{
 		case mtlpp::TextureType::Texture2D:
diff --git a/Engine/Source/Runtime/Apple/MetalRHI/Private/MetalStateCache.cpp b/Engine/Source/Runtime/Apple/MetalRHI/Private/MetalStateCache.cpp
--- a/Engine/Source/Runtime/Apple/MetalRHI/Private/MetalStateCache.cpp
+++ b/Engine/Source/Runtime/Apple/MetalRHI/Private/MetalStateCache.cpp
@@ -2765,7 +2765,7 @@ FTexture2DRHIRef FMetalStateCache::CreateFallbackDepthStencilSurface(uint32 Widt
 	if (!IsValidRef(FallbackDepthStencilSurface) || FallbackDepthStencilSurface->GetSizeX() != Width || FallbackDepthStencilSurface->GetSizeY() != Height)
 #endif
 	{
-		FRHIResourceCreateInfo TexInfo;
+		FRHIResourceCreateInfo TexInfo(TEXT("FallbackDepthStencilSurface"));
 		FallbackDepthStencilSurface = RHICreateTexture2D(Width, Height, PF_DepthStencil, 1, 1, TexCreate_DepthStencilTargetable, TexInfo);
 	}
 	check(IsValidRef(FallbackDepthStencilSurface));
diff --git a/Engine/Source/Runtime/Apple/MetalRHI/Private/MetalViewport.cpp b/Engine/Source/Runtime/Apple/MetalRHI/Private/MetalViewport.cpp
--- a/Engine/Source/Runtime/Apple/MetalRHI/Private/MetalViewport.cpp
+++ b/Engine/Source/Runtime/Apple/MetalRHI/Private/MetalViewport.cpp
@@ -238,7 +238,7 @@ void FMetalViewport::Resize(uint32 InSizeX, uint32 InSizeY, bool bInIsFullscreen
 
     {
         FScopeLock Lock(&Mutex);
-        FRHIResourceCreateInfo CreateInfo;
+        FRHIResourceCreateInfo CreateInfo(TEXT("BackBuffer"));
         FTexture2DRHIRef NewBackBuffer;
         FTexture2DRHIRef DoubleBuffer;
         if (GMetalSupportsIntermediateBackBuffer)
