Extend input/output bitmask in shader binding descriptors from 16 to 32 bits; in Metal, GLSL, ESSL, and SPIR-V backends.

#rnx

diff --git a/Engine/Shaders/Public/Platform/Metal/MetalCommon.ush b/Engine/Shaders/Public/Platform/Metal/MetalCommon.ush
--- a/Engine/Shaders/Public/Platform/Metal/MetalCommon.ush
+++ b/Engine/Shaders/Public/Platform/Metal/MetalCommon.ush
@@ -5,7 +5,7 @@
 =============================================================================*/
 
 // Update this GUID to improve shader recompilation for Metal only shaders
-// GUID = 0475930F-29E9-4E32-98DF-B2D12F73C042
+// GUID = 293FE8D0-534F-4A76-ACD0-2802BC36DEA7
 
 #pragma once
 
diff --git a/Engine/Shaders/Public/Platform/Vulkan/VulkanCommon.ush b/Engine/Shaders/Public/Platform/Vulkan/VulkanCommon.ush
--- a/Engine/Shaders/Public/Platform/Vulkan/VulkanCommon.ush
+++ b/Engine/Shaders/Public/Platform/Vulkan/VulkanCommon.ush
@@ -5,6 +5,6 @@
 =============================================================================*/
 
 // Update this GUID to improve shader recompilation for Vulkan only shaders
-// GUID = 0C614AF7-59B7-41BB-A89F-8C0709FCE94E
+// GUID = 05AA3A0F-8574-4976-9E80-9038FAF0A7D2
 
 #pragma once
diff --git a/Engine/Shaders/Public/ShaderVersion.ush b/Engine/Shaders/Public/ShaderVersion.ush
--- a/Engine/Shaders/Public/ShaderVersion.ush
+++ b/Engine/Shaders/Public/ShaderVersion.ush
@@ -3,4 +3,4 @@
 // in Platform.ush (which should be included in any shader) it allows to invalidate the shader DDC.
 //
 // If you are merging streams and there is a conflict with this GUID you should make a new GUID rather than taking one or the other.
-// GUID = 1A51199A-E12A-49BE-B98F-CBF652EDE4A0
+// GUID = B1C09C8B-698B-404B-84C8-81AE0AA9340D
diff --git a/Engine/Source/Developer/Apple/MetalShaderFormat/Private/MetalShaderCompiler.cpp b/Engine/Source/Developer/Apple/MetalShaderFormat/Private/MetalShaderCompiler.cpp
--- a/Engine/Source/Developer/Apple/MetalShaderFormat/Private/MetalShaderCompiler.cpp
+++ b/Engine/Source/Developer/Apple/MetalShaderFormat/Private/MetalShaderCompiler.cpp
@@ -549,7 +549,7 @@ void BuildMetalShaderOutput(
 			if (Input.Name.StartsWith(AttributePrefix))
 			{
 				uint8 AttributeIndex = ParseNumber(*Input.Name + AttributePrefix.Len());
-				Header.Bindings.InOutMask |= (1 << AttributeIndex);
+				Header.Bindings.InOutMask.EnableField(AttributeIndex);
 			}
 		}
 	}
@@ -566,19 +566,19 @@ void BuildMetalShaderOutput(
 			if (Output.Name.StartsWith(TargetPrefix))
 			{
 				uint8 TargetIndex = ParseNumber(*Output.Name + TargetPrefix.Len());
-				Header.Bindings.InOutMask |= (1 << TargetIndex);
+				Header.Bindings.InOutMask.EnableField(TargetIndex);
 			}
 			else if (Output.Name.StartsWith(TargetPrefix2))
 			{
 				uint8 TargetIndex = ParseNumber(*Output.Name + TargetPrefix2.Len());
-				Header.Bindings.InOutMask |= (1 << TargetIndex);
+				Header.Bindings.InOutMask.EnableField(TargetIndex);
 			}
         }
 		
 		// For fragment shaders that discard but don't output anything we need at least a depth-stencil surface, so we need a way to validate this at runtime.
 		if (FCStringAnsi::Strstr(USFSource, "[[ depth(") != nullptr || FCStringAnsi::Strstr(USFSource, "[[depth(") != nullptr)
 		{
-			Header.Bindings.InOutMask |= 0x8000;
+			Header.Bindings.InOutMask.EnableField(CrossCompiler::FShaderBindingInOutMask::DepthStencilMaskIndex);
 		}
         
         // For fragment shaders that discard but don't output anything we need at least a depth-stencil surface, so we need a way to validate this at runtime.
diff --git a/Engine/Source/Developer/ShaderFormatOpenGL/Private/OpenGLShaderCompiler.cpp b/Engine/Source/Developer/ShaderFormatOpenGL/Private/OpenGLShaderCompiler.cpp
--- a/Engine/Source/Developer/ShaderFormatOpenGL/Private/OpenGLShaderCompiler.cpp
+++ b/Engine/Source/Developer/ShaderFormatOpenGL/Private/OpenGLShaderCompiler.cpp
@@ -653,12 +653,12 @@ void FOpenGLFrontend::BuildShaderOutput(
 		if (Frequency == SF_Vertex && Input.Name.StartsWith(AttributePrefix))
 		{
 			int32 AttributeIndex = ParseNumber(*Input.Name + AttributePrefix.Len());
-			Header.Bindings.InOutMask |= (1 << AttributeIndex);
+			Header.Bindings.InOutMask.EnableField(AttributeIndex);
 		}
 		else if (Frequency == SF_Vertex && Input.Name.StartsWith(AttributeVarPrefix))
 		{
 			int32 AttributeIndex = ParseNumber(*Input.Name + AttributeVarPrefix.Len());
-			Header.Bindings.InOutMask |= (1 << AttributeIndex);
+			Header.Bindings.InOutMask.EnableField(AttributeIndex);
 		}
 		// Record user-defined input varyings
 		else if (!Input.Name.StartsWith(GL_Prefix))
@@ -678,12 +678,12 @@ void FOpenGLFrontend::BuildShaderOutput(
 		if (Frequency == SF_Pixel && Output.Name.StartsWith(TargetPrefix))
 		{
 			uint8 TargetIndex = ParseNumber(*Output.Name + TargetPrefix.Len());
-			Header.Bindings.InOutMask |= (1 << TargetIndex);
+			Header.Bindings.InOutMask.EnableField(TargetIndex);
 		}
 		// Only depth writes for pixel shaders must be tracked.
 		else if (Frequency == SF_Pixel && Output.Name.Equals(GL_FragDepth))
 		{
-			Header.Bindings.InOutMask |= 0x8000;
+			Header.Bindings.InOutMask.EnableField(CrossCompiler::FShaderBindingInOutMask::DepthStencilMaskIndex);
 		}
 		// Record user-defined output varyings
 		else if (!Output.Name.StartsWith(GL_Prefix))
diff --git a/Engine/Source/Developer/VulkanShaderFormat/Private/VulkanShaderCompiler.cpp b/Engine/Source/Developer/VulkanShaderFormat/Private/VulkanShaderCompiler.cpp
--- a/Engine/Source/Developer/VulkanShaderFormat/Private/VulkanShaderCompiler.cpp
+++ b/Engine/Source/Developer/VulkanShaderFormat/Private/VulkanShaderCompiler.cpp
@@ -343,7 +343,6 @@ class FVulkanShaderSerializedBindings : public CrossCompiler::FShaderBindings
 public:
 	FVulkanShaderSerializedBindings()
 	{
-		InOutMask = 0;
 		NumSamplers = 0;
 		NumUniformBuffers = 0;
 		NumUAVs = 0;
@@ -1048,7 +1047,7 @@ static void ConvertToNEWHeader(FOLDVulkanCodeHeader& OLDHeader,
 #if VULKAN_ENABLE_BINDING_DEBUG_NAMES
 	OutHeader.DebugName = OLDHeader.ShaderName;
 #endif
-	OutHeader.InOutMask = OLDHeader.SerializedBindings.InOutMask;
+	OutHeader.InOutMask = OLDHeader.SerializedBindings.InOutMask.Bitmask;
 	OutHeader.bHasRealUBs = bHasRealUBs;
 }
 
@@ -1099,7 +1098,7 @@ static void BuildShaderOutput(
 			int32 Count = FMath::Max(1, Input.ArrayCount);
 			for(int32 Index = 0; Index < Count; ++Index)
 			{
-				OLDHeader.SerializedBindings.InOutMask |= (1 << (Index + AttributeIndex));
+				OLDHeader.SerializedBindings.InOutMask.EnableField(Index + AttributeIndex);
 			}
 		}
 #if 0
@@ -1123,12 +1122,12 @@ static void BuildShaderOutput(
 		if (Frequency == SF_Pixel && Output.Name.StartsWith(TargetPrefix))
 		{
 			uint8 TargetIndex = ParseNumber(*Output.Name + TargetPrefix.Len(), /*bEmptyIsZero:*/ true);
-			OLDHeader.SerializedBindings.InOutMask |= (1 << TargetIndex);
+			OLDHeader.SerializedBindings.InOutMask.EnableField(TargetIndex);
 		}
 		// Only depth writes for pixel shaders must be tracked.
 		else if (Frequency == SF_Pixel && Output.Name.Equals(GL_FragDepth))
 		{
-			OLDHeader.SerializedBindings.InOutMask |= 0x8000;
+			OLDHeader.SerializedBindings.InOutMask.EnableField(CrossCompiler::FShaderBindingInOutMask::DepthStencilMaskIndex);
 		}
 #if 0
 		// Record user-defined output varyings
diff --git a/Engine/Source/Runtime/Apple/MetalRHI/Private/MetalContext.cpp b/Engine/Source/Runtime/Apple/MetalRHI/Private/MetalContext.cpp
--- a/Engine/Source/Runtime/Apple/MetalRHI/Private/MetalContext.cpp
+++ b/Engine/Source/Runtime/Apple/MetalRHI/Private/MetalContext.cpp
@@ -1297,8 +1297,7 @@ bool FMetalContext::PrepareToDraw(uint32 PrimitiveType, EMetalIndexType IndexTyp
 					
 					uint32 BufferIndex = METAL_TO_UNREAL_BUFFER_INDEX(Attribute.bufferIndex);
 					
-					uint32 InOutMask = CurrentPSO->VertexShader->Bindings.InOutMask;
-					if (InOutMask & (1 << BufferIndex))
+					if (CurrentPSO->VertexShader->Bindings.InOutMask.IsFieldEnabled(BufferIndex))
 					{
 						uint64 MetalSize = StateCache.GetVertexBufferSize(BufferIndex);
 						
@@ -1316,7 +1315,7 @@ bool FMetalContext::PrepareToDraw(uint32 PrimitiveType, EMetalIndexType IndexTyp
 #endif
 	
 	// @todo Handle the editor not setting a depth-stencil target for the material editor's tiles which render to depth even when they shouldn't.
-	bool const bNeedsDepthStencilWrite = (IsValidRef(CurrentPSO->PixelShader) && (CurrentPSO->PixelShader->Bindings.InOutMask & 0x8000));
+	bool const bNeedsDepthStencilWrite = (IsValidRef(CurrentPSO->PixelShader) && (CurrentPSO->PixelShader->Bindings.InOutMask.IsFieldEnabled(CrossCompiler::FShaderBindingInOutMask::DepthStencilMaskIndex)));
 	
 	// @todo Improve the way we handle binding a dummy depth/stencil so we can get pure UAV raster operations...
 	bool const bNeedsDepthStencilForUAVRaster = (StateCache.GetRenderPassInfo().GetNumColorRenderTargets() == 0);
diff --git a/Engine/Source/Runtime/Apple/MetalRHI/Private/MetalPipeline.cpp b/Engine/Source/Runtime/Apple/MetalRHI/Private/MetalPipeline.cpp
--- a/Engine/Source/Runtime/Apple/MetalRHI/Private/MetalPipeline.cpp
+++ b/Engine/Source/Runtime/Apple/MetalRHI/Private/MetalPipeline.cpp
@@ -389,7 +389,7 @@ struct FMetalGraphicsPipelineKey
 		// If the pixel shader writes depth then we must compile with depth access, so we may bind the dummy depth.
 		// If the pixel shader writes to UAVs but not target is bound we must also bind the dummy depth.
 		FMetalPixelShader* PixelShader = (FMetalPixelShader*)Init.BoundShaderState.PixelShaderRHI;
-		if (PixelShader && (((PixelShader->Bindings.InOutMask & 0x8000) && (DepthFormatKey == 0)) || (bHasActiveTargets == false && PixelShader->Bindings.NumUAVs > 0)))
+		if ( PixelShader && ( ( PixelShader->Bindings.InOutMask.IsFieldEnabled(CrossCompiler::FShaderBindingInOutMask::DepthStencilMaskIndex) && DepthFormatKey == 0 ) || (bHasActiveTargets == false && PixelShader->Bindings.NumUAVs > 0) ) )
 		{
 			mtlpp::PixelFormat MetalFormat = (mtlpp::PixelFormat)GPixelFormats[PF_DepthStencil].PlatformFormat;
 			DepthFormatKey = GetMetalPixelFormatKey(MetalFormat);
@@ -721,30 +721,30 @@ APPLE_PLATFORM_OBJECT_ALLOC_OVERRIDES(FMetalShaderPipeline)
 }
 @end
 
-static MTLVertexDescriptor* GetMaskedVertexDescriptor(MTLVertexDescriptor* InputDesc, uint32 InOutMask)
+static MTLVertexDescriptor* GetMaskedVertexDescriptor(MTLVertexDescriptor* InputDesc, const CrossCompiler::FShaderBindingInOutMask& InOutMask)
 {
 	for (uint32 Attr = 0; Attr < MaxMetalStreams; Attr++)
 	{
-		if (!(InOutMask & (1 << Attr)) && [InputDesc.attributes objectAtIndexedSubscript:Attr] != nil)
+		if (!InOutMask.IsFieldEnabled((int32)Attr) && [InputDesc.attributes objectAtIndexedSubscript:Attr] != nil)
 		{
 			MTLVertexDescriptor* Desc = [[InputDesc copy] autorelease];
-			uint32 BuffersUsed = 0;
-			for (uint32 i = 0; i < MaxMetalStreams; i++)
+			CrossCompiler::FShaderBindingInOutMask BuffersUsed;
+			for (int32 MetalStreamIndex = 0; MetalStreamIndex < MaxMetalStreams; ++MetalStreamIndex)
 			{
-				if (!(InOutMask & (1 << i)))
+				if (!InOutMask.IsFieldEnabled(MetalStreamIndex))
 				{
-					[Desc.attributes setObject:nil atIndexedSubscript:i];
+					[Desc.attributes setObject:nil atIndexedSubscript:MetalStreamIndex];
 				}
 				else
 				{
-					BuffersUsed |= (1 << [Desc.attributes objectAtIndexedSubscript:i].bufferIndex);
+					BuffersUsed.EnableField([Desc.attributes objectAtIndexedSubscript : MetalStreamIndex].bufferIndex);
 				}
 			}
-			for (uint32 i = 0; i < ML_MaxBuffers; i++)
+			for (int32 BufferIndex = 0; BufferIndex < ML_MaxBuffers; ++BufferIndex)
 			{
-				if (!(BuffersUsed & (1 << i)))
+				if (!BuffersUsed.IsFieldEnabled(BufferIndex))
 				{
-					[Desc.layouts setObject:nil atIndexedSubscript:i];
+					[Desc.layouts setObject:nil atIndexedSubscript:BufferIndex];
 				}
 			}
 			return Desc;
@@ -769,13 +769,14 @@ static bool ConfigureRenderPipelineDescriptor(mtlpp::RenderPipelineDescriptor& R
 	check(NumActiveTargets <= MaxSimultaneousRenderTargets);
 	if (PixelShader)
 	{
-		if ((PixelShader->Bindings.InOutMask & 0x8000) == 0 && (PixelShader->Bindings.InOutMask & 0x7fff) == 0 && PixelShader->Bindings.NumUAVs == 0 && PixelShader->Bindings.bDiscards == false)
+		if (PixelShader->Bindings.InOutMask.Bitmask == 0 && PixelShader->Bindings.NumUAVs == 0 && PixelShader->Bindings.bDiscards == false)
 		{
-			UE_LOG(LogMetal, Error, TEXT("Pixel shader has no outputs which is not permitted. No Discards, In-Out Mask: %x\nNumber UAVs: %d\nSource Code:\n%s"), PixelShader->Bindings.InOutMask, PixelShader->Bindings.NumUAVs, *FString(PixelShader->GetSourceCode()));
+			UE_LOG(LogMetal, Error, TEXT("Pixel shader has no outputs which is not permitted. No Discards, In-Out Mask: %x\nNumber UAVs: %d\nSource Code:\n%s"), PixelShader->Bindings.InOutMask.Bitmask, PixelShader->Bindings.NumUAVs, *FString(PixelShader->GetSourceCode()));
 			return false;
 		}
 		
-		UE_CLOG((NumActiveTargets < __builtin_popcount(PixelShader->Bindings.InOutMask & 0x7fff)), LogMetal, Verbose, TEXT("NumActiveTargets doesn't match pipeline's pixel shader output mask: %u, %hx"), NumActiveTargets, PixelShader->Bindings.InOutMask);
+		const uint32 MaxNumActiveTargets = __builtin_popcount(PixelShader->Bindings.InOutMask.Bitmask & ((1u << CrossCompiler::FShaderBindingInOutMask::MaxIndex) - 1));
+		UE_CLOG((NumActiveTargets < MaxNumActiveTargets), LogMetal, Verbose, TEXT("NumActiveTargets doesn't match pipeline's pixel shader output mask: %u, %hx"), NumActiveTargets, PixelShader->Bindings.InOutMask.Bitmask);
 	}
 	
 	FMetalBlendState* BlendState = (FMetalBlendState*)Init.BlendState;
@@ -786,16 +787,17 @@ static bool ConfigureRenderPipelineDescriptor(mtlpp::RenderPipelineDescriptor& R
 #endif
 	
 	uint32 TargetWidth = 0;
-	for (uint32 i = 0; i < NumActiveTargets; i++)
+	for (uint32 ActiveTargetIndex = 0; ActiveTargetIndex < NumActiveTargets; ActiveTargetIndex++)
 	{
-		EPixelFormat TargetFormat = (EPixelFormat)Init.RenderTargetFormats[i];
+		EPixelFormat TargetFormat = (EPixelFormat)Init.RenderTargetFormats[ActiveTargetIndex];
 		
-		METAL_FATAL_ASSERT(!(TargetFormat == PF_Unknown && PixelShader && (((PixelShader->Bindings.InOutMask & 0x7fff) & (1 << i)))), TEXT("Pipeline pixel shader expects target %u to be bound but it isn't: %s."), i, *FString(PixelShader->GetSourceCode()));
+		const bool bIsActiveTargetBound = (PixelShader && PixelShader->Bindings.InOutMask.IsFieldEnabled(ActiveTargetIndex));
+		METAL_FATAL_ASSERT(!(TargetFormat == PF_Unknown && bIsActiveTargetBound), TEXT("Pipeline pixel shader expects target %u to be bound but it isn't: %s."), ActiveTargetIndex, *FString(PixelShader->GetSourceCode()));
 		
 		TargetWidth += GPixelFormats[TargetFormat].BlockBytes;
 		
 		mtlpp::PixelFormat MetalFormat = (mtlpp::PixelFormat)GPixelFormats[TargetFormat].PlatformFormat;
-		ETextureCreateFlags Flags = (ETextureCreateFlags)Init.RenderTargetFlags[i];
+		ETextureCreateFlags Flags = (ETextureCreateFlags)Init.RenderTargetFlags[ActiveTargetIndex];
 		if (Flags & TexCreate_SRGB)
 		{
 #if PLATFORM_MAC // Expand as R8_sRGB is iOS only.
@@ -807,15 +809,15 @@ static bool ConfigureRenderPipelineDescriptor(mtlpp::RenderPipelineDescriptor& R
 			MetalFormat = ToSRGBFormat(MetalFormat);
 		}
 		
-		mtlpp::RenderPipelineColorAttachmentDescriptor Attachment = ColorAttachments[i];
+		mtlpp::RenderPipelineColorAttachmentDescriptor Attachment = ColorAttachments[ActiveTargetIndex];
 		Attachment.SetPixelFormat(MetalFormat);
 		
 #if !PLATFORM_TVOS
-		auto DebugAttachment = DebugColorAttachements[i];;
+		auto DebugAttachment = DebugColorAttachements[ActiveTargetIndex];
 		DebugAttachment.SetPixelFormat(MetalFormat);
 #endif
 		
-		mtlpp::RenderPipelineColorAttachmentDescriptor Blend = BlendState->RenderTargetStates[i].BlendState;
+		mtlpp::RenderPipelineColorAttachmentDescriptor Blend = BlendState->RenderTargetStates[ActiveTargetIndex].BlendState;
 		if(TargetFormat != PF_Unknown)
 		{
 			// assign each property manually, would be nice if this was faster
@@ -908,7 +910,8 @@ static bool ConfigureRenderPipelineDescriptor(mtlpp::RenderPipelineDescriptor& R
 	check(Init.BoundShaderState.GeometryShaderRHI == nullptr);
 #endif
 	
-	if(RenderPipelineDesc.GetDepthAttachmentPixelFormat() == mtlpp::PixelFormat::Invalid && PixelShader && ((PixelShader->Bindings.InOutMask & 0x8000) || (NumActiveTargets == 0 && (PixelShader->Bindings.NumUAVs > 0))))
+	if( RenderPipelineDesc.GetDepthAttachmentPixelFormat() == mtlpp::PixelFormat::Invalid &&
+		PixelShader && ( PixelShader->Bindings.InOutMask.IsFieldEnabled(CrossCompiler::FShaderBindingInOutMask::DepthStencilMaskIndex) || ( NumActiveTargets == 0 && PixelShader->Bindings.NumUAVs > 0) ) )
 	{
 		RenderPipelineDesc.SetDepthAttachmentPixelFormat((mtlpp::PixelFormat)GPixelFormats[PF_DepthStencil].PlatformFormat);
 		RenderPipelineDesc.SetStencilAttachmentPixelFormat((mtlpp::PixelFormat)GPixelFormats[PF_DepthStencil].PlatformFormat);
diff --git a/Engine/Source/Runtime/Apple/MetalRHI/Private/MetalRenderPass.cpp b/Engine/Source/Runtime/Apple/MetalRHI/Private/MetalRenderPass.cpp
--- a/Engine/Source/Runtime/Apple/MetalRHI/Private/MetalRenderPass.cpp
+++ b/Engine/Source/Runtime/Apple/MetalRHI/Private/MetalRenderPass.cpp
@@ -512,13 +512,13 @@ void FMetalRenderPass::DrawIndexedPrimitive(FMetalBuffer const& IndexBuffer, uin
 		
 		// Set our local copy and try to disprove the passed in value
 		uint32 ClampedNumInstances = NumInstances;
-		uint32 InOutMask = PipelineState->VertexShader->Bindings.InOutMask;
+		const CrossCompiler::FShaderBindingInOutMask& InOutMask = PipelineState->VertexShader->Bindings.InOutMask;
 
 		// I think it is valid to have no elements in this list
 		for(int VertexElemIdx = 0;VertexElemIdx < VertexDecl->Elements.Num();++VertexElemIdx)
 		{
 			FVertexElement const & VertexElem = VertexDecl->Elements[VertexElemIdx];
-			if(VertexElem.Stride > 0 && VertexElem.bUseInstanceIndex && ((InOutMask & (1 << VertexElem.AttributeIndex))))
+			if(VertexElem.Stride > 0 && VertexElem.bUseInstanceIndex && InOutMask.IsFieldEnabled(VertexElem.AttributeIndex))
 			{
 				uint32 AvailElementCount = 0;
 				
diff --git a/Engine/Source/Runtime/Apple/MetalRHI/Public/MetalShaderResources.h b/Engine/Source/Runtime/Apple/MetalRHI/Public/MetalShaderResources.h
--- a/Engine/Source/Runtime/Apple/MetalRHI/Public/MetalShaderResources.h
+++ b/Engine/Source/Runtime/Apple/MetalRHI/Public/MetalShaderResources.h
@@ -131,14 +131,14 @@ struct FMetalShaderBindings
 {
 	TArray<TArray<CrossCompiler::FPackedArrayInfo>>	PackedUniformBuffers;
 	TArray<CrossCompiler::FPackedArrayInfo>			PackedGlobalArrays;
-	FMetalShaderResourceTable				ShaderResourceTable;
-	TMap<uint8, TArray<uint8>>				ArgumentBufferMasks;
+	FMetalShaderResourceTable						ShaderResourceTable;
+	TMap<uint8, TArray<uint8>>						ArgumentBufferMasks;
+	CrossCompiler::FShaderBindingInOutMask			InOutMask;
 
     uint32  LinearBuffer;
 	uint32	TypedBuffers;
 	uint32 	ConstantBuffers;
 	uint32  ArgumentBuffers;
-	uint16	InOutMask;
 	uint8	NumSamplers;
 	uint8	NumUniformBuffers;
 	uint8	NumUAVs;
@@ -149,7 +149,6 @@ struct FMetalShaderBindings
         TypedBuffers(0),
 		ConstantBuffers(0),
 		ArgumentBuffers(0),
-		InOutMask(0),
 		NumSamplers(0),
 		NumUniformBuffers(0),
 		NumUAVs(0),
diff --git a/Engine/Source/Runtime/OpenGLDrv/Private/OpenGLCommands.cpp b/Engine/Source/Runtime/OpenGLDrv/Private/OpenGLCommands.cpp
--- a/Engine/Source/Runtime/OpenGLDrv/Private/OpenGLCommands.cpp
+++ b/Engine/Source/Runtime/OpenGLDrv/Private/OpenGLCommands.cpp
@@ -1967,8 +1967,7 @@ void FOpenGLDynamicRHI::SetupVertexArrays(FOpenGLContextState& ContextState, uin
 	{
 		FOpenGLVertexElement& VertexElement = VertexDeclaration->VertexElements[ElementIndex];
 		uint32 AttributeIndex = VertexElement.AttributeIndex;
-		const bool bAttribInUse = (Bindings.InOutMask & (0x1 << AttributeIndex)) != 0;
-		if (!bAttribInUse)
+		if (!Bindings.InOutMask.IsFieldEnabled((int32)AttributeIndex))
 		{
 			continue; // skip unused attributes.
 		}
@@ -2049,9 +2048,9 @@ void FOpenGLDynamicRHI::SetupVertexArraysVAB(FOpenGLContextState& ContextState,
 
 	check(IsValidRef(PendingState.BoundShaderState));
 	FOpenGLVertexDeclaration* VertexDeclaration = PendingState.BoundShaderState->VertexDeclaration;
-	uint32 AttributeMask = PendingState.BoundShaderState->GetVertexShader()->Bindings.InOutMask;
+	const CrossCompiler::FShaderBindingInOutMask& AttributeMask = PendingState.BoundShaderState->GetVertexShader()->Bindings.InOutMask;
 
-	if (ContextState.VertexDecl != VertexDeclaration || AttributeMask != ContextState.VertexAttrs_EnabledBits)
+	if (ContextState.VertexDecl != VertexDeclaration || AttributeMask.Bitmask != ContextState.VertexAttrs_EnabledBits)
 	{
 		StreamMask = 0;
 		UpdateDivisors = true;
@@ -2062,12 +2061,11 @@ void FOpenGLDynamicRHI::SetupVertexArraysVAB(FOpenGLContextState& ContextState,
 		{
 			FOpenGLVertexElement& VertexElement = VertexDeclaration->VertexElements[ElementIndex];
 			uint32 AttributeIndex = VertexElement.AttributeIndex;
-			const bool bAttribInUse = (AttributeMask & (0x1 << AttributeIndex)) != 0;
 			const uint32 StreamIndex = VertexElement.StreamIndex;
 
 			//only setup/track attributes actually in use
 			FOpenGLCachedAttr &Attr = ContextState.VertexAttrs[AttributeIndex];
-			if (bAttribInUse)
+			if (AttributeMask.IsFieldEnabled(AttributeIndex))
 			{
 				if (VertexElement.StreamIndex < NumStreams)
 				{
@@ -2138,7 +2136,7 @@ void FOpenGLDynamicRHI::SetupVertexArraysVAB(FOpenGLContextState& ContextState,
 	}
 
 	// Disable remaining vertex arrays
-	uint32 NotUsedButEnabledAttrMask = (ContextState.VertexAttrs_EnabledBits & ~(AttributeMask));
+	uint32 NotUsedButEnabledAttrMask = (ContextState.VertexAttrs_EnabledBits & ~(AttributeMask.Bitmask));
 	for (GLuint AttribIndex = 0; AttribIndex < NUM_OPENGL_VERTEX_STREAMS && NotUsedButEnabledAttrMask; AttribIndex++)
 	{
 		if (NotUsedButEnabledAttrMask & 1)
@@ -2233,8 +2231,7 @@ void FOpenGLDynamicRHI::SetupVertexArraysUP(FOpenGLContextState& ContextState, v
 		check(VertexElement.StreamIndex < 1);
 
 		uint32 AttributeIndex = VertexElement.AttributeIndex;
-		const bool bAttribInUse = (Bindings.InOutMask & (0x1 << AttributeIndex)) != 0;
-		if (bAttribInUse)
+		if (Bindings.InOutMask.IsFieldEnabled(AttributeIndex))
 		{
 			check(Stride > 0);
 			EnableVertexElementCached(
diff --git a/Engine/Source/Runtime/OpenGLDrv/Public/OpenGLShaderResources.h b/Engine/Source/Runtime/OpenGLDrv/Public/OpenGLShaderResources.h
--- a/Engine/Source/Runtime/OpenGLDrv/Public/OpenGLShaderResources.h
+++ b/Engine/Source/Runtime/OpenGLDrv/Public/OpenGLShaderResources.h
@@ -101,15 +101,14 @@ struct FOpenGLShaderBindings
 	TArray<FOpenGLShaderVarying>					InputVaryings;
 	TArray<FOpenGLShaderVarying>					OutputVaryings;
 	FOpenGLShaderResourceTable						ShaderResourceTable;
+	CrossCompiler::FShaderBindingInOutMask			InOutMask;
 
-	uint16	InOutMask;
 	uint8	NumSamplers;
 	uint8	NumUniformBuffers;
 	uint8	NumUAVs;
 	bool	bFlattenUB;
 
 	FOpenGLShaderBindings() :
-		InOutMask(0),
 		NumSamplers(0),
 		NumUniformBuffers(0),
 		NumUAVs(0),
@@ -163,9 +162,9 @@ struct FOpenGLShaderBindings
 	friend uint32 GetTypeHash(const FOpenGLShaderBindings &Binding)
 	{
 		uint32 Hash = 0;
-		Hash = Binding.InOutMask;
-		Hash |= Binding.NumSamplers << 16;
-		Hash |= Binding.NumUniformBuffers << 24;
+		Hash = Binding.InOutMask.Bitmask;
+		Hash ^= Binding.NumSamplers << 16;
+		Hash ^= Binding.NumUniformBuffers << 24;
 		Hash ^= Binding.NumUAVs;
 		Hash ^= Binding.bFlattenUB << 8;
 		Hash ^= FCrc::MemCrc_DEPRECATED( Binding.PackedGlobalArrays.GetData(), Binding.PackedGlobalArrays.GetTypeSize()*Binding.PackedGlobalArrays.Num());
diff --git a/Engine/Source/Runtime/RenderCore/Public/CrossCompilerCommon.h b/Engine/Source/Runtime/RenderCore/Public/CrossCompilerCommon.h
--- a/Engine/Source/Runtime/RenderCore/Public/CrossCompilerCommon.h
+++ b/Engine/Source/Runtime/RenderCore/Public/CrossCompilerCommon.h
@@ -189,17 +189,58 @@ namespace CrossCompiler
 		return Ar;
 	}
 
+	struct FShaderBindingInOutMask
+	{
+		uint32 Bitmask = 0;
+
+		/** Maximum value for a valid index in this bitmask. */
+		static constexpr int32 MaxIndex = (sizeof(Bitmask) * 8u) - 1u;
+
+		/** Index to mark the binding of a depth-stencil output resource. */
+		static constexpr int32 DepthStencilMaskIndex = FShaderBindingInOutMask::MaxIndex;
+
+		/** Sets the specified bitfield in this bitmask and validates its index boundary. */
+		FORCEINLINE void EnableField(int32 Index)
+		{
+			ensure(Index >= 0 && Index <= FShaderBindingInOutMask::MaxIndex);
+			Bitmask |= (1u << Index);
+		}
+
+		/** Returns whether the specified bitfield in this bitmask is set and validates its index boundary. */
+		FORCEINLINE bool IsFieldEnabled(int32 Index) const
+		{
+			ensure(Index >= 0 && Index <= FShaderBindingInOutMask::MaxIndex);
+			return (Bitmask & (1u << Index)) != 0;
+		}
+
+		friend bool operator == (const FShaderBindingInOutMask& Lhs, const FShaderBindingInOutMask& Rhs)
+		{
+			return Lhs.Bitmask == Rhs.Bitmask;
+		}
+
+		friend bool operator != (const FShaderBindingInOutMask& Lhs, const FShaderBindingInOutMask& Rhs)
+		{
+			return !(Lhs == Rhs);
+		}
+	};
+
+	inline FArchive& operator<<(FArchive& Ar, FShaderBindingInOutMask& BindingInOutMask)
+	{
+		Ar << BindingInOutMask.Bitmask;
+		return Ar;
+	}
+
 	struct FShaderBindings
 	{
 		TArray<TArray<FPackedArrayInfo>>	PackedUniformBuffers;
 		TArray<FPackedArrayInfo>			PackedGlobalArrays;
 		FShaderCompilerResourceTable		ShaderResourceTable;
 
-		uint16	InOutMask;
-		uint8	NumSamplers;
-		uint8	NumUniformBuffers;
-		uint8	NumUAVs;
-		bool	bHasRegularUniformBuffers;
+		FShaderBindingInOutMask	InOutMask;
+		uint8					NumSamplers;
+		uint8					NumUniformBuffers;
+		uint8					NumUAVs;
+		bool					bHasRegularUniformBuffers;
 	};
 
 	// Information for copying members from uniform buffers to packed
