CIS fix for Mac (missing paren on conditional)

diff --git a/Engine/Source/Runtime/Apple/AudioMixerCoreAudio/Private/AudioMixerPlatformCoreAudio.cpp b/Engine/Source/Runtime/Apple/AudioMixerCoreAudio/Private/AudioMixerPlatformCoreAudio.cpp
--- a/Engine/Source/Runtime/Apple/AudioMixerCoreAudio/Private/AudioMixerPlatformCoreAudio.cpp
+++ b/Engine/Source/Runtime/Apple/AudioMixerCoreAudio/Private/AudioMixerPlatformCoreAudio.cpp
@@ -417,7 +417,7 @@ namespace Audio
 	ICompressedAudioInfo* FMixerPlatformCoreAudio::CreateCompressedAudioInfo(const FSoundWaveProxyPtr& InSoundWave)
 	{
 #if WITH_ENGINE
-		if (!ensure(InSoundWave.IsValid))
+		if (!ensure(InSoundWave.IsValid()))
 		{
 			return nullptr;
 		}
