Mac-only CIS fix (missing #endif)

diff --git a/Engine/Source/Runtime/Apple/AudioMixerCoreAudio/Private/AudioMixerPlatformCoreAudio.cpp b/Engine/Source/Runtime/Apple/AudioMixerCoreAudio/Private/AudioMixerPlatformCoreAudio.cpp
--- a/Engine/Source/Runtime/Apple/AudioMixerCoreAudio/Private/AudioMixerPlatformCoreAudio.cpp
+++ b/Engine/Source/Runtime/Apple/AudioMixerCoreAudio/Private/AudioMixerPlatformCoreAudio.cpp
@@ -422,7 +422,6 @@ namespace Audio
 			return nullptr;
 		}
 
-
 		if (InSoundWave->IsSeekableStreaming())
 		{
 			return new FADPCMAudioInfo();
@@ -432,6 +431,7 @@ namespace Audio
 			return new FOpusAudioInfo();
 		}
 
+#endif // WITH_ENGINE
 		return nullptr;
 	}
 
