Remove all tessellation support from Metal RHI

- Also, remove _NOTESS Metal profile

#rnx

diff --git a/Engine/Shaders/Private/BasePassPixelShader.usf b/Engine/Shaders/Private/BasePassPixelShader.usf
--- a/Engine/Shaders/Private/BasePassPixelShader.usf
+++ b/Engine/Shaders/Private/BasePassPixelShader.usf
@@ -1449,7 +1449,7 @@ void FPixelShaderInOut_MainPS(
 	// The following block is disabled on Vulkan because it triggers an nvidia driver bug (UE-101609).
 #if USE_DEVELOPMENT_SHADERS && !VULKAN_PROFILE_SM5
 	// this feature is only needed for development/editor - we can compile it out for a shipping build (see r.CompileShadersForDevelopment cvar help)
-	#if METAL_SM5_PROFILE || SM6_PROFILE || SM5_PROFILE || SM4_PROFILE || METAL_SM5_NOTESS_PROFILE || VULKAN_PROFILE_SM5
+	#if METAL_SM5_PROFILE || SM6_PROFILE || SM5_PROFILE || SM4_PROFILE || VULKAN_PROFILE_SM5
 		BRANCH
 		if (View.OutOfBoundsMask > 0)
 		{
diff --git a/Engine/Shaders/Private/PostProcessHMD.usf b/Engine/Shaders/Private/PostProcessHMD.usf
--- a/Engine/Shaders/Private/PostProcessHMD.usf
+++ b/Engine/Shaders/Private/PostProcessHMD.usf
@@ -69,7 +69,7 @@ void MainPS(
 {
 	// Only testing 1 channel for optimization
 #if PIXELSHADER // Fix gles error since VS and PS are in the same file
-#if METAL_PROFILE || METAL_MRT_PROFILE || METAL_SM5_NOTESS_PROFILE || METAL_SM5_PROFILE // Fix Mac issue with bvec2
+#if METAL_PROFILE || METAL_MRT_PROFILE || METAL_SM5_PROFILE // Fix Mac issue with bvec2
 	const float2 OOBCoords = TexCoord1.zw - saturate(TexCoord1.zw);
 	if( dot(OOBCoords, OOBCoords) != 0.0 )
 #else
@@ -83,4 +83,4 @@ void MainPS(
 	float ResultG = Texture2DSample(InputTexture, InputSampler, TexCoord1.xy).g;
 	float ResultB = Texture2DSample(InputTexture, InputSampler, TexCoord2.xy).b;
 	OutColor = float4(ResultR * Color.r, ResultG * Color.g, ResultB * Color.b, 1.0);
-}
\ No newline at end of file
+}
diff --git a/Engine/Shaders/Public/Platform.ush b/Engine/Shaders/Public/Platform.ush
--- a/Engine/Shaders/Public/Platform.ush
+++ b/Engine/Shaders/Public/Platform.ush
@@ -113,10 +113,6 @@
 #define METAL_MRT_PROFILE 0
 #endif
 
-#ifndef METAL_SM5_NOTESS_PROFILE
-#define METAL_SM5_NOTESS_PROFILE 0
-#endif
-
 #ifndef METAL_SM5_PROFILE
 #define METAL_SM5_PROFILE 0
 #endif
@@ -264,12 +260,6 @@
 	#define FMANUALFETCH_COLOR_COMPONENT_SWIZZLE .bgra
 	#define STENCIL_COMPONENT_SWIZZLE .x
 
-#elif METAL_SM5_NOTESS_PROFILE
-	#define FEATURE_LEVEL FEATURE_LEVEL_SM5
-	#define FCOLOR_COMPONENT_SWIZZLE .rgba
-	#define FMANUALFETCH_COLOR_COMPONENT_SWIZZLE .bgra
-	#define STENCIL_COMPONENT_SWIZZLE .x
-
 #elif METAL_SM5_PROFILE
 	#define FEATURE_LEVEL FEATURE_LEVEL_SM5
 	#define FCOLOR_COMPONENT_SWIZZLE .rgba
@@ -344,13 +334,13 @@
 // Ideally this would just be enabled for all SM4 and above feature level platforms
 // @todo metal mrt: No reason this can't work with Metal, once cross compiler is fixed
 #ifndef SUPPORTS_INDEPENDENT_SAMPLERS
-#define SUPPORTS_INDEPENDENT_SAMPLERS (SM6_PROFILE || SM5_PROFILE || SM4_PROFILE || METAL_MRT_PROFILE || METAL_SM5_NOTESS_PROFILE || METAL_SM5_PROFILE || VULKAN_PROFILE_SM5 || VULKAN_PROFILE)
+#define SUPPORTS_INDEPENDENT_SAMPLERS (SM6_PROFILE || SM5_PROFILE || SM4_PROFILE || METAL_MRT_PROFILE || METAL_SM5_PROFILE || VULKAN_PROFILE_SM5 || VULKAN_PROFILE)
 #endif
 
 // Whether the platform supports a global clip plane through SV_ClipDistance
 // Ideally this would just be enabled for all SM4 and above feature level platforms, but not tested everywhere yet
 #ifndef PLATFORM_SUPPORTS_GLOBAL_CLIP_PLANE
-#define PLATFORM_SUPPORTS_GLOBAL_CLIP_PLANE (PS4_PROFILE || SM6_PROFILE || SM5_PROFILE || SM4_PROFILE || METAL_PROFILE || METAL_MRT_PROFILE || METAL_SM5_NOTESS_PROFILE || METAL_SM5_PROFILE || VULKAN_PROFILE_SM5)
+#define PLATFORM_SUPPORTS_GLOBAL_CLIP_PLANE (PS4_PROFILE || SM6_PROFILE || SM5_PROFILE || SM4_PROFILE || METAL_PROFILE || METAL_MRT_PROFILE || METAL_SM5_PROFILE || VULKAN_PROFILE_SM5)
 #endif
 
 // Whether the platform support pixel coverage on MSAA targets (SV_Coverage).
@@ -492,11 +482,11 @@ void ShaderYield()
 
 #endif // SM6_PROFILE || SM5_PROFILE || COMPILER_SUPPORTS_ATTRIBUTES
 
-#if SM6_PROFILE || SM5_PROFILE || METAL_MRT_PROFILE || METAL_SM5_PROFILE || METAL_SM5_NOTESS_PROFILE || ES3_1_PROFILE
+#if SM6_PROFILE || SM5_PROFILE || METAL_MRT_PROFILE || METAL_SM5_PROFILE || ES3_1_PROFILE
 #define EARLYDEPTHSTENCIL [earlydepthstencil]
 #endif
 
-#if METAL_MRT_PROFILE || METAL_SM5_PROFILE || METAL_SM5_NOTESS_PROFILE
+#if METAL_MRT_PROFILE || METAL_SM5_PROFILE
 #define DUAL_SOURCE_BLENDING_SLOT(SLOT) [[vk::location(0), vk::index(SLOT)]]
 #endif
 
diff --git a/Engine/Shaders/Public/ShaderVersion.ush b/Engine/Shaders/Public/ShaderVersion.ush
--- a/Engine/Shaders/Public/ShaderVersion.ush
+++ b/Engine/Shaders/Public/ShaderVersion.ush
@@ -2,4 +2,4 @@
 // Each time the console command is executed it generates a new GUID. As this file is included
 // in Platform.ush (which should be included in any shader) it allows to invalidate the shader DDC.
 // 
-// GUID = 8FF90D8281CB4FEB9C95FD3CF0E5BE07
\ No newline at end of file
+// GUID = 4B4392B38CAB46E39C0C10B2CA18FF52
diff --git a/Engine/Source/Developer/Apple/MetalShaderFormat/Private/MetalDerivedData.cpp b/Engine/Source/Developer/Apple/MetalShaderFormat/Private/MetalDerivedData.cpp
--- a/Engine/Source/Developer/Apple/MetalShaderFormat/Private/MetalDerivedData.cpp
+++ b/Engine/Source/Developer/Apple/MetalShaderFormat/Private/MetalDerivedData.cpp
@@ -34,7 +34,6 @@ extern void BuildMetalShaderOutput(
 	TCHAR const* MinOSVersion,
 	EMetalTypeBufferMode TypeMode,
 	TArray<FShaderCompilerError>& OutErrors,
-	FMetalTessellationOutputs const& TessOutputAttribs,
 	uint32 TypedBuffers,
 	uint32 InvariantBuffers,
 	uint32 TypedUAVs,
@@ -361,7 +360,6 @@ bool FMetalShaderOutputCooker::Build(TArray<uint8>& OutData)
 	}
 	
 	FMetalShaderOutputMetaData OutputData;
-	FMetalTessellationOutputs Attribs;
 	
 	uint32 CRCLen = 0;
 	uint32 CRC = 0;
@@ -377,11 +375,8 @@ bool FMetalShaderOutputCooker::Build(TArray<uint8>& OutData)
 	};
 	TMap<FString, TArray<FMetalResourceTableEntry>> IABs;
 	
-	bool bUsingTessellation = Input.IsUsingTessellation();
-	
 #if PLATFORM_MAC || PLATFORM_WINDOWS
 	{
-		const bool bRewriteHlslSource = !bUsingTessellation;
 		std::string EntryPointNameAnsi(TCHAR_TO_UTF8(*Input.EntryPointName));
 
 		CrossCompiler::FShaderConductorContext CompilerContext;
@@ -389,13 +384,9 @@ bool FMetalShaderOutputCooker::Build(TArray<uint8>& OutData)
 		// Initialize compilation options for ShaderConductor
 		CrossCompiler::FShaderConductorOptions Options;
 
-		// Always enable FMA pass for Hull shader, always disable for Pixel and Compute shader,
+		// Always disable FMA pass for Pixel and Compute shader,
 		// otherwise determine whether [[position, invariant]] qualifier is available in Metal or not.
-		if (Frequency == HSF_HullShader)
-		{
-			Options.bEnableFMAPass = true;
-		}
-		else if (Frequency == HSF_PixelShader || Frequency == HSF_ComputeShader)
+		if (Frequency == HSF_PixelShader || Frequency == HSF_ComputeShader)
 		{
 			Options.bEnableFMAPass = false;
 		}
@@ -407,16 +398,13 @@ bool FMetalShaderOutputCooker::Build(TArray<uint8>& OutData)
 		// Load shader source into compiler context
 		CompilerContext.LoadSource(PreprocessedShader, Input.VirtualSourceFilePath, Input.EntryPointName, Frequency);
 
-		if (bRewriteHlslSource)
+		// Rewrite HLSL source code to remove unused global resources and variables
+		Options.bRemoveUnusedGlobals = true;
+		if (!CompilerContext.RewriteHlsl(Options, &PreprocessedShader))
 		{
-			// Rewrite HLSL source code to remove unused global resources and variables
-			Options.bRemoveUnusedGlobals = true;
-			if (!CompilerContext.RewriteHlsl(Options, &PreprocessedShader))
-			{
-				CompilerContext.FlushErrors(Output.Errors);
-			}
-			Options.bRemoveUnusedGlobals = false;
+			CompilerContext.FlushErrors(Output.Errors);
 		}
+		Options.bRemoveUnusedGlobals = false;
 
 		// Convert shader source to ANSI string
 		std::string SourceData(CompilerContext.GetSourceString(), static_cast<size_t>(CompilerContext.GetSourceLength()));
@@ -433,7 +421,7 @@ bool FMetalShaderOutputCooker::Build(TArray<uint8>& OutData)
 			CompilerContext.LoadSource(SourceData.c_str(), /*Fileanme:*/ nullptr, /*EntryPoint:*/ nullptr, Frequency);
 		}
 
-		if (bDumpDebugInfo && bRewriteHlslSource)
+		if (bDumpDebugInfo)
 		{
 			DumpDebugShaderText(Input, &SourceData[0], SourceData.size(), TEXT("rewritten.hlsl"));
 		}
@@ -549,37 +537,6 @@ bool FMetalShaderOutputCooker::Build(TArray<uint8>& OutData)
 				}
 			}
 			
-			if(Frequency == HSF_HullShader)
-			{
-				uint32 PatchSize = 0;
-				size_t ThreadSizeIdx = SourceData.find("InputPatch<");
-				check(ThreadSizeIdx != std::string::npos);
-				char const* String = SourceData.c_str() + ThreadSizeIdx;
-				
-#if !PLATFORM_WINDOWS
-				size_t Found = sscanf(String, "InputPatch< %*s ,  %u >", &PatchSize);
-#else
-				size_t Found = sscanf_s(String, "InputPatch< %*s ,  %u >", &PatchSize);
-#endif
-				check(Found == 1);
-				
-				ThreadSizeIdx = SourceData.find("[maxtessfactor(");
-				check(ThreadSizeIdx != std::string::npos);
-				String = SourceData.c_str() + ThreadSizeIdx;
-				
-				uint32 MaxTessFactor = 0;
-#if !PLATFORM_WINDOWS
-				Found = sscanf(String, "[maxtessfactor(%u)]", &MaxTessFactor);
-#else
-				Found = sscanf_s(String, "[maxtessfactor(%u)]", &MaxTessFactor);
-#endif
-				check(Found == 1);
-				
-				CCHeaderWriter.WriteTessellationInputControlPoints(PatchSize);
-				CCHeaderWriter.WriteTessellationMaxTessFactor(MaxTessFactor);
-				CCHeaderWriter.WriteTessellationPatchesPerThreadGroup(1);
-			}
-			
 			{
 				Count = 0;
 				SPVRResult = Reflection.EnumerateExecutionModes(&Count, nullptr);
@@ -599,94 +556,6 @@ bool FMetalShaderOutputCooker::Build(TArray<uint8>& OutData)
 								CCHeaderWriter.WriteNumThreads(Mode->operands[0], Mode->operands[1], Mode->operands[2]);
 							}
 							break;
-						case SpvExecutionModeTriangles:
-							if (Frequency == HSF_HullShader || Frequency == HSF_DomainShader)
-							{
-								CCHeaderWriter.WriteTessellationDomainTri();
-								Attribs.HSTFOutSize = 8; // sizeof(MTLTriangleTessellationFactorsHalf);
-							}
-							break;
-						case SpvExecutionModeQuads:
-							if (Frequency == HSF_HullShader || Frequency == HSF_DomainShader)
-							{
-								CCHeaderWriter.WriteTessellationDomainQuad();
-								Attribs.HSTFOutSize = 12; // sizeof(MTLQuadTessellationFactorsHalf);
-							}
-							break;
-						case SpvExecutionModeIsolines:
-							if (Frequency == HSF_HullShader || Frequency == HSF_DomainShader)
-							{
-								check(0); // Not supported by Metal
-							}
-							break;
-						case SpvExecutionModeOutputVertices:
-							if (Frequency == HSF_HullShader || Frequency == HSF_DomainShader)
-							{
-								check(Mode->operands_count == 1);
-								CCHeaderWriter.WriteTessellationOutputControlPoints(Mode->operands[0]);
-							}
-							break;
-						case SpvExecutionModeVertexOrderCw:
-							if (Frequency == HSF_HullShader || Frequency == HSF_DomainShader)
-							{
-								CCHeaderWriter.WriteTessellationOutputWindingCW();
-							}
-							break;
-						case SpvExecutionModeVertexOrderCcw:
-							if (Frequency == HSF_HullShader || Frequency == HSF_DomainShader)
-							{
-								CCHeaderWriter.WriteTessellationOutputWindingCCW();
-							}
-							break;
-						case SpvExecutionModeSpacingEqual:
-							if (Frequency == HSF_HullShader || Frequency == HSF_DomainShader)
-							{
-								CCHeaderWriter.WriteTessellationPartitioningInteger();
-							}
-							break;
-						case SpvExecutionModeSpacingFractionalEven:
-							if (Frequency == HSF_HullShader || Frequency == HSF_DomainShader)
-							{
-								CCHeaderWriter.WriteTessellationPartitioningFractionalEven();
-							}
-							break;
-						case SpvExecutionModeSpacingFractionalOdd:
-							if (Frequency == HSF_HullShader || Frequency == HSF_DomainShader)
-							{
-								CCHeaderWriter.WriteTessellationPartitioningFractionalOdd();
-							}
-							break;
-						case SpvExecutionModeInvocations:
-						case SpvExecutionModePixelCenterInteger:
-						case SpvExecutionModeOriginUpperLeft:
-						case SpvExecutionModeOriginLowerLeft:
-						case SpvExecutionModeEarlyFragmentTests:
-						case SpvExecutionModePointMode:
-						case SpvExecutionModeXfb:
-						case SpvExecutionModeDepthReplacing:
-						case SpvExecutionModeDepthGreater:
-						case SpvExecutionModeDepthLess:
-						case SpvExecutionModeDepthUnchanged:
-						case SpvExecutionModeInputPoints:
-						case SpvExecutionModeInputLines:
-						case SpvExecutionModeInputLinesAdjacency:
-						case SpvExecutionModeInputTrianglesAdjacency:
-						case SpvExecutionModeOutputPoints:
-						case SpvExecutionModeOutputLineStrip:
-						case SpvExecutionModeOutputTriangleStrip:
-						case SpvExecutionModeVecTypeHint:
-						case SpvExecutionModeContractionOff:
-						case SpvExecutionModeInitializer:
-						case SpvExecutionModeFinalizer:
-						case SpvExecutionModeSubgroupSize:
-						case SpvExecutionModeSubgroupsPerWorkgroup:
-						case SpvExecutionModeSubgroupsPerWorkgroupId:
-						case SpvExecutionModeLocalSizeId:
-						case SpvExecutionModeLocalSizeHintId:
-						case SpvExecutionModePostDepthCoverage:
-						case SpvExecutionModeStencilRefReplacingEXT:
-							break;
-							
 						default:
 							break;
 					}
@@ -1296,321 +1165,11 @@ bool FMetalShaderOutputCooker::Build(TArray<uint8>& OutData)
 				}
 			}
 			
-			if (Frequency == HSF_VertexShader && bUsingTessellation)
-			{
-				Count = 0;
-				SPVRResult = Reflection.EnumerateOutputVariables(&Count, nullptr);
-				check(SPVRResult == SPV_REFLECT_RESULT_SUCCESS);
-				OutputVars.SetNum(Count);
-				SPVRResult = Reflection.EnumerateOutputVariables(&Count, OutputVars.GetData());
-				check(SPVRResult == SPV_REFLECT_RESULT_SUCCESS);
-				if (Count > 0)
-				{
-					
-					OutputVars.Sort([](SpvReflectInterfaceVariable& LHS, SpvReflectInterfaceVariable& RHS)  { return LHS.location < RHS.location; });
-					uint32 Offset = 0;
-					uint32 InitialAlign = 0;
-					for (auto const& Var : OutputVars)
-					{
-						if (Var->storage_class == SpvStorageClassOutput && Var->built_in == -1)
-						{
-							unsigned Location = Var->location;
-							
-							uint32 ArrayCount = 1;
-							for (uint32 Dim = 0; Dim < Var->array.dims_count; Dim++)
-							{
-								ArrayCount *= Var->array.dims[Dim];
-							}
-							
-							auto const type = *Var->type_description;
-							uint32_t masked_type = type.type_flags & 0xF;
-							
-							for (uint32 j = 0; j < ArrayCount; j++)
-							{
-								FMetalAttribute& Attrib = Attribs.HSOut.AddDefaulted_GetRef();
-								Attrib.Index = Location + j;
-								uint32 ElementSize = 0;
-								switch (masked_type) {
-									default: checkf(false, TEXT("unsupported component type %d"), masked_type); break;
-									case SPV_REFLECT_TYPE_FLAG_BOOL  : Attrib.Type = EMetalComponentType::Bool; ElementSize = 1; break;
-									case SPV_REFLECT_TYPE_FLAG_INT   : Attrib.Type = (type.traits.numeric.scalar.signedness ? EMetalComponentType::Int : EMetalComponentType::Uint); ElementSize = 4; break;
-									case SPV_REFLECT_TYPE_FLAG_FLOAT : Attrib.Type = (type.traits.numeric.scalar.width == 32 ? EMetalComponentType::Float : EMetalComponentType::Half); if(type.traits.numeric.scalar.width == 32) { ElementSize = 4; } else { ElementSize = 2; } break;
-								}
-								if (type.type_flags & SPV_REFLECT_TYPE_FLAG_MATRIX)
-								{
-									Attrib.Components = (type.traits.numeric.matrix.row_count * type.traits.numeric.matrix.column_count);
-								}
-								else if (type.type_flags & SPV_REFLECT_TYPE_FLAG_VECTOR)
-								{
-									Attrib.Components = (type.traits.numeric.vector.component_count);
-								}
-								else
-								{
-									Attrib.Components = 1;
-								}
-								
-								char const* NumberedSemantic = strstr(Var->name, Var->semantic);
-								FString Semantic = FString::Printf(TEXT("%s_%u"), UTF8_TO_TCHAR(NumberedSemantic), j);
-								Attrib.Semantic = GetTypeHash(Semantic);
-								
-								uint32 AlignedComponents = Attrib.Components == 3 ? 4 : Attrib.Components;
-								InitialAlign = InitialAlign == 0 ? ElementSize * AlignedComponents : InitialAlign;
-								Offset = Align(Offset, ElementSize * AlignedComponents);
-								Attrib.Offset = Offset;
-								Offset += ElementSize * AlignedComponents;
-
-								if (Input.Environment.CompilerFlags.Contains(CFLAG_KeepDebugInfo))
-									ALNString += FString::Printf(TEXT("%sHSOut.%s(%u)"), ALNString.Len() ? TEXT(",") : TEXT(""), *Semantic, Attrib.Offset);
-							}
-						}
-					}
-					Attribs.HSOutSize = Align(Offset, 16);
-					
-					if (Input.Environment.CompilerFlags.Contains(CFLAG_KeepDebugInfo))
-						ALNString += FString::Printf(TEXT(" = HSOutSize(%u)"), Attribs.HSOutSize);
-				}
-			}
-			
-			if (Frequency == HSF_HullShader)
-			{
-				Count = 0;
-				SPVRResult = Reflection.EnumerateInputVariables(&Count, nullptr);
-				check(SPVRResult == SPV_REFLECT_RESULT_SUCCESS);
-				InputVars.SetNum(Count);
-				SPVRResult = Reflection.EnumerateInputVariables(&Count, InputVars.GetData());
-				check(SPVRResult == SPV_REFLECT_RESULT_SUCCESS);
-				if (Count > 0)
-				{
-					
-					InputVars.Sort([](SpvReflectInterfaceVariable& LHS, SpvReflectInterfaceVariable& RHS)  { return LHS.location < RHS.location; });
-					uint32 Offset = 0;
-					for (auto const& Var : InputVars)
-					{
-						if (Var->storage_class == SpvStorageClassInput && Var->built_in == -1)
-						{
-							unsigned Location = Var->location;
-							
-							uint32 ArrayCount = 1;
-							for (uint32 Dim = 1; Dim < Var->array.dims_count; Dim++)
-							{
-								ArrayCount *= Var->array.dims[Dim];
-							}
-							
-							auto const type = *Var->type_description;
-							uint32_t masked_type = type.type_flags & 0xF;
-							
-							for (uint32 j = 0; j < ArrayCount; j++)
-							{
-								FMetalAttribute& Attrib = Attribs.HSIn.AddDefaulted_GetRef();
-								Attrib.Index = Location + j;
-								uint32 ElementSize = 0;
-								switch (masked_type) {
-									default: checkf(false, TEXT("unsupported component type %d"), masked_type); break;
-									case SPV_REFLECT_TYPE_FLAG_BOOL  : Attrib.Type = EMetalComponentType::Bool; ElementSize = 1; break;
-									case SPV_REFLECT_TYPE_FLAG_INT   : Attrib.Type = (type.traits.numeric.scalar.signedness ? EMetalComponentType::Int : EMetalComponentType::Uint); ElementSize = 4; break;
-									case SPV_REFLECT_TYPE_FLAG_FLOAT : Attrib.Type = (type.traits.numeric.scalar.width == 32 ? EMetalComponentType::Float : EMetalComponentType::Half); if(type.traits.numeric.scalar.width == 32) { ElementSize = 4; } else { ElementSize = 2; } break;
-								}
-								if (type.type_flags & SPV_REFLECT_TYPE_FLAG_MATRIX)
-								{
-									Attrib.Components = (type.traits.numeric.matrix.row_count * type.traits.numeric.matrix.column_count);
-								}
-								else if (type.type_flags & SPV_REFLECT_TYPE_FLAG_VECTOR)
-								{
-									Attrib.Components = (type.traits.numeric.vector.component_count);
-								}
-								else
-								{
-									Attrib.Components = 1;
-								}
-								
-								char const* NumberedSemantic = strstr(Var->name, Var->semantic);
-								FString Semantic = FString::Printf(TEXT("%s_%u"), UTF8_TO_TCHAR(NumberedSemantic), j);
-								Attrib.Semantic = GetTypeHash(Semantic);
-								
-								uint32 AlignedComponents = Attrib.Components == 3 ? 4 : Attrib.Components;
-								Offset = Align(Offset, ElementSize * AlignedComponents);
-								Attrib.Offset = Offset;
-								Offset += ElementSize * AlignedComponents;
-								
-								if (Input.Environment.CompilerFlags.Contains(CFLAG_KeepDebugInfo))
-									ALNString += FString::Printf(TEXT("%sHSIn.%s(%u)"), ALNString.Len() ? TEXT(",") : TEXT(""), *Semantic, Attrib.Offset);
-							}
-						}
-					}
-				}
-				
-				Count = 0;
-				SPVRResult = Reflection.EnumerateOutputVariables(&Count, nullptr);
-				check(SPVRResult == SPV_REFLECT_RESULT_SUCCESS);
-				OutputVars.SetNum(Count);
-				SPVRResult = Reflection.EnumerateOutputVariables(&Count, OutputVars.GetData());
-				check(SPVRResult == SPV_REFLECT_RESULT_SUCCESS);
-				if (Count > 0)
-				{
-					uint32 HSOutOffset = 0;
-					uint32 PatchControlPointOffset = 0;
-					uint32 HSOutAlign = 0;
-					uint32 PatchControlPointAlign = 0;
-					
-					OutputVars.Sort([](SpvReflectInterfaceVariable& LHS, SpvReflectInterfaceVariable& RHS)  { return LHS.location < RHS.location; });
-					
-					for (auto const& Var : OutputVars)
-					{
-						if (Var->storage_class == SpvStorageClassOutput && Var->built_in == -1 && !CrossCompiler::FShaderConductorContext::IsIntermediateSpirvOutputVariable(Var->name))
-						{
-							unsigned Location = Var->location;
-							uint32 ArrayCount = 1;
-							for (uint32 Dim = 1; Dim < Var->array.dims_count; Dim++)
-							{
-								ArrayCount *= Var->array.dims[Dim];
-							}
-							
-							auto const type = *Var->type_description;
-							uint32_t masked_type = type.type_flags & 0xF;
-							
-							for (uint32 j = 0; j < ArrayCount; j++)
-							{
-								FMetalAttribute& Attrib = Var->array.dims_count ? Attribs.PatchControlPointOut.AddDefaulted_GetRef() : Attribs.HSOut.AddDefaulted_GetRef();
-								Attrib.Index = Location + j;
-								uint32 ElementSize = 0;
-								switch (masked_type) {
-									default: checkf(false, TEXT("unsupported component type %d"), masked_type); break;
-									case SPV_REFLECT_TYPE_FLAG_BOOL  : Attrib.Type = EMetalComponentType::Bool; ElementSize = 1; break;
-									case SPV_REFLECT_TYPE_FLAG_INT   : Attrib.Type = (type.traits.numeric.scalar.signedness ? EMetalComponentType::Int : EMetalComponentType::Uint); ElementSize = 4; break;
-									case SPV_REFLECT_TYPE_FLAG_FLOAT : Attrib.Type = (type.traits.numeric.scalar.width == 32 ? EMetalComponentType::Float : EMetalComponentType::Half); if(type.traits.numeric.scalar.width == 32) { ElementSize = 4; } else { ElementSize = 2; } break;
-								}
-								if (type.type_flags & SPV_REFLECT_TYPE_FLAG_MATRIX)
-								{
-									Attrib.Components = (type.traits.numeric.matrix.row_count * type.traits.numeric.matrix.column_count);
-								}
-								else if (type.type_flags & SPV_REFLECT_TYPE_FLAG_VECTOR)
-								{
-									Attrib.Components = (type.traits.numeric.vector.component_count);
-								}
-								else
-								{
-									Attrib.Components = 1;
-								}
-								
-								uint32& Offset = Var->array.dims_count ? PatchControlPointOffset : HSOutOffset;
-								uint32& InitialAlign = Var->array.dims_count ? PatchControlPointAlign : HSOutAlign;
-								TCHAR const* Name = Var->array.dims_count ? TEXT("PatchOut") : TEXT("HSOut");
-								
-								char const* NumberedSemantic = strstr(Var->name, Var->semantic);
-								FString Semantic = FString::Printf(TEXT("%s_%u"), UTF8_TO_TCHAR(NumberedSemantic), j);
-								Attrib.Semantic = GetTypeHash(Semantic);
-								
-								uint32 AlignedComponents = Attrib.Components == 3 ? 4 : Attrib.Components;
-								InitialAlign = InitialAlign == 0 ? ElementSize * AlignedComponents : InitialAlign;
-								Offset = Align(Offset, ElementSize * AlignedComponents);
-								Attrib.Offset = Offset;
-								Offset += ElementSize * AlignedComponents;
-								
-								if (Input.Environment.CompilerFlags.Contains(CFLAG_KeepDebugInfo))
-									ALNString += FString::Printf(TEXT("%s%s.%s(%u)"), ALNString.Len() ? TEXT(",") : TEXT(""), Name, *Semantic, Attrib.Offset);
-							}
-						}
-					}
-					Attribs.HSOutSize = Align(HSOutOffset, 16);
-					Attribs.PatchControlPointOutSize = Align(PatchControlPointOffset, 16);
-					
-					if (Input.Environment.CompilerFlags.Contains(CFLAG_KeepDebugInfo))
-						ALNString += FString::Printf(TEXT(" = HSOutSize(%u), PatchControlPointOutSize(%u)"), Attribs.HSOutSize, Attribs.PatchControlPointOutSize);
-				}
-			}
-			
-			if (Frequency == HSF_DomainShader)
-			{
-				Count = 0;
-				SPVRResult = Reflection.EnumerateInputVariables(&Count, nullptr);
-				check(SPVRResult == SPV_REFLECT_RESULT_SUCCESS);
-				InputVars.SetNum(Count);
-				SPVRResult = Reflection.EnumerateInputVariables(&Count, InputVars.GetData());
-				check(SPVRResult == SPV_REFLECT_RESULT_SUCCESS);
-				if (Count > 0)
-				{
-					uint32 HSOutOffset = 0;
-					uint32 PatchControlPointOffset = 0;
-					uint32 HSOutAlign = 0;
-					uint32 PatchControlPointAlign = 0;
-					
-					InputVars.Sort([](SpvReflectInterfaceVariable& LHS, SpvReflectInterfaceVariable& RHS)  { return LHS.location < RHS.location; });
-					for (auto const& Var : InputVars)
-					{
-						if (Var->storage_class == SpvStorageClassInput && Var->built_in == -1)
-						{
-							unsigned Location = Var->location;
-							
-							uint32 ArrayCount = 1;
-							for (uint32 Dim = 1; Dim < Var->array.dims_count; Dim++)
-							{
-								ArrayCount *= Var->array.dims[Dim];
-							}
-							
-							auto const type = *Var->type_description;
-							uint32_t masked_type = type.type_flags & 0xF;
-							
-							for (uint32 j = 0; j < ArrayCount; j++)
-							{
-								FMetalAttribute& Attrib = Var->array.dims_count ? Attribs.PatchControlPointOut.AddDefaulted_GetRef() : Attribs.HSOut.AddDefaulted_GetRef();
-								Attrib.Index = Location + j;
-								uint32 ElementSize = 0;
-								switch (masked_type) {
-									default: checkf(false, TEXT("unsupported component type %d"), masked_type); break;
-									case SPV_REFLECT_TYPE_FLAG_BOOL  : Attrib.Type = EMetalComponentType::Bool; ElementSize = 1; break;
-									case SPV_REFLECT_TYPE_FLAG_INT   : Attrib.Type = (type.traits.numeric.scalar.signedness ? EMetalComponentType::Int : EMetalComponentType::Uint); ElementSize = 4; break;
-									case SPV_REFLECT_TYPE_FLAG_FLOAT : Attrib.Type = (type.traits.numeric.scalar.width == 32 ? EMetalComponentType::Float : EMetalComponentType::Half); if(type.traits.numeric.scalar.width == 32) { ElementSize = 4; } else { ElementSize = 2; } break;
-								}
-								if (type.type_flags & SPV_REFLECT_TYPE_FLAG_MATRIX)
-								{
-									Attrib.Components = (type.traits.numeric.matrix.row_count * type.traits.numeric.matrix.column_count);
-								}
-								else if (type.type_flags & SPV_REFLECT_TYPE_FLAG_VECTOR)
-								{
-									Attrib.Components = (type.traits.numeric.vector.component_count);
-								}
-								else
-								{
-									Attrib.Components = 1;
-								}
-								
-								uint32& Offset = Var->array.dims_count ? PatchControlPointOffset : HSOutOffset;
-								uint32& InitialAlign = Var->array.dims_count ? PatchControlPointAlign : HSOutAlign;
-								TCHAR const* Name = Var->array.dims_count ? TEXT("PatchOut") : TEXT("HSOut");
-								
-								char const* NumberedSemantic = strstr(Var->name, Var->semantic);
-								FString Semantic = FString::Printf(TEXT("%s_%u"), UTF8_TO_TCHAR(NumberedSemantic), j);
-								Attrib.Semantic = GetTypeHash(Semantic);
-								
-								uint32 AlignedComponents = Attrib.Components == 3 ? 4 : Attrib.Components;
-								InitialAlign = InitialAlign == 0 ? ElementSize * AlignedComponents : InitialAlign;
-								Offset = Align(Offset, ElementSize * AlignedComponents);
-								Attrib.Offset = Offset;
-								Offset += ElementSize * AlignedComponents;
-								
-								if (Input.Environment.CompilerFlags.Contains(CFLAG_KeepDebugInfo))
-									ALNString += FString::Printf(TEXT("%s%s.%s(%u)"), ALNString.Len() ? TEXT(",") : TEXT(""), Name, *Semantic, Attrib.Offset);
-							}
-						}
-					}
-					Attribs.HSOutSize = Align(HSOutOffset, 16);
-					Attribs.PatchControlPointOutSize = Align(PatchControlPointOffset, 16);
-					
-					if (Input.Environment.CompilerFlags.Contains(CFLAG_KeepDebugInfo))
-						ALNString += FString::Printf(TEXT(" = HSOutSize(%u), PatchControlPointOutSize(%u)"), Attribs.HSOutSize, Attribs.PatchControlPointOutSize);
-				}
-			}
-			
 			// Copy reflection code back to SPIR-V buffer
 			SpirvData = TArray<uint32>(Reflection.GetCode(), Reflection.GetCodeSize() / sizeof(uint32));
 		}
 		
 		uint32 SideTableIndex = 0;
-		uint32 OutputBufferIndex = 0;
-		uint32 IndirectParamsIndex = 0;
-		uint32 PatchBufferIndex = 0;
-		uint32 TessFactorBufferIndex = 0;
-		uint32 HullIndexBuffer = 0;
 
 		CrossCompiler::FShaderConductorTarget TargetDesc;
 
@@ -1668,70 +1227,6 @@ bool FMetalShaderOutputCooker::Build(TArray<uint8>& OutData)
 				}
 			}
 			
-			if (Frequency == HSF_VertexShader && bUsingTessellation)
-			{
-				TargetDesc.CompileFlags.SetDefine(TEXT("capture_output_to_buffer"), 1);
-			}
-			
-			if (Frequency == HSF_VertexShader && bUsingTessellation)
-			{
-				OutputBufferIndex = FPlatformMath::CountTrailingZeros64(BufferIndices);
-				BufferIndices &= ~(1ull << (uint64)OutputBufferIndex);
-				TargetDesc.CompileFlags.SetDefine(TEXT("shader_output_buffer_index"), OutputBufferIndex);
-
-				IndirectParamsIndex = FPlatformMath::CountTrailingZeros64(BufferIndices);
-				BufferIndices &= ~(1ull << (uint64)IndirectParamsIndex);
-				TargetDesc.CompileFlags.SetDefine(TEXT("indirect_params_buffer_index"), IndirectParamsIndex);
-				
-				CCHeaderWriter.WriteTessellationHSOutBuffer(OutputBufferIndex);
-				CCHeaderWriter.WriteTessellationPatchCountBuffer(IndirectParamsIndex);
-			}
-			
-			if (Frequency == HSF_DomainShader)
-			{
-				OutputBufferIndex = FPlatformMath::CountTrailingZeros64(BufferIndices);
-				BufferIndices &= ~(1ull << (uint64)OutputBufferIndex);
-				
-				IndirectParamsIndex = FPlatformMath::CountTrailingZeros64(BufferIndices);
-				BufferIndices &= ~(1ull << (uint64)IndirectParamsIndex);
-				
-				CCHeaderWriter.WriteTessellationHSOutBuffer(OutputBufferIndex);
-				CCHeaderWriter.WriteTessellationControlPointOutBuffer(IndirectParamsIndex);
-			}
-			
-			if (Frequency == HSF_HullShader)
-			{
-				OutputBufferIndex = FPlatformMath::CountTrailingZeros64(BufferIndices);
-				BufferIndices &= ~(1ull << (uint64)OutputBufferIndex);
-				TargetDesc.CompileFlags.SetDefine(TEXT("shader_output_buffer_index"), OutputBufferIndex);
-
-				IndirectParamsIndex = FPlatformMath::CountTrailingZeros64(BufferIndices);
-				BufferIndices &= ~(1ull << (uint64)IndirectParamsIndex);
-				TargetDesc.CompileFlags.SetDefine(TEXT("indirect_params_buffer_index"), IndirectParamsIndex);
-
-				PatchBufferIndex = FPlatformMath::CountTrailingZeros64(BufferIndices);
-				BufferIndices &= ~(1ull << (uint64)PatchBufferIndex);
-				TargetDesc.CompileFlags.SetDefine(TEXT("shader_patch_output_buffer_index"), PatchBufferIndex);
-
-				TessFactorBufferIndex = FPlatformMath::CountTrailingZeros64(BufferIndices);
-				BufferIndices &= ~(1ull << (uint64)TessFactorBufferIndex);
-				TargetDesc.CompileFlags.SetDefine(TEXT("shader_tess_factor_buffer_index"), TessFactorBufferIndex);
-				TargetDesc.CompileFlags.SetDefine(TEXT("shader_input_wg_index"), 0);
-
-				HullIndexBuffer = FPlatformMath::CountTrailingZeros64(BufferIndices);
-				BufferIndices &= ~(1ull << (uint64)HullIndexBuffer);
-				
-				uint32 HullVertexBuffer = FPlatformMath::CountTrailingZeros64(BufferIndices);
-				BufferIndices &= ~(1ull << (uint64)HullVertexBuffer);
-				
-				CCHeaderWriter.WriteTessellationHSOutBuffer(PatchBufferIndex);
-				CCHeaderWriter.WriteTessellationPatchCountBuffer(IndirectParamsIndex);
-				CCHeaderWriter.WriteTessellationControlPointIndexBuffer(HullVertexBuffer);
-				CCHeaderWriter.WriteTessellationHSTFOutBuffer(TessFactorBufferIndex);
-				CCHeaderWriter.WriteTessellationControlPointOutBuffer(OutputBufferIndex);
-				CCHeaderWriter.WriteTessellationIndexBuffer(HullIndexBuffer);
-			}
-			
 			if (VersionEnum >= 4)
 			{
 				if (IABTier >= 1)
@@ -1823,7 +1318,6 @@ bool FMetalShaderOutputCooker::Build(TArray<uint8>& OutData)
 			MetalSource = TCHAR_TO_UTF8(*MetaData);
 			MetalSource += ResultsTargetDataAsString;
 			
-			// Tessellation vertex & hull shaders must always use FMA
 			if (Options.bEnableFMAPass)
 			{
 				std::string FMADefine = std::string("\n"
@@ -1844,72 +1338,7 @@ bool FMetalShaderOutputCooker::Build(TArray<uint8>& OutData)
 				if (IncludePos != std::string::npos)
 					MetalSource.replace(IncludePos, IncludeString.length(), FMADefine);
 			}
-			
-			if (Frequency == HSF_DomainShader)
-			{
-				uint32 PatchSize = 0;
-				size_t ThreadSizeIdx = SourceData.find("OutputPatch<");
-				check(ThreadSizeIdx != std::string::npos);
-				char const* String = SourceData.c_str() + ThreadSizeIdx;
-				while(!isdigit((unsigned char)*String))
-				{
-					String++;
-				}
-#if !PLATFORM_WINDOWS
-				size_t Found = sscanf(String, "%u", &PatchSize);
-#else
-				size_t Found = sscanf_s(String, "%u", &PatchSize);
-#endif
-				check(Found == 1);
-				
-				std::string PatchAttr;
-				std::string PatchSearch = "[[ patch(triangle, 0) ]]";
-				size_t PatchIdx = MetalSource.find(PatchSearch);
-				if (PatchIdx != std::string::npos)
-				{
-					PatchAttr = "[[ patch(triangle, ";
-				}
-				else
-				{
-					PatchSearch = "[[ patch(quad, 0) ]]";
-					PatchIdx = MetalSource.find(PatchSearch);
-					PatchAttr = "[[ patch(quad, ";
-				}
-				if (PatchIdx != std::string::npos)
-				{
-					char BufferIdx[3];
-					FCStringAnsi::Snprintf(BufferIdx, 3, "%d", PatchSize);
-					PatchAttr += BufferIdx;
-					PatchAttr += ") ]]";
-					MetalSource.replace(PatchIdx, PatchSearch.length(), PatchAttr);
-				}
-				
-				CRCLen = MetalSource.length();
-				CRC = FCrc::MemCrc_DEPRECATED(MetalSource.c_str(), CRCLen);
-
-				std::string PatchInputName = EntryPointNameAnsi + "_patchIn";
-				std::string PatchNameDefine = "#define ";
-				PatchNameDefine += PatchInputName;
-				PatchNameDefine += " ";
-				PatchNameDefine += PatchInputName;
-				PatchNameDefine += std::to_string(CRC);
-				PatchNameDefine += std::to_string(CRCLen);
-				PatchNameDefine += "\n";
-				
-				std::string PatchInput = EntryPointNameAnsi + "_in";
-				std::string PatchDefine = "#define ";
-				PatchDefine += PatchInput;
-				PatchDefine += " ";
-				PatchDefine += PatchInput;
-				PatchDefine += std::to_string(CRC);
-				PatchDefine += std::to_string(CRCLen);
-				PatchDefine += "\n";
-				
-				size_t Pos = MetalSource.find("#include ");
-				MetalSource.insert(Pos, PatchNameDefine);
-				MetalSource.insert(Pos, PatchDefine);
-			}
-			
+
 			CRCLen = MetalSource.length();
 			CRC = FCrc::MemCrc_DEPRECATED(MetalSource.c_str(), CRCLen);
 			
@@ -2138,7 +1567,7 @@ bool FMetalShaderOutputCooker::Build(TArray<uint8>& OutData)
 	if (Result != 0)
 	{
 		Output.Target = Input.Target;
-		BuildMetalShaderOutput(Output, Input, GUIDHash, CCFlags, MetalSource.c_str(), MetalSource.length(), CRCLen, CRC, VersionEnum, *Standard, *MinOSVersion, TypeMode, Output.Errors, Attribs, OutputData.TypedBuffers, OutputData.InvariantBuffers, OutputData.TypedUAVs, OutputData.ConstantBuffers, OutputData.TypedBufferFormats, bAllowFastIntriniscs);
+		BuildMetalShaderOutput(Output, Input, GUIDHash, CCFlags, MetalSource.c_str(), MetalSource.length(), CRCLen, CRC, VersionEnum, *Standard, *MinOSVersion, TypeMode, Output.Errors, OutputData.TypedBuffers, OutputData.InvariantBuffers, OutputData.TypedUAVs, OutputData.ConstantBuffers, OutputData.TypedBufferFormats, bAllowFastIntriniscs);
 
 		FMemoryWriter Ar(OutData);
 		Ar << Output;
diff --git a/Engine/Source/Developer/Apple/MetalShaderFormat/Private/MetalShaderCompiler.cpp b/Engine/Source/Developer/Apple/MetalShaderFormat/Private/MetalShaderCompiler.cpp
--- a/Engine/Source/Developer/Apple/MetalShaderFormat/Private/MetalShaderCompiler.cpp
+++ b/Engine/Source/Developer/Apple/MetalShaderFormat/Private/MetalShaderCompiler.cpp
@@ -76,51 +76,21 @@ static inline uint32 ParseNumber(const ANSICHAR* Str)
 
 struct FHlslccMetalHeader : public CrossCompiler::FHlslccHeader
 {
-	FHlslccMetalHeader(uint8 const Version, bool const bUsingTessellation);
+	FHlslccMetalHeader(uint8 const Version);
 	virtual ~FHlslccMetalHeader();
 	
 	// After the standard header, different backends can output their own info
 	virtual bool ParseCustomHeaderEntries(const ANSICHAR*& ShaderSource) override;
 	
-	float  TessellationMaxTessFactor;
-	uint32 TessellationOutputControlPoints;
-	uint32 TessellationDomain; // 3 = tri, 4 = quad
-	uint32 TessellationInputControlPoints;
-	uint32 TessellationPatchesPerThreadGroup;
-	uint32 TessellationPatchCountBuffer;
-	uint32 TessellationIndexBuffer;
-	uint32 TessellationHSOutBuffer;
-	uint32 TessellationHSTFOutBuffer;
-	uint32 TessellationControlPointOutBuffer;
-	uint32 TessellationControlPointIndexBuffer;
-	EMetalOutputWindingMode TessellationOutputWinding;
-	EMetalPartitionMode TessellationPartitioning;
 	TMap<uint8, TArray<uint8>> ArgumentBuffers;
 	int8 SideTable;
 	uint8 Version;
-	bool bUsingTessellation;
 };
 
-FHlslccMetalHeader::FHlslccMetalHeader(uint8 const InVersion, bool const bInUsingTessellation)
+FHlslccMetalHeader::FHlslccMetalHeader(uint8 const InVersion)
 {
-	TessellationMaxTessFactor = 0.0f;
-	TessellationOutputControlPoints = 0;
-	TessellationDomain = 0;
-	TessellationInputControlPoints = 0;
-	TessellationPatchesPerThreadGroup = 0;
-	TessellationOutputWinding = EMetalOutputWindingMode::Clockwise;
-	TessellationPartitioning = EMetalPartitionMode::Pow2;
-	
-	TessellationPatchCountBuffer = UINT_MAX;
-	TessellationIndexBuffer = UINT_MAX;
-	TessellationHSOutBuffer = UINT_MAX;
-	TessellationHSTFOutBuffer = UINT_MAX;
-	TessellationControlPointOutBuffer = UINT_MAX;
-	TessellationControlPointIndexBuffer = UINT_MAX;
-	
 	SideTable = -1;
 	Version = InVersion;
-	bUsingTessellation = bInUsingTessellation;
 }
 
 FHlslccMetalHeader::~FHlslccMetalHeader()
@@ -133,19 +103,6 @@ bool FHlslccMetalHeader::ParseCustomHeaderEntries(const ANSICHAR*& ShaderSource)
 #define DEF_PREFIX_STR(Str) \
 static const ANSICHAR* Str##Prefix = "// @" #Str ": "; \
 static const int32 Str##PrefixLen = FCStringAnsi::Strlen(Str##Prefix)
-	DEF_PREFIX_STR(TessellationOutputControlPoints);
-	DEF_PREFIX_STR(TessellationDomain);
-	DEF_PREFIX_STR(TessellationInputControlPoints);
-	DEF_PREFIX_STR(TessellationMaxTessFactor);
-	DEF_PREFIX_STR(TessellationOutputWinding);
-	DEF_PREFIX_STR(TessellationPartitioning);
-	DEF_PREFIX_STR(TessellationPatchesPerThreadGroup);
-	DEF_PREFIX_STR(TessellationPatchCountBuffer);
-	DEF_PREFIX_STR(TessellationIndexBuffer);
-	DEF_PREFIX_STR(TessellationHSOutBuffer);
-	DEF_PREFIX_STR(TessellationHSTFOutBuffer);
-	DEF_PREFIX_STR(TessellationControlPointOutBuffer);
-	DEF_PREFIX_STR(TessellationControlPointIndexBuffer);
 	DEF_PREFIX_STR(ArgumentBuffers);
 	DEF_PREFIX_STR(SideTable);
 #undef DEF_PREFIX_STR
@@ -234,191 +191,6 @@ static const int32 Str##PrefixLen = FCStringAnsi::Strlen(Str##Prefix)
 		}
 	}
 	
-	// Early out for non-tessellation...
-	if (!bUsingTessellation)
-	{
-		return true;
-	}
-	
-	auto ParseUInt32Attribute = [&ShaderSource](const ANSICHAR* prefix, int32 prefixLen, uint32& attributeOut)
-	{
-		if (FCStringAnsi::Strncmp(ShaderSource, prefix, prefixLen) == 0)
-		{
-			ShaderSource += prefixLen;
-			
-			if (!CrossCompiler::ParseIntegerNumber(ShaderSource, attributeOut))
-			{
-				return false;
-			}
-			
-			if (!CrossCompiler::Match(ShaderSource, '\n'))
-			{
-				return false;
-			}
-		}
-		
-		return true;
-	};
- 
-	// Read number of tessellation output control points
-	if (!ParseUInt32Attribute(TessellationOutputControlPointsPrefix, TessellationOutputControlPointsPrefixLen, TessellationOutputControlPoints))
-	{
-		return false;
-	}
-	
-	// Read the tessellation domain (tri vs quad)
-	if (FCStringAnsi::Strncmp(ShaderSource, TessellationDomainPrefix, TessellationDomainPrefixLen) == 0)
-	{
-		ShaderSource += TessellationDomainPrefixLen;
-		
-		if (FCStringAnsi::Strncmp(ShaderSource, "tri", 3) == 0)
-		{
-			ShaderSource += 3;
-			TessellationDomain = 3;
-		}
-		else if (FCStringAnsi::Strncmp(ShaderSource, "quad", 4) == 0)
-		{
-			ShaderSource += 4;
-			TessellationDomain = 4;
-		}
-		else
-		{
-			return false;
-		}
-		
-		if (!CrossCompiler::Match(ShaderSource, '\n'))
-		{
-			return false;
-		}
-	}
-	
-	// Read number of tessellation input control points
-	if (!ParseUInt32Attribute(TessellationInputControlPointsPrefix, TessellationInputControlPointsPrefixLen, TessellationInputControlPoints))
-	{
-		return false;
-	}
-	
-	// Read max tessellation factor
-	if (FCStringAnsi::Strncmp(ShaderSource, TessellationMaxTessFactorPrefix, TessellationMaxTessFactorPrefixLen) == 0)
-	{
-		ShaderSource += TessellationMaxTessFactorPrefixLen;
-		
-#if PLATFORM_WINDOWS
-		if (sscanf_s(ShaderSource, "%g\n", &TessellationMaxTessFactor) != 1)
-#else
-		if (sscanf(ShaderSource, "%g\n", &TessellationMaxTessFactor) != 1)
-#endif
-		{
-			return false;
-		}
-		
-		while (*ShaderSource != '\n')
-		{
-			++ShaderSource;
-		}
-		++ShaderSource; // to match the newline
-	}
-	
-	// Read tessellation output winding mode
-	if (FCStringAnsi::Strncmp(ShaderSource, TessellationOutputWindingPrefix, TessellationOutputWindingPrefixLen) == 0)
-	{
-		ShaderSource += TessellationOutputWindingPrefixLen;
-		
-		if (FCStringAnsi::Strncmp(ShaderSource, "cw", 2) == 0)
-		{
-			ShaderSource += 2;
-			TessellationOutputWinding = EMetalOutputWindingMode::Clockwise;
-		}
-		else if (FCStringAnsi::Strncmp(ShaderSource, "ccw", 3) == 0)
-		{
-			ShaderSource += 3;
-			TessellationOutputWinding = EMetalOutputWindingMode::CounterClockwise;
-		}
-		else
-		{
-			return false;
-		}
-		
-		if (!CrossCompiler::Match(ShaderSource, '\n'))
-		{
-			return false;
-		}
-	}
-	
-	// Read tessellation partition mode
-	if (FCStringAnsi::Strncmp(ShaderSource, TessellationPartitioningPrefix, TessellationPartitioningPrefixLen) == 0)
-	{
-		ShaderSource += TessellationPartitioningPrefixLen;
-		
-		static char const* partitionModeNames[] =
-		{
-			// order match enum order
-			"pow2",
-			"integer",
-			"fractional_odd",
-			"fractional_even",
-		};
-		
-		bool match = false;
-		for (size_t i = 0; i < sizeof(partitionModeNames) / sizeof(partitionModeNames[0]); ++i)
-		{
-			size_t partitionModeNameLen = strlen(partitionModeNames[i]);
-			if (FCStringAnsi::Strncmp(ShaderSource, partitionModeNames[i], partitionModeNameLen) == 0)
-			{
-				ShaderSource += partitionModeNameLen;
-				TessellationPartitioning = (EMetalPartitionMode)i;
-				match = true;
-				break;
-			}
-		}
-		
-		if (!match)
-		{
-			return false;
-		}
-		
-		if (!CrossCompiler::Match(ShaderSource, '\n'))
-		{
-			return false;
-		}
-	}
-	
-	// Read number of tessellation patches per threadgroup
-	if (!ParseUInt32Attribute(TessellationPatchesPerThreadGroupPrefix, TessellationPatchesPerThreadGroupPrefixLen, TessellationPatchesPerThreadGroup))
-	{
-		return false;
-	}
-	
-	if (!ParseUInt32Attribute(TessellationPatchCountBufferPrefix, TessellationPatchCountBufferPrefixLen, TessellationPatchCountBuffer))
-	{
-		TessellationPatchCountBuffer = UINT_MAX;
-	}
-	
-	if (!ParseUInt32Attribute(TessellationIndexBufferPrefix, TessellationIndexBufferPrefixLen, TessellationIndexBuffer))
-	{
-		TessellationIndexBuffer = UINT_MAX;
-	}
-	
-	if (!ParseUInt32Attribute(TessellationHSOutBufferPrefix, TessellationHSOutBufferPrefixLen, TessellationHSOutBuffer))
-	{
-		TessellationHSOutBuffer = UINT_MAX;
-	}
-	
-	if (!ParseUInt32Attribute(TessellationControlPointOutBufferPrefix, TessellationControlPointOutBufferPrefixLen, TessellationControlPointOutBuffer))
-	{
-		TessellationControlPointOutBuffer = UINT_MAX;
-	}
-	
-	if (!ParseUInt32Attribute(TessellationHSTFOutBufferPrefix, TessellationHSTFOutBufferPrefixLen, TessellationHSTFOutBuffer))
-	{
-		TessellationHSTFOutBuffer = UINT_MAX;
-	}
-	
-	if (!ParseUInt32Attribute(TessellationControlPointIndexBufferPrefix, TessellationControlPointIndexBufferPrefixLen, TessellationControlPointIndexBuffer))
-	{
-		TessellationControlPointIndexBuffer = UINT_MAX;
-	}
-	
 	return true;
 }
 
@@ -442,7 +214,6 @@ void BuildMetalShaderOutput(
 	TCHAR const* MinOSVersion,
 	EMetalTypeBufferMode TypeMode,
 	TArray<FShaderCompilerError>& OutErrors,
-	FMetalTessellationOutputs const& TessOutputAttribs,
 	uint32 TypedBuffers,
 	uint32 InvariantBuffers,
 	uint32 TypedUAVs,
@@ -466,9 +237,7 @@ void BuildMetalShaderOutput(
 		Main++;
 	}
 	
-	bool bUsingTessellation = ShaderInput.IsUsingTessellation();
-	
-	FHlslccMetalHeader CCHeader(Version, bUsingTessellation);
+	FHlslccMetalHeader CCHeader(Version);
 	if (!CCHeader.Read(USFSource, SourceLen))
 	{
 		UE_LOG(LogMetalShaderCompiler, Fatal, TEXT("Bad hlslcc header found"));
@@ -733,27 +502,6 @@ void BuildMetalShaderOutput(
 	Header.NumThreadsY = CCHeader.NumThreads[1];
 	Header.NumThreadsZ = CCHeader.NumThreads[2];
 	
-	if (ShaderInput.Target.Platform == SP_METAL_SM5 && (Frequency == SF_Vertex || Frequency == SF_Hull || Frequency == SF_Domain))
-	{
-		FMetalTessellationHeader TessHeader;
-		TessHeader.TessellationOutputControlPoints 		= CCHeader.TessellationOutputControlPoints;
-		TessHeader.TessellationDomain					= CCHeader.TessellationDomain;
-		TessHeader.TessellationInputControlPoints       = CCHeader.TessellationInputControlPoints;
-		TessHeader.TessellationMaxTessFactor            = CCHeader.TessellationMaxTessFactor;
-		TessHeader.TessellationOutputWinding			= CCHeader.TessellationOutputWinding;
-		TessHeader.TessellationPartitioning				= CCHeader.TessellationPartitioning;
-		TessHeader.TessellationPatchesPerThreadGroup    = CCHeader.TessellationPatchesPerThreadGroup;
-		TessHeader.TessellationPatchCountBuffer         = CCHeader.TessellationPatchCountBuffer;
-		TessHeader.TessellationIndexBuffer              = CCHeader.TessellationIndexBuffer;
-		TessHeader.TessellationHSOutBuffer              = CCHeader.TessellationHSOutBuffer;
-		TessHeader.TessellationHSTFOutBuffer            = CCHeader.TessellationHSTFOutBuffer;
-		TessHeader.TessellationControlPointOutBuffer    = CCHeader.TessellationControlPointOutBuffer;
-		TessHeader.TessellationControlPointIndexBuffer  = CCHeader.TessellationControlPointIndexBuffer;
-		TessHeader.TessellationOutputAttribs            = TessOutputAttribs;
-
-		Header.Tessellation.Add(TessHeader);
-		
-	}
 	Header.bDeviceFunctionConstants				= (FCStringAnsi::Strstr(USFSource, "#define __METAL_DEVICE_CONSTANT_INDEX__ 1") != nullptr);
 	Header.SideTable 							= CCHeader.SideTable;
 	Header.Bindings.ArgumentBufferMasks			= CCHeader.ArgumentBuffers;
@@ -1190,15 +938,6 @@ void CompileShader_Metal(const FShaderCompilerInput& _Input,FShaderCompilerOutpu
 		MetalCompilerTarget = HCT_FeatureLevelES3_1;
 		Semantics = EMetalGPUSemanticsImmediateDesktop;
 	}
-	else if (Input.ShaderFormat == NAME_SF_METAL_SM5_NOTESS)
-	{
-		UE_CLOG(VersionEnum < 3, LogShaders, Warning, TEXT("Metal shader version must be Metal v2.0 or higher for format %s!"), VersionEnum, *Input.ShaderFormat.ToString());
-		AdditionalDefines.SetDefine(TEXT("METAL_SM5_NOTESS_PROFILE"), 1);
-		AdditionalDefines.SetDefine(TEXT("USING_VERTEX_SHADER_LAYER"), 1);
-		VersionEnum = VersionEnum > 3 ? VersionEnum : 3;
-		MetalCompilerTarget = HCT_FeatureLevelSM5;
-		Semantics = EMetalGPUSemanticsImmediateDesktop;
-	}
 	else if (Input.ShaderFormat == NAME_SF_METAL_SM5)
 	{
 		UE_CLOG(VersionEnum < 3, LogShaders, Warning, TEXT("Metal shader version must be Metal v2.0 or higher for format %s!"), VersionEnum, *Input.ShaderFormat.ToString());
@@ -1381,34 +1120,6 @@ void CompileShader_Metal(const FShaderCompilerInput& _Input,FShaderCompilerOutpu
 		AdditionalDefines.SetDefine(TEXT("COMPILER_SUPPORTS_ATTRIBUTES"), (uint32)1);
 	}
 
-	if (!Input.bSkipPreprocessedCache && !bDirectCompile)
-	{
-		bool bUsingTessellation = Input.IsUsingTessellation();
-		if (bUsingTessellation && (Input.Target.Frequency == SF_Vertex))
-		{
-			// force HULLSHADER on so that VS that is USING_TESSELLATION can be built together with the proper HS
-			FString const* VertexShaderDefine = Input.Environment.GetDefinitions().Find(TEXT("VERTEXSHADER"));
-			check(VertexShaderDefine && FString("1") == *VertexShaderDefine);
-			FString const* HullShaderDefine = Input.Environment.GetDefinitions().Find(TEXT("HULLSHADER"));
-			check(HullShaderDefine && FString("0") == *HullShaderDefine);
-			Input.Environment.SetDefine(TEXT("HULLSHADER"), 1u);
-		}
-		if (Input.Target.Frequency == SF_Hull)
-		{
-			check(bUsingTessellation);
-			// force VERTEXSHADER on so that HS that is USING_TESSELLATION can be built together with the proper VS
-			FString const* VertexShaderDefine = Input.Environment.GetDefinitions().Find(TEXT("VERTEXSHADER"));
-			check(VertexShaderDefine && FString("0") == *VertexShaderDefine);
-			FString const* HullShaderDefine = Input.Environment.GetDefinitions().Find(TEXT("HULLSHADER"));
-			check(HullShaderDefine && FString("1") == *HullShaderDefine);
-
-			// enable VERTEXSHADER so that this HS will hash uniquely with its associated VS
-			// We do not want a given HS to be shared among numerous VS'Sampler
-			// this should accomplish that goal -- see GenerateOutputHash
-			Input.Environment.SetDefine(TEXT("VERTEXSHADER"), 1u);
-		}
-	}
-
 	if (Input.bSkipPreprocessedCache)
 	{
 		if (!FFileHelper::LoadFileToString(PreprocessedShader, *Input.VirtualSourceFilePath))
diff --git a/Engine/Source/Developer/Apple/MetalShaderFormat/Private/MetalShaderFormat.cpp b/Engine/Source/Developer/Apple/MetalShaderFormat/Private/MetalShaderFormat.cpp
--- a/Engine/Source/Developer/Apple/MetalShaderFormat/Private/MetalShaderFormat.cpp
+++ b/Engine/Source/Developer/Apple/MetalShaderFormat/Private/MetalShaderFormat.cpp
@@ -62,14 +62,13 @@ public:
 		OutFormats.Add(NAME_SF_METAL_MRT);
 		OutFormats.Add(NAME_SF_METAL_TVOS);
 		OutFormats.Add(NAME_SF_METAL_MRT_TVOS);
-		OutFormats.Add(NAME_SF_METAL_SM5_NOTESS);
 		OutFormats.Add(NAME_SF_METAL_SM5);
 		OutFormats.Add(NAME_SF_METAL_MACES3_1);
 		OutFormats.Add(NAME_SF_METAL_MRT_MAC);
 	}
 	virtual void CompileShader(FName Format, const struct FShaderCompilerInput& Input, struct FShaderCompilerOutput& Output,const FString& WorkingDirectory) const override final
 	{
-		check(Format == NAME_SF_METAL || Format == NAME_SF_METAL_MRT || Format == NAME_SF_METAL_TVOS || Format == NAME_SF_METAL_MRT_TVOS || Format == NAME_SF_METAL_SM5_NOTESS || Format == NAME_SF_METAL_SM5 || Format == NAME_SF_METAL_MACES3_1 || Format == NAME_SF_METAL_MRT_MAC);
+		check(Format == NAME_SF_METAL || Format == NAME_SF_METAL_MRT || Format == NAME_SF_METAL_TVOS || Format == NAME_SF_METAL_MRT_TVOS || Format == NAME_SF_METAL_SM5 || Format == NAME_SF_METAL_MACES3_1 || Format == NAME_SF_METAL_MRT_MAC);
 		CompileShader_Metal(Input, Output, WorkingDirectory);
 	}
 	virtual bool CanStripShaderCode(bool const bNativeFormat) const override final
@@ -95,7 +94,7 @@ public:
     {
 		const int32 NumShadersPerLibrary = 10000;
 		check(LibraryName.Len() > 0);
-		check(Format == NAME_SF_METAL || Format == NAME_SF_METAL_MRT || Format == NAME_SF_METAL_TVOS || Format == NAME_SF_METAL_MRT_TVOS || Format == NAME_SF_METAL_SM5_NOTESS || Format == NAME_SF_METAL_SM5 || Format == NAME_SF_METAL_MACES3_1 || Format == NAME_SF_METAL_MRT_MAC);
+		check(Format == NAME_SF_METAL || Format == NAME_SF_METAL_MRT || Format == NAME_SF_METAL_TVOS || Format == NAME_SF_METAL_MRT_TVOS || Format == NAME_SF_METAL_SM5 || Format == NAME_SF_METAL_MACES3_1 || Format == NAME_SF_METAL_MRT_MAC);
 
 		const FString ArchivePath = (WorkingDirectory / Format.GetPlainNameString());
 		IFileManager::Get().DeleteDirectory(*ArchivePath, false, true);
@@ -494,7 +493,6 @@ EShaderPlatform FMetalCompilerToolchain::MetalShaderFormatToLegacyShaderPlatform
 	if (ShaderFormat == NAME_SF_METAL_MRT_TVOS)		return SP_METAL_MRT_TVOS;
 	if (ShaderFormat == NAME_SF_METAL_MRT_MAC)		return SP_METAL_MRT_MAC;
 	if (ShaderFormat == NAME_SF_METAL_SM5)			return SP_METAL_SM5;
-	if (ShaderFormat == NAME_SF_METAL_SM5_NOTESS)	return SP_METAL_SM5_NOTESS;
 	if (ShaderFormat == NAME_SF_METAL_MACES3_1)		return SP_METAL_MACES3_1;
 
 	return SP_NumPlatforms;
diff --git a/Engine/Source/Developer/Apple/MetalShaderFormat/Private/MetalShaderFormat.h b/Engine/Source/Developer/Apple/MetalShaderFormat/Private/MetalShaderFormat.h
--- a/Engine/Source/Developer/Apple/MetalShaderFormat/Private/MetalShaderFormat.h
+++ b/Engine/Source/Developer/Apple/MetalShaderFormat/Private/MetalShaderFormat.h
@@ -24,7 +24,6 @@ static FName NAME_SF_METAL(TEXT("SF_METAL"));
 static FName NAME_SF_METAL_MRT(TEXT("SF_METAL_MRT"));
 static FName NAME_SF_METAL_TVOS(TEXT("SF_METAL_TVOS"));
 static FName NAME_SF_METAL_MRT_TVOS(TEXT("SF_METAL_MRT_TVOS"));
-static FName NAME_SF_METAL_SM5_NOTESS(TEXT("SF_METAL_SM5_NOTESS"));
 static FName NAME_SF_METAL_SM5(TEXT("SF_METAL_SM5"));
 static FName NAME_SF_METAL_MACES3_1(TEXT("SF_METAL_MACES3_1"));
 static FName NAME_SF_METAL_MRT_MAC(TEXT("SF_METAL_MRT_MAC"));
diff --git a/Engine/Source/Developer/Mac/MacTargetPlatform/Private/GenericMacTargetPlatform.h b/Engine/Source/Developer/Mac/MacTargetPlatform/Private/GenericMacTargetPlatform.h
--- a/Engine/Source/Developer/Mac/MacTargetPlatform/Private/GenericMacTargetPlatform.h
+++ b/Engine/Source/Developer/Mac/MacTargetPlatform/Private/GenericMacTargetPlatform.h
@@ -113,8 +113,6 @@ return TSuper::SupportsFeature(Feature);
 		// no shaders needed for dedicated server target
 		if (!IS_DEDICATED_SERVER)
 		{
-			static FName NAME_SF_METAL_SM5_NOTESS(TEXT("SF_METAL_SM5_NOTESS"));
-			OutFormats.AddUnique(NAME_SF_METAL_SM5_NOTESS);
 			static FName NAME_SF_METAL_SM5(TEXT("SF_METAL_SM5"));
 			OutFormats.AddUnique(NAME_SF_METAL_SM5);
 			static FName NAME_SF_METAL_MACES3_1(TEXT("SF_METAL_MACES3_1"));
diff --git a/Engine/Source/Developer/ShaderCompilerCommon/Private/HlslccHeaderWriter.cpp b/Engine/Source/Developer/ShaderCompilerCommon/Private/HlslccHeaderWriter.cpp
--- a/Engine/Source/Developer/ShaderCompilerCommon/Private/HlslccHeaderWriter.cpp
+++ b/Engine/Source/Developer/ShaderCompilerCommon/Private/HlslccHeaderWriter.cpp
@@ -448,108 +448,6 @@ namespace CrossCompiler
 		Strings.ArgumentBuffers += TEXT("]");
 	}
 
-	void FHlslccHeaderWriter::WriteTessellationInputControlPoints(uint32 PatchSize)
-	{
-		checkf(Strings.TessellationInputControlPoints.IsEmpty(), TEXT("cannot define @TessellationInputControlPoints attribute more than once"));
-		Strings.TessellationInputControlPoints = FString::Printf(TEXT("%u"), PatchSize);
-	}
-
-	void FHlslccHeaderWriter::WriteTessellationOutputControlPoints(uint32 PatchSize)
-	{
-		checkf(Strings.TessellationOutputControlPoints.IsEmpty(), TEXT("cannot define @TessellationOutputControlPoints attribute more than once"));
-		Strings.TessellationOutputControlPoints = FString::Printf(TEXT("%u"), PatchSize);
-	}
-
-	void FHlslccHeaderWriter::WriteTessellationMaxTessFactor(uint32 MaxTessFactor)
-	{
-		checkf(Strings.TessellationMaxTessFactor.IsEmpty(), TEXT("cannot define @TessellationMaxTessFactor attribute more than once"));
-		Strings.TessellationMaxTessFactor = FString::Printf(TEXT("%u"), MaxTessFactor);
-	}
-
-	void FHlslccHeaderWriter::WriteTessellationDomainTri()
-	{
-		checkf(Strings.TessellationDomain.IsEmpty(), TEXT("cannot define @TessellationDomain attribute more than once"));
-		Strings.TessellationDomain = TEXT("tri");
-	}
-
-	void FHlslccHeaderWriter::WriteTessellationDomainQuad()
-	{
-		checkf(Strings.TessellationDomain.IsEmpty(), TEXT("cannot define @TessellationDomain attribute more than once"));
-		Strings.TessellationDomain = TEXT("quad");
-	}
-
-	void FHlslccHeaderWriter::WriteTessellationOutputWindingCW()
-	{
-		checkf(Strings.TessellationOutputWinding.IsEmpty(), TEXT("cannot define @TessellationOutputWinding attribute more than once"));
-		Strings.TessellationOutputWinding = TEXT("cw");
-	}
-
-	void FHlslccHeaderWriter::WriteTessellationOutputWindingCCW()
-	{
-		checkf(Strings.TessellationOutputWinding.IsEmpty(), TEXT("cannot define @TessellationOutputWinding attribute more than once"));
-		Strings.TessellationOutputWinding = TEXT("ccw");
-	}
-
-	void FHlslccHeaderWriter::WriteTessellationPartitioningInteger()
-	{
-		checkf(Strings.TessellationPartitioning.IsEmpty(), TEXT("cannot define @TessellationPartitioning attribute more than once"));
-		Strings.TessellationPartitioning = TEXT("integer");
-	}
-
-	void FHlslccHeaderWriter::WriteTessellationPartitioningFractionalOdd()
-	{
-		checkf(Strings.TessellationPartitioning.IsEmpty(), TEXT("cannot define @TessellationPartitioning attribute more than once"));
-		Strings.TessellationPartitioning = TEXT("fractional_odd");
-	}
-
-	void FHlslccHeaderWriter::WriteTessellationPartitioningFractionalEven()
-	{
-		checkf(Strings.TessellationPartitioning.IsEmpty(), TEXT("cannot define @TessellationPartitioning attribute more than once"));
-		Strings.TessellationPartitioning = TEXT("fractional_even");
-	}
-
-	void FHlslccHeaderWriter::WriteTessellationPatchesPerThreadGroup(uint32 NumPatches)
-	{
-		checkf(Strings.TessellationPatchesPerThreadGroup.IsEmpty(), TEXT("cannot define @TessellationPatchesPerThreadGroup attribute more than once"));
-		Strings.TessellationPatchesPerThreadGroup = FString::Printf(TEXT("%u"), NumPatches);
-	}
-
-	void FHlslccHeaderWriter::WriteTessellationPatchCountBuffer(uint32 BindingIndex)
-	{
-		checkf(Strings.TessellationPatchCountBuffer.IsEmpty(), TEXT("cannot define @TessellationPatchCountBuffer attribute more than once"));
-		Strings.TessellationPatchCountBuffer = FString::Printf(TEXT("%u"), BindingIndex);
-	}
-
-	void FHlslccHeaderWriter::WriteTessellationIndexBuffer(uint32 BindingIndex)
-	{
-		checkf(Strings.TessellationIndexBuffer.IsEmpty(), TEXT("cannot define @TessellationIndexBuffer attribute more than once"));
-		Strings.TessellationIndexBuffer = FString::Printf(TEXT("%u"), BindingIndex);
-	}
-
-	void FHlslccHeaderWriter::WriteTessellationHSOutBuffer(uint32 BindingIndex)
-	{
-		checkf(Strings.TessellationHSOutBuffer.IsEmpty(), TEXT("cannot define @TessellationHSOutBuffer attribute more than once"));
-		Strings.TessellationHSOutBuffer = FString::Printf(TEXT("%u"), BindingIndex);
-	}
-
-	void FHlslccHeaderWriter::WriteTessellationHSTFOutBuffer(uint32 BindingIndex)
-	{
-		checkf(Strings.TessellationHSTFOutBuffer.IsEmpty(), TEXT("cannot define @TessellationHSTFOutBuffer attribute more than once"));
-		Strings.TessellationHSTFOutBuffer = FString::Printf(TEXT("%u"), BindingIndex);
-	}
-
-	void FHlslccHeaderWriter::WriteTessellationControlPointOutBuffer(uint32 BindingIndex)
-	{
-		checkf(Strings.TessellationControlPointOutBuffer.IsEmpty(), TEXT("cannot define @TessellationControlPointOutBuffer attribute more than once"));
-		Strings.TessellationControlPointOutBuffer = FString::Printf(TEXT("%u"), BindingIndex);
-	}
-
-	void FHlslccHeaderWriter::WriteTessellationControlPointIndexBuffer(uint32 BindingIndex)
-	{
-		checkf(Strings.TessellationControlPointIndexBuffer.IsEmpty(), TEXT("cannot define @TessellationControlPointIndexBuffer attribute more than once"));
-		Strings.TessellationControlPointIndexBuffer = FString::Printf(TEXT("%u"), BindingIndex);
-	}
-
 	FString FHlslccHeaderWriter::ToString() const
 	{
 		FString MetaData;
@@ -574,19 +472,6 @@ namespace CrossCompiler
 		PrintAttributes(TEXT("NumThreads"), Strings.NumThreads);
 		PrintAttributes(TEXT("SideTable"), Strings.SideTable);
 		PrintAttributes(TEXT("ArgumentBuffers"), Strings.ArgumentBuffers);
-		PrintAttributes(TEXT("TessellationOutputControlPoints"), Strings.TessellationOutputControlPoints);
-		PrintAttributes(TEXT("TessellationDomain"), Strings.TessellationDomain);
-		PrintAttributes(TEXT("TessellationInputControlPoints"), Strings.TessellationInputControlPoints);
-		PrintAttributes(TEXT("TessellationMaxTessFactor"), Strings.TessellationMaxTessFactor);
-		PrintAttributes(TEXT("TessellationOutputWinding"), Strings.TessellationOutputWinding);
-		PrintAttributes(TEXT("TessellationPartitioning"), Strings.TessellationPartitioning);
-		PrintAttributes(TEXT("TessellationPatchesPerThreadGroup"), Strings.TessellationPatchesPerThreadGroup);
-		PrintAttributes(TEXT("TessellationPatchCountBuffer"), Strings.TessellationPatchCountBuffer);
-		PrintAttributes(TEXT("TessellationIndexBuffer"), Strings.TessellationIndexBuffer);
-		PrintAttributes(TEXT("TessellationHSOutBuffer"), Strings.TessellationHSOutBuffer);
-		PrintAttributes(TEXT("TessellationControlPointOutBuffer"), Strings.TessellationControlPointOutBuffer);
-		PrintAttributes(TEXT("TessellationHSTFOutBuffer"), Strings.TessellationHSTFOutBuffer);
-		PrintAttributes(TEXT("TessellationControlPointIndexBuffer"), Strings.TessellationControlPointIndexBuffer);
 
 		return MetaData;
 	}
diff --git a/Engine/Source/Developer/ShaderCompilerCommon/Public/HlslccHeaderWriter.h b/Engine/Source/Developer/ShaderCompilerCommon/Public/HlslccHeaderWriter.h
--- a/Engine/Source/Developer/ShaderCompilerCommon/Public/HlslccHeaderWriter.h
+++ b/Engine/Source/Developer/ShaderCompilerCommon/Public/HlslccHeaderWriter.h
@@ -33,28 +33,6 @@ namespace CrossCompiler
 		void WriteSideTable(const TCHAR* ResourceName, uint32 SideTableIndex);
 		void WriteArgumentBuffers(uint32 BindingIndex, const TArray<uint32>& ResourceIndices);
 
-		void WriteTessellationInputControlPoints(uint32 PatchSize);
-		void WriteTessellationOutputControlPoints(uint32 PatchSize);
-		void WriteTessellationMaxTessFactor(uint32 MaxTessFactor);
-
-		void WriteTessellationDomainTri();
-		void WriteTessellationDomainQuad();
-
-		void WriteTessellationOutputWindingCW();
-		void WriteTessellationOutputWindingCCW();
-
-		void WriteTessellationPartitioningInteger();
-		void WriteTessellationPartitioningFractionalOdd();
-		void WriteTessellationPartitioningFractionalEven();
-
-		void WriteTessellationPatchesPerThreadGroup(uint32 NumPatches);
-		void WriteTessellationPatchCountBuffer(uint32 BindingIndex);
-		void WriteTessellationIndexBuffer(uint32 BindingIndex);
-		void WriteTessellationHSOutBuffer(uint32 BindingIndex);
-		void WriteTessellationHSTFOutBuffer(uint32 BindingIndex);
-		void WriteTessellationControlPointOutBuffer(uint32 BindingIndex);
-		void WriteTessellationControlPointIndexBuffer(uint32 BindingIndex);
-
 		/** Returns the finalized meta data. */
 		FString ToString() const;
 
@@ -79,19 +57,6 @@ namespace CrossCompiler
 			FString ExternalTextures; // External texture resources (Vulkan ES3.1 profile only)
 			FString SideTable; // Side table for additional indices, e.,g. "spvBufferSizeConstants(31)" (Metal only)
 			FString ArgumentBuffers; // Indirect argument buffers (Metal only)
-			FString TessellationOutputControlPoints;
-			FString TessellationDomain;
-			FString TessellationInputControlPoints;
-			FString TessellationMaxTessFactor;
-			FString TessellationOutputWinding;
-			FString TessellationPartitioning;
-			FString TessellationPatchesPerThreadGroup;
-			FString TessellationPatchCountBuffer;
-			FString TessellationIndexBuffer;
-			FString TessellationHSOutBuffer;
-			FString TessellationControlPointOutBuffer;
-			FString TessellationHSTFOutBuffer;
-			FString TessellationControlPointIndexBuffer;
 		};
 
 	private:
diff --git a/Engine/Source/Developer/SharedSettingsWidgets/Private/ShaderFormatsPropertyDetails.cpp b/Engine/Source/Developer/SharedSettingsWidgets/Private/ShaderFormatsPropertyDetails.cpp
--- a/Engine/Source/Developer/SharedSettingsWidgets/Private/ShaderFormatsPropertyDetails.cpp
+++ b/Engine/Source/Developer/SharedSettingsWidgets/Private/ShaderFormatsPropertyDetails.cpp
@@ -53,11 +53,8 @@ FText FShaderFormatsPropertyDetails::GetFriendlyNameFromRHINameMac(const FString
 		case SP_METAL_MRT_TVOS:
 			FriendlyRHIName = LOCTEXT("MetalMRTTV", "tvOS Metal Desktop Renderer (SM5, Metal 1.2+, tvOS 10.0 or later)");
 			break;
-		case SP_METAL_SM5_NOTESS:
-			FriendlyRHIName = LOCTEXT("MetalSM5_NOTESS", "Mac Metal Desktop Renderer without Tessellation (SM5, Metal 2.0+, macOS High Sierra 10.13.6 or later)");
-			break;
 		case SP_METAL_SM5:
-			FriendlyRHIName = LOCTEXT("MetalSM5", "Mac Metal Desktop Renderer with Tessellation (SM5, Metal 2.0+, macOS High Sierra 10.13.6 or later)");
+			FriendlyRHIName = LOCTEXT("MetalSM5", "Mac Metal Desktop Renderer (SM5, Metal 2.1+, macOS Catalina 10.15.7 or later)");
 			break;
 		case SP_METAL_MACES3_1:
 			FriendlyRHIName = LOCTEXT("MetalES3.1", "Mac Metal High-End Mobile Preview (ES3.1)");
diff --git a/Engine/Source/Editor/MaterialEditor/Private/MaterialStatsCommon.cpp b/Engine/Source/Editor/MaterialEditor/Private/MaterialStatsCommon.cpp
--- a/Engine/Source/Editor/MaterialEditor/Private/MaterialStatsCommon.cpp
+++ b/Engine/Source/Editor/MaterialEditor/Private/MaterialStatsCommon.cpp
@@ -148,8 +148,6 @@ FString FMaterialStatsUtils::ShaderPlatformTypeName(const EShaderPlatform Platfo
 			return FString("METAL_SM5");
 		case SP_VULKAN_PCES3_1:
 			return FString("VULKAN_PCES3_1");
-		case SP_METAL_SM5_NOTESS:
-			return FString("METAL_SM5_NOTESS");
 		case SP_VULKAN_SM5:
 			return FString("VULKAN_SM5");
 		case SP_VULKAN_ES3_1_ANDROID:
@@ -223,7 +221,6 @@ bool FMaterialStatsUtils::PlatformNeedsOfflineCompiler(const EShaderPlatform Sha
 		case SP_METAL_MRT_TVOS:
 		case SP_PCD3D_ES3_1:
 		case SP_METAL_SM5:
-		case SP_METAL_SM5_NOTESS:
 		case SP_METAL_MACES3_1:
 		case SP_SWITCH:
 		case SP_SWITCH_FORWARD:
diff --git a/Engine/Source/Runtime/Apple/MetalRHI/Private/MetalCommandQueue.cpp b/Engine/Source/Runtime/Apple/MetalRHI/Private/MetalCommandQueue.cpp
--- a/Engine/Source/Runtime/Apple/MetalRHI/Private/MetalCommandQueue.cpp
+++ b/Engine/Source/Runtime/Apple/MetalRHI/Private/MetalCommandQueue.cpp
@@ -68,7 +68,7 @@ FMetalCommandQueue::FMetalCommandQueue(mtlpp::Device InDevice, uint32 const MaxN
 		
 		if(Device.SupportsFeatureSet(mtlpp::FeatureSet::tvOS_GPUFamily2_v1))
 		{
-			Features |= EMetalFeaturesCountingQueries | EMetalFeaturesBaseVertexInstance | EMetalFeaturesIndirectBuffer | EMetalFeaturesMSAADepthResolve | EMetalFeaturesTessellation | EMetalFeaturesMSAAStoreAndResolve;
+			Features |= EMetalFeaturesCountingQueries | EMetalFeaturesBaseVertexInstance | EMetalFeaturesIndirectBuffer | EMetalFeaturesMSAADepthResolve | EMetalFeaturesMSAAStoreAndResolve;
 		}
 		
 		if(Vers.majorVersion > 10)
@@ -126,7 +126,7 @@ FMetalCommandQueue::FMetalCommandQueue(mtlpp::Device InDevice, uint32 const MaxN
 		
 		if(Device.SupportsFeatureSet(mtlpp::FeatureSet::iOS_GPUFamily3_v2))
 		{
-			Features |= EMetalFeaturesTessellation | EMetalFeaturesMSAAStoreAndResolve;
+			Features |= EMetalFeaturesMSAAStoreAndResolve;
 		}
 		
 		if(Vers.majorVersion > 10 || (Vers.majorVersion == 10 && Vers.minorVersion >= 3))
@@ -197,7 +197,7 @@ FMetalCommandQueue::FMetalCommandQueue(mtlpp::Device InDevice, uint32 const MaxN
 	}
 	if (Device.SupportsFeatureSet(mtlpp::FeatureSet::macOS_GPUFamily1_v2))
     {
-        Features |= EMetalFeaturesTessellation | EMetalFeaturesMSAADepthResolve | EMetalFeaturesMSAAStoreAndResolve;
+        Features |= EMetalFeaturesMSAADepthResolve | EMetalFeaturesMSAAStoreAndResolve;
         
         // Assume that set*Bytes only works on macOS Sierra and above as no-one has tested it anywhere else.
 		Features |= EMetalFeaturesSetBytes;
@@ -278,7 +278,7 @@ FMetalCommandQueue::FMetalCommandQueue(mtlpp::Device InDevice, uint32 const MaxN
     
     if(Device.SupportsFeatureSet(mtlpp::FeatureSet::macOS_GPUFamily1_v3) && FPlatformMisc::MacOSXVersionCompare(10,13,0) >= 0)
     {
-        Features |= EMetalFeaturesMultipleViewports | EMetalFeaturesPipelineBufferMutability | EMetalFeaturesGPUCaptureManager | EMetalFeaturesSeparateTessellation;
+        Features |= EMetalFeaturesMultipleViewports | EMetalFeaturesPipelineBufferMutability | EMetalFeaturesGPUCaptureManager;
 		
 		if (FParse::Param(FCommandLine::Get(),TEXT("metalfence")))
 		{
diff --git a/Engine/Source/Runtime/Apple/MetalRHI/Private/MetalCommandQueue.h b/Engine/Source/Runtime/Apple/MetalRHI/Private/MetalCommandQueue.h
--- a/Engine/Source/Runtime/Apple/MetalRHI/Private/MetalCommandQueue.h
+++ b/Engine/Source/Runtime/Apple/MetalRHI/Private/MetalCommandQueue.h
@@ -29,8 +29,8 @@ typedef NS_OPTIONS(uint64, EMetalFeatures)
 	EMetalFeaturesLayeredRendering = 1 << 4,
 	/** Support for specifying small buffers as byte arrays */
 	EMetalFeaturesSetBytes = 1 << 5,
-	/** Supports tessellation rendering */
-	EMetalFeaturesTessellation = 1 << 6,
+	/** Unused Reserved Bit */
+	EMetalFeaturesUnusedReservedBit6 = 1 << 6, // was EMetalFeaturesTessellation
 	/** Supports framework-level validation */
 	EMetalFeaturesValidation = 1 << 7,
 	/** Supports detailed statistics */
@@ -71,8 +71,8 @@ typedef NS_OPTIONS(uint64, EMetalFeatures)
     EMetalFeaturesPipelineBufferMutability = 1llu << 25llu,
     /** Supports tile shaders */
     EMetalFeaturesTileShaders = 1llu << 26llu,
-	/** Supports separate tessellation shader execution */
-	EMetalFeaturesSeparateTessellation = 1llu << 27llu,
+	/** Unused Reserved Bit */
+	EMetalFeaturesUnusedReservedBit27 = 1llu << 27llu, // was EMetalFeaturesSeparateTessellation
 	/** Supports indirect argument buffers Tier 2 */
 	EMetalFeaturesTier2IABs = 1llu << 28llu,
 };
diff --git a/Engine/Source/Runtime/Apple/MetalRHI/Private/MetalCommands.cpp b/Engine/Source/Runtime/Apple/MetalRHI/Private/MetalCommands.cpp
--- a/Engine/Source/Runtime/Apple/MetalRHI/Private/MetalCommands.cpp
+++ b/Engine/Source/Runtime/Apple/MetalRHI/Private/MetalCommands.cpp
@@ -96,10 +96,6 @@ static FORCEINLINE EMetalShaderStages GetShaderStage(FRHIGraphicsShader* ShaderR
 	switch (ShaderRHI->GetFrequency())
 	{
 	case SF_Vertex:		Stage = EMetalShaderStages::Vertex; break;
-#if PLATFORM_SUPPORTS_TESSELLATION_SHADERS
-	case SF_Hull:		Stage = EMetalShaderStages::Hull; break;
-	case SF_Domain:		Stage = EMetalShaderStages::Domain; break;
-#endif
 	case SF_Pixel:		Stage = EMetalShaderStages::Pixel; break;
 	default:
 		checkf(0, TEXT("FRHIShader Type %d is invalid or unsupported!"), (int32)ShaderRHI->GetFrequency());
@@ -268,10 +264,6 @@ void FMetalRHICommandContext::RHISetGraphicsPipelineState(FRHIGraphicsPipelineSt
 		if (bApplyAdditionalState)
 		{
 			ApplyStaticUniformBuffers(PipelineState->VertexShader.GetReference());
-#if PLATFORM_SUPPORTS_TESSELLATION_SHADERS
-			ApplyStaticUniformBuffers(PipelineState->HullShader.GetReference());
-			ApplyStaticUniformBuffers(PipelineState->DomainShader.GetReference());
-#endif
 			ApplyStaticUniformBuffers(PipelineState->PixelShader.GetReference());
 		}
 	}
diff --git a/Engine/Source/Runtime/Apple/MetalRHI/Private/MetalContext.cpp b/Engine/Source/Runtime/Apple/MetalRHI/Private/MetalContext.cpp
--- a/Engine/Source/Runtime/Apple/MetalRHI/Private/MetalContext.cpp
+++ b/Engine/Source/Runtime/Apple/MetalRHI/Private/MetalContext.cpp
@@ -1260,7 +1260,7 @@ void FMetalContext::ResetRenderCommandEncoder()
 	SetRenderPassInfo(StateCache.GetRenderPassInfo(), true);
 }
 
-bool FMetalContext::PrepareToDraw(uint32 PrimitiveType, EMetalIndexType IndexType)
+bool FMetalContext::PrepareToDraw(uint32 PrimitiveType)
 {
 	SCOPE_CYCLE_COUNTER(STAT_MetalPrepareDrawTime);
 	TRefCountPtr<FMetalGraphicsPipelineState> CurrentPSO = StateCache.GetGraphicsPSO();
@@ -1407,9 +1407,6 @@ bool FMetalContext::PrepareToDraw(uint32 PrimitiveType, EMetalIndexType IndexTyp
 		check(StateCache.GetHasValidRenderTarget());
 	}
 	
-	// make sure the BSS has a valid pipeline state object
-	StateCache.SetIndexType(IndexType);
-	
 	return true;
 }
 
@@ -1547,7 +1544,7 @@ void FMetalContext::DrawPrimitiveIndirect(uint32 PrimitiveType, FMetalVertexBuff
 void FMetalContext::DrawIndexedPrimitive(FMetalBuffer const& IndexBuffer, uint32 IndexStride, mtlpp::IndexType IndexType, uint32 PrimitiveType, int32 BaseVertexIndex, uint32 FirstInstance, uint32 NumVertices, uint32 StartIndex, uint32 NumPrimitives, uint32 NumInstances)
 {
 	// finalize any pending state
-	if(!PrepareToDraw(PrimitiveType, GetRHIMetalIndexType(IndexType)))
+	if(!PrepareToDraw(PrimitiveType))
 	{
 		return;
 	}
diff --git a/Engine/Source/Runtime/Apple/MetalRHI/Private/MetalContext.h b/Engine/Source/Runtime/Apple/MetalRHI/Private/MetalContext.h
--- a/Engine/Source/Runtime/Apple/MetalRHI/Private/MetalContext.h
+++ b/Engine/Source/Runtime/Apple/MetalRHI/Private/MetalContext.h
@@ -46,10 +46,9 @@ public:
 	/**
 	 * Do anything necessary to prepare for any kind of draw call 
 	 * @param PrimitiveType The UE4 primitive type for the draw call, needed to compile the correct render pipeline.
-	 * @param IndexType The index buffer type (none, uint16, uint32), needed to compile the correct tessellation compute pipeline.
 	 * @returns True if the preparation completed and the draw call can be encoded, false to skip.
 	 */
-	bool PrepareToDraw(uint32 PrimitiveType, EMetalIndexType IndexType = EMetalIndexType_None);
+	bool PrepareToDraw(uint32 PrimitiveType);
 	
 	/**
 	 * Set the color, depth and stencil render targets, and then make the new command buffer/encoder
diff --git a/Engine/Source/Runtime/Apple/MetalRHI/Private/MetalDynamicRHI.cpp b/Engine/Source/Runtime/Apple/MetalRHI/Private/MetalDynamicRHI.cpp
--- a/Engine/Source/Runtime/Apple/MetalRHI/Private/MetalDynamicRHI.cpp
+++ b/Engine/Source/Runtime/Apple/MetalRHI/Private/MetalDynamicRHI.cpp
@@ -59,10 +59,6 @@ FGraphicsPipelineStateRHIRef FMetalDynamicRHI::RHICreateGraphicsPipelineState(co
 		State->VertexDeclaration = ResourceCast(Initializer.BoundShaderState.VertexDeclarationRHI);
 		State->VertexShader = ResourceCast(Initializer.BoundShaderState.VertexShaderRHI);
 		State->PixelShader = ResourceCast(Initializer.BoundShaderState.PixelShaderRHI);
-#if PLATFORM_SUPPORTS_TESSELLATION_SHADERS
-		State->HullShader = ResourceCast(Initializer.BoundShaderState.HullShaderRHI);
-		State->DomainShader = ResourceCast(Initializer.BoundShaderState.DomainShaderRHI);
-#endif // PLATFORM_SUPPORTS_TESSELLATION_SHADERS
 #if PLATFORM_SUPPORTS_GEOMETRY_SHADERS
 		State->GeometryShader = ResourceCast(Initializer.BoundShaderState.GeometryShaderRHI);
 #endif // PLATFORM_SUPPORTS_GEOMETRY_SHADERS
diff --git a/Engine/Source/Runtime/Apple/MetalRHI/Private/MetalDynamicRHI_Shaders.cpp b/Engine/Source/Runtime/Apple/MetalRHI/Private/MetalDynamicRHI_Shaders.cpp
--- a/Engine/Source/Runtime/Apple/MetalRHI/Private/MetalDynamicRHI_Shaders.cpp
+++ b/Engine/Source/Runtime/Apple/MetalRHI/Private/MetalDynamicRHI_Shaders.cpp
@@ -33,30 +33,12 @@ FPixelShaderRHIRef FMetalDynamicRHI::RHICreatePixelShader(TArrayView<const uint8
 
 FHullShaderRHIRef FMetalDynamicRHI::RHICreateHullShader(TArrayView<const uint8> Code, const FSHAHash& Hash)
 {
-	@autoreleasepool {
-#if PLATFORM_SUPPORTS_TESSELLATION_SHADERS
-		FMetalHullShader* Shader = new FMetalHullShader(Code);
-#else
-		FMetalHullShader* Shader = new FMetalHullShader;
-		FMetalCodeHeader Header;
-		Shader->Init(Code, Header);
-#endif // PLATFORM_SUPPORTS_TESSELLATION_SHADERS
-		return Shader;
-	}
+	return nullptr;
 }
 
 FDomainShaderRHIRef FMetalDynamicRHI::RHICreateDomainShader(TArrayView<const uint8> Code, const FSHAHash& Hash)
 {
-	@autoreleasepool {
-#if PLATFORM_SUPPORTS_TESSELLATION_SHADERS
-		FMetalDomainShader* Shader = new FMetalDomainShader(Code);
-#else
-		FMetalDomainShader* Shader = new FMetalDomainShader;
-		FMetalCodeHeader Header;
-		Shader->Init(Code, Header);
-#endif // PLATFORM_SUPPORTS_TESSELLATION_SHADERS
-		return Shader;
-	}
+	return nullptr;
 }
 
 FGeometryShaderRHIRef FMetalDynamicRHI::RHICreateGeometryShader(TArrayView<const uint8> Code, const FSHAHash& Hash)
diff --git a/Engine/Source/Runtime/Apple/MetalRHI/Private/MetalGraphicsPipelineState.cpp b/Engine/Source/Runtime/Apple/MetalRHI/Private/MetalGraphicsPipelineState.cpp
--- a/Engine/Source/Runtime/Apple/MetalRHI/Private/MetalGraphicsPipelineState.cpp
+++ b/Engine/Source/Runtime/Apple/MetalRHI/Private/MetalGraphicsPipelineState.cpp
@@ -17,7 +17,7 @@
 
 
 // From MetalPipeline.cpp:
-extern FMetalShaderPipeline* GetMTLRenderPipeline(bool const bSync, FMetalGraphicsPipelineState const* State, const FGraphicsPipelineStateInitializer& Init, EMetalIndexType const IndexType);
+extern FMetalShaderPipeline* GetMTLRenderPipeline(bool const bSync, FMetalGraphicsPipelineState const* State, const FGraphicsPipelineStateInitializer& Init);
 extern void ReleaseMTLRenderPipeline(FMetalShaderPipeline* Pipeline);
 
 
@@ -28,46 +28,34 @@ extern void ReleaseMTLRenderPipeline(FMetalShaderPipeline* Pipeline);
 
 FMetalGraphicsPipelineState::FMetalGraphicsPipelineState(const FGraphicsPipelineStateInitializer& Init)
 	: Initializer(Init)
+	, PipelineState(nil)
 {
 	// void
 }
 
 bool FMetalGraphicsPipelineState::Compile()
 {
-	FMemory::Memzero(PipelineStates);
-
-	for (uint32 i = 0; i < EMetalIndexType_Num; i++)
-	{
-		PipelineStates[i] = [GetMTLRenderPipeline(true, this, Initializer, (EMetalIndexType)i) retain];
-		if(!PipelineStates[i])
-		{
-			return false;
-		}
-	}
-
-	return true;
+	check(PipelineState == nil);
+	PipelineState = [GetMTLRenderPipeline(true, this, Initializer) retain];
+	return (PipelineState != nil);
 }
 
 FMetalGraphicsPipelineState::~FMetalGraphicsPipelineState()
 {
-	for (uint32 i = 0; i < EMetalIndexType_Num; i++)
+	if (PipelineState != nil)
 	{
-		ReleaseMTLRenderPipeline(PipelineStates[i]);
-		PipelineStates[i] = nil;
+		ReleaseMTLRenderPipeline(PipelineState);
+		PipelineState = nil;
 	}
 }
 
-FMetalShaderPipeline* FMetalGraphicsPipelineState::GetPipeline(EMetalIndexType IndexType)
+FMetalShaderPipeline* FMetalGraphicsPipelineState::GetPipeline()
 {
-	check(IndexType < EMetalIndexType_Num);
-
-	if (!PipelineStates[IndexType])
+	if (PipelineState == nil)
 	{
-		PipelineStates[IndexType] = [GetMTLRenderPipeline(true, this, Initializer, IndexType) retain];
+		PipelineState = [GetMTLRenderPipeline(true, this, Initializer) retain];
+		check(PipelineState != nil);
 	}
 
-	FMetalShaderPipeline* Pipe = PipelineStates[IndexType];
-	check(Pipe);
-
-    return Pipe;
+    return PipelineState;
 }
diff --git a/Engine/Source/Runtime/Apple/MetalRHI/Private/MetalGraphicsPipelineState.h b/Engine/Source/Runtime/Apple/MetalRHI/Private/MetalGraphicsPipelineState.h
--- a/Engine/Source/Runtime/Apple/MetalRHI/Private/MetalGraphicsPipelineState.h
+++ b/Engine/Source/Runtime/Apple/MetalRHI/Private/MetalGraphicsPipelineState.h
@@ -13,7 +13,7 @@ class FMetalGraphicsPipelineState : public FRHIGraphicsPipelineState
 public:
 	virtual ~FMetalGraphicsPipelineState();
 
-	FMetalShaderPipeline* GetPipeline(EMetalIndexType IndexType);
+	FMetalShaderPipeline* GetPipeline();
 
 	/** Cached vertex structure */
 	TRefCountPtr<FMetalVertexDeclaration> VertexDeclaration;
@@ -21,10 +21,6 @@ public:
 	/** Cached shaders */
 	TRefCountPtr<FMetalVertexShader> VertexShader;
 	TRefCountPtr<FMetalPixelShader> PixelShader;
-#if PLATFORM_SUPPORTS_TESSELLATION_SHADERS
-	TRefCountPtr<FMetalHullShader> HullShader;
-	TRefCountPtr<FMetalDomainShader> DomainShader;
-#endif
 #if PLATFORM_SUPPORTS_GEOMETRY_SHADERS
 	TRefCountPtr<FMetalGeometryShader> GeometryShader;
 #endif
@@ -53,6 +49,5 @@ private:
 	// Needed to runtime refine shaders currently.
 	FGraphicsPipelineStateInitializer Initializer;
 
-	// Tessellation pipelines have three different variations for the indexing-style.
-	FMetalShaderPipeline* PipelineStates[EMetalIndexType_Num];
+	FMetalShaderPipeline* PipelineState;
 };
diff --git a/Engine/Source/Runtime/Apple/MetalRHI/Private/MetalPipeline.cpp b/Engine/Source/Runtime/Apple/MetalRHI/Private/MetalPipeline.cpp
--- a/Engine/Source/Runtime/Apple/MetalRHI/Private/MetalPipeline.cpp
+++ b/Engine/Source/Runtime/Apple/MetalRHI/Private/MetalPipeline.cpp
@@ -26,12 +26,6 @@ static FAutoConsoleVariableRef CVarMetalCacheShaderPipelines(
 	TEXT("When enabled (1, default) cache all graphics pipeline state objects created in MetalRHI for the life of the program, this trades memory for performance as creating PSOs is expensive in Metal.\n")
 	TEXT("Disable in the project configuration to allow PSOs to be released to save memory at the expense of reduced performance and increased hitching in-game\n. (On by default (1))"), ECVF_ReadOnly);
 
-static int32 GMetalTessellationForcePartitionMode = 0;
-static FAutoConsoleVariableRef CVarMetalTessellationForcePartitionMode(
-	TEXT("rhi.Metal.TessellationForcePartitionMode"),
-	GMetalTessellationForcePartitionMode,
-	TEXT("The partition mode (+1) to force Metal to use for debugging or off (0). (Default: 0)"));
-
 static int32 GMetalCacheMinSize = 32;
 static FAutoConsoleVariableRef CVarMetalCacheMinSize(
 	TEXT("r.ShaderPipelineCache.MetalCacheMinSizeInMB"),
@@ -49,23 +43,6 @@ static float RoundUpNearestEven(const float f)
 	return ret + isOdd;
 }
 
-static float RoundTessLevel(float TessFactor, mtlpp::TessellationPartitionMode PartitionMode)
-{
-	switch(PartitionMode)
-	{
-		case mtlpp::TessellationPartitionMode::ModePow2:
-			return FMath::RoundUpToPowerOfTwo((uint32)TessFactor);
-		case mtlpp::TessellationPartitionMode::ModeInteger:
-			return FMath::CeilToFloat(TessFactor);
-		case mtlpp::TessellationPartitionMode::ModeFractionalEven:
-		case mtlpp::TessellationPartitionMode::ModeFractionalOdd: // these are handled the same way
-			return RoundUpNearestEven(TessFactor);
-		default:
-			check(false);
-			return 0.0f;
-	}
-}
-
 // A tile-based or vertex-based debug shader for trying to emulate Aftermath style failure reporting
 static NSString* GMetalDebugShader = @"#include <metal_stdlib>\n"
 "#include <metal_compute>\n"
@@ -125,49 +102,6 @@ static NSString* GMetalDebugMarkerComputeShader = @"#include <metal_stdlib>\n"
 "	debugBuffer[0] = debugTable[0];\n"
 "}";
 
-// A compute shader for copying indices for separate tessellation
-static NSString* GMetalCopyIndexComputeShader = @"#include <metal_stdlib>\n"
-"#include <metal_compute>\n"
-"#include <metal_stdlib>\n"
-"using namespace metal;\n"
-"\n"
-"// Executes once\n"
-"kernel void Main_CopyIndex32(const device uint* source [[ buffer(0) ]], device uint* dest [[ buffer(1) ]], constant uint2& controlPointCount [[ buffer(2) ]], constant MTLDrawIndexedPrimitivesIndirectArguments& Params [[ buffer(3) ]], uint2 threadgroup_position_in_grid [[ threadgroup_position_in_grid ]], uint2 thread_position_in_threadgroup [[ thread_position_in_threadgroup ]])\n"
-"{\n"
-"	uint i = thread_position_in_threadgroup.y;\n"
-"	uint j = threadgroup_position_in_grid.x;\n"
-"	uint k = thread_position_in_threadgroup.x;\n"
-"	if (k < controlPointCount.x) {\n"
-"		dest[i * Params.indexCount + j * controlPointCount.y + k] = source[Params.indexStart + j * controlPointCount.x + k] + i * Params.indexCount;\n"
-"	} else {\n"
-"		dest[i * Params.indexCount + j * controlPointCount.y + k] = 0;\n"
-"	}\n"
-"}\n"
-"\n"
-"// Executes once\n"
-"kernel void Main_CopyIndex16(const device ushort* source [[ buffer(0) ]], device uint* dest [[ buffer(1) ]], constant uint2& controlPointCount [[ buffer(2) ]], constant MTLDrawIndexedPrimitivesIndirectArguments& Params [[ buffer(3) ]], uint2 threadgroup_position_in_grid [[ threadgroup_position_in_grid ]], uint2 thread_position_in_threadgroup [[ thread_position_in_threadgroup ]])\n"
-"{\n"
-"	uint i = thread_position_in_threadgroup.y;\n"
-"	uint j = threadgroup_position_in_grid.x;\n"
-"	uint k = thread_position_in_threadgroup.x;\n"
-"	if (k < controlPointCount.x) {\n"
-"		dest[i * Params.indexCount + j * controlPointCount.y + k] = source[Params.indexStart + j * controlPointCount.x + k] + i * Params.indexCount;\n"
-"	} else {\n"
-"		dest[i * Params.indexCount + j * controlPointCount.y + k] = 0;\n"
-"	}\n"
-"}\n"
-"\n"
-"// Executes once\n"
-"kernel void Main_FlattenTess(device MTLTriangleTessellationFactorsHalf* dest [[ buffer(0) ]], constant MTLDrawIndexedPrimitivesIndirectArguments& Params [[ buffer(1) ]])\n"
-"{\n"
-"	for(uint i = 0; i < Params.indexCount; i++) {\n"
-"		dest[i].edgeTessellationFactor[0] = half(1.0);\n"
-"		dest[i].edgeTessellationFactor[1] = half(1.0);\n"
-"		dest[i].edgeTessellationFactor[2] = half(1.0);\n"
-"		dest[i].insideTessellationFactor = half(1.0);\n"
-"	}\n"
-"}";
-
 struct FMetalHelperFunctions
 {
     mtlpp::Library DebugShadersLib;
@@ -177,15 +111,6 @@ struct FMetalHelperFunctions
 	mtlpp::Function DebugComputeFunc;
 	mtlpp::ComputePipelineState DebugComputeState;
 	
-	mtlpp::Library CopyIndexLib;
-	mtlpp::Function CopyIndex32Func;
-	mtlpp::Function CopyIndex16Func;
-	mtlpp::ComputePipelineState CopyIndex32State;
-	mtlpp::ComputePipelineState CopyIndex16State;
-	
-	mtlpp::Function FlattenTessFunc;
-	mtlpp::ComputePipelineState FlattenTessState;
-	
     FMetalHelperFunctions()
     {
 #if !PLATFORM_TVOS
@@ -203,19 +128,6 @@ struct FMetalHelperFunctions
 			DebugComputeState = GetMetalDeviceContext().GetDevice().NewComputePipelineState(DebugComputeFunc, &Error);
         }
 #endif
-		{
-			mtlpp::CompileOptions CompileOptions;
-			ns::AutoReleasedError Error;
-			
-			CopyIndexLib = GetMetalDeviceContext().GetDevice().NewLibrary(GMetalCopyIndexComputeShader, CompileOptions, &Error);
-			CopyIndex32Func = CopyIndexLib.NewFunction(@"Main_CopyIndex32");
-			CopyIndex16Func = CopyIndexLib.NewFunction(@"Main_CopyIndex16");
-			CopyIndex32State = GetMetalDeviceContext().GetDevice().NewComputePipelineState(CopyIndex32Func, &Error);
-			CopyIndex16State = GetMetalDeviceContext().GetDevice().NewComputePipelineState(CopyIndex16Func, &Error);
-			
-			FlattenTessFunc = CopyIndexLib.NewFunction(@"Main_FlattenTess");
-			FlattenTessState = GetMetalDeviceContext().GetDevice().NewComputePipelineState(FlattenTessFunc, &Error);
-		}
     }
     
     static FMetalHelperFunctions& Get()
@@ -233,21 +145,6 @@ struct FMetalHelperFunctions
 	{
 		return DebugComputeState;
 	}
-	
-	mtlpp::ComputePipelineState GetCopyIndex32Function()
-	{
-		return CopyIndex32State;
-	}
-	
-	mtlpp::ComputePipelineState GetCopyIndex16Function()
-	{
-		return CopyIndex16State;
-	}
-	
-	mtlpp::ComputePipelineState GetFlattenTessState()
-	{
-		return FlattenTessState;
-	}
 };
 
 mtlpp::ComputePipelineState GetMetalDebugComputeState()
@@ -255,30 +152,11 @@ mtlpp::ComputePipelineState GetMetalDebugComputeState()
 	return FMetalHelperFunctions::Get().GetDebugComputeState();
 }
 
-mtlpp::ComputePipelineState GetMetalCopyIndex32Function()
-{
-	return FMetalHelperFunctions::Get().GetCopyIndex32Function();
-}
-
-mtlpp::ComputePipelineState GetMetalCopyIndex16Function()
-{
-	return FMetalHelperFunctions::Get().GetCopyIndex16Function();
-}
-
-mtlpp::ComputePipelineState GetMetalFlattenTessState()
-{
-	return FMetalHelperFunctions::Get().GetFlattenTessState();
-}
-
 struct FMetalGraphicsPipelineKey
 {
 	FMetalRenderPipelineHash RenderPipelineHash;
 	FMetalHashedVertexDescriptor VertexDescriptorHash;
 	FSHAHash VertexFunction;
-#if PLATFORM_SUPPORTS_TESSELLATION_SHADERS
-	FSHAHash HullFunction;
-	FSHAHash DomainFunction;
-#endif
 	FSHAHash PixelFunction;
 
 	template<typename Type>
@@ -302,10 +180,6 @@ struct FMetalGraphicsPipelineKey
 		return (RenderPipelineHash == Other.RenderPipelineHash
 		&& VertexDescriptorHash == Other.VertexDescriptorHash
 		&& VertexFunction == Other.VertexFunction
-#if PLATFORM_SUPPORTS_TESSELLATION_SHADERS
-		&& HullFunction == Other.HullFunction
-		&& DomainFunction == Other.DomainFunction
-#endif
 		&& PixelFunction == Other.PixelFunction);
 	}
 	
@@ -313,15 +187,11 @@ struct FMetalGraphicsPipelineKey
 	{
 		uint32 H = FCrc::MemCrc32(&Key.RenderPipelineHash, sizeof(Key.RenderPipelineHash), GetTypeHash(Key.VertexDescriptorHash));
 		H = FCrc::MemCrc32(Key.VertexFunction.Hash, sizeof(Key.VertexFunction.Hash), H);
-#if PLATFORM_SUPPORTS_TESSELLATION_SHADERS
-		H = FCrc::MemCrc32(Key.HullFunction.Hash, sizeof(Key.HullFunction.Hash), H);
-		H = FCrc::MemCrc32(Key.DomainFunction.Hash, sizeof(Key.DomainFunction.Hash), H);
-#endif
 		H = FCrc::MemCrc32(Key.PixelFunction.Hash, sizeof(Key.PixelFunction.Hash), H);
 		return H;
 	}
 	
-	friend void InitMetalGraphicsPipelineKey(FMetalGraphicsPipelineKey& Key, const FGraphicsPipelineStateInitializer& Init, EMetalIndexType const IndexType)
+	friend void InitMetalGraphicsPipelineKey(FMetalGraphicsPipelineKey& Key, const FGraphicsPipelineStateInitializer& Init)
 	{
 		uint32 const NumActiveTargets = Init.ComputeNumValidRenderTargets();
 		check(NumActiveTargets <= MaxSimultaneousRenderTargets);
@@ -412,20 +282,6 @@ struct FMetalGraphicsPipelineKey
 		FMetalVertexShader* VertexShader = (FMetalVertexShader*)Init.BoundShaderState.VertexShaderRHI;
 		Key.VertexFunction = VertexShader->GetHash();
 
-#if PLATFORM_SUPPORTS_TESSELLATION_SHADERS
-		FMetalDomainShader* DomainShader = (FMetalDomainShader*)Init.BoundShaderState.DomainShaderRHI;
-		FMetalHullShader* HullShader = (FMetalHullShader*)Init.BoundShaderState.HullShaderRHI;
-		if (DomainShader && HullShader)
-		{
-			Key.HullFunction = HullShader->GetHash();
-			Key.DomainFunction = DomainShader->GetHash();
-			Key.SetHashValue(Offset_IndexType, NumBits_IndexType, IndexType);
-		}
-		else
-#endif
-		{
-			Key.SetHashValue(Offset_IndexType, NumBits_IndexType, EMetalIndexType_None);
-		}
 		if (PixelShader)
 		{
 			Key.PixelFunction = PixelShader->GetHash();
@@ -433,7 +289,7 @@ struct FMetalGraphicsPipelineKey
 	}
 };
 
-static FMetalShaderPipeline* CreateMTLRenderPipeline(bool const bSync, FMetalGraphicsPipelineKey const& Key, const FGraphicsPipelineStateInitializer& Init, EMetalIndexType const IndexType);
+static FMetalShaderPipeline* CreateMTLRenderPipeline(bool const bSync, FMetalGraphicsPipelineKey const& Key, const FGraphicsPipelineStateInitializer& Init);
 
 class FMetalShaderPipelineCache
 {
@@ -444,12 +300,12 @@ public:
 		return sSelf;
 	}
 	
-	FMetalShaderPipeline* GetRenderPipeline(bool const bSync, FMetalGraphicsPipelineState const* State, const FGraphicsPipelineStateInitializer& Init, EMetalIndexType const IndexType)
+	FMetalShaderPipeline* GetRenderPipeline(bool const bSync, FMetalGraphicsPipelineState const* State, const FGraphicsPipelineStateInitializer& Init)
 	{
 		SCOPE_CYCLE_COUNTER(STAT_MetalPipelineStateTime);
 		
 		FMetalGraphicsPipelineKey Key;
-		InitMetalGraphicsPipelineKey(Key, Init, IndexType);
+		InitMetalGraphicsPipelineKey(Key, Init);
 		
 		// By default there'll be more threads trying to read this than to write it.
 		PipelineMutex.ReadLock();
@@ -490,7 +346,7 @@ public:
 
 			if (bCompile)
 			{
-				Desc = CreateMTLRenderPipeline(bSync, Key, Init, IndexType);
+				Desc = CreateMTLRenderPipeline(bSync, Key, Init);
 
 				if (Desc != nil)
 				{
@@ -761,8 +617,7 @@ static bool ConfigureRenderPipelineDescriptor(mtlpp::RenderPipelineDescriptor& R
 											  mtlpp::TileRenderPipelineDescriptor& DebugPipelineDesc,
 #endif
 											  FMetalGraphicsPipelineKey const& Key,
-											  const FGraphicsPipelineStateInitializer& Init,
-											  EMetalIndexType const IndexType)
+											  const FGraphicsPipelineStateInitializer& Init)
 {
 	FMetalPixelShader* PixelShader = (FMetalPixelShader*)Init.BoundShaderState.PixelShaderRHI;
 	uint32 const NumActiveTargets = Init.ComputeNumValidRenderTargets();
@@ -937,14 +792,8 @@ static bool ConfigureRenderPipelineDescriptor(mtlpp::RenderPipelineDescriptor& R
 		FMetalVertexShader* VertexShader = (FMetalVertexShader*)Init.BoundShaderState.VertexShaderRHI;
 		
 		ns::AutoReleased<ns::Array<mtlpp::PipelineBufferDescriptor>> VertexPipelineBuffers = RenderPipelineDesc.GetVertexBuffers();
-#if PLATFORM_SUPPORTS_TESSELLATION_SHADERS
-		FMetalDomainShader* DomainShader = (FMetalDomainShader*)Init.BoundShaderState.DomainShaderRHI;
-		FMetalShaderBindings& VertexBindings = DomainShader ? DomainShader->Bindings : VertexShader->Bindings;
-		int8 VertexSideTable = DomainShader ? DomainShader->SideTableBinding : VertexShader->SideTableBinding;
-#else
 		FMetalShaderBindings& VertexBindings = VertexShader->Bindings;
 		int8 VertexSideTable = VertexShader->SideTableBinding;
-#endif
 		{
 			uint32 ImmutableBuffers = VertexBindings.ConstantBuffers | VertexBindings.ArgumentBuffers;
 			while(ImmutableBuffers)
@@ -991,341 +840,7 @@ static bool ConfigureRenderPipelineDescriptor(mtlpp::RenderPipelineDescriptor& R
 	return true;
 }
 
-#if PLATFORM_SUPPORTS_TESSELLATION_SHADERS
-
-static mtlpp::VertexFormat Formats[(uint8)EMetalComponentType::Max][4] = {
-	{mtlpp::VertexFormat::UInt, mtlpp::VertexFormat::UInt2, mtlpp::VertexFormat::UInt3, mtlpp::VertexFormat::UInt4},
-	{mtlpp::VertexFormat::Int, mtlpp::VertexFormat::Int2, mtlpp::VertexFormat::Int3, mtlpp::VertexFormat::Int4},
-	{mtlpp::VertexFormat::Invalid, mtlpp::VertexFormat::Half2, mtlpp::VertexFormat::Half3, mtlpp::VertexFormat::Half4},
-	{mtlpp::VertexFormat::Float, mtlpp::VertexFormat::Float2, mtlpp::VertexFormat::Float3, mtlpp::VertexFormat::Float4},
-	{mtlpp::VertexFormat::Invalid, mtlpp::VertexFormat::UChar2, mtlpp::VertexFormat::UChar3, mtlpp::VertexFormat::UChar4},
-};
-
-static FMetalShaderPipeline* CreateSeparateMetalTessellationPipeline(bool const bSync, FMetalGraphicsPipelineKey const& Key, const FGraphicsPipelineStateInitializer& Init, EMetalIndexType const IndexType)
-{
-	FMetalShaderPipeline* Pipeline = [FMetalShaderPipeline new];
-	METAL_DEBUG_OPTION(FMemory::Memzero(Pipeline->ResourceMask, sizeof(Pipeline->ResourceMask)));
-	
-	FMetalVertexShader* VertexShader = (FMetalVertexShader*)Init.BoundShaderState.VertexShaderRHI;
-	mtlpp::Function vertexFunction = VertexShader->GetFunction();
-	
-	mtlpp::RenderPipelineDescriptor VertexRenderPipelineDesc;
-	VertexRenderPipelineDesc.SetRasterizationEnabled(false);
-	
-#if PLATFORM_MAC
-	VertexRenderPipelineDesc.SetInputPrimitiveTopology(TranslatePrimitiveTopology(Init.PrimitiveType));
-#endif
-	
-	FMetalVertexDeclaration* VertexDecl = (FMetalVertexDeclaration*)Init.BoundShaderState.VertexDeclarationRHI;
-	VertexRenderPipelineDesc.SetVertexDescriptor(GetMaskedVertexDescriptor(VertexDecl->Layout.VertexDesc, VertexShader->Bindings.InOutMask));
-	VertexRenderPipelineDesc.SetVertexFunction(vertexFunction);
-	
-	VertexRenderPipelineDesc.SetDepthAttachmentPixelFormat((mtlpp::PixelFormat)GPixelFormats[PF_DepthStencil].PlatformFormat);
-	VertexRenderPipelineDesc.SetStencilAttachmentPixelFormat((mtlpp::PixelFormat)GPixelFormats[PF_DepthStencil].PlatformFormat);
-	
-#if ENABLE_METAL_GPUPROFILE
-	ns::String VertexName = vertexFunction.GetName();
-	VertexRenderPipelineDesc.SetLabel([NSString stringWithFormat:@"%@", VertexName.GetPtr()]);
-#endif
-	
-	NSUInteger RenderOption = mtlpp::PipelineOption::NoPipelineOption;
-	mtlpp::AutoReleasedRenderPipelineReflection* Reflection = nullptr;
-	mtlpp::AutoReleasedRenderPipelineReflection OutReflection;
-	Reflection = &OutReflection;
-	if (GetMetalDeviceContext().GetCommandQueue().GetRuntimeDebuggingLevel() >= EMetalDebugLevelFastValidation)
-	{
-		RenderOption = mtlpp::PipelineOption::ArgumentInfo | mtlpp::PipelineOption::BufferTypeInfo;
-	}
-	
-	ns::Error Error;
-	mtlpp::Device Device = GetMetalDeviceContext().GetDevice();
-	mtlpp::RenderPipelineState& VertexRenderPipelineState = Pipeline->StreamPipelineState;
-	{
-		ns::AutoReleasedError RenderError;
-		METAL_GPUPROFILE(FScopedMetalCPUStats CPUStat(FString::Printf(TEXT("NewRenderPipeline: %s"), TEXT("")/**FString([RenderPipelineDesc.GetPtr() description])*/)));
-		VertexRenderPipelineState = Device.NewRenderPipelineState(VertexRenderPipelineDesc, (mtlpp::PipelineOption)RenderOption, Reflection, &RenderError);
-		if (Reflection)
-		{
-			Pipeline->StreamPipelineReflection = *Reflection;
-#if METAL_DEBUG_OPTIONS
-			Pipeline->StreamDesc = VertexRenderPipelineDesc;
-#endif
-		}
-		Error = RenderError;
-	}
-	
-	UE_CLOG((VertexRenderPipelineState == nil), LogMetal, Error, TEXT("Failed to generate a pipeline state object: %s"), *FString(Error.GetPtr().description));
-	UE_CLOG((VertexRenderPipelineState == nil), LogMetal, Error, TEXT("Vertex shader: %s"), *FString(VertexShader->GetSourceCode()));
-	UE_CLOG((VertexRenderPipelineState == nil), LogMetal, Error, TEXT("Descriptor: %s"), *FString(VertexRenderPipelineDesc.GetPtr().description));
-	UE_CLOG((VertexRenderPipelineState == nil), LogMetal, Error, TEXT("Failed to generate a render pipeline state object:\n\n %s\n\n"), *FString(Error.GetLocalizedDescription()));
-
-	FMetalHullShader* HullShader = (FMetalHullShader*)Init.BoundShaderState.HullShaderRHI;
-	mtlpp::ComputePipelineDescriptor HullShaderPipelineDesc;
-	
-	mtlpp::StageInputOutputDescriptor HullStageInOut;
-	
-	ns::Array<mtlpp::BufferLayoutDescriptor> HullVertexLayouts = HullStageInOut.GetLayouts();
-	ns::Array<mtlpp::AttributeDescriptor> HullVertexAttribs = HullStageInOut.GetAttributes();
-	
-	bool bFoundAttrib = false;
-	for (FMetalAttribute const& HSAttrib : HullShader->TessellationOutputAttribs.HSIn)
-	{
-		bFoundAttrib = false;
-		for (FMetalAttribute const& VSAttrib : VertexShader->TessellationOutputAttribs.HSOut)
-		{
-			if (HSAttrib.Semantic == VSAttrib.Semantic)
-			{
-				uint32 attributeIndex = HSAttrib.Index;
-
-				mtlpp::VertexFormat format = Formats[(uint8)HSAttrib.Type][HSAttrib.Components-1];
-				check(format != mtlpp::VertexFormat::Invalid); // TODO support more cases
-				HullVertexAttribs[attributeIndex].SetFormat((mtlpp::AttributeFormat)format);
-				HullVertexAttribs[attributeIndex].SetOffset(VSAttrib.Offset);
-				HullVertexAttribs[attributeIndex].SetBufferIndex(HullShader->TessellationControlPointIndexBuffer);
-				bFoundAttrib = true;
-				break;
-			}
-		}
-		check(bFoundAttrib);
-	}
-	
-	// if (IndexType == EMetalIndexType_UInt16 || IndexType == EMetalIndexType_UInt32)
-	{
-		HullStageInOut.SetIndexType(GetMetalIndexType(EMetalIndexType_UInt32));
-		HullStageInOut.SetIndexBufferIndex(HullShader->TessellationIndexBuffer);
-		HullVertexLayouts[HullShader->TessellationControlPointIndexBuffer].SetStepFunction(mtlpp::StepFunction::ThreadPositionInGridXIndexed);
-	}
-//	else
-//	{
-//		HullVertexLayouts[HullShader->TessellationControlPointIndexBuffer].SetStepFunction(mtlpp::StepFunction::ThreadPositionInGridX);
-//	}
-	HullVertexLayouts[HullShader->TessellationControlPointIndexBuffer].SetStride(VertexShader->TessellationOutputAttribs.HSOutSize);
-	
-	mtlpp::Function hullFunction = HullShader->GetFunction();
-	HullShaderPipelineDesc.SetComputeFunction(hullFunction);
-	check(HullShaderPipelineDesc.GetComputeFunction());
-	HullShaderPipelineDesc.SetStageInputDescriptor(HullStageInOut);
-	
-	mtlpp::ComputePipelineState& HullPipelineState = Pipeline->ComputePipelineState;
-	{
-		METAL_GPUPROFILE(FScopedMetalCPUStats CPUStat(FString::Printf(TEXT("NewComputePipelineState: %s"), TEXT("")/**FString([ComputePipelineDesc.GetPtr() description])*/)));
-		ns::AutoReleasedError AutoError;
-		NSUInteger ComputeOption = mtlpp::PipelineOption::NoPipelineOption;
-#if ENABLE_METAL_GPUPROFILE
-		{
-			ns::String HullName = hullFunction.GetName();
-			HullShaderPipelineDesc.SetLabel([NSString stringWithFormat:@"%@", HullName.GetPtr()]);
-		}
-#endif
-		if (GetMetalDeviceContext().GetCommandQueue().GetRuntimeDebuggingLevel() >= EMetalDebugLevelFastValidation)
-		{
-			mtlpp::AutoReleasedComputePipelineReflection HullReflection;
-			ComputeOption = mtlpp::PipelineOption::ArgumentInfo | mtlpp::PipelineOption::BufferTypeInfo;
-			HullPipelineState = Device.NewComputePipelineState(HullShaderPipelineDesc, (mtlpp::PipelineOption)ComputeOption, &HullReflection, &AutoError);
-			Pipeline->ComputePipelineReflection = HullReflection;
-		}
-		else
-		{
-			HullPipelineState = Device.NewComputePipelineState(HullShaderPipelineDesc, (mtlpp::PipelineOption)ComputeOption, nullptr, &AutoError);
-		}
-		Error = AutoError;
-		
-#if METAL_DEBUG_OPTIONS
-		Pipeline->ComputeDesc = HullShaderPipelineDesc;
-#endif
-		
-		UE_CLOG((HullPipelineState == nil), LogMetal, Error, TEXT("Failed to generate a pipeline state object: %s"), *FString([Error.GetPtr() description]));
-		UE_CLOG((HullPipelineState == nil), LogMetal, Error, TEXT("Hull shader: %s"), *FString(HullShader->GetSourceCode()));
-		UE_CLOG((HullPipelineState == nil), LogMetal, Error, TEXT("Descriptor: %s"), *FString(HullShaderPipelineDesc.GetPtr().description));
-		UE_CLOG((HullPipelineState == nil), LogMetal, Error, TEXT("Failed to generate a hull pipeline state object:\n\n %s\n\n"), *FString(Error.GetLocalizedDescription()));
-	}
-	
-#if PLATFORM_MAC
-	mtlpp::RenderPipelineDescriptor DebugPipelineDesc;
-#elif !PLATFORM_TVOS
-	mtlpp::TileRenderPipelineDescriptor DebugPipelineDesc;
-#endif
-
-	FMetalDomainShader* DomainShader = (FMetalDomainShader*)Init.BoundShaderState.DomainShaderRHI;
-	mtlpp::RenderPipelineDescriptor DomainRenderPipelineDesc;
-	{
-		mtlpp::VertexDescriptor DomainVertexDesc;
-		ns::Array<mtlpp::VertexBufferLayoutDescriptor> DomainVertexLayouts = DomainVertexDesc.GetLayouts();
-		ns::Array<mtlpp::VertexAttributeDescriptor> Attribs = DomainVertexDesc.GetAttributes();
-		
-		if (DomainShader->TessellationHSOutBuffer != UINT_MAX && HullShader->TessellationOutputAttribs.HSOutSize)
-		{
-			check(DomainShader->TessellationHSOutBuffer < ML_MaxBuffers);
-			uint32 bufferIndex = DomainShader->TessellationHSOutBuffer;
-			uint32 bufferSize = HullShader->TessellationOutputAttribs.HSOutSize;
-			
-			DomainVertexLayouts[bufferIndex].SetStride(bufferSize);
-			DomainVertexLayouts[bufferIndex].SetStepFunction(mtlpp::VertexStepFunction::PerPatch);
-			DomainVertexLayouts[bufferIndex].SetStepRate(1);
-			
-			bFoundAttrib = false;
-			for (FMetalAttribute const& VSAttrib : DomainShader->TessellationOutputAttribs.HSOut)
-			{
-				bFoundAttrib = false;
-				for (FMetalAttribute const& HSAttrib : HullShader->TessellationOutputAttribs.HSOut)
-				{
-					if (HSAttrib.Semantic == VSAttrib.Semantic)
-					{
-						uint32 attributeIndex = VSAttrib.Index;
-						
-						mtlpp::VertexFormat format = Formats[(uint8)VSAttrib.Type][VSAttrib.Components-1];
-						check(format != mtlpp::VertexFormat::Invalid); // TODO support more cases
-						Attribs[attributeIndex].SetFormat(format);
-						Attribs[attributeIndex].SetOffset(HSAttrib.Offset);
-						Attribs[attributeIndex].SetBufferIndex(DomainShader->TessellationHSOutBuffer);
-						bFoundAttrib = true;
-						break;
-					}
-				}
-				check(bFoundAttrib);
-			}
-		}
-		
-		if (DomainShader->TessellationControlPointOutBuffer != UINT_MAX && HullShader->TessellationOutputAttribs.PatchControlPointOutSize)
-		{
-			uint32 bufferIndex = DomainShader->TessellationControlPointOutBuffer;
-			uint32 bufferSize = HullShader->TessellationOutputAttribs.PatchControlPointOutSize;
-			
-			DomainVertexLayouts[bufferIndex].SetStride(bufferSize);
-			DomainVertexLayouts[bufferIndex].SetStepFunction(mtlpp::VertexStepFunction::PerPatchControlPoint);
-			DomainVertexLayouts[bufferIndex].SetStepRate(1);
-			
-			bFoundAttrib = false;
-			for (FMetalAttribute const& VSAttrib : DomainShader->TessellationOutputAttribs.PatchControlPointOut)
-			{
-				bFoundAttrib = false;
-				for (FMetalAttribute const& HSAttrib : HullShader->TessellationOutputAttribs.PatchControlPointOut)
-				{
-					if (HSAttrib.Semantic == VSAttrib.Semantic)
-					{
-						uint32 attributeIndex = VSAttrib.Index;
-						
-						mtlpp::VertexFormat format = Formats[(uint8)VSAttrib.Type][VSAttrib.Components-1];
-						check(format != mtlpp::VertexFormat::Invalid); // TODO support more cases
-						Attribs[attributeIndex].SetFormat(format);
-						Attribs[attributeIndex].SetOffset(HSAttrib.Offset);
-						Attribs[attributeIndex].SetBufferIndex(DomainShader->TessellationControlPointOutBuffer);
-						bFoundAttrib = true;
-						break;
-					}
-				}
-				check(bFoundAttrib);
-			}
-		}
-		
-		DomainRenderPipelineDesc.SetTessellationPartitionMode(GMetalTessellationForcePartitionMode == 0 ? HullShader->TessellationPartitioning : (mtlpp::TessellationPartitionMode)(GMetalTessellationForcePartitionMode - 1));
-		DomainRenderPipelineDesc.SetTessellationFactorStepFunction(mtlpp::TessellationFactorStepFunction::PerPatch);
-		DomainRenderPipelineDesc.SetTessellationOutputWindingOrder(HullShader->TessellationOutputWinding);
-		int FixedMaxTessFactor = (int)RoundTessLevel(HullShader->TessellationMaxTessFactor, DomainRenderPipelineDesc.GetTessellationPartitionMode());
-		DomainRenderPipelineDesc.SetMaxTessellationFactor(FixedMaxTessFactor);
-		DomainRenderPipelineDesc.SetTessellationFactorScaleEnabled(NO);
-		DomainRenderPipelineDesc.SetTessellationFactorFormat(mtlpp::TessellationFactorFormat::Half);
-		DomainRenderPipelineDesc.SetTessellationControlPointIndexType(mtlpp::TessellationControlPointIndexType::None);
-		DomainRenderPipelineDesc.SetVertexDescriptor(DomainVertexDesc);
-		
-#if PLATFORM_TVOS
-		if (!ConfigureRenderPipelineDescriptor(DomainRenderPipelineDesc, Key, Init, IndexType))
-#else
-		if (!ConfigureRenderPipelineDescriptor(DomainRenderPipelineDesc, DebugPipelineDesc, Key, Init, IndexType))
-#endif
-		{
-			[Pipeline release];
-			return nil;
-		}
-		
-		mtlpp::Function domainFunction = DomainShader ? DomainShader->GetFunction() : nil;
-
-		FMetalPixelShader* PixelShader = (FMetalPixelShader*)Init.BoundShaderState.PixelShaderRHI;
-		mtlpp::Function fragmentFunction = PixelShader ? PixelShader->GetFunction() : nil;
-		
-		DomainRenderPipelineDesc.SetVertexFunction(domainFunction);
-		DomainRenderPipelineDesc.SetFragmentFunction(fragmentFunction);
-#if ENABLE_METAL_GPUPROFILE
-		ns::String DomainName = domainFunction.GetName();
-		ns::String FragmentName = fragmentFunction ? fragmentFunction.GetName() : @"";
-		DomainRenderPipelineDesc.SetLabel([NSString stringWithFormat:@"%@+%@", DomainName.GetPtr(), FragmentName.GetPtr()]);
-#endif
-		
-		mtlpp::RenderPipelineState& RenderPipelineState = Pipeline->RenderPipelineState;
-		{
-			ns::AutoReleasedError RenderError;
-			METAL_GPUPROFILE(FScopedMetalCPUStats CPUStat(FString::Printf(TEXT("NewRenderPipeline: %s"), TEXT("")/**FString([RenderPipelineDesc.GetPtr() description])*/)));
-			RenderPipelineState = Device.NewRenderPipelineState(DomainRenderPipelineDesc, (mtlpp::PipelineOption)RenderOption, Reflection, &RenderError);
-			if (Reflection)
-			{
-				Pipeline->RenderPipelineReflection = *Reflection;
-#if METAL_DEBUG_OPTIONS
-				Pipeline->RenderDesc = DomainRenderPipelineDesc;
-#endif
-			}
-			Error = RenderError;
-		}
-		
-		UE_CLOG((RenderPipelineState == nil), LogMetal, Error, TEXT("Failed to generate a pipeline state object: %s"), *FString(Error.GetPtr().description));
-		UE_CLOG((RenderPipelineState == nil), LogMetal, Error, TEXT("Domain shader: %s"), DomainShader ? *FString(DomainShader->GetSourceCode()) : TEXT("NULL"));
-		UE_CLOG((RenderPipelineState == nil), LogMetal, Error, TEXT("Pixel shader: %s"), PixelShader ? *FString(PixelShader->GetSourceCode()) : TEXT("NULL"));
-		UE_CLOG((RenderPipelineState == nil), LogMetal, Error, TEXT("Descriptor: %s"), *FString(DomainRenderPipelineDesc.GetPtr().description));
-		UE_CLOG((RenderPipelineState == nil), LogMetal, Error, TEXT("Failed to generate a render pipeline state object:\n\n %s\n\n"), *FString(Error.GetLocalizedDescription()));
-	}
-	
-	// We need to pass a failure up the chain, so we'll clean up here.
-	if(Pipeline->StreamPipelineState == nil || Pipeline->ComputePipelineState == nil || Pipeline->RenderPipelineState == nil)
-	{
-		[Pipeline release];
-		return nil;
-	}
-	
-#if METAL_DEBUG_OPTIONS
-	Pipeline->ComputeSource = HullShader->GetSourceCode();
-	Pipeline->VertexSource = VertexShader->GetSourceCode();
-	Pipeline->DomainSource = DomainShader->GetSourceCode();
-	FMetalPixelShader* PixelShader = (FMetalPixelShader*)Init.BoundShaderState.PixelShaderRHI;
-	Pipeline->FragmentSource = PixelShader ? PixelShader->GetSourceCode() : nil;
-#endif
-	
-#if !PLATFORM_TVOS
-	if (GMetalCommandBufferDebuggingEnabled)
-	{
-#if PLATFORM_MAC
-		DebugPipelineDesc.SetVertexFunction(FMetalHelperFunctions::Get().GetDebugFunction());
-		DebugPipelineDesc.SetRasterizationEnabled(false);
-#else
-		DebugPipelineDesc.SetTileFunction(FMetalHelperFunctions::Get().GetDebugFunction());
-		DebugPipelineDesc.SetRasterSampleCount(RenderPipelineDesc.GetSampleCount());
-		DebugPipelineDesc.SetThreadgroupSizeMatchesTileSize(false);
-#endif
-#if ENABLE_METAL_GPUPROFILE
-		DebugPipelineDesc.SetLabel(@"Main_Debug");
-#endif
-		
-		ns::AutoReleasedError RenderError;
-		METAL_GPUPROFILE(FScopedMetalCPUStats CPUStat(FString::Printf(TEXT("NewDebugPipeline: %s"), TEXT("")/**FString([RenderPipelineDesc.GetPtr() description])*/)));
-		Pipeline->DebugPipelineState = Device.NewRenderPipelineState(DebugPipelineDesc, mtlpp::PipelineOption::NoPipelineOption, Reflection, nullptr);
-	}
-#endif
-	
-#if METAL_DEBUG_OPTIONS
-	if (GFrameCounter > 3)
-	{
-		UE_LOG(LogMetal, Verbose, TEXT("Created a hitchy pipeline state for hash %llx %llx %llx"), (uint64)Key.RenderPipelineHash.RasterBits, (uint64)(Key.RenderPipelineHash.TargetBits), (uint64)Key.VertexDescriptorHash.VertexDescHash);
-	}
-#endif
-
-	if (Pipeline && SafeGetRuntimeDebuggingLevel() >= EMetalDebugLevelFastValidation)
-	{
-		[Pipeline initResourceMask];
-	}
-	
-	return Pipeline;
-}
-#endif
-
-static FMetalShaderPipeline* CreateMTLRenderPipeline(bool const bSync, FMetalGraphicsPipelineKey const& Key, const FGraphicsPipelineStateInitializer& Init, EMetalIndexType const IndexType)
+static FMetalShaderPipeline* CreateMTLRenderPipeline(bool const bSync, FMetalGraphicsPipelineKey const& Key, const FGraphicsPipelineStateInitializer& Init)
 {
     FMetalVertexShader* VertexShader = (FMetalVertexShader*)Init.BoundShaderState.VertexShaderRHI;
     FMetalPixelShader* PixelShader = (FMetalPixelShader*)Init.BoundShaderState.PixelShaderRHI;
@@ -1333,21 +848,8 @@ static FMetalShaderPipeline* CreateMTLRenderPipeline(bool const bSync, FMetalGra
     mtlpp::Function vertexFunction = VertexShader->GetFunction();
     mtlpp::Function fragmentFunction = PixelShader ? PixelShader->GetFunction() : nil;
 
-#if PLATFORM_SUPPORTS_TESSELLATION_SHADERS
-	FMetalDomainShader* DomainShader = (FMetalDomainShader*)Init.BoundShaderState.DomainShaderRHI;
-    mtlpp::Function domainFunction = DomainShader ? DomainShader->GetFunction() : nil;
-	if (DomainShader && FMetalCommandQueue::SupportsFeature(EMetalFeaturesSeparateTessellation))
-	{
-		return CreateSeparateMetalTessellationPipeline(bSync, Key, Init, IndexType);
-	}
-#endif
-    
     FMetalShaderPipeline* Pipeline = nil;
-    if (vertexFunction && ((PixelShader != nullptr) == (fragmentFunction != nil))
-#if PLATFORM_SUPPORTS_TESSELLATION_SHADERS
-		&& ((DomainShader != nullptr) == (domainFunction != nil))
-#endif
-	)
+    if (vertexFunction && ((PixelShader != nullptr) == (fragmentFunction != nil)))
     {
 		ns::Error Error;
 		mtlpp::Device Device = GetMetalDeviceContext().GetDevice();
@@ -1367,9 +869,9 @@ static FMetalShaderPipeline* CreateMTLRenderPipeline(bool const bSync, FMetalGra
 #endif
 		
 #if PLATFORM_TVOS
-		if (!ConfigureRenderPipelineDescriptor(RenderPipelineDesc, Key, Init, IndexType))
+		if (!ConfigureRenderPipelineDescriptor(RenderPipelineDesc, Key, Init))
 #else
-		if (!ConfigureRenderPipelineDescriptor(RenderPipelineDesc, DebugPipelineDesc, Key, Init, IndexType))
+		if (!ConfigureRenderPipelineDescriptor(RenderPipelineDesc, DebugPipelineDesc, Key, Init))
 #endif
 		{
 			return nil;
@@ -1377,329 +879,13 @@ static FMetalShaderPipeline* CreateMTLRenderPipeline(bool const bSync, FMetalGra
         
         FMetalVertexDeclaration* VertexDecl = (FMetalVertexDeclaration*)Init.BoundShaderState.VertexDeclarationRHI;
 		
-#if PLATFORM_SUPPORTS_TESSELLATION_SHADERS
-		FMetalHullShader* HullShader = (FMetalHullShader*)Init.BoundShaderState.HullShaderRHI;
-		if (Init.BoundShaderState.HullShaderRHI == nullptr)
-#endif
-        {
-#if PLATFORM_SUPPORTS_TESSELLATION_SHADERS
-            check(Init.BoundShaderState.DomainShaderRHI == nullptr);
-#endif
-            RenderPipelineDesc.SetVertexDescriptor(GetMaskedVertexDescriptor(VertexDecl->Layout.VertexDesc, VertexShader->Bindings.InOutMask));
-            RenderPipelineDesc.SetVertexFunction(vertexFunction);
-            RenderPipelineDesc.SetFragmentFunction(fragmentFunction);
-#if ENABLE_METAL_GPUPROFILE
-			ns::String VertexName = vertexFunction.GetName();
-			ns::String FragmentName = fragmentFunction ? fragmentFunction.GetName() : @"";
-			RenderPipelineDesc.SetLabel([NSString stringWithFormat:@"%@+%@", VertexName.GetPtr(), FragmentName.GetPtr()]);
-#endif
-        }
-#if PLATFORM_SUPPORTS_TESSELLATION_SHADERS
-        else
-        {
-            check(Init.BoundShaderState.DomainShaderRHI != nullptr);
-            
-            RenderPipelineDesc.SetTessellationPartitionMode(GMetalTessellationForcePartitionMode == 0 ? DomainShader->TessellationPartitioning : (mtlpp::TessellationPartitionMode)(GMetalTessellationForcePartitionMode - 1));
-			RenderPipelineDesc.SetTessellationFactorStepFunction(mtlpp::TessellationFactorStepFunction::PerPatch);
-            RenderPipelineDesc.SetTessellationOutputWindingOrder(DomainShader->TessellationOutputWinding);
-            int FixedMaxTessFactor = (int)RoundTessLevel(VertexShader->TessellationMaxTessFactor, RenderPipelineDesc.GetTessellationPartitionMode());
-            RenderPipelineDesc.SetMaxTessellationFactor(FixedMaxTessFactor);
-			RenderPipelineDesc.SetTessellationFactorScaleEnabled(NO);
-			RenderPipelineDesc.SetTessellationFactorFormat(mtlpp::TessellationFactorFormat::Half);
-			RenderPipelineDesc.SetTessellationControlPointIndexType(mtlpp::TessellationControlPointIndexType::None);
-            
-            RenderPipelineDesc.SetVertexFunction(domainFunction);
-            RenderPipelineDesc.SetFragmentFunction(fragmentFunction);
-#if ENABLE_METAL_GPUPROFILE
-			{
-				ns::String VertexName = domainFunction.GetName();
-				ns::String FragmentName = fragmentFunction ? fragmentFunction.GetName() : @"";
-				RenderPipelineDesc.SetLabel([NSString stringWithFormat:@"%@+%@", VertexName.GetPtr(), FragmentName.GetPtr()]);
-			}
-#endif
-            
-			ComputePipelineDesc = mtlpp::ComputePipelineDescriptor();
-            check(ComputePipelineDesc);
-			
-			if (FMetalCommandQueue::SupportsFeature(EMetalFeaturesPipelineBufferMutability))
-			{
-				ns::AutoReleased<ns::Array<mtlpp::PipelineBufferDescriptor>> PipelineBuffers = ComputePipelineDesc.GetBuffers();
-				
-				uint32 ImmutableBuffers = VertexShader->Bindings.ConstantBuffers | VertexShader->Bindings.ArgumentBuffers;
-				while(ImmutableBuffers)
-				{
-					uint32 Index = __builtin_ctz(ImmutableBuffers);
-					ImmutableBuffers &= ~(1 << Index);
-					
-					if (Index < ML_MaxBuffers)
-					{
-						ns::AutoReleased<mtlpp::PipelineBufferDescriptor> PipelineBuffer = PipelineBuffers[Index];
-						PipelineBuffer.SetMutability(mtlpp::Mutability::Immutable);
-					}
-				}
-				if (VertexShader->SideTableBinding > 0)
-				{
-					ns::AutoReleased<mtlpp::PipelineBufferDescriptor> PipelineBuffer = PipelineBuffers[VertexShader->SideTableBinding];
-					PipelineBuffer.SetMutability(mtlpp::Mutability::Immutable);
-				}
-			}
-			
-			mtlpp::VertexDescriptor DomainVertexDesc;
-			mtlpp::StageInputOutputDescriptor ComputeStageInOut;
-			ComputeStageInOut.SetIndexBufferIndex(VertexShader->TessellationControlPointIndexBuffer);
-			
-            FMetalTessellationPipelineDesc& TessellationDesc = Pipeline->TessellationPipelineDesc;
-            TessellationDesc.TessellationInputControlPointBufferIndex = DomainShader->TessellationControlPointOutBuffer;
-            TessellationDesc.TessellationOutputControlPointBufferIndex = VertexShader->TessellationControlPointOutBuffer;
-            TessellationDesc.TessellationInputPatchConstBufferIndex = DomainShader->TessellationHSOutBuffer;
-            TessellationDesc.TessellationPatchConstBufferIndex = VertexShader->TessellationHSOutBuffer;
-            TessellationDesc.TessellationFactorBufferIndex = VertexShader->TessellationHSTFOutBuffer;
-            TessellationDesc.TessellationPatchCountBufferIndex = VertexShader->TessellationPatchCountBuffer;
-            TessellationDesc.TessellationIndexBufferIndex = VertexShader->TessellationIndexBuffer;
-            TessellationDesc.TessellationPatchConstOutSize = VertexShader->TessellationOutputAttribs.HSOutSize;
-			TessellationDesc.TessellationControlPointIndexBufferIndex = VertexShader->TessellationControlPointIndexBuffer;
-            TessellationDesc.TessellationPatchControlPointOutSize = VertexShader->TessellationOutputAttribs.PatchControlPointOutSize;
-            TessellationDesc.TessellationTessFactorOutSize = VertexShader->TessellationOutputAttribs.HSTFOutSize;
-            
-            check(TessellationDesc.TessellationOutputControlPointBufferIndex < ML_MaxBuffers);
-            check(TessellationDesc.TessellationFactorBufferIndex < ML_MaxBuffers);
-            check(TessellationDesc.TessellationPatchCountBufferIndex < ML_MaxBuffers);
-            check(TessellationDesc.TessellationTessFactorOutSize == 2*4 || TessellationDesc.TessellationTessFactorOutSize == 2*6);
-            
-            mtlpp::VertexStepFunction stepFunction = mtlpp::VertexStepFunction::PerPatch;
-			
-			ns::Array<mtlpp::VertexBufferLayoutDescriptor> DomainVertexLayouts = DomainVertexDesc.GetLayouts();
-			
-            if (DomainShader->TessellationHSOutBuffer != UINT_MAX)
-            {
-                check(DomainShader->TessellationHSOutBuffer < ML_MaxBuffers);
-                uint32 bufferIndex = DomainShader->TessellationHSOutBuffer;
-                uint32 bufferSize = VertexShader->TessellationOutputAttribs.HSOutSize;
-               
-                DomainVertexLayouts[bufferIndex].SetStride(bufferSize);
-                DomainVertexLayouts[bufferIndex].SetStepFunction(stepFunction);
-                DomainVertexLayouts[bufferIndex].SetStepRate(1);
-				
-				ns::Array<mtlpp::VertexAttributeDescriptor> Attribs = DomainVertexDesc.GetAttributes();
-				
-                for (FMetalAttribute const& Attrib : VertexShader->TessellationOutputAttribs.HSOut)
-                {
-                    int attributeIndex = Attrib.Index;
-                    check(attributeIndex >= 0 && attributeIndex <= 31);
-                    check(Attrib.Components > 0 && Attrib.Components <= 4);
-                    mtlpp::VertexFormat format = Formats[(uint8)Attrib.Type][Attrib.Components-1];
-                    check(format != mtlpp::VertexFormat::Invalid); // TODO support more cases
-                    Attribs[attributeIndex].SetFormat(format);
-                    Attribs[attributeIndex].SetOffset(Attrib.Offset);
-                    Attribs[attributeIndex].SetBufferIndex(bufferIndex);
-                }
-            }
-			
-            stepFunction = mtlpp::VertexStepFunction::PerPatchControlPoint;
-            uint32 bufferIndex = DomainShader->TessellationControlPointOutBuffer;
-            uint32 bufferSize = VertexShader->TessellationOutputAttribs.PatchControlPointOutSize;
-            
-            DomainVertexLayouts[bufferIndex].SetStride(bufferSize);
-            DomainVertexLayouts[bufferIndex].SetStepFunction(stepFunction);
-            DomainVertexLayouts[bufferIndex].SetStepRate(1);
-			
-			ns::Array<mtlpp::VertexAttributeDescriptor> DomainVertexAttribs = DomainVertexDesc.GetAttributes();
-            for (FMetalAttribute const& Attrib : VertexShader->TessellationOutputAttribs.PatchControlPointOut)
-            {
-                int attributeIndex = Attrib.Index;
-                check(attributeIndex >= 0 && attributeIndex <= 31);
-                check(Attrib.Components > 0 && Attrib.Components <= 4);
-                mtlpp::VertexFormat format = Formats[(uint8)Attrib.Type][Attrib.Components-1];
-                check(format != mtlpp::VertexFormat::Invalid); // TODO support more cases
-				DomainVertexAttribs[attributeIndex].SetFormat(format);
-				DomainVertexAttribs[attributeIndex].SetOffset(Attrib.Offset);
-				DomainVertexAttribs[attributeIndex].SetBufferIndex(bufferIndex);
-            }
-			
-			RenderPipelineDesc.SetVertexDescriptor(DomainVertexDesc);
-            
-            bool const bIsIndexed = (IndexType == EMetalIndexType_UInt16 || IndexType == EMetalIndexType_UInt32);
-            
-			mtlpp::VertexDescriptor VertexDesc = GetMaskedVertexDescriptor(VertexDecl->Layout.VertexDesc, VertexShader->Bindings.InOutMask);
-			ns::Array<mtlpp::VertexBufferLayoutDescriptor> VertexLayouts = VertexDesc.GetLayouts();
-			ns::Array<mtlpp::VertexAttributeDescriptor> VertexAttribs = VertexDesc.GetAttributes();
-			ns::Array<mtlpp::BufferLayoutDescriptor> ComputeLayouts = ComputeStageInOut.GetLayouts();
-			ns::Array<mtlpp::AttributeDescriptor> ComputeAttribs = ComputeStageInOut.GetAttributes();
-			for(int onIndex = 0; onIndex < MaxMetalStreams; onIndex++)
-            {
-                // NOTE: accessing the VertexDesc like this will end up allocating layouts/attributes
-                auto stride = VertexLayouts[onIndex].GetStride();
-                if(stride)
-                {
-                    ComputeLayouts[onIndex].SetStride(stride);
-                    auto InnerStepFunction = VertexLayouts[onIndex].GetStepFunction();
-                    switch(InnerStepFunction)
-                    {
-                        case mtlpp::VertexStepFunction::Constant:
-							ComputeLayouts[onIndex].SetStepFunction(mtlpp::StepFunction::Constant);
-                            break;
-                        case mtlpp::VertexStepFunction::PerVertex:
-                            ComputeLayouts[onIndex].SetStepFunction(bIsIndexed ? mtlpp::StepFunction::ThreadPositionInGridXIndexed : mtlpp::StepFunction::ThreadPositionInGridX);
-                            break;
-                        case mtlpp::VertexStepFunction::PerInstance:
-                            ComputeLayouts[onIndex].SetStepFunction(mtlpp::StepFunction::ThreadPositionInGridY);
-                            break;
-                        default:
-                            check(0);
-                    }
-                    ComputeLayouts[onIndex].SetStepRate(VertexLayouts[onIndex].GetStepRate());
-                }
-                auto format = VertexAttribs[onIndex].GetFormat();
-                if(format == mtlpp::VertexFormat::Invalid) continue;
-                {
-					ComputeAttribs[onIndex].SetFormat((mtlpp::AttributeFormat)format); // TODO FIXME currently these align perfectly (at least assert that is the case)
-                    ComputeAttribs[onIndex].SetOffset(VertexAttribs[onIndex].GetOffset());
-                    ComputeAttribs[onIndex].SetBufferIndex(VertexAttribs[onIndex].GetBufferIndex());
-                }
-            }
-            
-            // Disambiguated function name.
-            ComputePipelineDesc.SetComputeFunction(vertexFunction);
-            check(ComputePipelineDesc.GetComputeFunction());
-
-            // Don't set the index type if there isn't an index buffer.
-            if (IndexType != EMetalIndexType_None)
-            {
-                ComputeStageInOut.SetIndexType(GetMetalIndexType(IndexType));
-            }
-			ComputePipelineDesc.SetStageInputDescriptor(ComputeStageInOut);
-			
-            {
-				METAL_GPUPROFILE(FScopedMetalCPUStats CPUStat(FString::Printf(TEXT("NewComputePipelineState: %s"), TEXT("")/**FString([ComputePipelineDesc.GetPtr() description])*/)));
-				ns::AutoReleasedError AutoError;
-				NSUInteger ComputeOption = mtlpp::PipelineOption::NoPipelineOption;
+		RenderPipelineDesc.SetVertexDescriptor(GetMaskedVertexDescriptor(VertexDecl->Layout.VertexDesc, VertexShader->Bindings.InOutMask));
+		RenderPipelineDesc.SetVertexFunction(vertexFunction);
+		RenderPipelineDesc.SetFragmentFunction(fragmentFunction);
 #if ENABLE_METAL_GPUPROFILE
-				{
-					ns::String VertexName = vertexFunction.GetName();
-					RenderPipelineDesc.SetLabel([NSString stringWithFormat:@"%@", VertexName.GetPtr()]);
-				}
-#endif
-				if (GetMetalDeviceContext().GetCommandQueue().GetRuntimeDebuggingLevel() >= EMetalDebugLevelFastValidation)
-				{
-					mtlpp::AutoReleasedComputePipelineReflection Reflection;
-					ComputeOption = mtlpp::PipelineOption::ArgumentInfo | mtlpp::PipelineOption::BufferTypeInfo;
-					Pipeline->ComputePipelineState = Device.NewComputePipelineState(ComputePipelineDesc, (mtlpp::PipelineOption)ComputeOption, &Reflection, &AutoError);
-#if METAL_DEBUG_OPTIONS
-					Pipeline->ComputePipelineReflection = Reflection;
-#endif
-				}
-				else
-				{
-					Pipeline->ComputePipelineState = Device.NewComputePipelineState(ComputePipelineDesc, (mtlpp::PipelineOption)ComputeOption, nullptr, &AutoError);
-				}
-				Error = AutoError;
-                
-				UE_CLOG((Pipeline->ComputePipelineState == nil), LogMetal, Error, TEXT("Failed to generate a pipeline state object: %s"), *FString([Error.GetPtr() description]));
-				UE_CLOG((Pipeline->ComputePipelineState == nil), LogMetal, Error, TEXT("Vertex shader: %s"), *FString(VertexShader->GetSourceCode()));
-				UE_CLOG((Pipeline->ComputePipelineState == nil), LogMetal, Error, TEXT("Pixel shader: %s"), PixelShader ? *FString(PixelShader->GetSourceCode()) : TEXT("NULL"));
-				UE_CLOG((Pipeline->ComputePipelineState == nil), LogMetal, Error, TEXT("Hull shader: %s"), *FString(HullShader->GetSourceCode()));
-				UE_CLOG((Pipeline->ComputePipelineState == nil), LogMetal, Error, TEXT("Domain shader: %s"), *FString(DomainShader->GetSourceCode()));
-				UE_CLOG((Pipeline->ComputePipelineState == nil), LogMetal, Error, TEXT("Descriptor: %s"), *FString(ComputePipelineDesc.GetPtr().description));
-				UE_CLOG((Pipeline->ComputePipelineState == nil), LogMetal, Error, TEXT("Failed to generate a hull pipeline state object:\n\n %s\n\n"), *FString(Error.GetLocalizedDescription()));
-				
-#if METAL_DEBUG_OPTIONS
-				if (Pipeline->ComputePipelineReflection)
-				{
-					Pipeline->ComputeDesc = ComputePipelineDesc;
-					
-					bool found__HSTFOut = false;
-					for(mtlpp::Argument arg : Pipeline->ComputePipelineReflection.GetArguments())
-					{
-						bool addAttributes = false;
-						mtlpp::VertexStepFunction StepFunction = (mtlpp::VertexStepFunction)-1; // invalid
-						
-						uint32 BufferIndex = UINT_MAX;
-						
-						if([arg.GetName() isEqualToString: @"PatchControlPointOutBuffer"])
-						{
-							check((arg.GetBufferAlignment() & (arg.GetBufferAlignment() - 1)) == 0); // must be pow2
-							check((arg.GetBufferDataSize() & (arg.GetBufferAlignment() - 1)) == 0); // must be aligned
-							
-							check(arg.GetBufferDataSize() == VertexShader->TessellationOutputAttribs.PatchControlPointOutSize);
-							
-							addAttributes = true;
-							BufferIndex = DomainShader->TessellationControlPointOutBuffer;
-							StepFunction = mtlpp::VertexStepFunction::PerPatchControlPoint;
-							check(arg.GetIndex() == VertexShader->TessellationControlPointOutBuffer);
-						}
-						else if([arg.GetName() isEqualToString: @"__HSOut"])
-						{
-							check((arg.GetBufferAlignment() & (arg.GetBufferAlignment() - 1)) == 0); // must be pow2
-							check((arg.GetBufferDataSize() & (arg.GetBufferAlignment() - 1)) == 0); // must be aligned
-							
-							check(arg.GetBufferDataSize() == VertexShader->TessellationOutputAttribs.HSOutSize);
-							
-							addAttributes = true;
-							BufferIndex = DomainShader->TessellationHSOutBuffer;
-							StepFunction = mtlpp::VertexStepFunction::PerPatch;
-							check(arg.GetIndex() == VertexShader->TessellationHSOutBuffer);
-						}
-						else if([arg.GetName() isEqualToString: @"__HSTFOut"])
-						{
-							found__HSTFOut = true;
-							check((arg.GetBufferAlignment() & (arg.GetBufferAlignment() - 1)) == 0); // must be pow2
-							check((arg.GetBufferDataSize() & (arg.GetBufferAlignment() - 1)) == 0); // must be aligned
-							
-							check(arg.GetBufferDataSize() == VertexShader->TessellationOutputAttribs.HSTFOutSize);
-							
-							check(arg.GetIndex() == VertexShader->TessellationHSTFOutBuffer);
-						}
-						else if([arg.GetName() isEqualToString:@"patchCount"])
-						{
-							check(arg.GetIndex() == VertexShader->TessellationPatchCountBuffer);
-						}
-						else if([arg.GetName() isEqualToString:@"indexBuffer"])
-						{
-							check(arg.GetIndex() == VertexShader->TessellationIndexBuffer);
-						}
-						
-						// build the vertexDescriptor
-						if(addAttributes)
-						{
-							check(DomainVertexLayouts[BufferIndex].GetStride() == arg.GetBufferDataSize());
-							check(DomainVertexLayouts[BufferIndex].GetStepFunction() == StepFunction);
-							check(DomainVertexLayouts[BufferIndex].GetStepRate() == 1);
-							for(mtlpp::StructMember attribute : arg.GetBufferStructType().GetMembers())
-							{
-								int attributeIndex = -1;
-								sscanf([attribute.GetName() UTF8String], "OUT_ATTRIBUTE%d_", &attributeIndex);
-								check(attributeIndex >= 0 && attributeIndex <= 31);
-								mtlpp::VertexFormat format = mtlpp::VertexFormat::Invalid;
-								switch(attribute.GetDataType())
-								{
-									case mtlpp::DataType::Float:  format = mtlpp::VertexFormat::Float; break;
-									case mtlpp::DataType::Float2: format = mtlpp::VertexFormat::Float2; break;
-									case mtlpp::DataType::Float3: format = mtlpp::VertexFormat::Float3; break;
-									case mtlpp::DataType::Float4: format = mtlpp::VertexFormat::Float4; break;
-										
-									case mtlpp::DataType::Int:  format = mtlpp::VertexFormat::Int; break;
-									case mtlpp::DataType::Int2: format = mtlpp::VertexFormat::Int2; break;
-									case mtlpp::DataType::Int3: format = mtlpp::VertexFormat::Int3; break;
-									case mtlpp::DataType::Int4: format = mtlpp::VertexFormat::Int4; break;
-										
-									case mtlpp::DataType::UInt:  format = mtlpp::VertexFormat::UInt; break;
-									case mtlpp::DataType::UInt2: format = mtlpp::VertexFormat::UInt2; break;
-									case mtlpp::DataType::UInt3: format = mtlpp::VertexFormat::UInt3; break;
-									case mtlpp::DataType::UInt4: format = mtlpp::VertexFormat::UInt4; break;
-										
-									default: check(0); // TODO support more cases
-								}
-								check(DomainVertexAttribs[attributeIndex].GetFormat() == format);
-								check(DomainVertexAttribs[attributeIndex].GetOffset() == attribute.GetOffset());
-								check(DomainVertexAttribs[attributeIndex].GetBufferIndex() == BufferIndex);
-							}
-						}
-					}
-					check(found__HSTFOut);
-				}
-#endif
-            }
-        }
+		ns::String VertexName = vertexFunction.GetName();
+		ns::String FragmentName = fragmentFunction ? fragmentFunction.GetName() : @"";
+		RenderPipelineDesc.SetLabel([NSString stringWithFormat:@"%@+%@", VertexName.GetPtr(), FragmentName.GetPtr()]);
 #endif
 
 		NSUInteger RenderOption = mtlpp::PipelineOption::NoPipelineOption;
@@ -1728,10 +914,6 @@ static FMetalShaderPipeline* CreateMTLRenderPipeline(bool const bSync, FMetalGra
 		UE_CLOG((Pipeline->RenderPipelineState == nil), LogMetal, Error, TEXT("Failed to generate a pipeline state object: %s"), *FString(Error.GetPtr().description));
 		UE_CLOG((Pipeline->RenderPipelineState == nil), LogMetal, Error, TEXT("Vertex shader: %s"), *FString(VertexShader->GetSourceCode()));
 		UE_CLOG((Pipeline->RenderPipelineState == nil), LogMetal, Error, TEXT("Pixel shader: %s"), PixelShader ? *FString(PixelShader->GetSourceCode()) : TEXT("NULL"));
-#if PLATFORM_SUPPORTS_TESSELLATION_SHADERS
-		UE_CLOG((Pipeline->RenderPipelineState == nil), LogMetal, Error, TEXT("Hull shader: %s"), HullShader ? *FString(HullShader->GetSourceCode()) : TEXT("NULL"));
-		UE_CLOG((Pipeline->RenderPipelineState == nil), LogMetal, Error, TEXT("Domain shader: %s"), DomainShader ? *FString(DomainShader->GetSourceCode()) : TEXT("NULL"));
-#endif
 		UE_CLOG((Pipeline->RenderPipelineState == nil), LogMetal, Error, TEXT("Descriptor: %s"), *FString(RenderPipelineDesc.GetPtr().description));
 		UE_CLOG((Pipeline->RenderPipelineState == nil), LogMetal, Error, TEXT("Failed to generate a render pipeline state object:\n\n %s\n\n"), *FString(Error.GetLocalizedDescription()));
 		
@@ -1743,12 +925,7 @@ static FMetalShaderPipeline* CreateMTLRenderPipeline(bool const bSync, FMetalGra
 		}
 		
 #if METAL_DEBUG_OPTIONS
-#if PLATFORM_SUPPORTS_TESSELLATION_SHADERS
-		Pipeline->ComputeSource = DomainShader ? VertexShader->GetSourceCode() : nil;
-        Pipeline->VertexSource = DomainShader ? DomainShader->GetSourceCode() : VertexShader->GetSourceCode();
-#else
 		Pipeline->VertexSource = VertexShader->GetSourceCode();
-#endif
         Pipeline->FragmentSource = PixelShader ? PixelShader->GetSourceCode() : nil;
 #endif
 		
@@ -1789,9 +966,9 @@ static FMetalShaderPipeline* CreateMTLRenderPipeline(bool const bSync, FMetalGra
     return !bSync ? nil : Pipeline;
 }
 
-FMetalShaderPipeline* GetMTLRenderPipeline(bool const bSync, FMetalGraphicsPipelineState const* State, const FGraphicsPipelineStateInitializer& Init, EMetalIndexType const IndexType)
+FMetalShaderPipeline* GetMTLRenderPipeline(bool const bSync, FMetalGraphicsPipelineState const* State, const FGraphicsPipelineStateInitializer& Init)
 {
-	return FMetalShaderPipelineCache::Get().GetRenderPipeline(bSync, State, Init, IndexType);
+	return FMetalShaderPipelineCache::Get().GetRenderPipeline(bSync, State, Init);
 }
 
 void ReleaseMTLRenderPipeline(FMetalShaderPipeline* Pipeline)
diff --git a/Engine/Source/Runtime/Apple/MetalRHI/Private/MetalPipeline.h b/Engine/Source/Runtime/Apple/MetalRHI/Private/MetalPipeline.h
--- a/Engine/Source/Runtime/Apple/MetalRHI/Private/MetalPipeline.h
+++ b/Engine/Source/Runtime/Apple/MetalRHI/Private/MetalPipeline.h
@@ -18,7 +18,6 @@ enum EMetalPipelineHashBits
 
 	NumBits_BlendState = 7, //(x8=56),
 	NumBits_PrimitiveTopology = 2, //(x1=2)
-	NumBits_IndexType = 2, //(x1=2)
 	NumBits_AlphaToCoverage = 1, //(x1=1)
 };
 
@@ -33,8 +32,7 @@ enum EMetalPipelineHashOffsets
 	Offset_BlendState6 = Offset_BlendState5 + NumBits_BlendState,
 	Offset_BlendState7 = Offset_BlendState6 + NumBits_BlendState,
 	Offset_PrimitiveTopology = Offset_BlendState7 + NumBits_BlendState,
-	Offset_IndexType = Offset_PrimitiveTopology + NumBits_PrimitiveTopology,
-	Offset_RasterEnd = Offset_IndexType + NumBits_IndexType,
+	Offset_RasterEnd = Offset_PrimitiveTopology + NumBits_PrimitiveTopology,
 
 	Offset_RenderTargetFormat0 = 64,
 	Offset_RenderTargetFormat1 = Offset_RenderTargetFormat0 + NumBits_RenderTargetFormat,
@@ -68,24 +66,6 @@ private:
 	void OnShaderPipelineCachePrecompilationComplete(uint32 Count, double Seconds, const FShaderPipelineCache::FShaderCachePrecompileContext& ShaderCachePrecompileContext);
 };
 
-struct FMetalTessellationPipelineDesc
-{
-	FMetalTessellationPipelineDesc()
-	{
-	}
-	NSUInteger TessellationInputControlPointBufferIndex;
-	NSUInteger TessellationOutputControlPointBufferIndex;
-	NSUInteger TessellationPatchControlPointOutSize;
-	NSUInteger TessellationPatchConstBufferIndex;
-	NSUInteger TessellationInputPatchConstBufferIndex;
-	NSUInteger TessellationPatchConstOutSize;
-	NSUInteger TessellationTessFactorOutSize;
-	NSUInteger TessellationFactorBufferIndex;
-	NSUInteger TessellationPatchCountBufferIndex;
-	NSUInteger TessellationControlPointIndexBufferIndex;
-	NSUInteger TessellationIndexBufferIndex;
-};
-
 @interface FMetalShaderPipeline : FApplePlatformObject
 {
 @public
@@ -93,7 +73,6 @@ struct FMetalTessellationPipelineDesc
 	mtlpp::ComputePipelineState ComputePipelineState;
 	mtlpp::RenderPipelineState StreamPipelineState;
 	mtlpp::RenderPipelineState DebugPipelineState;
-	FMetalTessellationPipelineDesc TessellationPipelineDesc;
 	TArray<uint32> BufferDataSizes[EMetalShaderStagesNum];
 	TMap<uint8, uint8> TextureTypes[EMetalShaderStagesNum];
 	FMetalDebugShaderResourceMask ResourceMask[EMetalShaderStagesNum];
@@ -102,7 +81,6 @@ struct FMetalTessellationPipelineDesc
 	mtlpp::ComputePipelineReflection ComputePipelineReflection;
 #if METAL_DEBUG_OPTIONS
 	ns::String VertexSource;
-	ns::String DomainSource;
 	ns::String FragmentSource;
 	ns::String ComputeSource;
 	mtlpp::RenderPipelineDescriptor RenderDesc;
diff --git a/Engine/Source/Runtime/Apple/MetalRHI/Private/MetalProfiler.cpp b/Engine/Source/Runtime/Apple/MetalRHI/Private/MetalProfiler.cpp
--- a/Engine/Source/Runtime/Apple/MetalRHI/Private/MetalProfiler.cpp
+++ b/Engine/Source/Runtime/Apple/MetalRHI/Private/MetalProfiler.cpp
@@ -20,12 +20,10 @@ DEFINE_STAT(STAT_MetalPipelineStateTime);
 DEFINE_STAT(STAT_MetalPrepareDrawTime);
 
 DEFINE_STAT(STAT_MetalSwitchToRenderTime);
-DEFINE_STAT(STAT_MetalSwitchToTessellationTime);
 DEFINE_STAT(STAT_MetalSwitchToComputeTime);
 DEFINE_STAT(STAT_MetalSwitchToBlitTime);
 DEFINE_STAT(STAT_MetalSwitchToAsyncBlitTime);
 DEFINE_STAT(STAT_MetalPrepareToRenderTime);
-DEFINE_STAT(STAT_MetalPrepareToTessellateTime);
 DEFINE_STAT(STAT_MetalPrepareToDispatchTime);
 DEFINE_STAT(STAT_MetalCommitRenderResourceTablesTime);
 DEFINE_STAT(STAT_MetalSetRenderStateTime);
diff --git a/Engine/Source/Runtime/Apple/MetalRHI/Private/MetalProfiler.h b/Engine/Source/Runtime/Apple/MetalRHI/Private/MetalProfiler.h
--- a/Engine/Source/Runtime/Apple/MetalRHI/Private/MetalProfiler.h
+++ b/Engine/Source/Runtime/Apple/MetalRHI/Private/MetalProfiler.h
@@ -11,12 +11,10 @@ DECLARE_CYCLE_STAT_EXTERN(TEXT("MakeDrawable time"),STAT_MetalMakeDrawableTime,S
 DECLARE_CYCLE_STAT_EXTERN(TEXT("Draw call time"),STAT_MetalDrawCallTime,STATGROUP_MetalRHI, );
 DECLARE_CYCLE_STAT_EXTERN(TEXT("PrepareDraw time"),STAT_MetalPrepareDrawTime,STATGROUP_MetalRHI, );
 DECLARE_CYCLE_STAT_EXTERN(TEXT("SwitchToRender time"),STAT_MetalSwitchToRenderTime,STATGROUP_MetalRHI, );
-DECLARE_CYCLE_STAT_EXTERN(TEXT("SwitchToTessellation time"),STAT_MetalSwitchToTessellationTime,STATGROUP_MetalRHI, );
 DECLARE_CYCLE_STAT_EXTERN(TEXT("SwitchToCompute time"),STAT_MetalSwitchToComputeTime,STATGROUP_MetalRHI, );
 DECLARE_CYCLE_STAT_EXTERN(TEXT("SwitchToBlit time"),STAT_MetalSwitchToBlitTime,STATGROUP_MetalRHI, );
 DECLARE_CYCLE_STAT_EXTERN(TEXT("SwitchToAsyncBlit time"),STAT_MetalSwitchToAsyncBlitTime,STATGROUP_MetalRHI, );
 DECLARE_CYCLE_STAT_EXTERN(TEXT("PrepareToRender time"),STAT_MetalPrepareToRenderTime,STATGROUP_MetalRHI, );
-DECLARE_CYCLE_STAT_EXTERN(TEXT("PrepareToTessellate time"),STAT_MetalPrepareToTessellateTime,STATGROUP_MetalRHI, );
 DECLARE_CYCLE_STAT_EXTERN(TEXT("PrepareToDispatch time"),STAT_MetalPrepareToDispatchTime,STATGROUP_MetalRHI, );
 DECLARE_CYCLE_STAT_EXTERN(TEXT("CommitRenderResourceTables time"),STAT_MetalCommitRenderResourceTablesTime,STATGROUP_MetalRHI, );
 DECLARE_CYCLE_STAT_EXTERN(TEXT("SetRenderState time"),STAT_MetalSetRenderStateTime,STATGROUP_MetalRHI, );
diff --git a/Engine/Source/Runtime/Apple/MetalRHI/Private/MetalRHI.cpp b/Engine/Source/Runtime/Apple/MetalRHI/Private/MetalRHI.cpp
--- a/Engine/Source/Runtime/Apple/MetalRHI/Private/MetalRHI.cpp
+++ b/Engine/Source/Runtime/Apple/MetalRHI/Private/MetalRHI.cpp
@@ -366,13 +366,8 @@ FMetalDynamicRHI::FMetalDynamicRHI(ERHIFeatureLevel::Type RequestedFeatureLevel)
 	}
 	else
 	{
-		if (bRequestedSM5)
-		{
-			UE_LOG(LogMetal, Warning, TEXT("Metal Shader Model 5 w/tessellation support requires 10.12.6 for Nvidia, it is broken on 10.13.0+. Falling back to Metal Shader Model 5 without tessellation support."));
-		}
-	
 		GMaxRHIFeatureLevel = ERHIFeatureLevel::SM5;
-		GMaxRHIShaderPlatform = SP_METAL_SM5_NOTESS;
+		GMaxRHIShaderPlatform = SP_METAL_SM5;
 	}
 
 	ERHIFeatureLevel::Type PreviewFeatureLevel;
@@ -423,7 +418,7 @@ FMetalDynamicRHI::FMetalDynamicRHI(ERHIFeatureLevel::Type RequestedFeatureLevel)
 	}
 	
 	// Disable the distance field AO & shadowing effects on GPU drivers that don't currently execute the shaders correctly.
-	if ((GMaxRHIShaderPlatform == SP_METAL_SM5 || GMaxRHIShaderPlatform == SP_METAL_SM5_NOTESS) && !bSupportsDistanceFields && !FParse::Param(FCommandLine::Get(),TEXT("metaldistancefields")))
+	if ((GMaxRHIShaderPlatform == SP_METAL_SM5) && !bSupportsDistanceFields && !FParse::Param(FCommandLine::Get(),TEXT("metaldistancefields")))
 	{
 		static auto CVarDistanceFieldAO = IConsoleManager::Get().FindConsoleVariable(TEXT("r.DistanceFieldAO"));
 		if(CVarDistanceFieldAO && CVarDistanceFieldAO->GetInt() != 0)
diff --git a/Engine/Source/Runtime/Apple/MetalRHI/Private/MetalRHIPrivate.h b/Engine/Source/Runtime/Apple/MetalRHI/Private/MetalRHIPrivate.h
--- a/Engine/Source/Runtime/Apple/MetalRHI/Private/MetalRHIPrivate.h
+++ b/Engine/Source/Runtime/Apple/MetalRHI/Private/MetalRHIPrivate.h
@@ -146,12 +146,6 @@ struct FMetalDebugInfo
 // Get a compute pipeline state used to implement some debug features.
 mtlpp::ComputePipelineState GetMetalDebugComputeState();
 
-// Get the copy 32-bit indices function
-mtlpp::ComputePipelineState GetMetalCopyIndex32Function();
-
-// Get the copy 16-bit indices function
-mtlpp::ComputePipelineState GetMetalCopyIndex16Function();
-
 // Access the internal context for the device-owning DynamicRHI object
 FMetalDeviceContext& GetMetalDeviceContext();
 
@@ -189,31 +183,6 @@ FORCEINLINE bool MetalIsSafeToUseRHIThreadResources()
 			(IsInRenderingThread() && (!IsRunningRHIInSeparateThread() || !FRHICommandListExecutor::IsRHIThreadActive() || FRHICommandListImmediate::IsStalled() || FRHICommandListExecutor::IsRHIThreadCompletelyFlushed()));
 }
 
-FORCEINLINE mtlpp::IndexType GetMetalIndexType(EMetalIndexType IndexType)
-{
-	switch (IndexType)
-	{
-		case EMetalIndexType_UInt16: return mtlpp::IndexType::UInt16;
-		case EMetalIndexType_UInt32: return mtlpp::IndexType::UInt32;
-		case EMetalIndexType_None:
-		default:
-		{
-			METAL_FATAL_ERROR(TEXT("There is not equivalent mtlpp::IndexType for %d"), (uint32)IndexType);
-			return mtlpp::IndexType::UInt16;
-		}
-	}
-}
-
-FORCEINLINE EMetalIndexType GetRHIMetalIndexType(mtlpp::IndexType IndexType)
-{
-	switch (IndexType)
-	{
-		case mtlpp::IndexType::UInt16: return EMetalIndexType_UInt16;
-		case mtlpp::IndexType::UInt32: return EMetalIndexType_UInt32;
-		default: return EMetalIndexType_None;
-	}
-}
-
 FORCEINLINE int32 GetMetalCubeFace(ECubeFace Face)
 {
 	// According to Metal docs these should match now: https://developer.apple.com/library/prerelease/ios/documentation/Metal/Reference/MTLTexture_Ref/index.html#//apple_ref/c/tdef/MTLTextureType
@@ -266,10 +235,6 @@ mtlpp::LanguageVersion ValidateVersion(uint8 Version);
 enum EMetalShaderStages
 {
 	Vertex,
-#if PLATFORM_SUPPORTS_TESSELLATION_SHADERS
-	Hull,
-	Domain,
-#endif
 	Pixel,
 #if PLATFORM_SUPPORTS_GEOMETRY_SHADERS
 	Geometry,
@@ -285,12 +250,6 @@ FORCEINLINE EShaderFrequency GetRHIShaderFrequency(EMetalShaderStages Stage)
 	{
 		case EMetalShaderStages::Vertex:
 			return SF_Vertex;
-#if PLATFORM_SUPPORTS_TESSELLATION_SHADERS
-		case EMetalShaderStages::Hull:
-			return SF_Hull;
-		case EMetalShaderStages::Domain:
-			return SF_Domain;
-#endif
 		case EMetalShaderStages::Pixel:
 			return SF_Pixel;
 #if PLATFORM_SUPPORTS_GEOMETRY_SHADERS
@@ -310,12 +269,6 @@ FORCEINLINE EMetalShaderStages GetMetalShaderFrequency(EShaderFrequency Stage)
 	{
 		case SF_Vertex:
 			return EMetalShaderStages::Vertex;
-#if PLATFORM_SUPPORTS_TESSELLATION_SHADERS
-		case SF_Hull:
-			return EMetalShaderStages::Hull;
-		case SF_Domain:
-			return EMetalShaderStages::Domain;
-#endif
 		case SF_Pixel:
 			return EMetalShaderStages::Pixel;
 #if PLATFORM_SUPPORTS_GEOMETRY_SHADERS
diff --git a/Engine/Source/Runtime/Apple/MetalRHI/Private/MetalRenderCommandEncoder.cpp b/Engine/Source/Runtime/Apple/MetalRHI/Private/MetalRenderCommandEncoder.cpp
--- a/Engine/Source/Runtime/Apple/MetalRHI/Private/MetalRenderCommandEncoder.cpp
+++ b/Engine/Source/Runtime/Apple/MetalRHI/Private/MetalRenderCommandEncoder.cpp
@@ -584,102 +584,6 @@ void FMetalRenderCommandEncoderDebugging::DrawIndexed(mtlpp::PrimitiveType primi
 #endif
 }
 
-void FMetalRenderCommandEncoderDebugging::SetTessellationFactorBuffer( FMetalBuffer const& buffer, NSUInteger offset, NSUInteger instanceStride)
-{
-}
-
-void FMetalRenderCommandEncoderDebugging::DrawPatches(NSUInteger numberOfPatchControlPoints, NSUInteger patchStart, NSUInteger patchCount,  FMetalBuffer const& patchIndexBuffer, NSUInteger patchIndexBufferOffset, NSUInteger instanceCount, NSUInteger baseInstance)
-{
-#if METAL_DEBUG_OPTIONS
-	switch(((FMetalDebugRenderCommandEncoder*)m_ptr)->Buffer.GetPtr()->DebugLevel)
-	{
-		case EMetalDebugLevelConditionalSubmit:
-		case EMetalDebugLevelWaitForComplete:
-		case EMetalDebugLevelLogOperations:
-		{
-			((FMetalDebugRenderCommandEncoder*)m_ptr)->Buffer.Draw([NSString stringWithFormat:@"%s", __PRETTY_FUNCTION__]);
-		}
-		case EMetalDebugLevelValidation:
-		{
-			Validate();
-		}
-		default:
-		{
-			break;
-		}
-	}
-#endif
-}
-
-void FMetalRenderCommandEncoderDebugging::DrawPatches(NSUInteger numberOfPatchControlPoints,  FMetalBuffer const& patchIndexBuffer, NSUInteger patchIndexBufferOffset, FMetalBuffer const& indirectBuffer, NSUInteger indirectBufferOffset)
-{
-#if METAL_DEBUG_OPTIONS
-	switch(((FMetalDebugRenderCommandEncoder*)m_ptr)->Buffer.GetPtr()->DebugLevel)
-	{
-		case EMetalDebugLevelConditionalSubmit:
-		case EMetalDebugLevelWaitForComplete:
-		case EMetalDebugLevelLogOperations:
-		{
-			((FMetalDebugRenderCommandEncoder*)m_ptr)->Buffer.Draw([NSString stringWithFormat:@"%s", __PRETTY_FUNCTION__]);
-		}
-		case EMetalDebugLevelValidation:
-		{
-			Validate();
-		}
-		default:
-		{
-			break;
-		}
-	}
-#endif
-}
-
-void FMetalRenderCommandEncoderDebugging::DrawIndexedPatches(NSUInteger numberOfPatchControlPoints, NSUInteger patchStart, NSUInteger patchCount,  FMetalBuffer const& patchIndexBuffer, NSUInteger patchIndexBufferOffset, FMetalBuffer const& controlPointIndexBuffer, NSUInteger controlPointIndexBufferOffset, NSUInteger instanceCount, NSUInteger baseInstance)
-{
-#if METAL_DEBUG_OPTIONS
-	switch(((FMetalDebugRenderCommandEncoder*)m_ptr)->Buffer.GetPtr()->DebugLevel)
-	{
-		case EMetalDebugLevelConditionalSubmit:
-		case EMetalDebugLevelWaitForComplete:
-		case EMetalDebugLevelLogOperations:
-		{
-			((FMetalDebugRenderCommandEncoder*)m_ptr)->Buffer.Draw([NSString stringWithFormat:@"%s", __PRETTY_FUNCTION__]);
-		}
-		case EMetalDebugLevelValidation:
-		{
-			Validate();
-		}
-		default:
-		{
-			break;
-		}
-	}
-#endif
-}
-
-void FMetalRenderCommandEncoderDebugging::DrawIndexedPatches(NSUInteger numberOfPatchControlPoints,  FMetalBuffer const& patchIndexBuffer, NSUInteger patchIndexBufferOffset, FMetalBuffer const& controlPointIndexBuffer, NSUInteger controlPointIndexBufferOffset, FMetalBuffer const& indirectBuffer, NSUInteger indirectBufferOffset)
-{
-#if METAL_DEBUG_OPTIONS
-	switch(((FMetalDebugRenderCommandEncoder*)m_ptr)->Buffer.GetPtr()->DebugLevel)
-	{
-		case EMetalDebugLevelConditionalSubmit:
-		case EMetalDebugLevelWaitForComplete:
-		case EMetalDebugLevelLogOperations:
-		{
-			((FMetalDebugRenderCommandEncoder*)m_ptr)->Buffer.Draw([NSString stringWithFormat:@"%s", __PRETTY_FUNCTION__]);
-		}
-		case EMetalDebugLevelValidation:
-		{
-			Validate();
-		}
-		default:
-		{
-			break;
-		}
-	}
-#endif
-}
-
 /** Validates the pipeline/binding state */
 bool FMetalRenderCommandEncoderDebugging::ValidateFunctionBindings(EMetalShaderFrequency Frequency)
 {
diff --git a/Engine/Source/Runtime/Apple/MetalRHI/Private/MetalRenderCommandEncoder.h b/Engine/Source/Runtime/Apple/MetalRHI/Private/MetalRenderCommandEncoder.h
--- a/Engine/Source/Runtime/Apple/MetalRHI/Private/MetalRenderCommandEncoder.h
+++ b/Engine/Source/Runtime/Apple/MetalRHI/Private/MetalRenderCommandEncoder.h
@@ -47,16 +47,6 @@ public:
 	
 	void DrawIndexed(mtlpp::PrimitiveType primitiveType, mtlpp::IndexType indexType, FMetalBuffer const& indexBuffer, NSUInteger indexBufferOffset, FMetalBuffer const& indirectBuffer, NSUInteger indirectBufferOffset);
 	
-	void SetTessellationFactorBuffer( FMetalBuffer const& buffer, NSUInteger offset, NSUInteger instanceStride);
-	
-	void DrawPatches(NSUInteger numberOfPatchControlPoints, NSUInteger patchStart, NSUInteger patchCount,  FMetalBuffer const& patchIndexBuffer, NSUInteger patchIndexBufferOffset, NSUInteger instanceCount, NSUInteger baseInstance);
-	
-	void DrawPatches(NSUInteger numberOfPatchControlPoints,  FMetalBuffer const& patchIndexBuffer, NSUInteger patchIndexBufferOffset, FMetalBuffer const& indirectBuffer, NSUInteger indirectBufferOffset);
-	
-	void DrawIndexedPatches(NSUInteger numberOfPatchControlPoints, NSUInteger patchStart, NSUInteger patchCount,  FMetalBuffer const& patchIndexBuffer, NSUInteger patchIndexBufferOffset, FMetalBuffer const& controlPointIndexBuffer, NSUInteger controlPointIndexBufferOffset, NSUInteger instanceCount, NSUInteger baseInstance);
-	
-	void DrawIndexedPatches(NSUInteger numberOfPatchControlPoints,  FMetalBuffer const& patchIndexBuffer, NSUInteger patchIndexBufferOffset, FMetalBuffer const& controlPointIndexBuffer, NSUInteger controlPointIndexBufferOffset, FMetalBuffer const& indirectBuffer, NSUInteger indirectBufferOffset);
-	
 	/** Validates the pipeline/binding state */
 	bool ValidateFunctionBindings(EMetalShaderFrequency Frequency);
 	void Validate();
diff --git a/Engine/Source/Runtime/Apple/MetalRHI/Private/MetalRenderPass.cpp b/Engine/Source/Runtime/Apple/MetalRHI/Private/MetalRenderPass.cpp
--- a/Engine/Source/Runtime/Apple/MetalRHI/Private/MetalRenderPass.cpp
+++ b/Engine/Source/Runtime/Apple/MetalRHI/Private/MetalRenderPass.cpp
@@ -22,24 +22,6 @@ static FAutoConsoleVariableRef CVarMetalCommandBufferCommitThreshold(
 	GMetalCommandBufferCommitThreshold,
 	TEXT("When enabled (> 0) if the command buffer has more than this number of draw/dispatch command encoded then it will be committed at the next encoder boundary to keep the GPU busy. (Default: 0, set to <= 0 to disable)"));
 
-static int32 GMetalTessellationRunTessellationStage = 1;
-static FAutoConsoleVariableRef CVarMetalTessellationRunTessellationStage(
-	TEXT("rhi.Metal.RunTessellationStage"),
-	GMetalTessellationRunTessellationStage,
-	TEXT("Whether to run the VS+HS tessellation stage when performing tessellated draw calls in Metal or not. (Default: 1)"));
-
-static int32 GMetalTessellationRunDomainStage = 1;
-static FAutoConsoleVariableRef CVarMetalTessellationRunDomainStage(
-	TEXT("rhi.Metal.RunDomainStage"),
-	GMetalTessellationRunDomainStage,
-	TEXT("Whether to run the DS+PS domain stage when performing tessellated draw calls in Metal or not. (Default: 1)"));
-
-static int32 GMetalTessellationFlattenTriangleFactors = 0;
-static FAutoConsoleVariableRef CVarMetalTessellationFlattenTriangleFactors(
-	TEXT("rhi.Metal.FlattenTriangleFactors"),
-	GMetalTessellationFlattenTriangleFactors,
-	TEXT("Whether to flatten the tessellation factors to 1 when performing tessellated draw calls in Metal or not - used to debug tessellation factor generation. (Default: 0)"));
-
 static int32 GMetalDeferRenderPasses = 1;
 static FAutoConsoleVariableRef CVarMetalDeferRenderPasses(
 	TEXT("rhi.Metal.DeferRenderPasses"),
@@ -424,42 +406,31 @@ void FMetalRenderPass::DrawPrimitive(uint32 PrimitiveType, uint32 BaseVertexInde
 {
 	NumInstances = FMath::Max(NumInstances,1u);
 	
-#if PLATFORM_SUPPORTS_TESSELLATION_SHADERS
-	if(!State.GetUsingTessellation())
-#endif
-	{
-		ConditionalSwitchToRender();
-		check(CurrentEncoder.GetCommandBuffer());
-		check(CurrentEncoder.IsRenderCommandEncoderActive());
-		
-		PrepareToRender(PrimitiveType);
+	ConditionalSwitchToRender();
+	check(CurrentEncoder.GetCommandBuffer());
+	check(CurrentEncoder.IsRenderCommandEncoderActive());
 	
-		// draw!
-		// how many verts to render
-		uint32 NumVertices = GetVertexCountForPrimitiveCount(NumPrimitives, PrimitiveType);
-		
-		METAL_GPUPROFILE(FMetalProfiler::GetProfiler()->EncodeDraw(CurrentEncoder.GetCommandBufferStats(), __FUNCTION__, NumPrimitives, NumVertices, NumInstances));
-		CurrentEncoder.GetRenderCommandEncoder().Draw(TranslatePrimitiveType(PrimitiveType), BaseVertexIndex, NumVertices, NumInstances);
-		METAL_DEBUG_LAYER(EMetalDebugLevelFastValidation, CurrentEncoder.GetRenderCommandEncoderDebugging().Draw(TranslatePrimitiveType(PrimitiveType), BaseVertexIndex, NumVertices, NumInstances));
+	PrepareToRender(PrimitiveType);
 
-		if (GMetalCommandBufferDebuggingEnabled)
-		{
-			FMetalCommandData Data;
-			Data.CommandType = FMetalCommandData::Type::DrawPrimitive;
-			Data.Draw.BaseInstance = 0;
-			Data.Draw.InstanceCount = NumInstances;
-			Data.Draw.VertexCount = NumVertices;
-			Data.Draw.VertexStart = BaseVertexIndex;
-			
-			InsertDebugDraw(Data);
-		}
-	}
-#if PLATFORM_SUPPORTS_TESSELLATION_SHADERS
-	else
+	// draw!
+	// how many verts to render
+	uint32 NumVertices = GetVertexCountForPrimitiveCount(NumPrimitives, PrimitiveType);
+	
+	METAL_GPUPROFILE(FMetalProfiler::GetProfiler()->EncodeDraw(CurrentEncoder.GetCommandBufferStats(), __FUNCTION__, NumPrimitives, NumVertices, NumInstances));
+	CurrentEncoder.GetRenderCommandEncoder().Draw(TranslatePrimitiveType(PrimitiveType), BaseVertexIndex, NumVertices, NumInstances);
+	METAL_DEBUG_LAYER(EMetalDebugLevelFastValidation, CurrentEncoder.GetRenderCommandEncoderDebugging().Draw(TranslatePrimitiveType(PrimitiveType), BaseVertexIndex, NumVertices, NumInstances));
+
+	if (GMetalCommandBufferDebuggingEnabled)
 	{
-		DrawPatches(PrimitiveType, nullptr, 0, BaseVertexIndex, 0, 0, NumPrimitives, NumInstances);
+		FMetalCommandData Data;
+		Data.CommandType = FMetalCommandData::Type::DrawPrimitive;
+		Data.Draw.BaseInstance = 0;
+		Data.Draw.InstanceCount = NumInstances;
+		Data.Draw.VertexCount = NumVertices;
+		Data.Draw.VertexStart = BaseVertexIndex;
+		
+		InsertDebugDraw(Data);
 	}
-#endif
 	
 	ConditionalSubmit();	
 }
@@ -547,49 +518,38 @@ void FMetalRenderPass::DrawIndexedPrimitive(FMetalBuffer const& IndexBuffer, uin
 	}
 #endif
 	
-#if PLATFORM_SUPPORTS_TESSELLATION_SHADERS
-	if (!State.GetUsingTessellation())
-#endif
+	ConditionalSwitchToRender();
+	check(CurrentEncoder.GetCommandBuffer());
+	check(CurrentEncoder.IsRenderCommandEncoderActive());
+	
+	PrepareToRender(PrimitiveType);
+	
+	uint32 NumIndices = GetVertexCountForPrimitiveCount(NumPrimitives, PrimitiveType);
+	
+	METAL_GPUPROFILE(FMetalProfiler::GetProfiler()->EncodeDraw(CurrentEncoder.GetCommandBufferStats(), __FUNCTION__, NumPrimitives, NumVertices, NumInstances));
+	if (GRHISupportsBaseVertexIndex && GRHISupportsFirstInstance)
 	{
-		ConditionalSwitchToRender();
-		check(CurrentEncoder.GetCommandBuffer());
-		check(CurrentEncoder.IsRenderCommandEncoderActive());
-		
-		PrepareToRender(PrimitiveType);
-		
-		uint32 NumIndices = GetVertexCountForPrimitiveCount(NumPrimitives, PrimitiveType);
-		
-		METAL_GPUPROFILE(FMetalProfiler::GetProfiler()->EncodeDraw(CurrentEncoder.GetCommandBufferStats(), __FUNCTION__, NumPrimitives, NumVertices, NumInstances));
-		if (GRHISupportsBaseVertexIndex && GRHISupportsFirstInstance)
-		{
-			CurrentEncoder.GetRenderCommandEncoder().DrawIndexed(TranslatePrimitiveType(PrimitiveType), NumIndices, ((IndexStride == 2) ? mtlpp::IndexType::UInt16 : mtlpp::IndexType::UInt32), IndexBuffer, StartIndex * IndexStride, NumInstances, BaseVertexIndex, FirstInstance);
-			METAL_DEBUG_LAYER(EMetalDebugLevelFastValidation, CurrentEncoder.GetRenderCommandEncoderDebugging().DrawIndexed(TranslatePrimitiveType(PrimitiveType), NumIndices, ((IndexStride == 2) ? mtlpp::IndexType::UInt16 : mtlpp::IndexType::UInt32), IndexBuffer, StartIndex * IndexStride, NumInstances, BaseVertexIndex, FirstInstance));
-		}
-		else
-		{
-			CurrentEncoder.GetRenderCommandEncoder().DrawIndexed(TranslatePrimitiveType(PrimitiveType), NumIndices, ((IndexStride == 2) ? mtlpp::IndexType::UInt16 : mtlpp::IndexType::UInt32), IndexBuffer, StartIndex * IndexStride, NumInstances);
-			METAL_DEBUG_LAYER(EMetalDebugLevelFastValidation, CurrentEncoder.GetRenderCommandEncoderDebugging().DrawIndexed(TranslatePrimitiveType(PrimitiveType), NumIndices, ((IndexStride == 2) ? mtlpp::IndexType::UInt16 : mtlpp::IndexType::UInt32), IndexBuffer, StartIndex * IndexStride, NumInstances));
-		}
-		
-		if (GMetalCommandBufferDebuggingEnabled)
-		{
-			FMetalCommandData Data;
-			Data.CommandType = FMetalCommandData::Type::DrawPrimitiveIndexed;
-			Data.DrawIndexed.BaseInstance = FirstInstance;
-			Data.DrawIndexed.BaseVertex = BaseVertexIndex;
-			Data.DrawIndexed.IndexCount = NumIndices;
-			Data.DrawIndexed.IndexStart = StartIndex;
-			Data.DrawIndexed.InstanceCount = NumInstances;
-			
-			InsertDebugDraw(Data);
-		}
+		CurrentEncoder.GetRenderCommandEncoder().DrawIndexed(TranslatePrimitiveType(PrimitiveType), NumIndices, ((IndexStride == 2) ? mtlpp::IndexType::UInt16 : mtlpp::IndexType::UInt32), IndexBuffer, StartIndex * IndexStride, NumInstances, BaseVertexIndex, FirstInstance);
+		METAL_DEBUG_LAYER(EMetalDebugLevelFastValidation, CurrentEncoder.GetRenderCommandEncoderDebugging().DrawIndexed(TranslatePrimitiveType(PrimitiveType), NumIndices, ((IndexStride == 2) ? mtlpp::IndexType::UInt16 : mtlpp::IndexType::UInt32), IndexBuffer, StartIndex * IndexStride, NumInstances, BaseVertexIndex, FirstInstance));
 	}
-#if PLATFORM_SUPPORTS_TESSELLATION_SHADERS
 	else
 	{
-		DrawPatches(PrimitiveType, IndexBuffer, IndexStride, BaseVertexIndex, FirstInstance, StartIndex, NumPrimitives, NumInstances);
+		CurrentEncoder.GetRenderCommandEncoder().DrawIndexed(TranslatePrimitiveType(PrimitiveType), NumIndices, ((IndexStride == 2) ? mtlpp::IndexType::UInt16 : mtlpp::IndexType::UInt32), IndexBuffer, StartIndex * IndexStride, NumInstances);
+		METAL_DEBUG_LAYER(EMetalDebugLevelFastValidation, CurrentEncoder.GetRenderCommandEncoderDebugging().DrawIndexed(TranslatePrimitiveType(PrimitiveType), NumIndices, ((IndexStride == 2) ? mtlpp::IndexType::UInt16 : mtlpp::IndexType::UInt32), IndexBuffer, StartIndex * IndexStride, NumInstances));
+	}
+	
+	if (GMetalCommandBufferDebuggingEnabled)
+	{
+		FMetalCommandData Data;
+		Data.CommandType = FMetalCommandData::Type::DrawPrimitiveIndexed;
+		Data.DrawIndexed.BaseInstance = FirstInstance;
+		Data.DrawIndexed.BaseVertex = BaseVertexIndex;
+		Data.DrawIndexed.IndexCount = NumIndices;
+		Data.DrawIndexed.IndexStart = StartIndex;
+		Data.DrawIndexed.InstanceCount = NumInstances;
+		
+		InsertDebugDraw(Data);
 	}
-#endif
 	
 	ConditionalSubmit();
 }
@@ -668,508 +628,6 @@ void FMetalRenderPass::DrawIndexedPrimitiveIndirect(uint32 PrimitiveType,FMetalI
 	}
 }
 
-#if PLATFORM_SUPPORTS_TESSELLATION_SHADERS
-extern mtlpp::ComputePipelineState GetMetalFlattenTessState();
-
-void FMetalRenderPass::DrawPatches(uint32 PrimitiveType,FMetalBuffer const& IndexBuffer, uint32 IndexBufferStride, int32 BaseVertexIndex, uint32 FirstInstance, uint32 StartIndex,
-									uint32 NumPrimitives, uint32 NumInstances)
-{
-	if (GetMetalDeviceContext().SupportsFeature(EMetalFeaturesTessellation))
-	{
-		if (FMetalCommandQueue::SupportsFeature(EMetalFeaturesSeparateTessellation))
-		{
-			FMetalDeviceContext& deviceContext = (FMetalDeviceContext&)GetMetalDeviceContext();
-			mtlpp::Device device = deviceContext.GetDevice();
-			
-			FMetalGraphicsPipelineState* boundShaderState = State.GetGraphicsPSO();
-			FMetalShaderPipeline* Pipeline = State.GetPipelineState();
-			
-			ConditionalSwitchToSeparateTessellation();
-			check(CurrentEncoder.GetCommandBuffer());
-			check(PrologueEncoder.GetCommandBuffer());
-			check(PrologueEncoder.IsRenderCommandEncoderActive());
-			check(CurrentEncoder.IsRenderCommandEncoderActive());
-			
-			mtlpp::RenderCommandEncoder& StreamOutEncoder = PrologueEncoder.GetRenderCommandEncoder();
-			
-			PrepareToStreamOut(PrimitiveType);
-			
-			FMetalBuffer ShaderOutputBuffer = nil;
-			if (IndexBuffer)
-			{
-				uint32 NumIndices = GetVertexCountForPrimitiveCount(NumPrimitives, PrimitiveType);
-				
-				uint32 OutputBufferSize = (NumIndices * NumInstances) * boundShaderState->VertexShader->TessellationOutputAttribs.HSOutSize;
-				
-				ShaderOutputBuffer = deviceContext.CreatePooledBuffer(FMetalPooledBufferArgs(device, OutputBufferSize, BUF_Dynamic, mtlpp::StorageMode::Private));
-				
-				PrologueEncoder.SetShaderBuffer(mtlpp::FunctionType::Vertex, ShaderOutputBuffer, 0, ShaderOutputBuffer, boundShaderState->VertexShader->TessellationHSOutBuffer, mtlpp::ResourceUsage::Write);
-				State.SetShaderBuffer(EMetalShaderStages::Vertex, nil, nil, 0, 0, boundShaderState->VertexShader->TessellationHSOutBuffer, mtlpp::ResourceUsage(0));
-				
-				// set the patchCount
-				FMetalSubBufferRing& Ring = PrologueEncoder.GetRingBuffer();
-				FMetalBuffer PatchData = Ring.NewBuffer(sizeof(NumIndices), sizeof(NumIndices));
-				uint32* patchCountData = (uint32*)PatchData.GetContents();
-				patchCountData[0] = NumIndices;
-				PrologueEncoder.SetShaderBuffer(mtlpp::FunctionType::Vertex, PatchData, 0, PatchData, boundShaderState->VertexShader->TessellationPatchCountBuffer, mtlpp::ResourceUsage::Read);
-				State.SetShaderBuffer(EMetalShaderStages::Vertex, nil, nil, 0, 0, boundShaderState->VertexShader->TessellationPatchCountBuffer, mtlpp::ResourceUsage(0));
-				
-				if (boundShaderState->VertexShader->SideTableBinding >= 0)
-				{
-					PrologueEncoder.SetShaderSideTable(mtlpp::FunctionType::Vertex, boundShaderState->VertexShader->SideTableBinding);
-					State.SetShaderBuffer(EMetalShaderStages::Vertex, nil, nil, 0, 0, boundShaderState->VertexShader->SideTableBinding, mtlpp::ResourceUsage(0));
-				}
-				
-				StreamOutEncoder.DrawIndexed(TranslatePrimitiveType(PrimitiveType), NumIndices, ((IndexBufferStride == 2) ? mtlpp::IndexType::UInt16 : mtlpp::IndexType::UInt32), IndexBuffer, IndexBufferStride * StartIndex, NumInstances, BaseVertexIndex, FirstInstance);
-				METAL_DEBUG_LAYER(EMetalDebugLevelFastValidation, PrologueEncoder.GetRenderCommandEncoderDebugging().DrawIndexed(TranslatePrimitiveType(PrimitiveType), NumIndices, ((IndexBufferStride == 2) ? mtlpp::IndexType::UInt16 : mtlpp::IndexType::UInt32), IndexBuffer, IndexBufferStride * StartIndex, NumInstances, BaseVertexIndex, FirstInstance));
-			}
-			else
-			{
-				uint32 NumVertices = GetVertexCountForPrimitiveCount(NumPrimitives, PrimitiveType);
-				
-				uint32 OutputBufferSize = (NumVertices * NumInstances) * boundShaderState->VertexShader->TessellationOutputAttribs.HSOutSize;
-				
-				ShaderOutputBuffer = deviceContext.CreatePooledBuffer(FMetalPooledBufferArgs(device, OutputBufferSize, BUF_Dynamic, mtlpp::StorageMode::Private));
-
-				PrologueEncoder.SetShaderBuffer(mtlpp::FunctionType::Vertex, ShaderOutputBuffer, 0, ShaderOutputBuffer, boundShaderState->VertexShader->TessellationHSOutBuffer, mtlpp::ResourceUsage::Write);
-				State.SetShaderBuffer(EMetalShaderStages::Vertex, nil, nil, 0, 0, boundShaderState->VertexShader->TessellationHSOutBuffer, mtlpp::ResourceUsage(0));
-				
-				// set the patchCount
-				FMetalSubBufferRing& Ring = PrologueEncoder.GetRingBuffer();
-				FMetalBuffer PatchData = Ring.NewBuffer(sizeof(NumVertices), sizeof(NumVertices));
-				uint32* patchCountData = (uint32*)PatchData.GetContents();
-				patchCountData[0] = NumVertices;
-				PrologueEncoder.SetShaderBuffer(mtlpp::FunctionType::Vertex, PatchData, 0, PatchData, boundShaderState->VertexShader->TessellationPatchCountBuffer, mtlpp::ResourceUsage::Read);
-				State.SetShaderBuffer(EMetalShaderStages::Vertex, nil, nil, 0, 0, boundShaderState->VertexShader->TessellationPatchCountBuffer, mtlpp::ResourceUsage(0));
-				
-				if (boundShaderState->VertexShader->SideTableBinding >= 0)
-				{
-					PrologueEncoder.SetShaderSideTable(mtlpp::FunctionType::Vertex, boundShaderState->VertexShader->SideTableBinding);
-					State.SetShaderBuffer(EMetalShaderStages::Vertex, nil, nil, 0, 0, boundShaderState->VertexShader->SideTableBinding, mtlpp::ResourceUsage(0));
-				}
-				
-				METAL_GPUPROFILE(FMetalProfiler::GetProfiler()->EncodeDraw(CurrentEncoder.GetCommandBufferStats(), __FUNCTION__, NumPrimitives, NumVertices, NumInstances));
-				StreamOutEncoder.Draw(TranslatePrimitiveType(PrimitiveType), BaseVertexIndex, NumVertices, NumInstances, FirstInstance);
-				METAL_DEBUG_LAYER(EMetalDebugLevelFastValidation, PrologueEncoder.GetRenderCommandEncoderDebugging().Draw(TranslatePrimitiveType(PrimitiveType), BaseVertexIndex, NumVertices, NumInstances, FirstInstance));
-			}
-			
-			ConditionalSwitchToAsyncCompute();
-			check(CurrentEncoder.GetCommandBuffer());
-			check(PrologueEncoder.GetCommandBuffer());
-			check(PrologueEncoder.IsComputeCommandEncoderActive());
-			check(CurrentEncoder.IsRenderCommandEncoderActive());
-			
-			auto& computeEncoder = PrologueEncoder.GetComputeCommandEncoder();
-			
-			uint32 NumVertices = GetVertexCountForPrimitiveCount(NumPrimitives, PrimitiveType);
-			
-			uint32 InputCPs = boundShaderState->HullShader->TessellationInputControlPoints;
-			uint32 OutputCPs = boundShaderState->HullShader->TessellationOutputControlPoints;
-			
-			uint32 PatchCount = (NumVertices + (InputCPs - 1)) / InputCPs;
-			
-			FMetalBuffer HullIndexBuffer;
-			if (IndexBuffer)
-			{
-				HullIndexBuffer = deviceContext.CreatePooledBuffer(FMetalPooledBufferArgs(device, NumInstances * PatchCount * FMath::Max(OutputCPs, InputCPs) * sizeof(uint32), BUF_Static, mtlpp::StorageMode::Private));
-				if (IndexBufferStride == 2)
-				{
-					computeEncoder.SetComputePipelineState(GetMetalCopyIndex16Function());
-				}
-				else
-				{
-					computeEncoder.SetComputePipelineState(GetMetalCopyIndex32Function());
-				}
-				
-				PrologueEncoder.SetShaderBuffer(mtlpp::FunctionType::Kernel, IndexBuffer, 0, IndexBuffer, 0, mtlpp::ResourceUsage::Read);
-				State.SetShaderBuffer(EMetalShaderStages::Compute, nil, nil, 0, 0, 0, mtlpp::ResourceUsage(0));
-				
-				PrologueEncoder.SetShaderBuffer(mtlpp::FunctionType::Kernel, HullIndexBuffer, 0, HullIndexBuffer, 1, mtlpp::ResourceUsage::Write);
-				State.SetShaderBuffer(EMetalShaderStages::Compute, nil, nil, 0, 0, 1, mtlpp::ResourceUsage(0));
-				
-				uint32 patchCountData[] = { InputCPs, OutputCPs };
-				if (InputCPs > OutputCPs)
-					patchCountData[1] = InputCPs;
-				
-				PrologueEncoder.SetShaderBytes(mtlpp::FunctionType::Kernel, (const uint8*)&patchCountData[0], sizeof(patchCountData), 2);
-				State.SetShaderBuffer(EMetalShaderStages::Compute, nil, nil, 0, 0, 2, mtlpp::ResourceUsage(0));
-				
-				mtlpp::DrawIndexedPrimitivesIndirectArguments Args;
-				Args.IndexCount = NumVertices;
-				Args.InstanceCount = NumInstances;
-				Args.IndexStart = StartIndex;
-				PrologueEncoder.SetShaderBytes(mtlpp::FunctionType::Kernel, (const uint8*)&Args, sizeof(Args), 3);
-				State.SetShaderBuffer(EMetalShaderStages::Compute, nil, nil, 0, 0, 3, mtlpp::ResourceUsage(0));
-				
-				if(GMetalTessellationRunTessellationStage)
-				{
-					METAL_GPUPROFILE(FMetalProfiler::GetProfiler()->EncodeDispatch(PrologueEncoder.GetCommandBufferStats(), __FUNCTION__));
-					computeEncoder.DispatchThreadgroups(mtlpp::Size(PatchCount, 1, 1), mtlpp::Size(FMath::Max(InputCPs, OutputCPs), NumInstances, 1));
-					METAL_DEBUG_LAYER(EMetalDebugLevelFastValidation, PrologueEncoder.GetComputeCommandEncoderDebugging().DispatchThreadgroups(mtlpp::Size(PatchCount, 1, 1), mtlpp::Size(FMath::Max(InputCPs, OutputCPs), NumInstances, 1)));
-				}
-			}
-			else
-			{
-				HullIndexBuffer = deviceContext.CreatePooledBuffer(FMetalPooledBufferArgs(device, NumInstances * PatchCount * FMath::Max(OutputCPs, InputCPs) * sizeof(uint32), BUF_Static, mtlpp::StorageMode::Shared));
-				uint32* Data = (uint32*)HullIndexBuffer.GetContents();
-				uint32 Index = 0;
-				if (OutputCPs > InputCPs)
-				{
-					for (uint32 i = 0; i < NumInstances; i++)
-					{
-						for (uint32 j = 0; j < PatchCount * OutputCPs; j++)
-						{
-							if ((j % OutputCPs) < InputCPs)
-							{
-								Data[j] = Index++;
-							}
-							else
-							{
-								Data[j] = Index;
-							}
-						}
-					}
-				}
-				else
-				{
-					for (uint32 i = 0; i < NumInstances; i++)
-					{
-						for (uint32 j = 0; j < PatchCount; j++)
-						{
-							for (uint32 k = 0; k < InputCPs; k++)
-							{
-								Data[(i * NumVertices) + j * InputCPs + k] = (i * NumVertices) + j * InputCPs + k;
-							}
-						}
-					}
-				}
-			}
-			
-			PrepareToSeparateTessellate(PrimitiveType);
-			
-			if (HullIndexBuffer)
-			{
-				PrologueEncoder.SetShaderBuffer(mtlpp::FunctionType::Kernel, HullIndexBuffer, 0, HullIndexBuffer, boundShaderState->HullShader->TessellationIndexBuffer, mtlpp::ResourceUsage::Read);
-				State.SetShaderBuffer(EMetalShaderStages::Compute, nil, nil, 0, 0, boundShaderState->HullShader->TessellationIndexBuffer, mtlpp::ResourceUsage(0));
-			}
-			else
-			{
-				check(IndexBuffer);
-				PrologueEncoder.SetShaderBuffer(mtlpp::FunctionType::Kernel, IndexBuffer, 0, IndexBuffer, boundShaderState->HullShader->TessellationIndexBuffer, mtlpp::ResourceUsage::Read);
-				State.SetShaderBuffer(EMetalShaderStages::Compute, nil, nil, 0, 0, boundShaderState->HullShader->TessellationIndexBuffer, mtlpp::ResourceUsage(0));
-			}
-			
-			computeEncoder.SetStageInRegion(mtlpp::Region(0, FMath::Max(NumVertices, OutputCPs * PatchCount) * NumInstances));
-			
-			FMetalBuffer HSOutBuffer;
-			if (boundShaderState->HullShader->TessellationHSOutBuffer != UINT_MAX && boundShaderState->HullShader->TessellationOutputAttribs.HSOutSize)
-			{
-				HSOutBuffer = deviceContext.CreatePooledBuffer(FMetalPooledBufferArgs(device, Align(NumInstances * PatchCount * boundShaderState->HullShader->TessellationOutputAttribs.HSOutSize, BufferOffsetAlignment), BUF_Static, mtlpp::StorageMode::Private));
-				
-				PrologueEncoder.SetShaderBuffer(mtlpp::FunctionType::Kernel, HSOutBuffer, 0, HSOutBuffer, boundShaderState->HullShader->TessellationHSOutBuffer, mtlpp::ResourceUsage::Write);
-				State.SetShaderBuffer(EMetalShaderStages::Compute, nil, nil, 0, 0, boundShaderState->HullShader->TessellationHSOutBuffer, mtlpp::ResourceUsage(0));
-			}
-			
-			FMetalBuffer CPOutBuffer;
-			if (boundShaderState->HullShader->TessellationControlPointOutBuffer != UINT_MAX && boundShaderState->HullShader->TessellationOutputAttribs.PatchControlPointOutSize)
-			{
-				CPOutBuffer = deviceContext.CreatePooledBuffer(FMetalPooledBufferArgs(device, NumInstances * PatchCount * OutputCPs * boundShaderState->HullShader->TessellationOutputAttribs.PatchControlPointOutSize, BUF_Static, mtlpp::StorageMode::Private));
-				
-				PrologueEncoder.SetShaderBuffer(mtlpp::FunctionType::Kernel, CPOutBuffer, 0, CPOutBuffer, boundShaderState->HullShader->TessellationControlPointOutBuffer, mtlpp::ResourceUsage::Write);
-				State.SetShaderBuffer(EMetalShaderStages::Compute, nil, nil, 0, 0, boundShaderState->HullShader->TessellationControlPointOutBuffer, mtlpp::ResourceUsage(0));
-			}
-			
-			FMetalBuffer HSTFBuffer;
-			if (boundShaderState->HullShader->TessellationHSTFOutBuffer != UINT_MAX)
-			{
-				HSTFBuffer = deviceContext.CreatePooledBuffer(FMetalPooledBufferArgs(device, NumInstances * PatchCount * boundShaderState->HullShader->TessellationOutputAttribs.HSTFOutSize, BUF_Static, mtlpp::StorageMode::Private));
-				
-				PrologueEncoder.SetShaderBuffer(mtlpp::FunctionType::Kernel, HSTFBuffer, 0, HSTFBuffer, boundShaderState->HullShader->TessellationHSTFOutBuffer, mtlpp::ResourceUsage::Write);
-				State.SetShaderBuffer(EMetalShaderStages::Compute, nil, nil, 0, 0, boundShaderState->HullShader->TessellationHSTFOutBuffer, mtlpp::ResourceUsage(0));
-			}
-			
-			PrologueEncoder.SetShaderBuffer(mtlpp::FunctionType::Kernel, ShaderOutputBuffer, 0, ShaderOutputBuffer, boundShaderState->HullShader->TessellationControlPointIndexBuffer, mtlpp::ResourceUsage::Read);
-			State.SetShaderBuffer(EMetalShaderStages::Compute, nil, nil, 0, 0, boundShaderState->HullShader->TessellationControlPointIndexBuffer, mtlpp::ResourceUsage(0));
-			
-			FMetalSubBufferRing& Ring = PrologueEncoder.GetRingBuffer();
-			FMetalBuffer PatchData = Ring.NewBuffer(sizeof(InputCPs), sizeof(InputCPs));
-			uint32* patchCountData = (uint32*)PatchData.GetContents();
-			patchCountData[0] = InputCPs;
-			PrologueEncoder.SetShaderBuffer(mtlpp::FunctionType::Kernel, PatchData, 0, PatchData, boundShaderState->HullShader->TessellationPatchCountBuffer, mtlpp::ResourceUsage::Read);
-			State.SetShaderBuffer(EMetalShaderStages::Compute, nil, nil, 0, 0, boundShaderState->HullShader->TessellationPatchCountBuffer, mtlpp::ResourceUsage(0));
-			
-			if (boundShaderState->HullShader->SideTableBinding >= 0)
-			{
-				PrologueEncoder.SetShaderSideTable(mtlpp::FunctionType::Kernel, boundShaderState->HullShader->SideTableBinding);
-				State.SetShaderBuffer(EMetalShaderStages::Compute, nil, nil, 0, 0, boundShaderState->HullShader->SideTableBinding, mtlpp::ResourceUsage(0));
-			}
-			
-			computeEncoder.SetThreadgroupMemory(Align(boundShaderState->VertexShader->TessellationOutputAttribs.HSOutSize * InputCPs, 16), 0);
-			
-			if(GMetalTessellationRunTessellationStage)
-			{
-				METAL_GPUPROFILE(FMetalProfiler::GetProfiler()->EncodeDispatch(PrologueEncoder.GetCommandBufferStats(), __FUNCTION__));
-				computeEncoder.DispatchThreadgroups(mtlpp::Size(NumInstances * PatchCount, 1, 1), mtlpp::Size(FMath::Max(InputCPs, OutputCPs), 1, 1));
-				METAL_DEBUG_LAYER(EMetalDebugLevelFastValidation, PrologueEncoder.GetComputeCommandEncoderDebugging().DispatchThreadgroups(mtlpp::Size(NumInstances * PatchCount, 1, 1), mtlpp::Size(FMath::Max(InputCPs, OutputCPs), 1, 1)));
-			}
-			
-			computeEncoder.SetThreadgroupMemory(0, 0);
-			
-			if (HSTFBuffer && GMetalTessellationFlattenTriangleFactors)
-			{
-				computeEncoder.SetComputePipelineState(GetMetalFlattenTessState());
-				PrologueEncoder.SetShaderBuffer(mtlpp::FunctionType::Kernel, HSTFBuffer, 0, HSTFBuffer, 0, mtlpp::ResourceUsage::Write);
-				State.SetShaderBuffer(EMetalShaderStages::Compute, nil, nil, 0, 0, 0, mtlpp::ResourceUsage(0));
-				
-				mtlpp::DrawIndexedPrimitivesIndirectArguments Args;
-				Args.IndexCount = NumInstances * PatchCount;
-				Args.InstanceCount = 0;
-				Args.IndexStart = 0;
-				PrologueEncoder.SetShaderBytes(mtlpp::FunctionType::Kernel, (const uint8*)&Args, sizeof(Args), 1);
-				State.SetShaderBuffer(EMetalShaderStages::Compute, nil, nil, 0, 0, 1, mtlpp::ResourceUsage(0));
-				
-				METAL_GPUPROFILE(FMetalProfiler::GetProfiler()->EncodeDispatch(PrologueEncoder.GetCommandBufferStats(), __FUNCTION__));
-				computeEncoder.DispatchThreadgroups(mtlpp::Size(1, 1, 1), mtlpp::Size(1, 1, 1));
-				METAL_DEBUG_LAYER(EMetalDebugLevelFastValidation, PrologueEncoder.GetComputeCommandEncoderDebugging().DispatchThreadgroups(mtlpp::Size(1, 1, 1), mtlpp::Size(1, 1, 1)));
-			}
-			
-			auto& renderEncoder = CurrentEncoder.GetRenderCommandEncoder();
-			
-			if (HSOutBuffer)
-			{
-				CurrentEncoder.SetShaderBuffer(mtlpp::FunctionType::Vertex, HSOutBuffer, 0, HSOutBuffer, boundShaderState->DomainShader->TessellationHSOutBuffer, mtlpp::ResourceUsage::Read);
-				State.SetShaderBuffer(EMetalShaderStages::Domain, nil, nil, 0, 0, boundShaderState->DomainShader->TessellationHSOutBuffer, mtlpp::ResourceUsage(0));
-			}
-			
-			if (CPOutBuffer)
-			{
-				CurrentEncoder.SetShaderBuffer(mtlpp::FunctionType::Vertex, CPOutBuffer, 0, CPOutBuffer, boundShaderState->DomainShader->TessellationControlPointOutBuffer, mtlpp::ResourceUsage::Read);
-				State.SetShaderBuffer(EMetalShaderStages::Domain, nil, nil, 0, 0, boundShaderState->DomainShader->TessellationControlPointOutBuffer, mtlpp::ResourceUsage(0));
-			}
-			
-			if (HSTFBuffer)
-			{
-				renderEncoder.SetTessellationFactorBuffer(HSTFBuffer, 0, 0);
-				METAL_DEBUG_LAYER(EMetalDebugLevelFastValidation, CurrentEncoder.GetRenderCommandEncoderDebugging().SetTessellationFactorBuffer(HSTFBuffer, 0, 0));
-			}
-			
-			if (boundShaderState->DomainShader->SideTableBinding >= 0)
-			{
-				CurrentEncoder.SetShaderSideTable(mtlpp::FunctionType::Vertex, boundShaderState->DomainShader->SideTableBinding);
-				State.SetShaderBuffer(EMetalShaderStages::Domain, nil, nil, 0, 0, boundShaderState->DomainShader->SideTableBinding, mtlpp::ResourceUsage(0));
-			}
-			
-			if (IsValidRef(boundShaderState->PixelShader) && boundShaderState->PixelShader->SideTableBinding >= 0)
-			{
-				CurrentEncoder.SetShaderSideTable(mtlpp::FunctionType::Fragment, boundShaderState->PixelShader->SideTableBinding);
-				State.SetShaderBuffer(EMetalShaderStages::Pixel, nil, nil, 0, 0, boundShaderState->PixelShader->SideTableBinding, mtlpp::ResourceUsage(0));
-			}
-			
-			if (GMetalTessellationRunDomainStage)
-			{
-				METAL_GPUPROFILE(FMetalProfiler::GetProfiler()->EncodeDraw(CurrentEncoder.GetCommandBufferStats(), __FUNCTION__, OutputCPs, PatchCount, NumInstances));
-				renderEncoder.DrawPatches(OutputCPs, 0, NumInstances * PatchCount, nil, 0, 1, 0);
-				METAL_DEBUG_LAYER(EMetalDebugLevelFastValidation, CurrentEncoder.GetRenderCommandEncoderDebugging().DrawPatches(OutputCPs, 0, NumInstances * PatchCount, nil, 0, 1, 0));
-			}
-			
-			if (GMetalCommandBufferDebuggingEnabled)
-			{
-				FMetalCommandData Data;
-				Data.CommandType = FMetalCommandData::Type::DrawPrimitivePatch;
-				Data.DrawPatch.BaseInstance = FirstInstance;
-				Data.DrawPatch.InstanceCount = NumInstances;
-				Data.DrawPatch.PatchCount = PatchCount;
-				Data.DrawPatch.PatchStart = 0;
-				
-				InsertDebugDraw(Data);
-			}
-			
-			SafeReleaseMetalBuffer(ShaderOutputBuffer);
-			SafeReleaseMetalBuffer(HullIndexBuffer);
-			SafeReleaseMetalBuffer(HSOutBuffer);
-			SafeReleaseMetalBuffer(CPOutBuffer);
-			SafeReleaseMetalBuffer(HSTFBuffer);
-			return;
-		}
-		
-		ConditionalSwitchToTessellation();
-		check(CurrentEncoder.GetCommandBuffer());
-		check(PrologueEncoder.GetCommandBuffer());
-		check(PrologueEncoder.IsComputeCommandEncoderActive());
-		check(CurrentEncoder.IsRenderCommandEncoderActive());
-		
-		size_t hullShaderOutputOffset = 0;
-		size_t hullConstShaderOutputOffset = 0;
-		size_t tessellationFactorsOffset = 0;
-		
-		FMetalDeviceContext& deviceContext = (FMetalDeviceContext&)GetMetalDeviceContext();
-		mtlpp::Device device = deviceContext.GetDevice();
-		
-		FMetalGraphicsPipelineState* boundShaderState = State.GetGraphicsPSO();
-		FMetalShaderPipeline* Pipeline = State.GetPipelineState();
-		
-		// TODO could allocate this as 1 buffer and use the sizes to make the offsets we need...
-		auto hullShaderOutputBufferSize = (Pipeline->TessellationPipelineDesc.TessellationPatchControlPointOutSize * boundShaderState->VertexShader->TessellationOutputControlPoints) * NumPrimitives * NumInstances;
-		auto hullConstShaderOutputBufferSize = (Pipeline->TessellationPipelineDesc.TessellationPatchConstOutSize) * NumPrimitives * NumInstances;
-		auto tessellationFactorBufferSize = (Pipeline->TessellationPipelineDesc.TessellationTessFactorOutSize) * NumPrimitives * NumInstances;
-		
-		FMetalBuffer hullShaderOutputBuffer = nil;
-		if(hullShaderOutputBufferSize)
-		{
-			hullShaderOutputBuffer = deviceContext.GetResourceHeap().CreateBuffer(hullShaderOutputBufferSize, 16, BUF_Dynamic, FMetalCommandQueue::GetCompatibleResourceOptions((mtlpp::ResourceOptions)( mtlpp::ResourceOptions::HazardTrackingModeUntracked | ((NSUInteger)mtlpp::StorageMode::Private << mtlpp::ResourceStorageModeShift))), true);
-		}
-		
-		FMetalBuffer hullConstShaderOutputBuffer = nil;
-		if(hullConstShaderOutputBufferSize)
-		{
-			hullConstShaderOutputBuffer = deviceContext.GetResourceHeap().CreateBuffer(hullConstShaderOutputBufferSize, 16, BUF_Dynamic, FMetalCommandQueue::GetCompatibleResourceOptions((mtlpp::ResourceOptions)( mtlpp::ResourceOptions::HazardTrackingModeUntracked | ((NSUInteger)mtlpp::StorageMode::Private << mtlpp::ResourceStorageModeShift))), true);
-		}
-		
-		FMetalBuffer tessellationFactorBuffer = nil;
-		if(tessellationFactorBufferSize)
-		{
-			tessellationFactorBuffer = deviceContext.GetResourceHeap().CreateBuffer(tessellationFactorBufferSize, 16, BUF_Dynamic, FMetalCommandQueue::GetCompatibleResourceOptions((mtlpp::ResourceOptions)( mtlpp::ResourceOptions::HazardTrackingModeUntracked | ((NSUInteger)mtlpp::StorageMode::Private << mtlpp::ResourceStorageModeShift))), true);
-		}
-	
-		auto& computeEncoder = PrologueEncoder.GetComputeCommandEncoder();
-		auto& renderEncoder = CurrentEncoder.GetRenderCommandEncoder();
-		
-		PrepareToTessellate(PrimitiveType);
-		
-		// Per-draw call bindings should *not* be cached in the StateCache - causes absolute chaos.
-		if(IndexBuffer && Pipeline->TessellationPipelineDesc.TessellationControlPointIndexBufferIndex != UINT_MAX)
-		{
-			PrologueEncoder.SetShaderBuffer(mtlpp::FunctionType::Kernel, IndexBuffer, StartIndex * IndexBufferStride, IndexBuffer.GetLength() - (StartIndex * IndexBufferStride), Pipeline->TessellationPipelineDesc.TessellationControlPointIndexBufferIndex, mtlpp::ResourceUsage::Read);
-			
-			State.SetShaderBuffer(EMetalShaderStages::Vertex, nil, nil, 0, 0, Pipeline->TessellationPipelineDesc.TessellationControlPointIndexBufferIndex, mtlpp::ResourceUsage(0));
-		}
-		
-		if (Pipeline->TessellationPipelineDesc.TessellationIndexBufferIndex != UINT_MAX)
-		{
-			if (IndexBuffer)
-			{
-				PrologueEncoder.SetShaderBuffer(mtlpp::FunctionType::Kernel, IndexBuffer, StartIndex * IndexBufferStride, IndexBuffer.GetLength() - (StartIndex * IndexBufferStride), Pipeline->TessellationPipelineDesc.TessellationIndexBufferIndex, mtlpp::ResourceUsage::Read, IndexBufferStride == 2 ? PF_R16_UINT : PF_R32_UINT);
-		}
-			else
-			{
-				PrologueEncoder.SetShaderBuffer(mtlpp::FunctionType::Kernel, nil, 0, 0, Pipeline->TessellationPipelineDesc.TessellationIndexBufferIndex, mtlpp::ResourceUsage::Read, (EPixelFormat)0);
-			}
-			State.SetShaderBuffer(EMetalShaderStages::Vertex, nil, nil, 0, 0, Pipeline->TessellationPipelineDesc.TessellationControlPointIndexBufferIndex, mtlpp::ResourceUsage(0));
-		}
-		
-		if(Pipeline->TessellationPipelineDesc.TessellationOutputControlPointBufferIndex != UINT_MAX) //TessellationOutputControlPointBufferIndex -> hullShaderOutputBuffer
-		{
-			PrologueEncoder.SetShaderBuffer(mtlpp::FunctionType::Kernel, hullShaderOutputBuffer, hullShaderOutputOffset, hullShaderOutputBuffer.GetLength() - hullShaderOutputOffset, Pipeline->TessellationPipelineDesc.TessellationOutputControlPointBufferIndex, mtlpp::ResourceUsage(0));
-			State.SetShaderBuffer(EMetalShaderStages::Vertex, nil, nil, 0, 0, Pipeline->TessellationPipelineDesc.TessellationOutputControlPointBufferIndex, mtlpp::ResourceUsage(0));
-		}
-		
-		if(Pipeline->TessellationPipelineDesc.TessellationPatchConstBufferIndex != UINT_MAX) //TessellationPatchConstBufferIndex -> hullConstShaderOutputBuffer
-		{
-			PrologueEncoder.SetShaderBuffer(mtlpp::FunctionType::Kernel, hullConstShaderOutputBuffer, hullConstShaderOutputOffset, hullConstShaderOutputBuffer.GetLength() - hullConstShaderOutputOffset, Pipeline->TessellationPipelineDesc.TessellationPatchConstBufferIndex, mtlpp::ResourceUsage(0));
-			State.SetShaderBuffer(EMetalShaderStages::Vertex, nil, nil, 0, 0, Pipeline->TessellationPipelineDesc.TessellationPatchConstBufferIndex, mtlpp::ResourceUsage(0));
-		}
-		
-		if(Pipeline->TessellationPipelineDesc.TessellationFactorBufferIndex != UINT_MAX) // TessellationFactorBufferIndex->tessellationFactorBuffer
-		{
-			PrologueEncoder.SetShaderBuffer(mtlpp::FunctionType::Kernel, tessellationFactorBuffer, tessellationFactorsOffset, tessellationFactorBuffer.GetLength() - tessellationFactorsOffset, Pipeline->TessellationPipelineDesc.TessellationFactorBufferIndex, mtlpp::ResourceUsage(0));
-			State.SetShaderBuffer(EMetalShaderStages::Vertex, nil, nil, 0, 0, Pipeline->TessellationPipelineDesc.TessellationFactorBufferIndex, mtlpp::ResourceUsage(0));
-		}
-		
-		if(Pipeline->TessellationPipelineDesc.TessellationInputControlPointBufferIndex != UINT_MAX) //TessellationInputControlPointBufferIndex->hullShaderOutputBuffer
-		{
-			CurrentEncoder.SetShaderBuffer(mtlpp::FunctionType::Vertex, hullShaderOutputBuffer, hullShaderOutputOffset, hullShaderOutputBuffer.GetLength() - hullShaderOutputOffset, Pipeline->TessellationPipelineDesc.TessellationInputControlPointBufferIndex, mtlpp::ResourceUsage(0));
-			State.SetShaderBuffer(EMetalShaderStages::Domain, nil, nil, 0, 0, Pipeline->TessellationPipelineDesc.TessellationInputControlPointBufferIndex, mtlpp::ResourceUsage(0));
-		}
-		if(Pipeline->TessellationPipelineDesc.TessellationInputPatchConstBufferIndex != UINT_MAX) //TessellationInputPatchConstBufferIndex->hullConstShaderOutputBuffer
-		{
-			CurrentEncoder.SetShaderBuffer(mtlpp::FunctionType::Vertex, hullConstShaderOutputBuffer, hullConstShaderOutputOffset, hullConstShaderOutputBuffer.GetLength() - hullConstShaderOutputOffset, Pipeline->TessellationPipelineDesc.TessellationInputPatchConstBufferIndex, mtlpp::ResourceUsage(0));
-			State.SetShaderBuffer(EMetalShaderStages::Domain, nil, nil, 0, 0, Pipeline->TessellationPipelineDesc.TessellationInputPatchConstBufferIndex, mtlpp::ResourceUsage(0));
-		}
-		
-		// set the patchCount
-		uint32 patchCountData[] = { NumPrimitives, StartIndex };
-		PrologueEncoder.SetShaderBytes(mtlpp::FunctionType::Kernel, (const uint8*)&patchCountData[0], sizeof(patchCountData), Pipeline->TessellationPipelineDesc.TessellationPatchCountBufferIndex);
-		State.SetShaderBuffer(EMetalShaderStages::Vertex, nil, nil, 0, 0, Pipeline->TessellationPipelineDesc.TessellationPatchCountBufferIndex, mtlpp::ResourceUsage(0));
-		
-		if (boundShaderState->VertexShader->SideTableBinding >= 0)
-		{
-			PrologueEncoder.SetShaderSideTable(mtlpp::FunctionType::Kernel, boundShaderState->VertexShader->SideTableBinding);
-			State.SetShaderBuffer(EMetalShaderStages::Vertex, nil, nil, 0, 0, boundShaderState->VertexShader->SideTableBinding, mtlpp::ResourceUsage(0));
-		}
-		
-		if (boundShaderState->DomainShader->SideTableBinding >= 0)
-		{
-			CurrentEncoder.SetShaderSideTable(mtlpp::FunctionType::Vertex, boundShaderState->DomainShader->SideTableBinding);
-			State.SetShaderBuffer(EMetalShaderStages::Domain, nil, nil, 0, 0, boundShaderState->DomainShader->SideTableBinding, mtlpp::ResourceUsage(0));
-		}
-		
-		if (IsValidRef(boundShaderState->PixelShader) && boundShaderState->PixelShader->SideTableBinding >= 0)
-		{
-			CurrentEncoder.SetShaderSideTable(mtlpp::FunctionType::Fragment, boundShaderState->PixelShader->SideTableBinding);
-			State.SetShaderBuffer(EMetalShaderStages::Pixel, nil, nil, 0, 0, boundShaderState->PixelShader->SideTableBinding, mtlpp::ResourceUsage(0));
-		}
-		
-		auto patchesPerThreadGroup = boundShaderState->VertexShader->TessellationPatchesPerThreadGroup;
-		auto threadgroups = mtlpp::Size((NumPrimitives + (patchesPerThreadGroup - 1)) / patchesPerThreadGroup, NumInstances, 1);
-		auto threadsPerThreadgroup = mtlpp::Size(boundShaderState->VertexShader->TessellationInputControlPoints * patchesPerThreadGroup, 1, 1);
-		
-		computeEncoder.SetStageInRegion(mtlpp::Region(BaseVertexIndex, FirstInstance, boundShaderState->VertexShader->TessellationInputControlPoints * NumPrimitives, NumInstances));
-		if(GMetalTessellationRunTessellationStage)
-		{
-			METAL_GPUPROFILE(FMetalProfiler::GetProfiler()->EncodeDispatch(PrologueEncoder.GetCommandBufferStats(), __FUNCTION__));
-			computeEncoder.DispatchThreadgroups(threadgroups, threadsPerThreadgroup);
-			METAL_DEBUG_LAYER(EMetalDebugLevelFastValidation, PrologueEncoder.GetComputeCommandEncoderDebugging().DispatchThreadgroups(threadgroups, threadsPerThreadgroup));
-		}
-		
-		check(computeEncoder.GetPtr() != nil);
-		check(renderEncoder.GetPtr() != nil);
-		
-		if(tessellationFactorBuffer)
-		{
-			renderEncoder.SetTessellationFactorBuffer(tessellationFactorBuffer, tessellationFactorsOffset, 0);
-			METAL_DEBUG_LAYER(EMetalDebugLevelFastValidation, CurrentEncoder.GetRenderCommandEncoderDebugging().SetTessellationFactorBuffer(tessellationFactorBuffer, tessellationFactorsOffset, 0));
-		}
-		if(GMetalTessellationRunDomainStage)
-		{
-			METAL_GPUPROFILE(FMetalProfiler::GetProfiler()->EncodeDraw(CurrentEncoder.GetCommandBufferStats(), __FUNCTION__, NumPrimitives, GetVertexCountForPrimitiveCount(NumPrimitives, PrimitiveType), NumInstances));
-			renderEncoder.DrawPatches(boundShaderState->VertexShader->TessellationOutputControlPoints, 0, NumPrimitives * NumInstances, nil, 0, 1, 0);
-			METAL_DEBUG_LAYER(EMetalDebugLevelFastValidation, CurrentEncoder.GetRenderCommandEncoderDebugging().DrawPatches(boundShaderState->VertexShader->TessellationOutputControlPoints, 0, NumPrimitives * NumInstances, nil, 0, 1, 0));
-		}
-		
-#if METAL_DEBUG_OPTIONS
-		if (GMetalCommandBufferDebuggingEnabled)
-		{
-			FMetalCommandData Data;
-			Data.CommandType = FMetalCommandData::Type::DrawPrimitivePatch;
-			Data.DrawPatch.BaseInstance = FirstInstance;
-			Data.DrawPatch.InstanceCount = NumInstances;
-			Data.DrawPatch.PatchCount = NumPrimitives * NumInstances;
-			Data.DrawPatch.PatchStart = 0;
-			
-			InsertDebugDraw(Data);
-		}
-#endif
-		
-		if(hullShaderOutputBufferSize)
-		{
-			deviceContext.ReleaseBuffer(hullShaderOutputBuffer);
-		}
-		if(hullConstShaderOutputBufferSize)
-		{
-			deviceContext.ReleaseBuffer(hullConstShaderOutputBuffer);
-		}
-		if(tessellationFactorBufferSize)
-		{
-			deviceContext.ReleaseBuffer(tessellationFactorBuffer);
-		}
-	}
-	else
-	{
-		NOT_SUPPORTED("DrawPatches");
-	}
-}
-#endif
-
 void FMetalRenderPass::Dispatch(uint32 ThreadGroupCountX, uint32 ThreadGroupCountY, uint32 ThreadGroupCountZ)
 {
     if (CurrentEncoder.IsParallel() || CurrentEncoder.NumEncodedPasses() == 0)
@@ -1735,210 +1193,6 @@ void FMetalRenderPass::ConditionalSwitchToRender(void)
 	check(CurrentEncoder.IsRenderCommandEncoderActive());
 }
 
-void FMetalRenderPass::ConditionalSwitchToTessellation(void)
-{
-	SCOPE_CYCLE_COUNTER(STAT_MetalSwitchToTessellationTime);
-	
-	check(bWithinRenderPass);
-	check(RenderPassDesc);
-	check(CurrentEncoder.GetCommandBuffer());
-	
-	// End all current encoders that don't match required compute/raster setup.
-	if (PrologueEncoder.IsBlitCommandEncoderActive())
-	{
-		PrologueEncoderFence = PrologueEncoder.EndEncoding();
-	}
-	if (CurrentEncoder.IsComputeCommandEncoderActive() || CurrentEncoder.IsBlitCommandEncoderActive())
-	{
-		CurrentEncoderFence = CurrentEncoder.EndEncoding();
-	}
-	
-	// Create a new prologue compute encoder if needed
-	if (!PrologueEncoder.IsComputeCommandEncoderActive())
-	{
-		State.SetStateDirty();
-		if (!PrologueEncoder.GetCommandBuffer())
-		{
-			PrologueEncoder.StartCommandBuffer();
-		}
-		PrologueEncoder.BeginComputeCommandEncoding(ComputeDispatchType);
-		
-		// Wait on the pass start fence to ensure proper ordering.
-		if (PrologueStartEncoderFence)
-		{
-			if (PrologueStartEncoderFence->NeedsWait(mtlpp::RenderStages::Vertex))
-			{
-				PrologueEncoder.WaitForFence(PrologueStartEncoderFence);
-			}
-			else
-			{
-				PrologueEncoder.WaitAndUpdateFence(PrologueStartEncoderFence);
-			}
-			PrologueStartEncoderFence = nullptr;
-		}
-		// Wait on previous prologue encoder fence and consume it, we'll replace it with the new one later.
-		if (PrologueEncoderFence)
-		{
-			if (PrologueEncoderFence->NeedsWait(mtlpp::RenderStages::Vertex))
-			{
-				PrologueEncoder.WaitForFence(PrologueEncoderFence);
-			}
-			else
-			{
-				PrologueEncoder.WaitAndUpdateFence(PrologueEncoderFence);
-			}
-			PrologueEncoderFence = nullptr;
-		}
-#if METAL_DEBUG_OPTIONS
-//		if (GetEmitDrawEvents() && PrologueEncoder.GetEncoderFence())
-//		{
-//			for (uint32 i = mtlpp::RenderStages::Vertex; i <= mtlpp::RenderStages::Fragment && PrologueEncoder.GetEncoderFence()->Get((mtlpp::RenderStages)i).GetPtr(); i++)
-//			{
-//				if (CmdList.GetCommandQueue().GetRuntimeDebuggingLevel() >= EMetalDebugLevelValidation)
-//				{
-//					PrologueEncoder.GetEncoderFence()->Get((mtlpp::RenderStages)i).GetPtr().label = [NSString stringWithFormat:@"Prologue %@", PrologueEncoderFence->Get((mtlpp::RenderStages)i).GetLabel().GetPtr()];
-//				}
-//				else
-//				{
-//					PrologueEncoder.GetEncoderFence()->Get((mtlpp::RenderStages)i).SetLabel([NSString stringWithFormat:@"Prologue %@", PrologueEncoderFence->Get((mtlpp::RenderStages)i).GetLabel().GetPtr()]);
-//				}
-//			}
-//		}
-#endif
-	}
-	
-	// Restart the render pass to ensure we have a raster encoder
-	if (!CurrentEncoder.IsRenderCommandEncoderActive())
-	{
-		RestartRenderPass(nil);
-		
-		check(CurrentEncoder.IsRenderCommandEncoderActive());
-		check(PrologueEncoder.IsComputeCommandEncoderActive());
-	}
-	else
-	{
-		check(CurrentEncoder.IsRenderCommandEncoderActive());
-		check(PrologueEncoder.IsComputeCommandEncoderActive());
-
-		// Encode a wait to the current encoder for the necessary prologue encoder
-		CurrentEncoder.WaitForFence(PrologueEncoder.GetEncoderFence());
-	}
-}
-
-void FMetalRenderPass::ConditionalSwitchToSeparateTessellation(void)
-{
-	SCOPE_CYCLE_COUNTER(STAT_MetalSwitchToTessellationTime);
-	
-	check(bWithinRenderPass);
-	check(RenderPassDesc);
-	check(CurrentEncoder.GetCommandBuffer());
-	
-	// End all current encoders that don't match required compute/raster setup.
-	if (PrologueEncoder.IsBlitCommandEncoderActive() || PrologueEncoder.IsComputeCommandEncoderActive())
-	{
-		PrologueEncoderFence = PrologueEncoder.EndEncoding();
-	}
-	if (CurrentEncoder.IsComputeCommandEncoderActive() || CurrentEncoder.IsBlitCommandEncoderActive())
-	{
-		CurrentEncoderFence = CurrentEncoder.EndEncoding();
-	}
-	
-	// Create a new prologue stream-out encoder if needed
-	if (!PrologueEncoder.IsRenderCommandEncoderActive())
-	{
-		State.SetStateDirty();
-		if (!PrologueEncoder.GetCommandBuffer())
-		{
-			PrologueEncoder.StartCommandBuffer();
-		}
-
-		// Deliberately allocate a new pass with outstanding ref-count as it will be pooled.
-		mtlpp::RenderPassDescriptor Desc([MTLRenderPassDescriptor new]);
-		mtlpp::RenderPassDepthAttachmentDescriptor DepthAttachment;
-		mtlpp::RenderPassStencilAttachmentDescriptor StencilAttachment;
-		
-		CGSize FBSize = CGSizeMake(State.GetViewport(0).width, State.GetViewport(0).height);
-		
-		FTexture2DRHIRef FallbackDepthStencilSurface = State.CreateFallbackDepthStencilSurface(FBSize.width, FBSize.height);
-		check(IsValidRef(FallbackDepthStencilSurface));
-		
-		FMetalSurface& Surface = *GetMetalSurfaceFromRHITexture(FallbackDepthStencilSurface.GetReference());
-		
-		DepthAttachment.SetTexture(Surface.Texture);
-		DepthAttachment.SetLoadAction(mtlpp::LoadAction::DontCare);
-		DepthAttachment.SetStoreAction(mtlpp::StoreAction::DontCare);
-		Desc.SetDepthAttachment(DepthAttachment);
-		
-		StencilAttachment.SetTexture(Surface.Texture);
-		StencilAttachment.SetLoadAction(mtlpp::LoadAction::DontCare);
-		StencilAttachment.SetStoreAction(mtlpp::StoreAction::DontCare);
-		Desc.SetStencilAttachment(StencilAttachment);
-		
-		PrologueEncoder.SetRenderPassDescriptor(Desc);
-		PrologueEncoder.BeginRenderCommandEncoding();
-		
-		// Wait on the pass start fence to ensure proper ordering.
-		if (PrologueStartEncoderFence)
-		{
-			if (PrologueStartEncoderFence->NeedsWait(mtlpp::RenderStages::Vertex))
-			{
-				PrologueEncoder.WaitForFence(PrologueStartEncoderFence);
-			}
-			else
-			{
-				PrologueEncoder.WaitAndUpdateFence(PrologueStartEncoderFence);
-			}
-			PrologueStartEncoderFence = nullptr;
-		}
-		// Wait on previous prologue encoder fence and consume it, we'll replace it with the new one later.
-		if (PrologueEncoderFence)
-		{
-			if (PrologueEncoderFence->NeedsWait(mtlpp::RenderStages::Vertex))
-			{
-				PrologueEncoder.WaitForFence(PrologueEncoderFence);
-			}
-			else
-			{
-				PrologueEncoder.WaitAndUpdateFence(PrologueEncoderFence);
-			}
-			PrologueEncoderFence = nullptr;
-		}
-#if METAL_DEBUG_OPTIONS
-		//		if (GetEmitDrawEvents() && PrologueEncoder.GetEncoderFence())
-		//		{
-		//			for (uint32 i = mtlpp::RenderStages::Vertex; i <= mtlpp::RenderStages::Fragment && PrologueEncoder.GetEncoderFence()->Get((mtlpp::RenderStages)i).GetPtr(); i++)
-		//			{
-		//				if (CmdList.GetCommandQueue().GetRuntimeDebuggingLevel() >= EMetalDebugLevelValidation)
-		//				{
-		//					PrologueEncoder.GetEncoderFence()->Get((mtlpp::RenderStages)i).GetPtr().label = [NSString stringWithFormat:@"Prologue %@", PrologueEncoderFence->Get((mtlpp::RenderStages)i).GetLabel().GetPtr()];
-		//				}
-		//				else
-		//				{
-		//					PrologueEncoder.GetEncoderFence()->Get((mtlpp::RenderStages)i).SetLabel([NSString stringWithFormat:@"Prologue %@", PrologueEncoderFence->Get((mtlpp::RenderStages)i).GetLabel().GetPtr()]);
-		//				}
-		//			}
-		//		}
-#endif
-	}
-	
-	// Restart the render pass to ensure we have a raster encoder
-	if (!CurrentEncoder.IsRenderCommandEncoderActive())
-	{
-		RestartRenderPass(nil);
-		
-		check(CurrentEncoder.IsRenderCommandEncoderActive());
-		check(PrologueEncoder.IsRenderCommandEncoderActive());
-	}
-	else
-	{
-		check(CurrentEncoder.IsRenderCommandEncoderActive());
-		check(PrologueEncoder.IsRenderCommandEncoderActive());
-		
-		// Encode a wait to the current encoder for the necessary prologue encoder
-		CurrentEncoder.WaitForFence(PrologueEncoder.GetEncoderFence());
-	}
-}
-
 void FMetalRenderPass::ConditionalSwitchToCompute(void)
 {
 	SCOPE_CYCLE_COUNTER(STAT_MetalSwitchToComputeTime);
@@ -2213,46 +1467,6 @@ void FMetalRenderPass::CommitRenderResourceTables(void)
 	}
 }
 
-#if PLATFORM_SUPPORTS_TESSELLATION_SHADERS
-void FMetalRenderPass::CommitTessellationResourceTables(void)
-{
-	State.CommitTessellationResources(&CurrentEncoder, &PrologueEncoder);
-	
-	State.CommitResourceTable(EMetalShaderStages::Vertex, mtlpp::FunctionType::Kernel, PrologueEncoder);
-	
-	State.CommitResourceTable(EMetalShaderStages::Hull, mtlpp::FunctionType::Kernel, PrologueEncoder);
-	
-	State.CommitResourceTable(EMetalShaderStages::Domain, mtlpp::FunctionType::Vertex, CurrentEncoder);
-	
-	TRefCountPtr<FMetalGraphicsPipelineState> CurrentBoundShaderState = State.GetGraphicsPSO();
-	if (IsValidRef(CurrentBoundShaderState->PixelShader))
-	{
-		State.CommitResourceTable(EMetalShaderStages::Pixel, mtlpp::FunctionType::Fragment, CurrentEncoder);
-	}
-}
-
-void FMetalRenderPass::CommitStreamOutResourceTables(void)
-{
-	State.CommitVertexStreamResources(&PrologueEncoder);
-	State.CommitResourceTable(EMetalShaderStages::Vertex, mtlpp::FunctionType::Vertex, PrologueEncoder);
-}
-
-void FMetalRenderPass::CommitSeparateTessellationResourceTables(void)
-{
-	State.CommitSeparateTessellationResources(&CurrentEncoder, &PrologueEncoder);
-	
-	State.CommitResourceTable(EMetalShaderStages::Hull, mtlpp::FunctionType::Kernel, PrologueEncoder);
-	
-	State.CommitResourceTable(EMetalShaderStages::Domain, mtlpp::FunctionType::Vertex, CurrentEncoder);
-	
-	TRefCountPtr<FMetalGraphicsPipelineState> CurrentBoundShaderState = State.GetGraphicsPSO();
-	if (IsValidRef(CurrentBoundShaderState->PixelShader))
-	{
-		State.CommitResourceTable(EMetalShaderStages::Pixel, mtlpp::FunctionType::Fragment, CurrentEncoder);
-	}
-}
-#endif
-
 void FMetalRenderPass::CommitDispatchResourceTables(void)
 {
 	State.CommitComputeResources(&CurrentEncoder);
@@ -2297,57 +1511,6 @@ void FMetalRenderPass::PrepareToRender(uint32 PrimitiveType)
     State.SetRenderPipelineState(CurrentEncoder, nullptr);
 }
 
-#if PLATFORM_SUPPORTS_TESSELLATION_SHADERS
-void FMetalRenderPass::PrepareToTessellate(uint32 PrimitiveType)
-{
-	SCOPE_CYCLE_COUNTER(STAT_MetalPrepareToTessellateTime);
-	
-	check(CurrentEncoder.GetCommandBuffer());
-	check(PrologueEncoder.GetCommandBuffer());
-	check(CurrentEncoder.IsRenderCommandEncoderActive());
-	check(PrologueEncoder.IsComputeCommandEncoderActive());
-	
-	// Set raster state
-	State.SetRenderState(CurrentEncoder, &PrologueEncoder);
-	
-	// Bind shader resources
-	CommitTessellationResourceTables();
-    
-    State.SetRenderPipelineState(CurrentEncoder, &PrologueEncoder);
-}
-
-void FMetalRenderPass::PrepareToStreamOut(uint32 PrimType)
-{
-	SCOPE_CYCLE_COUNTER(STAT_MetalPrepareToTessellateTime);
-	
-	check(PrologueEncoder.GetCommandBuffer());
-	check(PrologueEncoder.IsRenderCommandEncoderActive());
-	
-	// Bind shader resources
-	CommitStreamOutResourceTables();
-	
-	State.SetStreamOutPipelineState(PrologueEncoder);
-}
-
-void FMetalRenderPass::PrepareToSeparateTessellate(uint32 PrimType)
-{
-	SCOPE_CYCLE_COUNTER(STAT_MetalPrepareToTessellateTime);
-	
-	check(CurrentEncoder.GetCommandBuffer());
-	check(PrologueEncoder.GetCommandBuffer());
-	check(CurrentEncoder.IsRenderCommandEncoderActive());
-	check(PrologueEncoder.IsComputeCommandEncoderActive());
-	
-	// Set raster state
-	State.SetRenderState(CurrentEncoder, &PrologueEncoder);
-	
-	// Bind shader resources
-	CommitSeparateTessellationResourceTables();
-	
-	State.SetTessellationPipelineState(CurrentEncoder, &PrologueEncoder);
-}
-#endif
-
 void FMetalRenderPass::PrepareToDispatch(void)
 {
 	SCOPE_CYCLE_COUNTER(STAT_MetalPrepareToDispatchTime);
diff --git a/Engine/Source/Runtime/Apple/MetalRHI/Private/MetalRenderPass.h b/Engine/Source/Runtime/Apple/MetalRHI/Private/MetalRenderPass.h
--- a/Engine/Source/Runtime/Apple/MetalRHI/Private/MetalRenderPass.h
+++ b/Engine/Source/Runtime/Apple/MetalRHI/Private/MetalRenderPass.h
@@ -47,11 +47,6 @@ public:
     
     void DrawIndexedPrimitiveIndirect(uint32 PrimitiveType,FMetalIndexBuffer* IndexBufferRHI,FMetalVertexBuffer* VertexBufferRHI,uint32 ArgumentOffset);
 	
-#if PLATFORM_SUPPORTS_TESSELLATION_SHADERS
-    void DrawPatches(uint32 PrimitiveType, FMetalBuffer const& IndexBuffer, uint32 IndexBufferStride, int32 BaseVertexIndex, uint32 FirstInstance, uint32 StartIndex,
-                     uint32 NumPrimitives, uint32 NumInstances);
-#endif
-	
     void Dispatch(uint32 ThreadGroupCountX, uint32 ThreadGroupCountY, uint32 ThreadGroupCountZ);
     
     void DispatchIndirect(FMetalVertexBuffer* ArgumentBufferRHI, uint32 ArgumentOffset);
@@ -161,30 +156,17 @@ public:
 	void InsertTextureBarrier();
 	
 private:
-#pragma mark -
     void ConditionalSwitchToRender(void);
-    void ConditionalSwitchToTessellation(void);
-	void ConditionalSwitchToSeparateTessellation(void);
     void ConditionalSwitchToCompute(void);
 	void ConditionalSwitchToBlit(void);
 	void ConditionalSwitchToAsyncBlit(void);
 	void ConditionalSwitchToAsyncCompute(void);
 	
     void PrepareToRender(uint32 PrimType);
-#if PLATFORM_SUPPORTS_TESSELLATION_SHADERS
-    void PrepareToTessellate(uint32 PrimType);
-	void PrepareToStreamOut(uint32 PrimType);
-	void PrepareToSeparateTessellate(uint32 PrimType);
-#endif
     void PrepareToDispatch(void);
 	void PrepareToAsyncDispatch(void);
 
     void CommitRenderResourceTables(void);
-#if PLATFORM_SUPPORTS_TESSELLATION_SHADERS
-    void CommitTessellationResourceTables(void);
-	void CommitStreamOutResourceTables(void);
-	void CommitSeparateTessellationResourceTables(void);
-#endif
     void CommitDispatchResourceTables(void);
 	void CommitAsyncDispatchResourceTables(void);
     
@@ -195,8 +177,8 @@ private:
 	
 	void InsertDebugDraw(FMetalCommandData& Data);
 	void InsertDebugDispatch(FMetalCommandData& Data);
+
 private:
-#pragma mark -
 	FMetalCommandList& CmdList;
     FMetalStateCache& State;
     
diff --git a/Engine/Source/Runtime/Apple/MetalRHI/Private/MetalShaderTypes.h b/Engine/Source/Runtime/Apple/MetalRHI/Private/MetalShaderTypes.h
--- a/Engine/Source/Runtime/Apple/MetalRHI/Private/MetalShaderTypes.h
+++ b/Engine/Source/Runtime/Apple/MetalRHI/Private/MetalShaderTypes.h
@@ -8,8 +8,6 @@
 
 #include "Shaders/Types/Templates/MetalBaseShader.h"
 #include "Shaders/Types/MetalVertexShader.h"
-#include "Shaders/Types/MetalHullShader.h"
-#include "Shaders/Types/MetalDomainShader.h"
 #include "Shaders/Types/MetalGeometryShader.h"
 #include "Shaders/Types/MetalPixelShader.h"
 #include "Shaders/Types/MetalComputeShader.h"
diff --git a/Engine/Source/Runtime/Apple/MetalRHI/Private/MetalStateCache.cpp b/Engine/Source/Runtime/Apple/MetalRHI/Private/MetalStateCache.cpp
--- a/Engine/Source/Runtime/Apple/MetalRHI/Private/MetalStateCache.cpp
+++ b/Engine/Source/Runtime/Apple/MetalRHI/Private/MetalStateCache.cpp
@@ -176,9 +176,6 @@ FMetalStateCache::FMetalStateCache(bool const bInImmediate)
 , bHasValidRenderTarget(false)
 , bHasValidColorTarget(false)
 , bScissorRectEnabled(false)
-#if PLATFORM_SUPPORTS_TESSELLATION_SHADERS
-, bUsingTessellation(false)
-#endif
 , bCanRestartRenderPass(false)
 , bImmediate(bInImmediate)
 , bFallbackDepthStencilBound(false)
@@ -240,7 +237,6 @@ FMetalStateCache::~FMetalStateCache()
 
 void FMetalStateCache::Reset(void)
 {
-	IndexType = EMetalIndexType_None;
 	SampleCount = 0;
 	
 	FMemory::Memzero(Viewport);
@@ -314,9 +310,6 @@ void FMetalStateCache::Reset(void)
 	BlendFactor = FLinearColor::Transparent;
 	FrameBufferSize = CGSizeMake(0.0, 0.0);
 	RenderTargetArraySize = 0;
-#if PLATFORM_SUPPORTS_TESSELLATION_SHADERS
-    bUsingTessellation = false;
-#endif
     bCanRestartRenderPass = false;
     
     RasterBits = EMetalRenderFlagMask;
@@ -1098,62 +1091,7 @@ void FMetalStateCache::SetGraphicsPipelineState(FMetalGraphicsPipelineState* Sta
 	if (GraphicsPSO != State)
 	{
 		GraphicsPSO = State;
-		
-#if PLATFORM_SUPPORTS_TESSELLATION_SHADERS
-		bool bNewUsingTessellation = (State && IsValidRef(State->DomainShader));
-		if (bNewUsingTessellation != bUsingTessellation)
-		{
-			for (uint32 i = 0; i < EMetalShaderStages::Num; i++)
-			{
-				ShaderBuffers[i].Bound = UINT32_MAX;
-				ShaderTextures[i].Bound = FMETALTEXTUREMASK_MAX;
-				ShaderSamplers[i].Bound = UINT16_MAX;
-			}
-		}
-		// Whenever the pipeline changes & a Hull shader is bound clear the Hull shader bindings, otherwise the Hull resources from a
-		// previous pipeline with different binding table will overwrite the vertex shader bindings for the current pipeline.
-		if (bNewUsingTessellation)
-		{
-			ShaderBuffers[EMetalShaderStages::Hull].Bound = UINT32_MAX;
-			ShaderTextures[EMetalShaderStages::Hull].Bound = FMETALTEXTUREMASK_MAX;
-			ShaderSamplers[EMetalShaderStages::Hull].Bound = UINT16_MAX;
-			
-			for (uint32 i = 0; i < ML_MaxBuffers; i++)
-			{
-				BoundUniformBuffers[EMetalShaderStages::Hull][i] = nullptr;
-				ShaderBuffers[EMetalShaderStages::Hull].Buffers[i].Buffer = nil;
-				ShaderBuffers[EMetalShaderStages::Hull].Buffers[i].Bytes = nil;
-				ShaderBuffers[EMetalShaderStages::Hull].Buffers[i].Length = 0;
-                ShaderBuffers[EMetalShaderStages::Hull].Buffers[i].Offset = 0;
-                ShaderBuffers[EMetalShaderStages::Hull].Buffers[i].Usage = mtlpp::ResourceUsage(0);
-				ShaderBuffers[EMetalShaderStages::Hull].Formats[i] = PF_Unknown;
-			}
-			for (uint32 i = 0; i < ML_MaxTextures; i++)
-			{
-				ShaderTextures[EMetalShaderStages::Hull].Textures[i] = nil;
-				ShaderTextures[EMetalShaderStages::Hull].Usage[i] = mtlpp::ResourceUsage(0);
-			}
-			
-			for (uint32 i = 0; i < ML_MaxSamplers; i++)
-			{
-				ShaderSamplers[EMetalShaderStages::Hull].Samplers[i] = nil;
-			}
-
-			for (const auto& PackedGlobalArray : State->HullShader->Bindings.PackedGlobalArrays)
-			{
-				ShaderParameters[EMetalShaderStages::Hull].PrepareGlobalUniforms(CrossCompiler::PackedTypeNameToTypeIndex(PackedGlobalArray.TypeName), PackedGlobalArray.Size);
-			}
-
-			for (const auto& PackedGlobalArray : State->DomainShader->Bindings.PackedGlobalArrays)
-			{
-				ShaderParameters[EMetalShaderStages::Domain].PrepareGlobalUniforms(CrossCompiler::PackedTypeNameToTypeIndex(PackedGlobalArray.TypeName), PackedGlobalArray.Size);
-			}
-		}
-		bUsingTessellation = bNewUsingTessellation;
-		DirtyUniformBuffers[EMetalShaderStages::Hull] = 0xffffffff;
-		DirtyUniformBuffers[EMetalShaderStages::Domain] = 0xffffffff;
-#endif
-		
+				
 		DirtyUniformBuffers[EMetalShaderStages::Vertex] = 0xffffffff;
 		DirtyUniformBuffers[EMetalShaderStages::Pixel] = 0xffffffff;
 #if PLATFORM_SUPPORTS_GEOMETRY_SHADERS
@@ -1192,7 +1130,7 @@ void FMetalStateCache::SetGraphicsPipelineState(FMetalGraphicsPipelineState* Sta
 
 FMetalShaderPipeline* FMetalStateCache::GetPipelineState() const
 {
-	return GraphicsPSO->GetPipeline(GetIndexType());
+	return GraphicsPSO->GetPipeline();
 }
 
 EPrimitiveType FMetalStateCache::GetPrimitiveType()
@@ -1201,16 +1139,6 @@ EPrimitiveType FMetalStateCache::GetPrimitiveType()
 	return GraphicsPSO->GetPrimitiveType();
 }
 
-void FMetalStateCache::SetIndexType(EMetalIndexType InIndexType)
-{
-	if (IndexType != InIndexType)
-	{
-		IndexType = InIndexType;
-		
-		PipelineBits |= EMetalPipelineFlagPipelineState;
-	}
-}
-
 void FMetalStateCache::BindUniformBuffer(EMetalShaderStages const Freq, uint32 const BufferIndex, FRHIUniformBuffer* BufferRHI)
 {
 	check(BufferIndex < ML_MaxBuffers);
@@ -1486,15 +1414,6 @@ void FMetalStateCache::SetResource(uint32 ShaderStage, uint32 BindIndex, FRHITex
 		case CrossCompiler::SHADER_STAGE_COMPUTE:
 			SetShaderTexture(EMetalShaderStages::Compute, Texture, BindIndex, Usage);
 			break;
-#if PLATFORM_SUPPORTS_TESSELLATION_SHADERS
-		case CrossCompiler::SHADER_STAGE_HULL:
-			SetShaderTexture(EMetalShaderStages::Hull, Texture, BindIndex, Usage);
-			break;
-			
-		case CrossCompiler::SHADER_STAGE_DOMAIN:
-			SetShaderTexture(EMetalShaderStages::Domain, Texture, BindIndex, Usage);
-			break;
-#endif
 			
 		default:
 			check(0);
@@ -1566,18 +1485,6 @@ bool FMetalStateCache::IsLinearBuffer(EMetalShaderStages ShaderStage, uint32 Bin
             return (GraphicsPSO->PixelShader->Bindings.LinearBuffer & (1 << BindIndex)) != 0;
             break;
         }
-#if PLATFORM_SUPPORTS_TESSELLATION_SHADERS
-        case EMetalShaderStages::Hull:
-        {
-            return (GraphicsPSO->HullShader->Bindings.LinearBuffer & (1 << BindIndex)) != 0;
-            break;
-        }
-        case EMetalShaderStages::Domain:
-        {
-            return (GraphicsPSO->DomainShader->Bindings.LinearBuffer & (1 << BindIndex)) != 0;
-            break;
-        }
-#endif
         case EMetalShaderStages::Compute:
         {
             return (ComputeShader->Bindings.LinearBuffer & (1 << BindIndex)) != 0;
@@ -1681,16 +1588,6 @@ void FMetalStateCache::SetResource(uint32 ShaderStage, uint32 BindIndex, FMetalS
 			SetShaderResourceView(nullptr, EMetalShaderStages::Compute, BindIndex, SRV);
 			break;
 			
-#if PLATFORM_SUPPORTS_TESSELLATION_SHADERS
-		case CrossCompiler::SHADER_STAGE_HULL:
-			SetShaderResourceView(nullptr, EMetalShaderStages::Hull, BindIndex, SRV);
-			break;
-			
-		case CrossCompiler::SHADER_STAGE_DOMAIN:
-			SetShaderResourceView(nullptr, EMetalShaderStages::Domain, BindIndex, SRV);
-			break;
-#endif
-			
 		default:
 			check(0);
 			break;
@@ -1714,16 +1611,6 @@ void FMetalStateCache::SetResource(uint32 ShaderStage, uint32 BindIndex, FMetalS
 			SetShaderSamplerState(EMetalShaderStages::Compute, SamplerState, BindIndex);
 			break;
 			
-#if PLATFORM_SUPPORTS_TESSELLATION_SHADERS
-		case CrossCompiler::SHADER_STAGE_HULL:
-			SetShaderSamplerState(EMetalShaderStages::Hull, SamplerState, BindIndex);
-			break;
-			
-		case CrossCompiler::SHADER_STAGE_DOMAIN:
-			SetShaderSamplerState(EMetalShaderStages::Domain, SamplerState, BindIndex);
-			break;
-#endif
-			
 		default:
 			check(0);
 			break;
@@ -1746,16 +1633,6 @@ void FMetalStateCache::SetResource(uint32 ShaderStage, uint32 BindIndex, FMetalU
 			SetShaderUnorderedAccessView(EMetalShaderStages::Compute, BindIndex, UAV);
 			break;
 			
-#if PLATFORM_SUPPORTS_TESSELLATION_SHADERS
-		case CrossCompiler::SHADER_STAGE_HULL:
-			SetShaderUnorderedAccessView(EMetalShaderStages::Hull, BindIndex, UAV);
-			break;
-			
-		case CrossCompiler::SHADER_STAGE_DOMAIN:
-			SetShaderUnorderedAccessView(EMetalShaderStages::Domain, BindIndex, UAV);
-			break;
-#endif
-			
 		default:
 			check(0);
 			break;
@@ -1802,14 +1679,6 @@ void FMetalStateCache::SetResourcesFromTables(ShaderType Shader, uint32 ShaderSt
 		case CrossCompiler::SHADER_STAGE_VERTEX:
 			Frequency = EMetalShaderStages::Vertex;
 			break;
-#if PLATFORM_SUPPORTS_TESSELLATION_SHADERS
-		case CrossCompiler::SHADER_STAGE_HULL:
-			Frequency = EMetalShaderStages::Hull;
-			break;
-		case CrossCompiler::SHADER_STAGE_DOMAIN:
-			Frequency = EMetalShaderStages::Domain;
-			break;
-#endif
 		case CrossCompiler::SHADER_STAGE_PIXEL:
 			Frequency = EMetalShaderStages::Pixel;
 			break;
@@ -1862,55 +1731,6 @@ void FMetalStateCache::CommitRenderResources(FMetalCommandEncoder* Raster)
     }
 }
 
-#if PLATFORM_SUPPORTS_TESSELLATION_SHADERS
-void FMetalStateCache::CommitTessellationResources(FMetalCommandEncoder* Raster, FMetalCommandEncoder* Compute)
-{
-	check(IsValidRef(GraphicsPSO));
-    check(IsValidRef(GraphicsPSO->HullShader) && IsValidRef(GraphicsPSO->DomainShader));
-    
-    SetResourcesFromTables(GraphicsPSO->VertexShader, CrossCompiler::SHADER_STAGE_VERTEX);
-    GetShaderParameters(EMetalShaderStages::Vertex).CommitPackedGlobals(this, Compute, EMetalShaderStages::Vertex, GraphicsPSO->VertexShader->Bindings);
-	
-    if (IsValidRef(GraphicsPSO->PixelShader))
-    {
-    	SetResourcesFromTables(GraphicsPSO->PixelShader, CrossCompiler::SHADER_STAGE_PIXEL);
-        GetShaderParameters(EMetalShaderStages::Pixel).CommitPackedGlobals(this, Raster, EMetalShaderStages::Pixel, GraphicsPSO->PixelShader->Bindings);
-    }
-    
-    SetResourcesFromTables(GraphicsPSO->HullShader, CrossCompiler::SHADER_STAGE_HULL);
-	
-	SetResourcesFromTables(GraphicsPSO->DomainShader, CrossCompiler::SHADER_STAGE_DOMAIN);
-    GetShaderParameters(EMetalShaderStages::Domain).CommitPackedGlobals(this, Raster, EMetalShaderStages::Domain, GraphicsPSO->DomainShader->Bindings);
-}
-
-void FMetalStateCache::CommitVertexStreamResources(FMetalCommandEncoder* Raster)
-{
-	check(IsValidRef(GraphicsPSO));
-	check(IsValidRef(GraphicsPSO->HullShader) && IsValidRef(GraphicsPSO->DomainShader));
-	
-	SetResourcesFromTables(GraphicsPSO->VertexShader, CrossCompiler::SHADER_STAGE_VERTEX);
-	GetShaderParameters(EMetalShaderStages::Vertex).CommitPackedGlobals(this, Raster, EMetalShaderStages::Vertex, GraphicsPSO->VertexShader->Bindings);
-}
-
-void FMetalStateCache::CommitSeparateTessellationResources(FMetalCommandEncoder* Raster, FMetalCommandEncoder* Compute)
-{
-	check(IsValidRef(GraphicsPSO));
-	check(IsValidRef(GraphicsPSO->HullShader) && IsValidRef(GraphicsPSO->DomainShader));
-	
-	if (IsValidRef(GraphicsPSO->PixelShader))
-	{
-		SetResourcesFromTables(GraphicsPSO->PixelShader, CrossCompiler::SHADER_STAGE_PIXEL);
-		GetShaderParameters(EMetalShaderStages::Pixel).CommitPackedGlobals(this, Raster, EMetalShaderStages::Pixel, GraphicsPSO->PixelShader->Bindings);
-	}
-	
-	SetResourcesFromTables(GraphicsPSO->HullShader, CrossCompiler::SHADER_STAGE_HULL);
-	GetShaderParameters(EMetalShaderStages::Hull).CommitPackedGlobals(this, Compute, EMetalShaderStages::Hull, GraphicsPSO->HullShader->Bindings);
-	
-	SetResourcesFromTables(GraphicsPSO->DomainShader, CrossCompiler::SHADER_STAGE_DOMAIN);
-	GetShaderParameters(EMetalShaderStages::Domain).CommitPackedGlobals(this, Raster, EMetalShaderStages::Domain, GraphicsPSO->DomainShader->Bindings);
-}
-#endif
-
 void FMetalStateCache::CommitComputeResources(FMetalCommandEncoder* Compute)
 {
 	check(IsValidRef(ComputeShader));
@@ -2122,78 +1942,6 @@ void FMetalStateCache::SetRenderState(FMetalCommandEncoder& CommandEncoder, FMet
 	}
 }
 
-#if PLATFORM_SUPPORTS_TESSELLATION_SHADERS
-void FMetalStateCache::SetStreamOutPipelineState(FMetalCommandEncoder& CommandEncoder)
-{
-	SCOPE_CYCLE_COUNTER(STAT_MetalSetRenderPipelineStateTime);
-	FMetalShaderPipeline* Pipeline = GetPipelineState();
-	
-	auto& StreamOutEncoder = CommandEncoder.GetRenderCommandEncoder();
-	StreamOutEncoder.SetRenderPipelineState(Pipeline->StreamPipelineState);
-	
-	FMetalGraphicsPipelineState* boundShaderState = GetGraphicsPSO();
-	
-	if (SafeGetRuntimeDebuggingLevel() >= EMetalDebugLevelFastValidation)
-	{
-		EMetalShaderStages VertexStage = EMetalShaderStages::Vertex;
-
-		FMetalDebugShaderResourceMask VertexMask = Pipeline->ResourceMask[EMetalShaderStream];
-		TArray<uint32>& MinVertexBufferSizes = Pipeline->BufferDataSizes[EMetalShaderStream];
-		TMap<uint8, uint8> VertexTexTypes = Pipeline->TextureTypes[EMetalShaderStream];
-		while(VertexMask.BufferMask)
-		{
-			uint32 Index = __builtin_ctz(VertexMask.BufferMask);
-			VertexMask.BufferMask &= ~(1 << Index);
-			
-			if ((boundShaderState->VertexShader->TessellationPatchCountBuffer != Index
-																&& boundShaderState->VertexShader->TessellationHSOutBuffer != Index))
-			{
-				FMetalBufferBinding const& Binding = ShaderBuffers[VertexStage].Buffers[Index];
-				ensure(Binding.Buffer || Binding.Bytes);
-				ensure(MinVertexBufferSizes.Num() > Index);
-				ensure(Binding.Length >= MinVertexBufferSizes[Index]);
-			}
-		}
-#if PLATFORM_MAC
-		{
-			uint64 LoTextures = (uint64)VertexMask.TextureMask;
-			while(LoTextures)
-			{
-				uint32 Index = __builtin_ctzll(LoTextures);
-				LoTextures &= ~(uint64(1) << uint64(Index));
-				ensure(ShaderTextures[VertexStage].Textures[Index]);
-				ensure(ShaderTextures[VertexStage].Textures[Index].GetTextureType() == (mtlpp::TextureType)VertexTexTypes.FindRef(Index));
-			}
-			
-			uint64 HiTextures = (uint64)(VertexMask.TextureMask >> FMetalTextureMask(64));
-			while(HiTextures)
-			{
-				uint32 Index = __builtin_ctzll(HiTextures);
-				HiTextures &= ~(uint64(1) << uint64(Index));
-				ensure(ShaderTextures[VertexStage].Textures[Index + 64]);
-				ensure(ShaderTextures[VertexStage].Textures[Index + 64].GetTextureType() == (mtlpp::TextureType)VertexTexTypes.FindRef(Index + 64));
-			}
-		}
-#else
-		while(VertexMask.TextureMask)
-		{
-			uint32 Index = __builtin_ctz(VertexMask.TextureMask);
-			VertexMask.TextureMask &= ~(1 << Index);
-			
-			ensure(ShaderTextures[VertexStage].Textures[Index]);
-			ensure(ShaderTextures[VertexStage].Textures[Index].GetTextureType() == (mtlpp::TextureType)VertexTexTypes.FindRef(Index));
-		}
-#endif
-		while(VertexMask.SamplerMask)
-		{
-			uint32 Index = __builtin_ctz(VertexMask.SamplerMask);
-			VertexMask.SamplerMask &= ~(1 << Index);
-			ensure(ShaderSamplers[VertexStage].Samplers[Index]);
-		}
-	}
-}
-#endif
-
 void FMetalStateCache::SetRenderPipelineState(FMetalCommandEncoder& CommandEncoder, FMetalCommandEncoder* PrologueEncoder)
 {
 	SCOPE_CYCLE_COUNTER(STAT_MetalSetRenderPipelineStateTime);
@@ -2219,74 +1967,6 @@ void FMetalStateCache::SetRenderPipelineState(FMetalCommandEncoder& CommandEncod
 	{
 		FMetalShaderPipeline* Pipeline = GetPipelineState();
 		EMetalShaderStages VertexStage = EMetalShaderStages::Vertex;
-#if PLATFORM_SUPPORTS_TESSELLATION_SHADERS
-		if (GraphicsPSO->DomainShader)
-		{
-			VertexStage = EMetalShaderStages::Domain;
-			
-			FMetalDebugShaderResourceMask ComputeMask = Pipeline->ResourceMask[EMetalShaderCompute];
-			TArray<uint32>& MinComputeBufferSizes = Pipeline->BufferDataSizes[EMetalShaderCompute];
-			TMap<uint8, uint8> ComputeTexTypes = Pipeline->TextureTypes[EMetalShaderCompute];
-			while(ComputeMask.BufferMask)
-			{
-				uint32 Index = __builtin_ctz(ComputeMask.BufferMask);
-				ComputeMask.BufferMask &= ~(1 << Index);
-				
-				if (Pipeline->TessellationPipelineDesc.TessellationControlPointIndexBufferIndex != Index
-					&& Pipeline->TessellationPipelineDesc.TessellationIndexBufferIndex != Index
-					&& Pipeline->TessellationPipelineDesc.TessellationOutputControlPointBufferIndex != Index
-					&& Pipeline->TessellationPipelineDesc.TessellationPatchConstBufferIndex != Index
-					&& Pipeline->TessellationPipelineDesc.TessellationFactorBufferIndex != Index
-					&& Pipeline->TessellationPipelineDesc.TessellationPatchCountBufferIndex != Index)
-				{
-					FMetalBufferBinding const& Binding = ShaderBuffers[EMetalShaderStages::Vertex].Buffers[Index];
-					FMetalBufferBinding const& HullBinding = ShaderBuffers[EMetalShaderStages::Hull].Buffers[Index];
-					ensure(Binding.Buffer || Binding.Bytes || HullBinding.Buffer || HullBinding.Bytes);
-					ensure(MinComputeBufferSizes.Num() > Index);
-					ensure(Binding.Length >= MinComputeBufferSizes[Index] || HullBinding.Length >= MinComputeBufferSizes[Index]);
-				}
-			}
-#if PLATFORM_MAC
-			{
-				uint64 LoTextures = (uint64)ComputeMask.TextureMask;
-				while(LoTextures)
-				{
-					uint32 Index = __builtin_ctzll(LoTextures);
-					LoTextures &= ~(uint64(1) << uint64(Index));
-					ensure(ShaderTextures[EMetalShaderStages::Vertex].Textures[Index] || ShaderTextures[EMetalShaderStages::Hull].Textures[Index + 64]);
-					ensure(!ShaderTextures[EMetalShaderStages::Vertex].Textures[Index] || ShaderTextures[EMetalShaderStages::Vertex].Textures[Index].GetTextureType() == (mtlpp::TextureType)ComputeTexTypes.FindRef(Index));
-					ensure(!ShaderTextures[EMetalShaderStages::Hull].Textures[Index] || ShaderTextures[EMetalShaderStages::Hull].Textures[Index].GetTextureType() == (mtlpp::TextureType)ComputeTexTypes.FindRef(Index));
-				}
-				
-				uint64 HiTextures = (uint64)(ComputeMask.TextureMask >> FMetalTextureMask(64));
-				while(HiTextures)
-				{
-					uint32 Index = __builtin_ctzll(HiTextures);
-					HiTextures &= ~(uint64(1) << uint64(Index));
-					ensure(ShaderTextures[EMetalShaderStages::Vertex].Textures[Index + 64] || ShaderTextures[EMetalShaderStages::Hull].Textures[Index + 64]);
-					ensure(!ShaderTextures[EMetalShaderStages::Vertex].Textures[Index + 64] || ShaderTextures[EMetalShaderStages::Vertex].Textures[Index + 64].GetTextureType() == (mtlpp::TextureType)ComputeTexTypes.FindRef(Index + 64));
-					ensure(!ShaderTextures[EMetalShaderStages::Hull].Textures[Index + 64] || ShaderTextures[EMetalShaderStages::Hull].Textures[Index + 64].GetTextureType() == (mtlpp::TextureType)ComputeTexTypes.FindRef(Index + 64));
-				}
-			}
-#else
-			while(ComputeMask.TextureMask)
-			{
-				uint32 Index = __builtin_ctz(ComputeMask.TextureMask);
-				ComputeMask.TextureMask &= ~(1 << Index);
-				
-				ensure(ShaderTextures[EMetalShaderStages::Vertex].Textures[Index] || ShaderTextures[EMetalShaderStages::Hull].Textures[Index]);
-				ensure(!ShaderTextures[EMetalShaderStages::Vertex].Textures[Index] || ShaderTextures[EMetalShaderStages::Vertex].Textures[Index].GetTextureType() == (mtlpp::TextureType)ComputeTexTypes.FindRef(Index));
-				ensure(!ShaderTextures[EMetalShaderStages::Hull].Textures[Index] || ShaderTextures[EMetalShaderStages::Hull].Textures[Index].GetTextureType() == (mtlpp::TextureType)ComputeTexTypes.FindRef(Index));
-			}
-#endif
-			while(ComputeMask.SamplerMask)
-			{
-				uint32 Index = __builtin_ctz(ComputeMask.SamplerMask);
-				ComputeMask.SamplerMask &= ~(1 << Index);
-				ensure(ShaderSamplers[EMetalShaderStages::Vertex].Samplers[Index] || ShaderSamplers[EMetalShaderStages::Hull].Samplers[Index]);
-			}
-		}
-#endif
 		
 		FMetalDebugShaderResourceMask VertexMask = Pipeline->ResourceMask[EMetalShaderVertex];
 		TArray<uint32>& MinVertexBufferSizes = Pipeline->BufferDataSizes[EMetalShaderVertex];
@@ -2296,12 +1976,7 @@ void FMetalStateCache::SetRenderPipelineState(FMetalCommandEncoder& CommandEncod
 			uint32 Index = __builtin_ctz(VertexMask.BufferMask);
 			VertexMask.BufferMask &= ~(1 << Index);
 			
-			if (VertexStage == EMetalShaderStages::Vertex
-#if PLATFORM_SUPPORTS_TESSELLATION_SHADERS
-				|| (Pipeline->TessellationPipelineDesc.TessellationInputPatchConstBufferIndex != Index
-				&& Pipeline->TessellationPipelineDesc.TessellationInputControlPointBufferIndex != Index)
-#endif // PLATFORM_SUPPORTS_TESSELLATION_SHADERS
-				)
+			if (VertexStage == EMetalShaderStages::Vertex)
 			{
 				FMetalBufferBinding const& Binding = ShaderBuffers[VertexStage].Buffers[Index];
 				ensure(Binding.Buffer || Binding.Bytes);
@@ -2399,203 +2074,6 @@ void FMetalStateCache::SetRenderPipelineState(FMetalCommandEncoder& CommandEncod
 #endif // METAL_DEBUG_OPTIONS
 }
 
-#if PLATFORM_SUPPORTS_TESSELLATION_SHADERS
-void FMetalStateCache::SetTessellationPipelineState(FMetalCommandEncoder& CommandEncoder, FMetalCommandEncoder* PrologueEncoder)
-{
-	SCOPE_CYCLE_COUNTER(STAT_MetalSetRenderPipelineStateTime);
-	
-	if ((PipelineBits & EMetalPipelineFlagRasterMask) != 0)
-	{
-		// Some Intel drivers need RenderPipeline state to be set after DepthStencil state to work properly
-		FMetalShaderPipeline* Pipeline = GetPipelineState();
-		
-		check(Pipeline);
-		CommandEncoder.SetRenderPipelineState(Pipeline);
-		if (Pipeline->ComputePipelineState)
-		{
-			check(PrologueEncoder);
-			PrologueEncoder->SetComputePipelineState(Pipeline);
-		}
-		
-		PipelineBits &= EMetalPipelineFlagComputeMask;
-	}
-	
-	if (SafeGetRuntimeDebuggingLevel() >= EMetalDebugLevelFastValidation)
-	{
-		FMetalShaderPipeline* Pipeline = GetPipelineState();
-		EMetalShaderStages VertexStage = EMetalShaderStages::Vertex;
-#if PLATFORM_SUPPORTS_TESSELLATION_SHADERS
-		if (GraphicsPSO->DomainShader)
-		{
-			VertexStage = EMetalShaderStages::Domain;
-			
-			FMetalDebugShaderResourceMask ComputeMask = Pipeline->ResourceMask[EMetalShaderCompute];
-			TArray<uint32>& MinComputeBufferSizes = Pipeline->BufferDataSizes[EMetalShaderCompute];
-			TMap<uint8, uint8> ComputeTexTypes = Pipeline->TextureTypes[EMetalShaderCompute];
-			while(ComputeMask.BufferMask)
-			{
-				uint32 Index = __builtin_ctz(ComputeMask.BufferMask);
-				ComputeMask.BufferMask &= ~(1 << Index);
-				
-				if (GraphicsPSO->HullShader->TessellationPatchCountBuffer != Index
-					&& GraphicsPSO->HullShader->TessellationIndexBuffer != Index
-					&& GraphicsPSO->HullShader->TessellationHSOutBuffer != Index
-					&& GraphicsPSO->HullShader->TessellationHSTFOutBuffer != Index
-					&& GraphicsPSO->HullShader->TessellationControlPointOutBuffer != Index
-					&& GraphicsPSO->HullShader->TessellationControlPointIndexBuffer != Index)
-				{
-					FMetalBufferBinding const& HullBinding = ShaderBuffers[EMetalShaderStages::Hull].Buffers[Index];
-					ensure(HullBinding.Buffer || HullBinding.Bytes);
-					ensure(MinComputeBufferSizes.Num() > Index);
-					ensure(HullBinding.Length >= MinComputeBufferSizes[Index]);
-				}
-			}
-#if PLATFORM_MAC
-			{
-				uint64 LoTextures = (uint64)ComputeMask.TextureMask;
-				while(LoTextures)
-				{
-					uint32 Index = __builtin_ctzll(LoTextures);
-					LoTextures &= ~(uint64(1) << uint64(Index));
-					ensure(ShaderTextures[EMetalShaderStages::Hull].Textures[Index + 64]);
-					ensure(ShaderTextures[EMetalShaderStages::Hull].Textures[Index].GetTextureType() == (mtlpp::TextureType)ComputeTexTypes.FindRef(Index));
-				}
-				
-				uint64 HiTextures = (uint64)(ComputeMask.TextureMask >> FMetalTextureMask(64));
-				while(HiTextures)
-				{
-					uint32 Index = __builtin_ctzll(HiTextures);
-					HiTextures &= ~(uint64(1) << uint64(Index));
-					ensure(ShaderTextures[EMetalShaderStages::Hull].Textures[Index + 64]);
-					ensure(ShaderTextures[EMetalShaderStages::Hull].Textures[Index + 64].GetTextureType() == (mtlpp::TextureType)ComputeTexTypes.FindRef(Index + 64));
-				}
-			}
-#else
-			while(ComputeMask.TextureMask)
-			{
-				uint32 Index = __builtin_ctz(ComputeMask.TextureMask);
-				ComputeMask.TextureMask &= ~(1 << Index);
-				
-				ensure(ShaderTextures[EMetalShaderStages::Hull].Textures[Index]);
-				ensure(ShaderTextures[EMetalShaderStages::Hull].Textures[Index].GetTextureType() == (mtlpp::TextureType)ComputeTexTypes.FindRef(Index));
-			}
-#endif
-			while(ComputeMask.SamplerMask)
-			{
-				uint32 Index = __builtin_ctz(ComputeMask.SamplerMask);
-				ComputeMask.SamplerMask &= ~(1 << Index);
-				ensure(ShaderSamplers[EMetalShaderStages::Hull].Samplers[Index]);
-			}
-		}
-#endif
-		
-		FMetalDebugShaderResourceMask VertexMask = Pipeline->ResourceMask[EMetalShaderVertex];
-		TArray<uint32>& MinVertexBufferSizes = Pipeline->BufferDataSizes[EMetalShaderVertex];
-		TMap<uint8, uint8> VertexTexTypes = Pipeline->TextureTypes[EMetalShaderVertex];
-		while(VertexMask.BufferMask)
-		{
-			uint32 Index = __builtin_ctz(VertexMask.BufferMask);
-			VertexMask.BufferMask &= ~(1 << Index);
-			
-			if ((VertexStage == EMetalShaderStages::Vertex) || (GraphicsPSO->DomainShader->TessellationHSOutBuffer != Index
-																&& GraphicsPSO->DomainShader->TessellationControlPointOutBuffer != Index))
-			{
-				FMetalBufferBinding const& Binding = ShaderBuffers[VertexStage].Buffers[Index];
-				ensure(Binding.Buffer || Binding.Bytes);
-				ensure(MinVertexBufferSizes.Num() > Index);
-				ensure(Binding.Length >= MinVertexBufferSizes[Index]);
-			}
-		}
-#if PLATFORM_MAC
-		{
-			uint64 LoTextures = (uint64)VertexMask.TextureMask;
-			while(LoTextures)
-			{
-				uint32 Index = __builtin_ctzll(LoTextures);
-				LoTextures &= ~(uint64(1) << uint64(Index));
-				ensure(ShaderTextures[VertexStage].Textures[Index]);
-				ensure(ShaderTextures[VertexStage].Textures[Index].GetTextureType() == (mtlpp::TextureType)VertexTexTypes.FindRef(Index));
-			}
-			
-			uint64 HiTextures = (uint64)(VertexMask.TextureMask >> FMetalTextureMask(64));
-			while(HiTextures)
-			{
-				uint32 Index = __builtin_ctzll(HiTextures);
-				HiTextures &= ~(uint64(1) << uint64(Index));
-				ensure(ShaderTextures[VertexStage].Textures[Index + 64]);
-				ensure(ShaderTextures[VertexStage].Textures[Index + 64].GetTextureType() == (mtlpp::TextureType)VertexTexTypes.FindRef(Index + 64));
-			}
-		}
-#else
-		while(VertexMask.TextureMask)
-		{
-			uint32 Index = __builtin_ctz(VertexMask.TextureMask);
-			VertexMask.TextureMask &= ~(1 << Index);
-			
-			ensure(ShaderTextures[VertexStage].Textures[Index]);
-			ensure(ShaderTextures[VertexStage].Textures[Index].GetTextureType() == (mtlpp::TextureType)VertexTexTypes.FindRef(Index));
-		}
-#endif
-		while(VertexMask.SamplerMask)
-		{
-			uint32 Index = __builtin_ctz(VertexMask.SamplerMask);
-			VertexMask.SamplerMask &= ~(1 << Index);
-			ensure(ShaderSamplers[VertexStage].Samplers[Index]);
-		}
-		
-		FMetalDebugShaderResourceMask FragmentMask = Pipeline->ResourceMask[EMetalShaderFragment];
-		TArray<uint32>& MinFragmentBufferSizes = Pipeline->BufferDataSizes[EMetalShaderFragment];
-		TMap<uint8, uint8> FragmentTexTypes = Pipeline->TextureTypes[EMetalShaderFragment];
-		while(FragmentMask.BufferMask)
-		{
-			uint32 Index = __builtin_ctz(FragmentMask.BufferMask);
-			FragmentMask.BufferMask &= ~(1 << Index);
-			
-			FMetalBufferBinding const& Binding = ShaderBuffers[EMetalShaderStages::Pixel].Buffers[Index];
-			ensure(Binding.Buffer || Binding.Bytes);
-			ensure(MinFragmentBufferSizes.Num() > Index);
-			ensure(Binding.Length >= MinFragmentBufferSizes[Index]);
-		}
-#if PLATFORM_MAC
-		{
-			uint64 LoTextures = (uint64)FragmentMask.TextureMask;
-			while(LoTextures)
-			{
-				uint32 Index = __builtin_ctzll(LoTextures);
-				LoTextures &= ~(uint64(1) << uint64(Index));
-				ensure(ShaderTextures[EMetalShaderStages::Pixel].Textures[Index]);
-				ensure(ShaderTextures[EMetalShaderStages::Pixel].Textures[Index].GetTextureType() == (mtlpp::TextureType)FragmentTexTypes.FindRef(Index));
-			}
-			
-			uint64 HiTextures = (uint64)(FragmentMask.TextureMask >> FMetalTextureMask(64));
-			while(HiTextures)
-			{
-				uint32 Index = __builtin_ctzll(HiTextures);
-				HiTextures &= ~(uint64(1) << uint64(Index));
-				ensure(ShaderTextures[EMetalShaderStages::Pixel].Textures[Index + 64]);
-				ensure(ShaderTextures[EMetalShaderStages::Pixel].Textures[Index + 64].GetTextureType() == (mtlpp::TextureType)FragmentTexTypes.FindRef(Index + 64));
-			}
-		}
-#else
-		while(FragmentMask.TextureMask)
-		{
-			uint32 Index = __builtin_ctz(FragmentMask.TextureMask);
-			FragmentMask.TextureMask &= ~(1 << Index);
-			
-			ensure(ShaderTextures[EMetalShaderStages::Pixel].Textures[Index]);
-			ensure(ShaderTextures[EMetalShaderStages::Pixel].Textures[Index].GetTextureType() == (mtlpp::TextureType)FragmentTexTypes.FindRef(Index));
-		}
-#endif
-		while(FragmentMask.SamplerMask)
-		{
-			uint32 Index = __builtin_ctz(FragmentMask.SamplerMask);
-			FragmentMask.SamplerMask &= ~(1 << Index);
-			ensure(ShaderSamplers[EMetalShaderStages::Pixel].Samplers[Index]);
-		}
-	}
-}
-#endif
-
 void FMetalStateCache::SetComputePipelineState(FMetalCommandEncoder& CommandEncoder)
 {
 	if ((PipelineBits & EMetalPipelineFlagComputeMask) != 0)
diff --git a/Engine/Source/Runtime/Apple/MetalRHI/Private/MetalStateCache.h b/Engine/Source/Runtime/Apple/MetalRHI/Private/MetalStateCache.h
--- a/Engine/Source/Runtime/Apple/MetalRHI/Private/MetalStateCache.h
+++ b/Engine/Source/Runtime/Apple/MetalRHI/Private/MetalStateCache.h
@@ -55,7 +55,6 @@ public:
 	void SetViewports(const mtlpp::Viewport InViewport[], uint32 Count);
 	void SetVertexStream(uint32 const Index, FMetalBuffer* Buffer, FMetalBufferData* Bytes, uint32 const Offset, uint32 const Length);
 	void SetGraphicsPipelineState(FMetalGraphicsPipelineState* State);
-	void SetIndexType(EMetalIndexType IndexType);
 	void BindUniformBuffer(EMetalShaderStages const Freq, uint32 const BufferIndex, FRHIUniformBuffer* BufferRHI);
 	void SetDirtyUniformBuffers(EMetalShaderStages const Freq, uint32 const Dirty);
 	
@@ -111,13 +110,6 @@ public:
 
 	void CommitRenderResources(FMetalCommandEncoder* Raster);
 
-#if PLATFORM_SUPPORTS_TESSELLATION_SHADERS
-	void CommitTessellationResources(FMetalCommandEncoder* Raster, FMetalCommandEncoder* Compute);
-	
-	void CommitVertexStreamResources(FMetalCommandEncoder* Raster);
-	void CommitSeparateTessellationResources(FMetalCommandEncoder* Raster, FMetalCommandEncoder* Compute);
-#endif
-
 	void CommitComputeResources(FMetalCommandEncoder* Compute);
 	
 	void CommitResourceTable(EMetalShaderStages const Frequency, mtlpp::FunctionType const Type, FMetalCommandEncoder& CommandEncoder);
@@ -146,10 +138,6 @@ public:
 	bool GetScissorRectEnabled() const { return bScissorRectEnabled; }
 	bool NeedsToSetRenderTarget(const FRHIRenderPassInfo& RenderPassInfo);
 	bool HasValidDepthStencilSurface() const { return IsValidRef(DepthStencilSurface); }
-	EMetalIndexType GetIndexType() const { return IndexType; }
-#if PLATFORM_SUPPORTS_TESSELLATION_SHADERS
-	bool GetUsingTessellation() const { return bUsingTessellation; }
-#endif
     bool CanRestartRenderPass() const { return bCanRestartRenderPass; }
 	mtlpp::RenderPassDescriptor GetRenderPassDescriptor(void) const { return RenderPassDesc; }
 	uint32 GetSampleCount(void) const { return SampleCount; }
@@ -163,10 +151,6 @@ public:
 	bool GetFallbackDepthStencilBound(void) const { return bFallbackDepthStencilBound; }
 	
 	void SetRenderPipelineState(FMetalCommandEncoder& CommandEncoder, FMetalCommandEncoder* PrologueEncoder);
-#if PLATFORM_SUPPORTS_TESSELLATION_SHADERS
-	void SetTessellationPipelineState(FMetalCommandEncoder& CommandEncoder, FMetalCommandEncoder* PrologueEncoder);
-	void SetStreamOutPipelineState(FMetalCommandEncoder& CommandEncoder);
-#endif
     void SetComputePipelineState(FMetalCommandEncoder& CommandEncoder);
 	void FlushVisibilityResults(FMetalCommandEncoder& CommandEncoder);
 
@@ -248,7 +232,6 @@ private:
 private:
 	FMetalShaderParameterCache ShaderParameters[EMetalShaderStages::Num];
 
-	EMetalIndexType IndexType;
 	uint32 SampleCount;
 
 	TSet<TRefCountPtr<FRHIUniformBuffer>> ActiveUniformBuffers;
@@ -306,9 +289,6 @@ private:
 	bool bHasValidRenderTarget;
 	bool bHasValidColorTarget;
 	bool bScissorRectEnabled;
-#if PLATFORM_SUPPORTS_TESSELLATION_SHADERS
-    bool bUsingTessellation;
-#endif
     bool bCanRestartRenderPass;
     bool bImmediate;
 	bool bFallbackDepthStencilBound;
diff --git a/Engine/Source/Runtime/Apple/MetalRHI/Private/MetalUniformBufferCommands.cpp b/Engine/Source/Runtime/Apple/MetalRHI/Private/MetalUniformBufferCommands.cpp
--- a/Engine/Source/Runtime/Apple/MetalRHI/Private/MetalUniformBufferCommands.cpp
+++ b/Engine/Source/Runtime/Apple/MetalRHI/Private/MetalUniformBufferCommands.cpp
@@ -83,16 +83,6 @@ void FMetalRHICommandContext::RHISetShaderUniformBuffer(FRHIGraphicsShader* Shad
 			SetUniformBufferInternal<EMetalShaderStages::Vertex, FRHIVertexShader>(Context, static_cast<FRHIVertexShader*>(Shader), BufferIndex, Buffer);
 			break;
 
-#if PLATFORM_SUPPORTS_TESSELLATION_SHADERS
-		case SF_Hull:
-			SetUniformBufferInternal<EMetalShaderStages::Hull, FRHIHullShader>(Context, static_cast<FRHIHullShader*>(Shader), BufferIndex, Buffer);
-			break;
-
-		case SF_Domain:
-			SetUniformBufferInternal<EMetalShaderStages::Domain, FRHIDomainShader>(Context, static_cast<FRHIDomainShader*>(Shader), BufferIndex, Buffer);
-			break;
-#endif // PLATFORM_SUPPORTS_TESSELLATION_SHADERS
-
 		case SF_Geometry:
 			NOT_SUPPORTED("RHISetShaderUniformBuffer-Geometry");
 			break;
diff --git a/Engine/Source/Runtime/Apple/MetalRHI/Private/MetalVertexBuffer.cpp b/Engine/Source/Runtime/Apple/MetalRHI/Private/MetalVertexBuffer.cpp
--- a/Engine/Source/Runtime/Apple/MetalRHI/Private/MetalVertexBuffer.cpp
+++ b/Engine/Source/Runtime/Apple/MetalRHI/Private/MetalVertexBuffer.cpp
@@ -90,10 +90,6 @@ static uint32 MetalBufferUsage(uint32 InUsage)
 
 	if (InUsage & BUF_IndexBuffer)
 	{
-		if (RHISupportsTessellation(GMaxRHIShaderPlatform))
-		{
-			Usage |= BUF_ShaderResource;
-		}
 		Usage |= (EMetalBufferUsage_GPUOnly | EMetalBufferUsage_LinearTex);
 	}
 
@@ -700,18 +696,7 @@ FMetalResourceMultiBuffer::FMetalResourceMultiBuffer(uint32 InSize, uint32 InUsa
 	, FMetalRHIBuffer(InSize, InUsage, Type)
 	, IndexType((InStride == 2) ? mtlpp::IndexType::UInt16 : mtlpp::IndexType::UInt32)
 {
-	if (InUsage & BUF_IndexBuffer)
-	{
-		if (InSize)
-		{
-			if (RHISupportsTessellation(GMaxRHIShaderPlatform))
-			{
-				EPixelFormat Format = IndexType == mtlpp::IndexType::UInt16 ? PF_R16_UINT : PF_R32_UINT;
-				CreateLinearTexture(Format, this);
-			}
-		}
-	}
-	else if (InUsage & BUF_StructuredBuffer)
+	if (InUsage & BUF_StructuredBuffer)
 	{
 		check((InSize % InStride) == 0);
 
diff --git a/Engine/Source/Runtime/Apple/MetalRHI/Private/Shaders/MetalShaderLibrary.cpp b/Engine/Source/Runtime/Apple/MetalRHI/Private/Shaders/MetalShaderLibrary.cpp
--- a/Engine/Source/Runtime/Apple/MetalRHI/Private/Shaders/MetalShaderLibrary.cpp
+++ b/Engine/Source/Runtime/Apple/MetalRHI/Private/Shaders/MetalShaderLibrary.cpp
@@ -143,16 +143,6 @@ TRefCountPtr<FRHIShader> FMetalShaderLibrary::CreateShader(int32 Index)
 			checkf(false, TEXT("Geometry shaders not supported"));
 			break;
 
-#if PLATFORM_SUPPORTS_TESSELLATION_SHADERS
-		case SF_Hull:
-			Shader = CreateMetalShader<FMetalHullShader>(Code, Library[LibraryIndex]);
-			break;
-
-		case SF_Domain:
-			Shader = CreateMetalShader<FMetalDomainShader>(Code, Library[LibraryIndex]);
-			break;
-#endif // PLATFORM_SUPPORTS_TESSELLATION_SHADERS
-
 		case SF_Compute:
 			Shader = CreateMetalShader<FMetalComputeShader>(Code, Library[LibraryIndex]);
 			break;
diff --git a/Engine/Source/Runtime/Apple/MetalRHI/Private/Shaders/Types/MetalDomainShader.cpp b/Engine/Source/Runtime/Apple/MetalRHI/Private/Shaders/Types/MetalDomainShader.cpp
deleted file mode 100644
diff --git a/Engine/Source/Runtime/Apple/MetalRHI/Private/Shaders/Types/MetalDomainShader.h b/Engine/Source/Runtime/Apple/MetalRHI/Private/Shaders/Types/MetalDomainShader.h
deleted file mode 100644
diff --git a/Engine/Source/Runtime/Apple/MetalRHI/Private/Shaders/Types/MetalHullShader.cpp b/Engine/Source/Runtime/Apple/MetalRHI/Private/Shaders/Types/MetalHullShader.cpp
deleted file mode 100644
diff --git a/Engine/Source/Runtime/Apple/MetalRHI/Private/Shaders/Types/MetalHullShader.h b/Engine/Source/Runtime/Apple/MetalRHI/Private/Shaders/Types/MetalHullShader.h
deleted file mode 100644
diff --git a/Engine/Source/Runtime/Apple/MetalRHI/Private/Shaders/Types/MetalVertexShader.cpp b/Engine/Source/Runtime/Apple/MetalRHI/Private/Shaders/Types/MetalVertexShader.cpp
--- a/Engine/Source/Runtime/Apple/MetalRHI/Private/Shaders/Types/MetalVertexShader.cpp
+++ b/Engine/Source/Runtime/Apple/MetalRHI/Private/Shaders/Types/MetalVertexShader.cpp
@@ -19,50 +19,12 @@ FMetalVertexShader::FMetalVertexShader(TArrayView<const uint8> InCode)
 {
 	FMetalCodeHeader Header;
 	Init(InCode, Header);
-
-#if PLATFORM_SUPPORTS_TESSELLATION_SHADERS
-	if (Header.Tessellation.Num())
-	{
-		auto const& Tess = Header.Tessellation[0];
-		TessellationOutputAttribs = Tess.TessellationOutputAttribs;
-		TessellationPatchCountBuffer = Tess.TessellationPatchCountBuffer;
-		TessellationIndexBuffer = Tess.TessellationIndexBuffer;
-		TessellationHSOutBuffer = Tess.TessellationHSOutBuffer;
-		TessellationHSTFOutBuffer = Tess.TessellationHSTFOutBuffer;
-		TessellationControlPointOutBuffer = Tess.TessellationControlPointOutBuffer;
-		TessellationControlPointIndexBuffer = Tess.TessellationControlPointIndexBuffer;
-		TessellationOutputControlPoints = Tess.TessellationOutputControlPoints;
-		TessellationDomain = Tess.TessellationDomain;
-		TessellationInputControlPoints = Tess.TessellationInputControlPoints;
-		TessellationMaxTessFactor = Tess.TessellationMaxTessFactor;
-		TessellationPatchesPerThreadGroup = Tess.TessellationPatchesPerThreadGroup;
-	}
-#endif // PLATFORM_SUPPORTS_TESSELLATION_SHADERS
 }
 
 FMetalVertexShader::FMetalVertexShader(TArrayView<const uint8> InCode, mtlpp::Library InLibrary)
 {
 	FMetalCodeHeader Header;
 	Init(InCode, Header, InLibrary);
-
-#if PLATFORM_SUPPORTS_TESSELLATION_SHADERS
-	if (Header.Tessellation.Num())
-	{
-		auto const& Tess = Header.Tessellation[0];
-		TessellationOutputAttribs = Tess.TessellationOutputAttribs;
-		TessellationPatchCountBuffer = Tess.TessellationPatchCountBuffer;
-		TessellationIndexBuffer = Tess.TessellationIndexBuffer;
-		TessellationHSOutBuffer = Tess.TessellationHSOutBuffer;
-		TessellationHSTFOutBuffer = Tess.TessellationHSTFOutBuffer;
-		TessellationControlPointOutBuffer = Tess.TessellationControlPointOutBuffer;
-		TessellationControlPointIndexBuffer = Tess.TessellationControlPointIndexBuffer;
-		TessellationOutputControlPoints = Tess.TessellationOutputControlPoints;
-		TessellationDomain = Tess.TessellationDomain;
-		TessellationInputControlPoints = Tess.TessellationInputControlPoints;
-		TessellationMaxTessFactor = Tess.TessellationMaxTessFactor;
-		TessellationPatchesPerThreadGroup = Tess.TessellationPatchesPerThreadGroup;
-	}
-#endif // PLATFORM_SUPPORTS_TESSELLATION_SHADERS
 }
 
 mtlpp::Function FMetalVertexShader::GetFunction()
diff --git a/Engine/Source/Runtime/Apple/MetalRHI/Private/Shaders/Types/MetalVertexShader.h b/Engine/Source/Runtime/Apple/MetalRHI/Private/Shaders/Types/MetalVertexShader.h
--- a/Engine/Source/Runtime/Apple/MetalRHI/Private/Shaders/Types/MetalVertexShader.h
+++ b/Engine/Source/Runtime/Apple/MetalRHI/Private/Shaders/Types/MetalVertexShader.h
@@ -19,20 +19,4 @@ public:
 	FMetalVertexShader(TArrayView<const uint8> InCode, mtlpp::Library InLibrary);
 
 	mtlpp::Function GetFunction();
-
-#if PLATFORM_SUPPORTS_TESSELLATION_SHADERS
-	// for VSHS
-	FMetalTessellationOutputs TessellationOutputAttribs;
-	float  TessellationMaxTessFactor;
-	uint32 TessellationOutputControlPoints;
-	uint32 TessellationDomain;
-	uint32 TessellationInputControlPoints;
-	uint32 TessellationPatchesPerThreadGroup;
-	uint32 TessellationPatchCountBuffer;
-	uint32 TessellationIndexBuffer;
-	uint32 TessellationHSOutBuffer;
-	uint32 TessellationHSTFOutBuffer;
-	uint32 TessellationControlPointOutBuffer;
-	uint32 TessellationControlPointIndexBuffer;
-#endif // PLATFORM_SUPPORTS_TESSELLATION_SHADERS
 };
diff --git a/Engine/Source/Runtime/Apple/MetalRHI/Public/MetalResources.h b/Engine/Source/Runtime/Apple/MetalRHI/Public/MetalResources.h
--- a/Engine/Source/Runtime/Apple/MetalRHI/Public/MetalResources.h
+++ b/Engine/Source/Runtime/Apple/MetalRHI/Public/MetalResources.h
@@ -22,15 +22,6 @@ class FMetalContext;
 
 extern NSString* DecodeMetalSourceCode(uint32 CodeSize, TArray<uint8> const& CompressedSource);
 
-enum EMetalIndexType
-{
-	EMetalIndexType_None   = 0,
-	EMetalIndexType_UInt16 = 1,
-	EMetalIndexType_UInt32 = 2,
-	EMetalIndexType_Num	   = 3
-};
-
-
 struct FMetalRenderPipelineHash
 {
 	friend uint32 GetTypeHash(FMetalRenderPipelineHash const& Hash)
@@ -675,8 +666,6 @@ class FMetalGraphicsPipelineState;
 class FMetalComputePipelineState;
 class FMetalVertexDeclaration;
 class FMetalVertexShader;
-class FMetalHullShader;
-class FMetalDomainShader;
 class FMetalGeometryShader;
 class FMetalPixelShader;
 class FMetalComputeShader;
@@ -709,16 +698,6 @@ struct TMetalResourceTraits<FRHIGeometryShader>
 	typedef FMetalGeometryShader TConcreteType;
 };
 template<>
-struct TMetalResourceTraits<FRHIHullShader>
-{
-	typedef FMetalHullShader TConcreteType;
-};
-template<>
-struct TMetalResourceTraits<FRHIDomainShader>
-{
-	typedef FMetalDomainShader TConcreteType;
-};
-template<>
 struct TMetalResourceTraits<FRHIPixelShader>
 {
 	typedef FMetalPixelShader TConcreteType;
diff --git a/Engine/Source/Runtime/Apple/MetalRHI/Public/MetalShaderResources.h b/Engine/Source/Runtime/Apple/MetalRHI/Public/MetalShaderResources.h
--- a/Engine/Source/Runtime/Apple/MetalRHI/Public/MetalShaderResources.h
+++ b/Engine/Source/Runtime/Apple/MetalRHI/Public/MetalShaderResources.h
@@ -228,100 +228,10 @@ struct FMetalAttribute
 	}
 };
 
-struct FMetalTessellationOutputs
-{
-	TArray<FMetalAttribute> HSIn;
-	TArray<FMetalAttribute> HSOut;
-	TArray<FMetalAttribute> PatchControlPointOut;
-	uint32 HSOutSize;
-	uint32 HSTFOutSize;
-	uint32 PatchControlPointOutSize;
-	
-	FMetalTessellationOutputs()
-	: HSOutSize(0)
-	, HSTFOutSize(0)
-	, PatchControlPointOutSize(0)
-	{
-		
-	}
-	
-	friend FArchive& operator<<(FArchive& Ar, FMetalTessellationOutputs& Attrs)
-	{
-		Ar << Attrs.HSIn;
-		Ar << Attrs.HSOut;
-		Ar << Attrs.PatchControlPointOut;
-		Ar << Attrs.HSOutSize;
-		Ar << Attrs.HSTFOutSize;
-		Ar << Attrs.PatchControlPointOutSize;
-		return Ar;
-	}
-};
-
-struct FMetalTessellationHeader
-{
-	FMetalTessellationOutputs TessellationOutputAttribs;
-	
-	uint32 TessellationOutputControlPoints;
-	uint32 TessellationDomain; // 3 = tri, 4 = quad // TODO unused
-	uint32 TessellationInputControlPoints; // TODO unused
-	uint32 TessellationPatchesPerThreadGroup;
-	uint32 TessellationPatchCountBuffer;
-	uint32 TessellationIndexBuffer;
-	uint32 TessellationHSOutBuffer;
-	uint32 TessellationHSTFOutBuffer;
-	uint32 TessellationControlPointOutBuffer;
-	uint32 TessellationControlPointIndexBuffer;
-	float  TessellationMaxTessFactor;
-	
-	EMetalOutputWindingMode TessellationOutputWinding;
-	EMetalPartitionMode TessellationPartitioning;
-	
-	FMetalTessellationHeader()
-	: TessellationOutputControlPoints(0)
-	, TessellationDomain(0)
-	, TessellationInputControlPoints(0)
-	, TessellationPatchesPerThreadGroup(0)
-	, TessellationPatchCountBuffer(~0u)
-	, TessellationIndexBuffer(~0u)
-	, TessellationHSOutBuffer(~0u)
-	, TessellationHSTFOutBuffer(~0u)
-	, TessellationControlPointOutBuffer(~0u)
-	, TessellationControlPointIndexBuffer(~0u)
-	, TessellationMaxTessFactor(0)
-	, TessellationOutputWinding(EMetalOutputWindingMode::Clockwise)
-	, TessellationPartitioning(EMetalPartitionMode::Pow2)
-	{
-		
-	}
-	
-	friend FArchive& operator<<(FArchive& Ar, FMetalTessellationHeader& Header)
-	{
-		Ar << Header.TessellationOutputAttribs;
-		
-		Ar << Header.TessellationOutputControlPoints;
-		Ar << Header.TessellationDomain;
-		Ar << Header.TessellationInputControlPoints;
-		Ar << Header.TessellationPatchesPerThreadGroup;
-		Ar << Header.TessellationMaxTessFactor;
-		
-		Ar << Header.TessellationPatchCountBuffer;
-		Ar << Header.TessellationIndexBuffer;
-		Ar << Header.TessellationHSOutBuffer;
-		Ar << Header.TessellationHSTFOutBuffer;
-		Ar << Header.TessellationControlPointOutBuffer;
-		Ar << Header.TessellationControlPointIndexBuffer;
-		
-		Ar << Header.TessellationOutputWinding;
-		Ar << Header.TessellationPartitioning;
-		return Ar;
-	}
-};
-
 struct FMetalCodeHeader
 {
 	FMetalShaderBindings Bindings;
 	TArray<CrossCompiler::FUniformBufferCopyInfo> UniformBuffersCopyInfo;
-	TArray<FMetalTessellationHeader> Tessellation;
 	FString ShaderName;
 
 	uint64 CompilerBuild;
@@ -378,7 +288,6 @@ inline FArchive& operator<<(FArchive& Ar, FMetalCodeHeader& Header)
 		}
 	}
 	
-	Ar << Header.Tessellation;
 	Ar << Header.ShaderName;
 
 	Ar << Header.CompilerBuild;
diff --git a/Engine/Source/Runtime/Core/Public/Mac/MacPlatform.h b/Engine/Source/Runtime/Core/Public/Mac/MacPlatform.h
--- a/Engine/Source/Runtime/Core/Public/Mac/MacPlatform.h
+++ b/Engine/Source/Runtime/Core/Public/Mac/MacPlatform.h
@@ -78,6 +78,7 @@ typedef FMacPlatformTypes FPlatformTypes;
 #define PLATFORM_SUPPORTS_STACK_SYMBOLS					1
 #define PLATFORM_HAS_BSD_SOCKET_FEATURE_MSG_DONTWAIT	1
 #define PLATFORM_IS_ANSI_MALLOC_THREADSAFE				1
+#define PLATFORM_SUPPORTS_TESSELLATION_SHADERS			0
 #define PLATFORM_SUPPORTS_VIRTUAL_TEXTURE_STREAMING		1
 
 #define PLATFORM_RHITHREAD_DEFAULT_BYPASS				WITH_EDITOR
diff --git a/Engine/Source/Runtime/Core/Public/Mac/MacPlatformProperties.h b/Engine/Source/Runtime/Core/Public/Mac/MacPlatformProperties.h
--- a/Engine/Source/Runtime/Core/Public/Mac/MacPlatformProperties.h
+++ b/Engine/Source/Runtime/Core/Public/Mac/MacPlatformProperties.h
@@ -123,11 +123,6 @@ struct FMacPlatformProperties
 		return 1.0f;
 	}
 	
-	static FORCEINLINE bool SupportsTessellation()
-	{
-		return true;
-	}
-
 	static FORCEINLINE bool SupportsAudioStreaming()
 	{
 		return !IsServerOnly();
diff --git a/Engine/Source/Runtime/Engine/Private/DebugViewModeHelpers.cpp b/Engine/Source/Runtime/Engine/Private/DebugViewModeHelpers.cpp
--- a/Engine/Source/Runtime/Engine/Private/DebugViewModeHelpers.cpp
+++ b/Engine/Source/Runtime/Engine/Private/DebugViewModeHelpers.cpp
@@ -23,7 +23,7 @@
 static bool PlatformSupportsDebugViewShaders(EShaderPlatform Platform)
 {
 	// List of platforms that have been tested and proved functional.
-	return Platform == SP_VULKAN_SM5 || Platform == SP_PCD3D_SM5 || Platform == SP_METAL_SM5_NOTESS || Platform == SP_METAL_SM5;
+	return Platform == SP_VULKAN_SM5 || Platform == SP_PCD3D_SM5 || Platform == SP_METAL_SM5;
 }
 
 bool AllowDebugViewVSDSHS(EShaderPlatform Platform)
diff --git a/Engine/Source/Runtime/RHI/Private/RHI.cpp b/Engine/Source/Runtime/RHI/Private/RHI.cpp
--- a/Engine/Source/Runtime/RHI/Private/RHI.cpp
+++ b/Engine/Source/Runtime/RHI/Private/RHI.cpp
@@ -1209,7 +1209,6 @@ FName ShaderPlatformToPlatformName(EShaderPlatform Platform)
 	case SP_METAL_MRT:
 		return NAME_PLATFORM_IOS;
 	case SP_METAL_SM5:
-	case SP_METAL_SM5_NOTESS:
 	case SP_METAL_MACES3_1:
 	case SP_METAL_MRT_MAC:
 		return NAME_PLATFORM_MAC;
@@ -1344,10 +1343,9 @@ RHI_API bool RHISupportsTessellation(const FStaticShaderPlatform Platform)
 	{
 		return true;
 	}
-
 	if (IsFeatureLevelSupported(Platform, ERHIFeatureLevel::SM5))
 	{
-		return (Platform == SP_PCD3D_SM5) || (Platform == SP_METAL_SM5) || (IsVulkanSM5Platform(Platform));
+		return (Platform == SP_PCD3D_SM5) || (IsVulkanSM5Platform(Platform));
 	}
 	return false;
 }
@@ -1623,7 +1621,6 @@ FString LexToString(EShaderPlatform Platform, bool bError)
 	case SP_METAL_MRT_TVOS: return TEXT("METAL_MRT_TVOS");
 	case SP_METAL_MRT_MAC: return TEXT("METAL_MRT_MAC");
 	case SP_METAL_SM5: return TEXT("METAL_SM5");
-	case SP_METAL_SM5_NOTESS: return TEXT("METAL_SM5_NOTESS");
 	case SP_METAL_MACES3_1: return TEXT("METAL_MACES3_1");
 	case SP_VULKAN_ES3_1_ANDROID: return TEXT("VULKAN_ES3_1_ANDROID");
 	case SP_VULKAN_ES3_1_LUMIN: return TEXT("VULKAN_ES3_1_LUMIN");
diff --git a/Engine/Source/Runtime/RHI/Public/RHI.h b/Engine/Source/Runtime/RHI/Public/RHI.h
--- a/Engine/Source/Runtime/RHI/Public/RHI.h
+++ b/Engine/Source/Runtime/RHI/Public/RHI.h
@@ -144,13 +144,13 @@ RHI_API int32 RHIGetPreferredClearUAVRectPSResourceType(const FStaticShaderPlatf
 
 inline bool RHISupportsInstancedStereo(const FStaticShaderPlatform Platform)
 {
-	return Platform == EShaderPlatform::SP_PCD3D_SM5 || Platform == EShaderPlatform::SP_METAL_SM5 || Platform == EShaderPlatform::SP_METAL_SM5_NOTESS
+	return Platform == EShaderPlatform::SP_PCD3D_SM5 || Platform == EShaderPlatform::SP_METAL_SM5
 		|| Platform == EShaderPlatform::SP_PCD3D_ES3_1 || FDataDrivenShaderPlatformInfo::GetSupportsInstancedStereo(Platform);
 }
 
 inline bool RHISupportsMultiView(const FStaticShaderPlatform Platform)
 {
-	return ((Platform == EShaderPlatform::SP_METAL_SM5 || Platform == SP_METAL_SM5_NOTESS))
+	return ((Platform == EShaderPlatform::SP_METAL_SM5))
 		|| FDataDrivenShaderPlatformInfo::GetSupportsMultiView(Platform);
 }
 
diff --git a/Engine/Source/Runtime/RHI/Public/RHIDefinitions.h b/Engine/Source/Runtime/RHI/Public/RHIDefinitions.h
--- a/Engine/Source/Runtime/RHI/Public/RHIDefinitions.h
+++ b/Engine/Source/Runtime/RHI/Public/RHIDefinitions.h
@@ -68,7 +68,7 @@ enum EShaderPlatform
 	SP_OPENGL_PCES3_1				= 15,
 	SP_METAL_SM5					= 16,
 	SP_VULKAN_PCES3_1				= 17,
-	SP_METAL_SM5_NOTESS				= 18,
+	SP_METAL_SM5_NOTESS_REMOVED		= 18,
 	SP_VULKAN_SM4_REMOVED			UE_DEPRECATED(4.27, "ShaderPlatform is removed; please don't use.") = 19,
 	SP_VULKAN_SM5					= 20,
 	SP_VULKAN_ES3_1_ANDROID			= 21,
@@ -1613,7 +1613,7 @@ inline bool IsPCPlatform(const FStaticShaderPlatform Platform)
 {
 	return Platform == SP_PCD3D_SM5 || Platform == SP_PCD3D_ES3_1 ||
 		Platform == SP_OPENGL_PCES3_1 ||
-		Platform == SP_METAL_SM5_NOTESS || Platform == SP_METAL_SM5 ||
+		Platform == SP_METAL_SM5 ||
 		Platform == SP_VULKAN_PCES3_1 || Platform == SP_VULKAN_SM5 || Platform == SP_METAL_MACES3_1 || Platform == SP_METAL_MRT_MAC 
 		|| FDataDrivenShaderPlatformInfo::GetIsPC(Platform);
 }
@@ -1640,7 +1640,7 @@ inline bool IsOpenGLPlatform(const FStaticShaderPlatform Platform)
 inline bool IsMetalPlatform(const FStaticShaderPlatform Platform)
 {
 	return Platform == SP_METAL || Platform == SP_METAL_MRT || Platform == SP_METAL_TVOS || Platform == SP_METAL_MRT_TVOS 
-		|| Platform == SP_METAL_SM5_NOTESS || Platform == SP_METAL_SM5
+		|| Platform == SP_METAL_SM5
 		|| Platform == SP_METAL_MACES3_1 || Platform == SP_METAL_MRT_MAC
 		|| FDataDrivenShaderPlatformInfo::GetIsLanguageMetal(Platform);
 }
@@ -1659,7 +1659,7 @@ inline bool IsMetalMRTPlatform(const FStaticShaderPlatform Platform)
 
 inline bool IsMetalSM5Platform(const FStaticShaderPlatform Platform)
 {
-	return Platform == SP_METAL_MRT || Platform == SP_METAL_MRT_TVOS || Platform == SP_METAL_SM5_NOTESS || Platform == SP_METAL_SM5 || Platform == SP_METAL_MRT_MAC
+	return Platform == SP_METAL_MRT || Platform == SP_METAL_MRT_TVOS || Platform == SP_METAL_SM5 || Platform == SP_METAL_MRT_MAC
 		|| (FDataDrivenShaderPlatformInfo::GetIsLanguageMetal(Platform) && FDataDrivenShaderPlatformInfo::GetMaxFeatureLevel(Platform) == ERHIFeatureLevel::SM5);
 }
 
@@ -1735,7 +1735,6 @@ inline FStaticFeatureLevel GetMaxSupportedFeatureLevel(const FStaticShaderPlatfo
 	case SP_METAL_MRT:
 	case SP_METAL_MRT_TVOS:
 	case SP_METAL_MRT_MAC:
-	case SP_METAL_SM5_NOTESS:
 	case SP_VULKAN_SM5:
 	case SP_VULKAN_SM5_LUMIN:
 	case SP_SWITCH:
diff --git a/Engine/Source/Runtime/RHI/Public/RHIShaderFormatDefinitions.inl b/Engine/Source/Runtime/RHI/Public/RHIShaderFormatDefinitions.inl
--- a/Engine/Source/Runtime/RHI/Public/RHIShaderFormatDefinitions.inl
+++ b/Engine/Source/Runtime/RHI/Public/RHIShaderFormatDefinitions.inl
@@ -26,7 +26,6 @@ static FName NAME_SF_METAL_TVOS(TEXT("SF_METAL_TVOS"));
 static FName NAME_SF_METAL_MRT_TVOS(TEXT("SF_METAL_MRT_TVOS"));
 static FName NAME_SF_METAL_MRT_MAC(TEXT("SF_METAL_MRT_MAC"));
 static FName NAME_SF_METAL_SM5(TEXT("SF_METAL_SM5"));
-static FName NAME_SF_METAL_SM5_NOTESS(TEXT("SF_METAL_SM5_NOTESS"));
 static FName NAME_SF_METAL_MACES3_1(TEXT("SF_METAL_MACES3_1"));
 
 static FName NAME_VULKAN_ES3_1_ANDROID(TEXT("SF_VULKAN_ES31_ANDROID"));
@@ -70,8 +69,6 @@ static FName ShaderPlatformToShaderFormatName(EShaderPlatform Platform)
 		return NAME_SF_METAL_MRT_MAC;
 	case SP_METAL_SM5:
 		return NAME_SF_METAL_SM5;
-	case SP_METAL_SM5_NOTESS:
-		return NAME_SF_METAL_SM5_NOTESS;
 	case SP_METAL_MACES3_1:
 		return NAME_SF_METAL_MACES3_1;
 	case SP_VULKAN_ES3_1_ANDROID:
@@ -118,7 +115,6 @@ static EShaderPlatform ShaderFormatNameToShaderPlatform(FName ShaderFormat)
 	if (ShaderFormat == NAME_SF_METAL_MRT_TVOS)			return SP_METAL_MRT_TVOS;
 	if (ShaderFormat == NAME_SF_METAL_MRT_MAC)			return SP_METAL_MRT_MAC;
 	if (ShaderFormat == NAME_SF_METAL_SM5)				return SP_METAL_SM5;
-	if (ShaderFormat == NAME_SF_METAL_SM5_NOTESS)		return SP_METAL_SM5_NOTESS;
 	if (ShaderFormat == NAME_SF_METAL_MACES3_1)			return SP_METAL_MACES3_1;
 
 	if (ShaderFormat == NAME_VULKAN_ES3_1_ANDROID)		return SP_VULKAN_ES3_1_ANDROID;
diff --git a/Engine/Source/Runtime/RHI/Public/RHIShaderPlatformDefinitions.inl b/Engine/Source/Runtime/RHI/Public/RHIShaderPlatformDefinitions.inl
--- a/Engine/Source/Runtime/RHI/Public/RHIShaderPlatformDefinitions.inl
+++ b/Engine/Source/Runtime/RHI/Public/RHIShaderPlatformDefinitions.inl
@@ -19,7 +19,6 @@ static FText GetFriendlyShaderPlatformName(const EShaderPlatform InShaderPlatfor
 	}
 	break;
 
-	case SP_METAL_SM5_NOTESS:
 	case SP_METAL_MRT_MAC:
 	case SP_METAL_MRT:
 	{
diff --git a/Engine/Source/Runtime/Renderer/Private/ShadowRendering.h b/Engine/Source/Runtime/Renderer/Private/ShadowRendering.h
--- a/Engine/Source/Runtime/Renderer/Private/ShadowRendering.h
+++ b/Engine/Source/Runtime/Renderer/Private/ShadowRendering.h
@@ -1610,7 +1610,7 @@ public:
 	static bool ShouldCompilePermutation(const FGlobalShaderPermutationParameters& Parameters)
 	{
 		return TShadowProjectionPS<Quality, bUseFadePlane>::ShouldCompilePermutation(Parameters)
-			&& (Parameters.Platform == SP_PCD3D_SM5 || IsVulkanSM5Platform(Parameters.Platform) || Parameters.Platform == SP_METAL_SM5 || Parameters.Platform == SP_METAL_SM5_NOTESS);
+			&& (Parameters.Platform == SP_PCD3D_SM5 || IsVulkanSM5Platform(Parameters.Platform) || Parameters.Platform == SP_METAL_SM5);
 	}
 
 	void SetParameters(
@@ -1660,7 +1660,7 @@ public:
 	{
 		return IsFeatureLevelSupported(Parameters.Platform, ERHIFeatureLevel::SM5);
 		// TODO: Abstract properly with FDataDrivenShaderPlatformInfo. Platforms like Switch need the below check apparently.
-		//	&& (Parameters.Platform == SP_PCD3D_SM5 || IsVulkanSM5Platform(Parameters.Platform) || Parameters.Platform == SP_METAL_SM5 || Parameters.Platform == SP_METAL_SM5_NOTESS);
+		//	&& (Parameters.Platform == SP_PCD3D_SM5 || IsVulkanSM5Platform(Parameters.Platform) || Parameters.Platform == SP_METAL_SM5);
 	}
 
 	static void ModifyCompilationEnvironment(const FGlobalShaderPermutationParameters& Parameters, FShaderCompilerEnvironment& OutEnvironment)
diff --git a/Engine/Source/Runtime/Renderer/Private/VolumetricFog.h b/Engine/Source/Runtime/Renderer/Private/VolumetricFog.h
--- a/Engine/Source/Runtime/Renderer/Private/VolumetricFog.h
+++ b/Engine/Source/Runtime/Renderer/Private/VolumetricFog.h
@@ -13,7 +13,7 @@
 inline bool DoesPlatformSupportVolumetricFog(const FStaticShaderPlatform Platform)
 {
 	return Platform == SP_PCD3D_SM5
-		|| Platform == SP_METAL_SM5 || Platform == SP_METAL_SM5_NOTESS
+		|| Platform == SP_METAL_SM5
 		|| IsVulkanSM5Platform(Platform)
 		|| FDataDrivenShaderPlatformInfo::GetSupportsVolumetricFog(Platform);
 }
@@ -21,7 +21,7 @@ inline bool DoesPlatformSupportVolumetricFog(const FStaticShaderPlatform Platfor
 inline bool DoesPlatformSupportVolumetricFogVoxelization(const FStaticShaderPlatform Platform)
 {
 	return Platform == SP_PCD3D_SM5
-		|| Platform == SP_METAL_SM5 || Platform == SP_METAL_SM5_NOTESS
+		|| Platform == SP_METAL_SM5
 		|| IsVulkanSM5Platform(Platform)
 		|| FDataDrivenShaderPlatformInfo::GetSupportsVolumetricFog(Platform);
 }
