Tessellation Removal: Removing Tessellation from MetalRHI

diff --git a/Engine/Source/Runtime/Apple/MetalRHI/Private/MetalDynamicRHI.h b/Engine/Source/Runtime/Apple/MetalRHI/Private/MetalDynamicRHI.h
--- a/Engine/Source/Runtime/Apple/MetalRHI/Private/MetalDynamicRHI.h
+++ b/Engine/Source/Runtime/Apple/MetalRHI/Private/MetalDynamicRHI.h
@@ -47,8 +47,6 @@ public:
 	virtual FVertexDeclarationRHIRef RHICreateVertexDeclaration(const FVertexDeclarationElementList& Elements) final override;
 	virtual FPixelShaderRHIRef RHICreatePixelShader(TArrayView<const uint8> Code, const FSHAHash& Hash) final override;
 	virtual FVertexShaderRHIRef RHICreateVertexShader(TArrayView<const uint8> Code, const FSHAHash& Hash) final override;
-	virtual FHullShaderRHIRef RHICreateHullShader(TArrayView<const uint8> Code, const FSHAHash& Hash) final override;
-	virtual FDomainShaderRHIRef RHICreateDomainShader(TArrayView<const uint8> Code, const FSHAHash& Hash) final override;
 	virtual FGeometryShaderRHIRef RHICreateGeometryShader(TArrayView<const uint8> Code, const FSHAHash& Hash) final override;
 	virtual FComputeShaderRHIRef RHICreateComputeShader(TArrayView<const uint8> Code, const FSHAHash& Hash) final override;
 	virtual FBoundShaderStateRHIRef RHICreateBoundShaderState(FRHIVertexDeclaration* VertexDeclaration, FRHIVertexShader* VertexShader, FRHIPixelShader* PixelShader, FRHIGeometryShader* GeometryShader) final override;
@@ -184,8 +182,6 @@ public:
 	virtual FPixelShaderRHIRef CreatePixelShader_RenderThread(class FRHICommandListImmediate& RHICmdList, TArrayView<const uint8> Code, const FSHAHash& Hash) final override;
 	virtual FGeometryShaderRHIRef CreateGeometryShader_RenderThread(class FRHICommandListImmediate& RHICmdList, TArrayView<const uint8> Code, const FSHAHash& Hash) final override;
 	virtual FComputeShaderRHIRef CreateComputeShader_RenderThread(class FRHICommandListImmediate& RHICmdList, TArrayView<const uint8> Code, const FSHAHash& Hash) final override;
-	virtual FHullShaderRHIRef CreateHullShader_RenderThread(class FRHICommandListImmediate& RHICmdList, TArrayView<const uint8> Code, const FSHAHash& Hash) final override;
-	virtual FDomainShaderRHIRef CreateDomainShader_RenderThread(class FRHICommandListImmediate& RHICmdList, TArrayView<const uint8> Code, const FSHAHash& Hash) final override;
 	virtual FRHIShaderLibraryRef RHICreateShaderLibrary_RenderThread(class FRHICommandListImmediate& RHICmdList, EShaderPlatform Platform, FString FilePath, FString Name) final override;
 
 	virtual void RHIUpdateUniformBuffer(FRHIUniformBuffer* UniformBufferRHI, const void* Contents) final override;
diff --git a/Engine/Source/Runtime/Apple/MetalRHI/Private/MetalDynamicRHI_Shaders.cpp b/Engine/Source/Runtime/Apple/MetalRHI/Private/MetalDynamicRHI_Shaders.cpp
--- a/Engine/Source/Runtime/Apple/MetalRHI/Private/MetalDynamicRHI_Shaders.cpp
+++ b/Engine/Source/Runtime/Apple/MetalRHI/Private/MetalDynamicRHI_Shaders.cpp
@@ -31,16 +31,6 @@ FPixelShaderRHIRef FMetalDynamicRHI::RHICreatePixelShader(TArrayView<const uint8
 	}
 }
 
-FHullShaderRHIRef FMetalDynamicRHI::RHICreateHullShader(TArrayView<const uint8> Code, const FSHAHash& Hash)
-{
-	return nullptr;
-}
-
-FDomainShaderRHIRef FMetalDynamicRHI::RHICreateDomainShader(TArrayView<const uint8> Code, const FSHAHash& Hash)
-{
-	return nullptr;
-}
-
 FGeometryShaderRHIRef FMetalDynamicRHI::RHICreateGeometryShader(TArrayView<const uint8> Code, const FSHAHash& Hash)
 {
 	@autoreleasepool {
@@ -164,16 +154,6 @@ FVertexShaderRHIRef FMetalDynamicRHI::CreateVertexShader_RenderThread(class FRHI
 	return RHICreateVertexShader(Code, Hash);
 }
 
-FHullShaderRHIRef FMetalDynamicRHI::CreateHullShader_RenderThread(class FRHICommandListImmediate& RHICmdList, TArrayView<const uint8> Code, const FSHAHash& Hash)
-{
-	return RHICreateHullShader(Code, Hash);
-}
-
-FDomainShaderRHIRef FMetalDynamicRHI::CreateDomainShader_RenderThread(class FRHICommandListImmediate& RHICmdList, TArrayView<const uint8> Code, const FSHAHash& Hash)
-{
-	return RHICreateDomainShader(Code, Hash);
-}
-
 FGeometryShaderRHIRef FMetalDynamicRHI::CreateGeometryShader_RenderThread(class FRHICommandListImmediate& RHICmdList, TArrayView<const uint8> Code, const FSHAHash& Hash)
 {
 	return RHICreateGeometryShader(Code, Hash);
