Rework sRGB SRV handling to distinguish between M1 (Apple) GPUs and Intel based Macs.

#rnx

diff --git a/Engine/Source/Runtime/Apple/MetalRHI/Private/MetalTexture.cpp b/Engine/Source/Runtime/Apple/MetalRHI/Private/MetalTexture.cpp
--- a/Engine/Source/Runtime/Apple/MetalRHI/Private/MetalTexture.cpp
+++ b/Engine/Source/Runtime/Apple/MetalRHI/Private/MetalTexture.cpp
@@ -160,11 +160,10 @@ void SafeReleaseMetalTexture(FMetalSurface* Surface, FMetalTexture& Texture)
 	}
 }
 
-mtlpp::PixelFormat ToSRGBFormat(mtlpp::PixelFormat LinMTLFormat)
+#if !PLATFORM_IOS
+mtlpp::PixelFormat ToSRGBFormat_NonAppleMacGPU(mtlpp::PixelFormat MTLFormat)
 {
-	mtlpp::PixelFormat MTLFormat = LinMTLFormat;
-	
-	switch (LinMTLFormat)
+	switch (MTLFormat)
 	{
 		case mtlpp::PixelFormat::RGBA8Unorm:
 			MTLFormat = mtlpp::PixelFormat::RGBA8Unorm_sRGB;
@@ -172,7 +171,6 @@ mtlpp::PixelFormat ToSRGBFormat(mtlpp::PixelFormat LinMTLFormat)
 		case mtlpp::PixelFormat::BGRA8Unorm:
 			MTLFormat = mtlpp::PixelFormat::BGRA8Unorm_sRGB;
 			break;
-#if PLATFORM_MAC
 		case mtlpp::PixelFormat::BC1_RGBA:
 			MTLFormat = mtlpp::PixelFormat::BC1_RGBA_sRGB;
 			break;
@@ -185,8 +183,23 @@ mtlpp::PixelFormat ToSRGBFormat(mtlpp::PixelFormat LinMTLFormat)
 		case mtlpp::PixelFormat::BC7_RGBAUnorm:
 			MTLFormat = mtlpp::PixelFormat::BC7_RGBAUnorm_sRGB;
 			break;
-#endif //PLATFORM_MAC
-#if PLATFORM_IOS
+		default:
+			break;
+	}
+	return MTLFormat;
+}
+#endif
+
+mtlpp::PixelFormat ToSRGBFormat_AppleGPU(mtlpp::PixelFormat MTLFormat)
+{
+	switch (MTLFormat)
+	{
+		case mtlpp::PixelFormat::RGBA8Unorm:
+			MTLFormat = mtlpp::PixelFormat::RGBA8Unorm_sRGB;
+			break;
+		case mtlpp::PixelFormat::BGRA8Unorm:
+			MTLFormat = mtlpp::PixelFormat::BGRA8Unorm_sRGB;
+			break;
 		case mtlpp::PixelFormat::R8Unorm:
 			MTLFormat = mtlpp::PixelFormat::R8Unorm_sRGB;
 			break;
@@ -211,32 +224,26 @@ mtlpp::PixelFormat ToSRGBFormat(mtlpp::PixelFormat LinMTLFormat)
 		case mtlpp::PixelFormat::ASTC_12x12_LDR:
 			MTLFormat = mtlpp::PixelFormat::ASTC_12x12_sRGB;
 			break;
-#endif //PLATFORM_IOS
 		default:
 			break;
 	}
-	
 	return MTLFormat;
 }
 
-static inline mtlpp::PixelFormat ToSRGBFormat(mtlpp::PixelFormat InMTLFormat, ERHIResourceType InResourceType)
+mtlpp::PixelFormat ToSRGBFormat(mtlpp::PixelFormat MTLFormat)
 {
-#if PLATFORM_MAC
-	// R8Unorm_sRGB is not available on macOS
-	// For now R8 sRGB expansion is 2D only, log other usage for later.
-	if (InMTLFormat == mtlpp::PixelFormat::R8Unorm)
+	if([GetMetalDeviceContext().GetDevice().GetPtr() supportsFamily:MTLGPUFamilyApple1])
 	{
-		if (InResourceType == RRT_Texture2D)
-		{
-			InMTLFormat = mtlpp::PixelFormat::RGBA8Unorm;
-		}
-		else
-		{
-			UE_LOG(LogMetal, Error, TEXT("Attempting to use unsupported mtlpp::PixelFormat::R8Unorm_sRGB on Mac with texture type: %d, no format expansion will be provided so rendering errors may occur."), InResourceType);
-		}
+		MTLFormat = ToSRGBFormat_AppleGPU(MTLFormat);
+	}
+#if !PLATFORM_IOS
+	else if([GetMetalDeviceContext().GetDevice().GetPtr() supportsFamily:MTLGPUFamilyMac1])
+	{
+		MTLFormat = ToSRGBFormat_NonAppleMacGPU(MTLFormat);
 	}
 #endif
-	return ToSRGBFormat(InMTLFormat);
+	
+	return MTLFormat;
 }
 
 void FMetalSurface::PrepareTextureView()
@@ -480,7 +487,7 @@ void FMetalSurface::Init(FMetalSurface& Source, NSRange MipRange)
 	if (Flags & TexCreate_SRGB)
 	{
 		// Ensure we have the correct sRGB target format if we create a new texture view rather than using the source texture
-		MetalFormat = ToSRGBFormat(MetalFormat, Type);
+		MetalFormat = ToSRGBFormat(MetalFormat);
 	}
 	
 	bool const bUseSourceTex = (Source.PixelFormat != PF_DepthStencil) && MipRange.location == 0 && MipRange.length == Source.Texture.GetMipmapLevelCount();
@@ -536,7 +543,7 @@ void FMetalSurface::Init(FMetalSurface& Source, NSRange MipRange, EPixelFormat F
 		else
 		{
 			// Ensure we have the correct sRGB target format if we create a new texture view rather than using the source texture
-			MetalFormat = ToSRGBFormat(MetalFormat, Type);
+			MetalFormat = ToSRGBFormat(MetalFormat);
 		}
 	}
 	
@@ -690,7 +697,7 @@ FMetalSurface::FMetalSurface(ERHIResourceType ResourceType, EPixelFormat Format,
 	
 	if (Flags & TexCreate_SRGB)
 	{
-		MTLFormat = ToSRGBFormat(MTLFormat, Type);
+		MTLFormat = ToSRGBFormat(MTLFormat);
 	}
 	
 	// set the key
