Fix-up Volumetric Clouds (and the default TimeOfDay level) on macOS

- Implement RHIIsTypedUAVLoadSupported for Metal RHI
- Provide a default PrimitiveUniformBufferResource on the SingleTriangleMesh object.

#rnx

diff --git a/Engine/Source/Runtime/Apple/MetalRHI/Private/MetalDynamicRHI.h b/Engine/Source/Runtime/Apple/MetalRHI/Private/MetalDynamicRHI.h
--- a/Engine/Source/Runtime/Apple/MetalRHI/Private/MetalDynamicRHI.h
+++ b/Engine/Source/Runtime/Apple/MetalRHI/Private/MetalDynamicRHI.h
@@ -186,6 +186,8 @@ public:
 
 	virtual void RHIUpdateUniformBuffer(FRHIUniformBuffer* UniformBufferRHI, const void* Contents) final override;
 
+	virtual bool RHIIsTypedUAVLoadSupported(EPixelFormat PixelFormat) override;
+
 	virtual uint16 RHIGetPlatformTextureMaxSampleCount() override;
 
 private:
diff --git a/Engine/Source/Runtime/Apple/MetalRHI/Private/MetalRHI.cpp b/Engine/Source/Runtime/Apple/MetalRHI/Private/MetalRHI.cpp
--- a/Engine/Source/Runtime/Apple/MetalRHI/Private/MetalRHI.cpp
+++ b/Engine/Source/Runtime/Apple/MetalRHI/Private/MetalRHI.cpp
@@ -29,6 +29,8 @@ bool GIsMetalInitialized = false;
 
 FMetalBufferFormat GMetalBufferFormats[PF_MAX];
 
+static bool GFormatSupportsTypedUAVLoad[PF_MAX] = { false };
+
 static TAutoConsoleVariable<int32> CVarUseRHIThread(
 													TEXT("r.Metal.IOSRHIThread"),
 													0,
@@ -895,7 +897,34 @@ FMetalDynamicRHI::FMetalDynamicRHI(ERHIFeatureLevel::Type RequestedFeatureLevel)
 		checkf(GMetalBufferFormats[i].DataFormat != 255, TEXT("Metal data buffer format for pixel-format %s (%d) is not configured!"), GPixelFormats[i].Name, i);
 	}
 #endif
-		
+
+	switch (Device.GetReadWriteTextureSupport())
+	{
+		case mtlpp::ReadWriteTextureTier::Tier2:
+			GFormatSupportsTypedUAVLoad[PF_A32B32G32R32F]			= true;
+			GFormatSupportsTypedUAVLoad[PF_R32G32B32A32_UINT]		= true;
+			GFormatSupportsTypedUAVLoad[PF_FloatRGBA]				= true;
+			GFormatSupportsTypedUAVLoad[PF_R16G16B16A16_UINT]		= true;
+			GFormatSupportsTypedUAVLoad[PF_R16G16B16A16_SINT]		= true;
+			GFormatSupportsTypedUAVLoad[PF_R8G8B8A8]				= true;
+			GFormatSupportsTypedUAVLoad[PF_R8G8B8A8_UINT]			= true;
+			GFormatSupportsTypedUAVLoad[PF_R16F]					= true;
+			GFormatSupportsTypedUAVLoad[PF_R16_UINT]				= true;
+			GFormatSupportsTypedUAVLoad[PF_R16_SINT]				= true;
+			GFormatSupportsTypedUAVLoad[PF_R8]						= true;
+			GFormatSupportsTypedUAVLoad[PF_R8_UINT]					= true;
+			// Fall through
+
+		case mtlpp::ReadWriteTextureTier::Tier1:
+			GFormatSupportsTypedUAVLoad[PF_R32_FLOAT]				= true;
+			GFormatSupportsTypedUAVLoad[PF_R32_UINT]				= true;
+			GFormatSupportsTypedUAVLoad[PF_R32_SINT]				= true;
+			// Fall through
+
+		case mtlpp::ReadWriteTextureTier::None:
+			break;
+	};
+
 	// get driver version (todo: share with other RHIs)
 	{
 		FGPUDriverInfo GPUDriverInfo = FPlatformMisc::GetGPUDriverInfo(GRHIAdapterName);
@@ -1204,6 +1233,10 @@ void* FMetalDynamicRHI::RHIGetNativeInstance()
 	return nullptr;
 }
 
+bool FMetalDynamicRHI::RHIIsTypedUAVLoadSupported(EPixelFormat PixelFormat)
+{
+	return GFormatSupportsTypedUAVLoad[PixelFormat];
+}
 
 uint16 FMetalDynamicRHI::RHIGetPlatformTextureMaxSampleCount()
 {
