Follow-on to CL 16186890 to remove deprecation warning from MetalRHI.

#rnx

diff --git a/Engine/Source/Runtime/Apple/MetalRHI/Private/MetalRHIRenderQuery.cpp b/Engine/Source/Runtime/Apple/MetalRHI/Private/MetalRHIRenderQuery.cpp
--- a/Engine/Source/Runtime/Apple/MetalRHI/Private/MetalRHIRenderQuery.cpp
+++ b/Engine/Source/Runtime/Apple/MetalRHI/Private/MetalRHIRenderQuery.cpp
@@ -19,7 +19,8 @@
 
 
 FMetalQueryBuffer::FMetalQueryBuffer(FMetalContext* InContext, FMetalBuffer InBuffer)
-	: Pool(InContext->GetQueryBufferPool())
+	: FRHIResource(RRT_TimestampCalibrationQuery)
+	, Pool(InContext->GetQueryBufferPool())
 	, Buffer(MoveTemp(InBuffer))
 	, WriteOffset(0)
 {
