Initial import of AGXRHI

AGXRHI is starting out as a clone of MetalRHI, but will be refactored and overhauled with several goals in mind:

- Reduce/Eliminate technical debt
- Substantially improve readability and maintainability
- Introduce clear layers for platform independent vs dependent code
- Reduce memory footprint and generally improve memory usage
- Remove legacy concepts (like manual reference counting) and enable better integration with platform tools
- Implement newer and more modern features provided by Metal (such as ray tracing and enhanced command buffer debugging)
- Target/tune the software architecture and API usage for performance on Apple designed GPUs

#rnx

diff --git a/Engine/Source/Developer/Apple/MetalShaderFormat/MetalShaderFormat.Build.cs b/Engine/Source/Developer/Apple/MetalShaderFormat/MetalShaderFormat.Build.cs
--- a/Engine/Source/Developer/Apple/MetalShaderFormat/MetalShaderFormat.Build.cs
+++ b/Engine/Source/Developer/Apple/MetalShaderFormat/MetalShaderFormat.Build.cs
@@ -7,7 +7,7 @@ public class MetalShaderFormat : ModuleRules
 	public MetalShaderFormat(ReadOnlyTargetRules Target) : base(Target)
 	{
 		PrivateIncludePathModuleNames.Add("TargetPlatform");
-		PublicIncludePaths.Add("Runtime/Apple/MetalRHI/Public");
+		PublicIncludePaths.Add("Runtime/Apple/Common/Public");
 
 		PrivateIncludePaths.AddRange(
 			new string[] {
diff --git a/Engine/Source/Programs/UnrealBuildTool/Configuration/UEBuildModuleCPP.cs b/Engine/Source/Programs/UnrealBuildTool/Configuration/UEBuildModuleCPP.cs
--- a/Engine/Source/Programs/UnrealBuildTool/Configuration/UEBuildModuleCPP.cs
+++ b/Engine/Source/Programs/UnrealBuildTool/Configuration/UEBuildModuleCPP.cs
@@ -1275,6 +1275,7 @@ namespace UnrealBuildTool
 				case "VulkanRHI":
 				case "OpenGLDrv":
 				case "MetalRHI":
+				case "AGXRHI":
 					return true;
 			}
 			return false;
diff --git a/Engine/Source/Programs/UnrealBuildTool/Platform/IOS/UEBuildIOS.cs b/Engine/Source/Programs/UnrealBuildTool/Platform/IOS/UEBuildIOS.cs
--- a/Engine/Source/Programs/UnrealBuildTool/Platform/IOS/UEBuildIOS.cs
+++ b/Engine/Source/Programs/UnrealBuildTool/Platform/IOS/UEBuildIOS.cs
@@ -738,6 +738,7 @@ namespace UnrealBuildTool
 			{
 				Target.GlobalDefinitions.Add("HAS_METAL=1");
 				Target.ExtraModuleNames.Add("MetalRHI");
+				Target.ExtraModuleNames.Add("AGXRHI");
 			}
 			else
 			{
diff --git a/Engine/Source/Programs/UnrealBuildTool/Platform/Mac/UEBuildMac.cs b/Engine/Source/Programs/UnrealBuildTool/Platform/Mac/UEBuildMac.cs
--- a/Engine/Source/Programs/UnrealBuildTool/Platform/Mac/UEBuildMac.cs
+++ b/Engine/Source/Programs/UnrealBuildTool/Platform/Mac/UEBuildMac.cs
@@ -129,6 +129,7 @@ namespace UnrealBuildTool
 			{
 				Target.GlobalDefinitions.Add("HAS_METAL=1");
 				Target.ExtraModuleNames.Add("MetalRHI");
+				Target.ExtraModuleNames.Add("AGXRHI");
 			}
 			else
 			{
diff --git a/Engine/Source/Runtime/Apple/AGXRHI/AGXRHI.Build.cs b/Engine/Source/Runtime/Apple/AGXRHI/AGXRHI.Build.cs
new file mode 100644
--- /dev/null
+++ b/Engine/Source/Runtime/Apple/AGXRHI/AGXRHI.Build.cs
@@ -0,0 +1,34 @@
+// Copyright Epic Games, Inc. All Rights Reserved.
+
+using System.IO;
+using UnrealBuildTool;
+
+public class AGXRHI : ModuleRules
+{	
+	public AGXRHI(ReadOnlyTargetRules Target) : base(Target)
+	{
+		PublicIncludePaths.Add("Runtime/Apple/Common/Public");
+
+		PrivateDependencyModuleNames.AddRange(
+			new string[] {
+				"Core",
+				"ApplicationCore",
+				"Engine",
+				"RHI",
+				"RHICore",
+				"RenderCore",
+			}
+			);
+
+		AddEngineThirdPartyPrivateStaticDependencies(Target,
+			"MTLPP"
+		);
+
+		PublicWeakFrameworks.Add("Metal");
+
+		if (Target.Platform == UnrealTargetPlatform.Mac)
+		{
+			PublicFrameworks.Add("QuartzCore");
+		}
+	}
+}
diff --git a/Engine/Source/Runtime/Apple/AGXRHI/Private/AGXBlitCommandEncoder.cpp b/Engine/Source/Runtime/Apple/AGXRHI/Private/AGXBlitCommandEncoder.cpp
new file mode 100644
--- /dev/null
+++ b/Engine/Source/Runtime/Apple/AGXRHI/Private/AGXBlitCommandEncoder.cpp
@@ -0,0 +1,280 @@
+// Copyright Epic Games, Inc. All Rights Reserved.
+
+/*=============================================================================
+	AGXBlitCommandEncoder.cpp: AGX command encoder wrapper.
+=============================================================================*/
+
+#include "AGXRHIPrivate.h"
+
+#include "AGXBlitCommandEncoder.h"
+#include "AGXCommandBuffer.h"
+#include "AGXFence.h"
+
+#if MTLPP_CONFIG_VALIDATE && METAL_DEBUG_OPTIONS
+extern int32 GAGXRuntimeDebugLevel;
+
+@interface FAGXDebugBlitCommandEncoder : FAGXDebugCommandEncoder
+{
+	@public
+	id<MTLBlitCommandEncoder> Inner;
+	FAGXCommandBufferDebugging Buffer;
+}
+/** Initialise the wrapper with the provided command-buffer. */
+-(id)initWithEncoder:(id<MTLBlitCommandEncoder>)Encoder andCommandBuffer:(FAGXCommandBufferDebugging const&)Buffer;
+
+@end
+
+@implementation FAGXDebugBlitCommandEncoder
+
+APPLE_PLATFORM_OBJECT_ALLOC_OVERRIDES(FAGXDebugBlitCommandEncoder)
+
+-(id)initWithEncoder:(id<MTLBlitCommandEncoder>)Encoder andCommandBuffer:(FAGXCommandBufferDebugging const&)SourceBuffer
+{
+	id Self = [super init];
+	if (Self)
+	{
+        Inner = Encoder;
+		Buffer = SourceBuffer;
+	}
+	return Self;
+}
+
+-(void)dealloc
+{
+	[super dealloc];
+}
+
+@end
+
+FAGXBlitCommandEncoderDebugging::FAGXBlitCommandEncoderDebugging()
+{
+	
+}
+FAGXBlitCommandEncoderDebugging::FAGXBlitCommandEncoderDebugging(mtlpp::BlitCommandEncoder& Encoder, FAGXCommandBufferDebugging& Buffer)
+: FAGXCommandEncoderDebugging((FAGXDebugCommandEncoder*)[[[FAGXDebugBlitCommandEncoder alloc] initWithEncoder:Encoder.GetPtr() andCommandBuffer:Buffer] autorelease])
+{
+	Buffer.BeginBlitCommandEncoder([NSString stringWithFormat:@"Blit: %@", Encoder.GetLabel().GetPtr()]);
+	Encoder.SetAssociatedObject((void const*)&FAGXBlitCommandEncoderDebugging::Get, (FAGXCommandEncoderDebugging const&)*this);
+}
+FAGXBlitCommandEncoderDebugging::FAGXBlitCommandEncoderDebugging(FAGXDebugCommandEncoder* handle)
+: FAGXCommandEncoderDebugging((FAGXDebugCommandEncoder*)handle)
+{
+	
+}
+
+FAGXBlitCommandEncoderDebugging FAGXBlitCommandEncoderDebugging::Get(mtlpp::BlitCommandEncoder& Encoder)
+{
+	return Encoder.GetAssociatedObject<FAGXBlitCommandEncoderDebugging>((void const*)&FAGXBlitCommandEncoderDebugging::Get);
+}
+
+void FAGXBlitCommandEncoderDebugging::InsertDebugSignpost(ns::String const& Label)
+{
+	((FAGXDebugBlitCommandEncoder*)m_ptr)->Buffer.InsertDebugSignpost(Label);
+}
+void FAGXBlitCommandEncoderDebugging::PushDebugGroup(ns::String const& Group)
+{
+	((FAGXDebugBlitCommandEncoder*)m_ptr)->Buffer.PushDebugGroup(Group);
+}
+void FAGXBlitCommandEncoderDebugging::PopDebugGroup()
+{
+	((FAGXDebugBlitCommandEncoder*)m_ptr)->Buffer.PopDebugGroup();
+}
+
+void FAGXBlitCommandEncoderDebugging::EndEncoder()
+{
+	((FAGXDebugBlitCommandEncoder*)m_ptr)->Buffer.EndCommandEncoder();
+}
+
+#if PLATFORM_MAC
+void FAGXBlitCommandEncoderDebugging::Synchronize(mtlpp::Resource const& resource)
+{
+#if METAL_DEBUG_OPTIONS
+	switch(((FAGXDebugBlitCommandEncoder*)m_ptr)->Buffer.GetPtr()->DebugLevel)
+	{
+		case EAGXDebugLevelConditionalSubmit:
+		case EAGXDebugLevelWaitForComplete:
+		case EAGXDebugLevelLogOperations:
+		{
+			((FAGXDebugBlitCommandEncoder*)m_ptr)->Buffer.Blit([NSString stringWithFormat:@"%s", __PRETTY_FUNCTION__]);
+		}
+		default:
+		{
+			break;
+		}
+	}
+#endif
+}
+
+void FAGXBlitCommandEncoderDebugging::Synchronize(FAGXTexture const& texture, NSUInteger slice, NSUInteger level)
+{
+#if METAL_DEBUG_OPTIONS
+	switch(((FAGXDebugBlitCommandEncoder*)m_ptr)->Buffer.GetPtr()->DebugLevel)
+	{
+		case EAGXDebugLevelConditionalSubmit:
+		case EAGXDebugLevelWaitForComplete:
+		case EAGXDebugLevelLogOperations:
+		{
+			((FAGXDebugBlitCommandEncoder*)m_ptr)->Buffer.Blit([NSString stringWithFormat:@"%s", __PRETTY_FUNCTION__]);
+		}
+		default:
+		{
+			break;
+		}
+	}
+#endif
+}
+#endif
+
+void FAGXBlitCommandEncoderDebugging::Copy(FAGXTexture const& sourceTexture, NSUInteger sourceSlice, NSUInteger sourceLevel, mtlpp::Origin const& sourceOrigin, mtlpp::Size const& sourceSize, FAGXTexture const& destinationTexture, NSUInteger destinationSlice, NSUInteger destinationLevel, mtlpp::Origin const& destinationOrigin)
+{
+#if METAL_DEBUG_OPTIONS
+	switch(((FAGXDebugBlitCommandEncoder*)m_ptr)->Buffer.GetPtr()->DebugLevel)
+	{
+		case EAGXDebugLevelConditionalSubmit:
+		case EAGXDebugLevelWaitForComplete:
+		case EAGXDebugLevelLogOperations:
+		{
+			((FAGXDebugBlitCommandEncoder*)m_ptr)->Buffer.Blit([NSString stringWithFormat:@"%s", __PRETTY_FUNCTION__]);
+		}
+		default:
+		{
+			break;
+		}
+	}
+#endif
+}
+
+void FAGXBlitCommandEncoderDebugging::Copy(FAGXBuffer const& sourceBuffer, NSUInteger sourceOffset, NSUInteger sourceBytesPerRow, NSUInteger sourceBytesPerImage, mtlpp::Size const& sourceSize, FAGXTexture const& destinationTexture, NSUInteger destinationSlice, NSUInteger destinationLevel, mtlpp::Origin const& destinationOrigin)
+{
+#if METAL_DEBUG_OPTIONS
+	switch(((FAGXDebugBlitCommandEncoder*)m_ptr)->Buffer.GetPtr()->DebugLevel)
+	{
+		case EAGXDebugLevelConditionalSubmit:
+		case EAGXDebugLevelWaitForComplete:
+		case EAGXDebugLevelLogOperations:
+		{
+			((FAGXDebugBlitCommandEncoder*)m_ptr)->Buffer.Blit([NSString stringWithFormat:@"%s", __PRETTY_FUNCTION__]);
+		}
+		default:
+		{
+			break;
+		}
+	}
+#endif
+}
+
+void FAGXBlitCommandEncoderDebugging::Copy(FAGXBuffer const& sourceBuffer, NSUInteger sourceOffset, NSUInteger sourceBytesPerRow, NSUInteger sourceBytesPerImage, mtlpp::Size const& sourceSize, FAGXTexture const& destinationTexture, NSUInteger destinationSlice, NSUInteger destinationLevel, mtlpp::Origin const& destinationOrigin, mtlpp::BlitOption options)
+{
+#if METAL_DEBUG_OPTIONS
+	switch(((FAGXDebugBlitCommandEncoder*)m_ptr)->Buffer.GetPtr()->DebugLevel)
+	{
+		case EAGXDebugLevelConditionalSubmit:
+		case EAGXDebugLevelWaitForComplete:
+		case EAGXDebugLevelLogOperations:
+		{
+			((FAGXDebugBlitCommandEncoder*)m_ptr)->Buffer.Blit([NSString stringWithFormat:@"%s", __PRETTY_FUNCTION__]);
+		}
+		default:
+		{
+			break;
+		}
+	}
+#endif
+}
+
+void FAGXBlitCommandEncoderDebugging::Copy(FAGXTexture const& sourceTexture, NSUInteger sourceSlice, NSUInteger sourceLevel, mtlpp::Origin const& sourceOrigin, mtlpp::Size const& sourceSize, FAGXBuffer const& destinationBuffer, NSUInteger destinationOffset, NSUInteger destinationBytesPerRow, NSUInteger destinationBytesPerImage)
+{
+#if METAL_DEBUG_OPTIONS
+	switch(((FAGXDebugBlitCommandEncoder*)m_ptr)->Buffer.GetPtr()->DebugLevel)
+	{
+		case EAGXDebugLevelConditionalSubmit:
+		case EAGXDebugLevelWaitForComplete:
+		case EAGXDebugLevelLogOperations:
+		{
+			((FAGXDebugBlitCommandEncoder*)m_ptr)->Buffer.Blit([NSString stringWithFormat:@"%s", __PRETTY_FUNCTION__]);
+		}
+		default:
+		{
+			break;
+		}
+	}
+#endif
+}
+
+void FAGXBlitCommandEncoderDebugging::Copy(FAGXTexture const& sourceTexture, NSUInteger sourceSlice, NSUInteger sourceLevel, mtlpp::Origin const& sourceOrigin, mtlpp::Size const& sourceSize, FAGXBuffer const& destinationBuffer, NSUInteger destinationOffset, NSUInteger destinationBytesPerRow, NSUInteger destinationBytesPerImage, mtlpp::BlitOption options)
+{
+#if METAL_DEBUG_OPTIONS
+	switch(((FAGXDebugBlitCommandEncoder*)m_ptr)->Buffer.GetPtr()->DebugLevel)
+	{
+		case EAGXDebugLevelConditionalSubmit:
+		case EAGXDebugLevelWaitForComplete:
+		case EAGXDebugLevelLogOperations:
+		{
+			((FAGXDebugBlitCommandEncoder*)m_ptr)->Buffer.Blit([NSString stringWithFormat:@"%s", __PRETTY_FUNCTION__]);
+		}
+		default:
+		{
+			break;
+		}
+	}
+#endif
+}
+
+void FAGXBlitCommandEncoderDebugging::GenerateMipmaps(FAGXTexture const& texture)
+{
+#if METAL_DEBUG_OPTIONS
+	switch(((FAGXDebugBlitCommandEncoder*)m_ptr)->Buffer.GetPtr()->DebugLevel)
+	{
+		case EAGXDebugLevelConditionalSubmit:
+		case EAGXDebugLevelWaitForComplete:
+		case EAGXDebugLevelLogOperations:
+		{
+			((FAGXDebugBlitCommandEncoder*)m_ptr)->Buffer.Blit([NSString stringWithFormat:@"%s", __PRETTY_FUNCTION__]);
+		}
+		default:
+		{
+			break;
+		}
+	}
+#endif
+}
+
+void FAGXBlitCommandEncoderDebugging::Fill(FAGXBuffer const& buffer, ns::Range const& range, uint8_t value)
+{
+#if METAL_DEBUG_OPTIONS
+	switch(((FAGXDebugBlitCommandEncoder*)m_ptr)->Buffer.GetPtr()->DebugLevel)
+	{
+		case EAGXDebugLevelConditionalSubmit:
+		case EAGXDebugLevelWaitForComplete:
+		case EAGXDebugLevelLogOperations:
+		{
+			((FAGXDebugBlitCommandEncoder*)m_ptr)->Buffer.Blit([NSString stringWithFormat:@"%s", __PRETTY_FUNCTION__]);
+		}
+		default:
+		{
+			break;
+		}
+	}
+#endif
+}
+
+void FAGXBlitCommandEncoderDebugging::Copy(FAGXBuffer const& sourceBuffer, NSUInteger sourceOffset, FAGXBuffer const& destinationBuffer, NSUInteger destinationOffset, NSUInteger size)
+{
+#if METAL_DEBUG_OPTIONS
+	switch(((FAGXDebugBlitCommandEncoder*)m_ptr)->Buffer.GetPtr()->DebugLevel)
+	{
+		case EAGXDebugLevelConditionalSubmit:
+		case EAGXDebugLevelWaitForComplete:
+		case EAGXDebugLevelLogOperations:
+		{
+			((FAGXDebugBlitCommandEncoder*)m_ptr)->Buffer.Blit([NSString stringWithFormat:@"%s", __PRETTY_FUNCTION__]);
+		}
+		default:
+		{
+			break;
+		}
+	}
+#endif
+}
+
+#endif
diff --git a/Engine/Source/Runtime/Apple/AGXRHI/Private/AGXBlitCommandEncoder.h b/Engine/Source/Runtime/Apple/AGXRHI/Private/AGXBlitCommandEncoder.h
new file mode 100644
--- /dev/null
+++ b/Engine/Source/Runtime/Apple/AGXRHI/Private/AGXBlitCommandEncoder.h
@@ -0,0 +1,53 @@
+// Copyright Epic Games, Inc. All Rights Reserved.
+
+#pragma once
+
+#include <Metal/Metal.h>
+#include "AGXDebugCommandEncoder.h"
+
+#if MTLPP_CONFIG_VALIDATE && METAL_DEBUG_OPTIONS
+NS_ASSUME_NONNULL_BEGIN
+
+@class FAGXDebugCommandBuffer;
+@class FAGXDebugBlitCommandEncoder;
+
+class FAGXBlitCommandEncoderDebugging : public FAGXCommandEncoderDebugging
+{
+public:
+	FAGXBlitCommandEncoderDebugging();
+	FAGXBlitCommandEncoderDebugging(mtlpp::BlitCommandEncoder& Encoder, FAGXCommandBufferDebugging& Buffer);
+	FAGXBlitCommandEncoderDebugging(FAGXDebugCommandEncoder* handle);
+	
+	static FAGXBlitCommandEncoderDebugging Get(mtlpp::BlitCommandEncoder& Buffer);
+	
+	void InsertDebugSignpost(ns::String const& Label);
+	void PushDebugGroup(ns::String const& Group);
+	void PopDebugGroup();
+	
+	void EndEncoder();
+	
+#if PLATFORM_MAC
+	void Synchronize(mtlpp::Resource const& resource);
+	
+	void Synchronize(FAGXTexture const& texture, NSUInteger slice, NSUInteger level);
+#endif
+	
+	void Copy(FAGXTexture const& sourceTexture, NSUInteger sourceSlice, NSUInteger sourceLevel, mtlpp::Origin const& sourceOrigin, mtlpp::Size const& sourceSize, FAGXTexture const& destinationTexture, NSUInteger destinationSlice, NSUInteger destinationLevel, mtlpp::Origin const& destinationOrigin);
+	
+	void Copy(FAGXBuffer const& sourceBuffer, NSUInteger sourceOffset, NSUInteger sourceBytesPerRow, NSUInteger sourceBytesPerImage, mtlpp::Size const& sourceSize, FAGXTexture const& destinationTexture, NSUInteger destinationSlice, NSUInteger destinationLevel, mtlpp::Origin const& destinationOrigin);
+	
+	void Copy(FAGXBuffer const& sourceBuffer, NSUInteger sourceOffset, NSUInteger sourceBytesPerRow, NSUInteger sourceBytesPerImage, mtlpp::Size const& sourceSize, FAGXTexture const& destinationTexture, NSUInteger destinationSlice, NSUInteger destinationLevel, mtlpp::Origin const& destinationOrigin, mtlpp::BlitOption options);
+	
+	void Copy(FAGXTexture const& sourceTexture, NSUInteger sourceSlice, NSUInteger sourceLevel, mtlpp::Origin const& sourceOrigin, mtlpp::Size const& sourceSize, FAGXBuffer const& destinationBuffer, NSUInteger destinationOffset, NSUInteger destinationBytesPerRow, NSUInteger destinationBytesPerImage);
+	
+	void Copy(FAGXTexture const& sourceTexture, NSUInteger sourceSlice, NSUInteger sourceLevel, mtlpp::Origin const& sourceOrigin, mtlpp::Size const& sourceSize, FAGXBuffer const& destinationBuffer, NSUInteger destinationOffset, NSUInteger destinationBytesPerRow, NSUInteger destinationBytesPerImage, mtlpp::BlitOption options);
+	
+	void GenerateMipmaps(FAGXTexture const& texture);
+	
+	void Fill(FAGXBuffer const& buffer, ns::Range const& range, uint8_t value);
+	
+	void Copy(FAGXBuffer const& sourceBuffer, NSUInteger sourceOffset, FAGXBuffer const& destinationBuffer, NSUInteger destinationOffset, NSUInteger size);
+};
+
+NS_ASSUME_NONNULL_END
+#endif
diff --git a/Engine/Source/Runtime/Apple/AGXRHI/Private/AGXBuffer.cpp b/Engine/Source/Runtime/Apple/AGXRHI/Private/AGXBuffer.cpp
new file mode 100644
--- /dev/null
+++ b/Engine/Source/Runtime/Apple/AGXRHI/Private/AGXBuffer.cpp
@@ -0,0 +1,1848 @@
+// Copyright Epic Games, Inc. All Rights Reserved.
+
+#include "AGXRHIPrivate.h"
+#include "AGXBuffer.h"
+#include "Templates/AlignmentTemplates.h"
+#include "AGXLLM.h"
+#include <objc/runtime.h>
+#include "AGXCommandBuffer.h"
+#include "AGXProfiler.h"
+#include "AGXRenderPass.h"
+
+DECLARE_MEMORY_STAT(TEXT("Used Device Buffer Memory"), STAT_AGXDeviceBufferMemory, STATGROUP_AGXRHI);
+DECLARE_MEMORY_STAT(TEXT("Used Pooled Buffer Memory"), STAT_AGXPooledBufferMemory, STATGROUP_AGXRHI);
+DECLARE_MEMORY_STAT(TEXT("Used Magazine Buffer Memory"), STAT_AGXMagazineBufferMemory, STATGROUP_AGXRHI);
+DECLARE_MEMORY_STAT(TEXT("Used Heap Buffer Memory"), STAT_AGXHeapBufferMemory, STATGROUP_AGXRHI);
+DECLARE_MEMORY_STAT(TEXT("Used Linear Buffer Memory"), STAT_AGXLinearBufferMemory, STATGROUP_AGXRHI);
+
+DECLARE_MEMORY_STAT(TEXT("Unused Pooled Buffer Memory"), STAT_AGXPooledBufferUnusedMemory, STATGROUP_AGXRHI);
+DECLARE_MEMORY_STAT(TEXT("Unused Magazine Buffer Memory"), STAT_AGXMagazineBufferUnusedMemory, STATGROUP_AGXRHI);
+DECLARE_MEMORY_STAT(TEXT("Unused Heap Buffer Memory"), STAT_AGXHeapBufferUnusedMemory, STATGROUP_AGXRHI);
+DECLARE_MEMORY_STAT(TEXT("Unused Linear Buffer Memory"), STAT_AGXLinearBufferUnusedMemory, STATGROUP_AGXRHI);
+
+static int32 GAGXHeapBufferBytesToCompact = 0;
+static FAutoConsoleVariableRef CVarAGXHeapBufferBytesToCompact(
+    TEXT("rhi.AGX.HeapBufferBytesToCompact"),
+    GAGXHeapBufferBytesToCompact,
+    TEXT("When enabled (> 0) this will force AGXRHI to compact the given number of bytes each frame into older buffer heaps from newer ones in order to defragment memory and reduce wastage.\n")
+    TEXT("(Off by default (0))"));
+
+static int32 GAGXResourcePurgeInPool = 0;
+static FAutoConsoleVariableRef CVarAGXResourcePurgeInPool(
+	TEXT("rhi.AGX.ResourcePurgeInPool"),
+	GAGXResourcePurgeInPool,
+	TEXT("Use the SetPurgeableState function to allow the OS to reclaim memory from resources while they are unused in the pools. (Default: 0, Off)"));
+
+#if METAL_DEBUG_OPTIONS
+extern int32 GAGXBufferScribble;
+#endif
+
+FAGXBuffer::~FAGXBuffer()
+{
+}
+
+FAGXBuffer::FAGXBuffer(ns::Protocol<id<MTLBuffer>>::type handle, ns::Ownership retain)
+: mtlpp::Buffer(handle, nullptr, retain)
+, Heap(nullptr)
+, Linear(nullptr)
+, Magazine(nullptr)
+, bPooled(false)
+, bSingleUse(false)
+{
+}
+
+FAGXBuffer::FAGXBuffer(mtlpp::Buffer&& rhs, FAGXSubBufferHeap* heap)
+: mtlpp::Buffer((mtlpp::Buffer&&)rhs)
+, Heap(heap)
+, Linear(nullptr)
+, Magazine(nullptr)
+, bPooled(false)
+, bSingleUse(false)
+{
+}
+
+
+FAGXBuffer::FAGXBuffer(mtlpp::Buffer&& rhs, FAGXSubBufferLinear* heap)
+: mtlpp::Buffer((mtlpp::Buffer&&)rhs)
+, Heap(nullptr)
+, Linear(heap)
+, Magazine(nullptr)
+, bPooled(false)
+, bSingleUse(false)
+{
+}
+
+FAGXBuffer::FAGXBuffer(mtlpp::Buffer&& rhs, FAGXSubBufferMagazine* magazine)
+: mtlpp::Buffer((mtlpp::Buffer&&)rhs)
+, Heap(nullptr)
+, Linear(nullptr)
+, Magazine(magazine)
+, bPooled(false)
+, bSingleUse(false)
+{
+}
+
+FAGXBuffer::FAGXBuffer(mtlpp::Buffer&& rhs, bool bInPooled)
+: mtlpp::Buffer((mtlpp::Buffer&&)rhs)
+, Heap(nullptr)
+, Linear(nullptr)
+, Magazine(nullptr)
+, bPooled(bInPooled)
+, bSingleUse(false)
+{
+}
+
+FAGXBuffer::FAGXBuffer(const FAGXBuffer& rhs)
+: mtlpp::Buffer(rhs)
+, Heap(rhs.Heap)
+, Linear(rhs.Linear)
+, Magazine(rhs.Magazine)
+, bPooled(rhs.bPooled)
+, bSingleUse(false)
+{
+}
+
+FAGXBuffer::FAGXBuffer(FAGXBuffer&& rhs)
+: mtlpp::Buffer((mtlpp::Buffer&&)rhs)
+, Heap(rhs.Heap)
+, Linear(rhs.Linear)
+, Magazine(rhs.Magazine)
+, bPooled(rhs.bPooled)
+, bSingleUse(false)
+{
+}
+
+FAGXBuffer& FAGXBuffer::operator=(const FAGXBuffer& rhs)
+{
+	if(this != &rhs)
+	{
+		mtlpp::Buffer::operator=(rhs);
+        Heap = rhs.Heap;
+		Linear = rhs.Linear;
+		Magazine = rhs.Magazine;
+		bPooled = rhs.bPooled;
+		bSingleUse = rhs.bSingleUse;
+	}
+	return *this;
+}
+
+FAGXBuffer& FAGXBuffer::operator=(FAGXBuffer&& rhs)
+{
+	mtlpp::Buffer::operator=((mtlpp::Buffer&&)rhs);
+	Heap = rhs.Heap;
+	Linear = rhs.Linear;
+	Magazine = rhs.Magazine;
+	bPooled = rhs.bPooled;
+	bSingleUse = rhs.bSingleUse;
+	return *this;
+}
+
+void FAGXBuffer::Release()
+{
+	if (Heap)
+	{
+		Heap->FreeRange(ns::Range(GetOffset(), GetLength()));
+		Heap = nullptr;
+	}
+	else if (Linear)
+	{
+		Linear->FreeRange(ns::Range(GetOffset(), GetLength()));
+		Linear = nullptr;
+	}
+	else if (Magazine)
+	{
+		Magazine->FreeRange(ns::Range(GetOffset(), GetLength()));
+		Magazine = nullptr;
+	}
+}
+
+void FAGXBuffer::SetOwner(class FAGXRHIBuffer* Owner, bool bIsSwap)
+{
+	check(Owner == nullptr);
+	if (Heap)
+	{
+		Heap->SetOwner(ns::Range(GetOffset(), GetLength()), Owner, bIsSwap);
+	}
+}
+
+FAGXSubBufferHeap::FAGXSubBufferHeap(NSUInteger Size, NSUInteger Alignment, mtlpp::ResourceOptions Options, FCriticalSection& InPoolMutex)
+: PoolMutex(InPoolMutex)
+, OutstandingAllocs(0)
+, MinAlign(Alignment)
+, UsedSize(0)
+{
+	Options = (mtlpp::ResourceOptions)FAGXCommandQueue::GetCompatibleResourceOptions(Options);
+	static bool bSupportsHeaps = GetAGXDeviceContext().SupportsFeature(EAGXFeaturesHeaps);
+	NSUInteger FullSize = Align(Size, Alignment);
+	METAL_GPUPROFILE(FAGXScopedCPUStats CPUStat(FString::Printf(TEXT("AllocBuffer: %llu, %llu"), FullSize, Options)));
+	
+	mtlpp::StorageMode Storage = (mtlpp::StorageMode)((Options & mtlpp::ResourceStorageModeMask) >> mtlpp::ResourceStorageModeShift);
+#if PLATFORM_MAC
+	check(Storage != mtlpp::StorageMode::Managed /* Managed memory cannot be safely suballocated! When you overwrite existing data the GPU buffer is immediately disposed of! */);
+#endif
+
+	if (bSupportsHeaps && Storage == mtlpp::StorageMode::Private)
+	{
+		mtlpp::HeapDescriptor Desc;
+		Desc.SetSize(FullSize);
+		Desc.SetStorageMode(Storage);
+		ParentHeap = GetAGXDeviceContext().GetDevice().NewHeap(Desc);
+		check(ParentHeap.GetPtr());
+#if STATS || ENABLE_LOW_LEVEL_MEM_TRACKER
+		AGXLLM::LogAllocHeap(GetAGXDeviceContext().GetDevice(), ParentHeap);
+#endif
+	}
+	else
+	{
+		ParentBuffer = MTLPP_VALIDATE(mtlpp::Device, GetAGXDeviceContext().GetDevice(), AGXSafeGetRuntimeDebuggingLevel() >= EAGXDebugLevelValidation, NewBuffer(FullSize, Options));
+		check(ParentBuffer.GetPtr());
+		check(ParentBuffer.GetLength() >= FullSize);
+#if STATS || ENABLE_LOW_LEVEL_MEM_TRACKER
+		AGXLLM::LogAllocBuffer(GetAGXDeviceContext().GetDevice(), ParentBuffer);
+#endif
+		FreeRanges.Add(ns::Range(0, FullSize));
+	}
+	INC_MEMORY_STAT_BY(STAT_AGXBufferUnusedMemory, FullSize);
+	INC_MEMORY_STAT_BY(STAT_AGXHeapBufferUnusedMemory, FullSize);
+}
+
+FAGXSubBufferHeap::~FAGXSubBufferHeap()
+{
+	if (ParentHeap)
+	{
+		DEC_MEMORY_STAT_BY(STAT_AGXBufferUnusedMemory, ParentHeap.GetSize());
+		DEC_MEMORY_STAT_BY(STAT_AGXHeapBufferUnusedMemory, ParentHeap.GetSize());
+	}
+	else
+	{
+		DEC_MEMORY_STAT_BY(STAT_AGXBufferUnusedMemory, ParentBuffer.GetLength());
+		DEC_MEMORY_STAT_BY(STAT_AGXHeapBufferUnusedMemory, ParentBuffer.GetLength());
+	}
+}
+
+void FAGXSubBufferHeap::SetOwner(ns::Range const& Range, FAGXRHIBuffer* Owner, bool bIsSwap)
+{
+	check(Owner == nullptr);
+	FScopeLock Lock(&PoolMutex);
+	for (uint32 i = 0; i < AllocRanges.Num(); i++)
+	{
+		if (AllocRanges[i].Range.Location == Range.Location)
+		{
+			check(AllocRanges[i].Range.Length == Range.Length);
+			check(AllocRanges[i].Owner == nullptr || Owner == nullptr || bIsSwap);
+			AllocRanges[i].Owner = Owner;
+			break;
+		}
+	}
+}
+
+void FAGXSubBufferHeap::FreeRange(ns::Range const& Range)
+{
+	FPlatformAtomics::InterlockedDecrement(&OutstandingAllocs);
+	{
+		FScopeLock Lock(&PoolMutex);
+		for (uint32 i = 0; i < AllocRanges.Num(); i++)
+		{
+			if (AllocRanges[i].Range.Location == Range.Location)
+			{
+				check(AllocRanges[i].Range.Length == Range.Length);
+				AllocRanges.RemoveAt(i);
+				break;
+			}
+		}
+	}
+	if (ParentHeap)
+	{
+		INC_MEMORY_STAT_BY(STAT_AGXBufferUnusedMemory, Range.Length);
+		INC_MEMORY_STAT_BY(STAT_AGXHeapBufferUnusedMemory, Range.Length);
+		DEC_MEMORY_STAT_BY(STAT_AGXHeapBufferMemory, Range.Length);
+	}
+	else
+	{
+#if METAL_DEBUG_OPTIONS
+		if (GIsRHIInitialized)
+		{
+			MTLPP_VALIDATE_ONLY(mtlpp::Buffer, ParentBuffer, AGXSafeGetRuntimeDebuggingLevel() >= EAGXDebugLevelValidation, ReleaseRange(Range));
+			FAGXBuffer Buf(ParentBuffer.NewBuffer(Range), false);
+			GetAGXDeviceContext().ValidateIsInactiveBuffer(Buf);
+		}
+#endif
+    
+		FScopeLock Lock(&PoolMutex);
+		{
+			ns::Range CompactRange = Range;
+			for (uint32 i = 0; i < FreeRanges.Num(); )
+			{
+				if (FreeRanges[i].Location == (CompactRange.Location + CompactRange.Length))
+				{
+					ns::Range PrevRange = FreeRanges[i];
+					FreeRanges.RemoveAt(i);
+					
+					CompactRange.Length += PrevRange.Length;
+				}
+				else if (CompactRange.Location == (FreeRanges[i].Location + FreeRanges[i].Length))
+				{
+					ns::Range PrevRange = FreeRanges[i];
+					FreeRanges.RemoveAt(i);
+					
+					CompactRange.Location = PrevRange.Location;
+					CompactRange.Length += PrevRange.Length;
+				}
+				else
+				{
+					i++;
+				}
+			}
+		
+			uint32 i = 0;
+			for (; i < FreeRanges.Num(); i++)
+			{
+				if (FreeRanges[i].Length >= CompactRange.Length)
+				{
+					break;
+				}
+			}
+			FreeRanges.Insert(CompactRange, i);
+		
+			UsedSize -= Range.Length;
+			
+			INC_MEMORY_STAT_BY(STAT_AGXBufferUnusedMemory, Range.Length);
+			INC_MEMORY_STAT_BY(STAT_AGXHeapBufferUnusedMemory, Range.Length);
+			DEC_MEMORY_STAT_BY(STAT_AGXHeapBufferMemory, Range.Length);
+		
+#if METAL_DEBUG_OPTIONS
+			uint64 LostSize = GetSize() - UsedSize;
+			for (ns::Range const& FreeRange : FreeRanges)
+			{
+				LostSize -= FreeRange.Length;
+			}
+			check(LostSize == 0);
+#endif
+		}
+	}
+}
+
+ns::String FAGXSubBufferHeap::GetLabel() const
+{
+	if (ParentHeap)
+	{
+		return ParentHeap.GetLabel();
+	}
+	else
+	{
+		return ParentBuffer.GetLabel();
+	}
+}
+
+mtlpp::Device FAGXSubBufferHeap::GetDevice() const
+{
+	if (ParentHeap)
+	{
+		return ParentHeap.GetDevice();
+	}
+	else
+	{
+		return ParentBuffer.GetDevice();
+	}
+}
+
+mtlpp::StorageMode FAGXSubBufferHeap::GetStorageMode() const
+{
+	if (ParentHeap)
+	{
+		return ParentHeap.GetStorageMode();
+	}
+	else
+	{
+		return ParentBuffer.GetStorageMode();
+	}
+}
+
+mtlpp::CpuCacheMode FAGXSubBufferHeap::GetCpuCacheMode() const
+{
+	if (ParentHeap)
+	{
+		return ParentHeap.GetCpuCacheMode();
+	}
+	else
+	{
+		return ParentBuffer.GetCpuCacheMode();
+	}
+}
+
+NSUInteger FAGXSubBufferHeap::GetSize() const
+{
+	if (ParentHeap)
+	{
+		return ParentHeap.GetSize();
+	}
+	else
+	{
+		return ParentBuffer.GetLength();
+	}
+}
+
+NSUInteger FAGXSubBufferHeap::GetUsedSize() const
+{
+	if (ParentHeap)
+	{
+		return ParentHeap.GetUsedSize();
+	}
+	else
+	{
+		return UsedSize;
+	}
+}
+
+int64 FAGXSubBufferHeap::NumCurrentAllocations() const
+{
+	return OutstandingAllocs;
+}
+
+void FAGXSubBufferHeap::SetLabel(const ns::String& label)
+{
+	if (ParentHeap)
+	{
+		ParentHeap.SetLabel(label);
+	}
+	else
+	{
+		ParentBuffer.SetLabel(label);
+	}
+}
+
+NSUInteger FAGXSubBufferHeap::MaxAvailableSize() const
+{
+	if (ParentHeap)
+	{
+		return ParentHeap.MaxAvailableSizeWithAlignment(MinAlign);
+	}
+	else
+	{
+		if (UsedSize < GetSize())
+		{
+			return FreeRanges.Last().Length;
+		}
+		else
+		{
+			return 0;
+		}
+	}
+}
+
+bool FAGXSubBufferHeap::CanAllocateSize(NSUInteger Size) const
+{
+	if (ParentHeap)
+	{
+		NSUInteger Storage = (NSUInteger(GetStorageMode()) << mtlpp::ResourceStorageModeShift);
+		NSUInteger Cache = (NSUInteger(GetCpuCacheMode()) << mtlpp::ResourceCpuCacheModeShift);
+		mtlpp::ResourceOptions Opt = mtlpp::ResourceOptions(Storage | Cache);
+		
+		NSUInteger Align = ParentHeap.GetDevice().HeapBufferSizeAndAlign(Size, Opt).Align;
+		return Size <= ParentHeap.MaxAvailableSizeWithAlignment(Align);
+	}
+	else
+	{
+		return Size <= MaxAvailableSize();
+	}
+}
+
+FAGXBuffer FAGXSubBufferHeap::NewBuffer(NSUInteger length)
+{
+	NSUInteger Size = Align(length, MinAlign);
+	FAGXBuffer Result;
+	
+	if (ParentHeap)
+	{
+		NSUInteger Storage = (NSUInteger(GetStorageMode()) << mtlpp::ResourceStorageModeShift);
+		NSUInteger Cache = (NSUInteger(GetCpuCacheMode()) << mtlpp::ResourceCpuCacheModeShift);
+		mtlpp::ResourceOptions Opt = mtlpp::ResourceOptions(Storage | Cache);
+		
+		Result = FAGXBuffer(ParentHeap.NewBuffer(Size, Opt), this);
+#if STATS || ENABLE_LOW_LEVEL_MEM_TRACKER
+		AGXLLM::LogAllocBuffer(GetAGXDeviceContext().GetDevice(), Result);
+#endif
+		DEC_MEMORY_STAT_BY(STAT_AGXBufferUnusedMemory, Result.GetLength());
+		DEC_MEMORY_STAT_BY(STAT_AGXHeapBufferUnusedMemory, Result.GetLength());
+		INC_MEMORY_STAT_BY(STAT_AGXHeapBufferMemory, Result.GetLength());
+	}
+	else
+	{
+		check(ParentBuffer && ParentBuffer.GetPtr());
+	
+		FScopeLock Lock(&PoolMutex);
+		if (MaxAvailableSize() >= Size)
+		{
+			for (uint32 i = 0; i < FreeRanges.Num(); i++)
+			{
+				if (FreeRanges[i].Length >= Size)
+				{
+					ns::Range Range = FreeRanges[i];
+					FreeRanges.RemoveAt(i);
+					
+					UsedSize += Range.Length;
+					
+					DEC_MEMORY_STAT_BY(STAT_AGXBufferUnusedMemory, Range.Length);
+					DEC_MEMORY_STAT_BY(STAT_AGXHeapBufferUnusedMemory, Range.Length);
+					INC_MEMORY_STAT_BY(STAT_AGXHeapBufferMemory, Range.Length);
+				
+					if (Range.Length > Size)
+					{
+						ns::Range Split = ns::Range(Range.Location + Size, Range.Length - Size);
+						FPlatformAtomics::InterlockedIncrement(&OutstandingAllocs);
+						FreeRange(Split);
+						
+						Range.Length = Size;
+					}
+					
+#if METAL_DEBUG_OPTIONS
+					uint64 LostSize = GetSize() - UsedSize;
+					for (ns::Range const& FreeRange : FreeRanges)
+					{
+						LostSize -= FreeRange.Length;
+					}
+					check(LostSize == 0);
+#endif
+					
+					Result = FAGXBuffer(MTLPP_VALIDATE(mtlpp::Buffer, ParentBuffer, AGXSafeGetRuntimeDebuggingLevel() >= EAGXDebugLevelValidation, NewBuffer(Range)), this);
+				
+                    Allocation Alloc;
+                    Alloc.Range = Range;
+                    Alloc.Resource = Result.GetPtr();
+                    Alloc.Owner = nullptr;
+                    AllocRanges.Add(Alloc);
+
+					break;
+				}
+			}
+		}
+	}
+	FPlatformAtomics::InterlockedIncrement(&OutstandingAllocs);
+	check(Result && Result.GetPtr());
+	return Result;
+}
+
+mtlpp::PurgeableState FAGXSubBufferHeap::SetPurgeableState(mtlpp::PurgeableState state)
+{
+	if (ParentHeap)
+	{
+		return ParentHeap.SetPurgeableState(state);
+	}
+	else
+	{
+		return ParentBuffer.SetPurgeableState(state);
+	}
+}
+
+#pragma mark --
+
+FAGXSubBufferLinear::FAGXSubBufferLinear(NSUInteger Size, NSUInteger Alignment, mtlpp::ResourceOptions Options, FCriticalSection& InPoolMutex)
+: PoolMutex(InPoolMutex)
+, MinAlign(Alignment)
+, WriteHead(0)
+, UsedSize(0)
+, FreedSize(0)
+{
+	Options = (mtlpp::ResourceOptions)FAGXCommandQueue::GetCompatibleResourceOptions(Options);
+	NSUInteger FullSize = Align(Size, Alignment);
+	METAL_GPUPROFILE(FAGXScopedCPUStats CPUStat(FString::Printf(TEXT("AllocBuffer: %llu, %llu"), FullSize, Options)));
+	
+	mtlpp::StorageMode Storage = (mtlpp::StorageMode)((Options & mtlpp::ResourceStorageModeMask) >> mtlpp::ResourceStorageModeShift);
+	ParentBuffer = MTLPP_VALIDATE(mtlpp::Device, GetAGXDeviceContext().GetDevice(), AGXSafeGetRuntimeDebuggingLevel() >= EAGXDebugLevelValidation, NewBuffer(FullSize, Options));
+	check(ParentBuffer.GetPtr());
+	check(ParentBuffer.GetLength() >= FullSize);
+#if STATS || ENABLE_LOW_LEVEL_MEM_TRACKER
+	AGXLLM::LogAllocBuffer(GetAGXDeviceContext().GetDevice(), ParentBuffer);
+#endif
+	INC_MEMORY_STAT_BY(STAT_AGXBufferUnusedMemory, FullSize);
+	INC_MEMORY_STAT_BY(STAT_AGXLinearBufferUnusedMemory, FullSize);
+}
+
+FAGXSubBufferLinear::~FAGXSubBufferLinear()
+{
+	DEC_MEMORY_STAT_BY(STAT_AGXBufferUnusedMemory, ParentBuffer.GetLength());
+	DEC_MEMORY_STAT_BY(STAT_AGXLinearBufferUnusedMemory, ParentBuffer.GetLength());
+}
+
+void FAGXSubBufferLinear::FreeRange(ns::Range const& Range)
+{
+#if METAL_DEBUG_OPTIONS
+	if (GIsRHIInitialized)
+	{
+		MTLPP_VALIDATE_ONLY(mtlpp::Buffer, ParentBuffer, AGXSafeGetRuntimeDebuggingLevel() >= EAGXDebugLevelValidation, ReleaseRange(Range));
+		FAGXBuffer Buf(ParentBuffer.NewBuffer(Range), false);
+		GetAGXDeviceContext().ValidateIsInactiveBuffer(Buf);
+	}
+#endif
+	
+	FScopeLock Lock(&PoolMutex);
+	{
+		FreedSize += Range.Length;
+		INC_MEMORY_STAT_BY(STAT_AGXBufferUnusedMemory, Range.Length);
+		INC_MEMORY_STAT_BY(STAT_AGXLinearBufferUnusedMemory, Range.Length);
+		DEC_MEMORY_STAT_BY(STAT_AGXLinearBufferMemory, Range.Length);
+		if (FreedSize == UsedSize)
+		{
+			UsedSize = 0;
+			FreedSize = 0;
+			WriteHead = 0;
+		}
+	}
+}
+
+ns::String FAGXSubBufferLinear::GetLabel() const
+{
+	return ParentBuffer.GetLabel();
+}
+
+mtlpp::Device FAGXSubBufferLinear::GetDevice() const
+{
+	return ParentBuffer.GetDevice();
+}
+
+mtlpp::StorageMode FAGXSubBufferLinear::GetStorageMode() const
+{
+	return ParentBuffer.GetStorageMode();
+}
+
+mtlpp::CpuCacheMode FAGXSubBufferLinear::GetCpuCacheMode() const
+{
+	return ParentBuffer.GetCpuCacheMode();
+}
+
+NSUInteger FAGXSubBufferLinear::GetSize() const
+{
+	return ParentBuffer.GetLength();
+}
+
+NSUInteger FAGXSubBufferLinear::GetUsedSize() const
+{
+	return UsedSize;
+}
+
+void FAGXSubBufferLinear::SetLabel(const ns::String& label)
+{
+	ParentBuffer.SetLabel(label);
+}
+
+bool FAGXSubBufferLinear::CanAllocateSize(NSUInteger Size) const
+{
+	if (WriteHead < GetSize())
+	{
+		NSUInteger Alignment = FMath::Max(NSUInteger(MinAlign), NSUInteger(Size & ~(Size - 1llu)));
+		NSUInteger NewWriteHead = Align(WriteHead, Alignment);
+		return (GetSize() - NewWriteHead) > Size;
+	}
+	else
+	{
+		return 0;
+	}
+}
+
+FAGXBuffer FAGXSubBufferLinear::NewBuffer(NSUInteger length)
+{
+	FScopeLock Lock(&PoolMutex);
+	NSUInteger Alignment = FMath::Max(NSUInteger(MinAlign), NSUInteger(length & ~(length - 1llu)));
+	NSUInteger Size = Align(length, Alignment);
+	NSUInteger NewWriteHead = Align(WriteHead, Alignment);
+	
+	FAGXBuffer Result;
+	if ((GetSize() - NewWriteHead) > Size)
+	{
+		ns::Range Range(NewWriteHead, Size);
+		DEC_MEMORY_STAT_BY(STAT_AGXBufferUnusedMemory, Range.Length);
+		DEC_MEMORY_STAT_BY(STAT_AGXLinearBufferUnusedMemory, Range.Length);
+		INC_MEMORY_STAT_BY(STAT_AGXLinearBufferMemory, Range.Length);
+		Result = FAGXBuffer(MTLPP_VALIDATE(mtlpp::Buffer, ParentBuffer, AGXSafeGetRuntimeDebuggingLevel() >= EAGXDebugLevelValidation, NewBuffer(Range)), this);
+		UsedSize += Size;
+		WriteHead = NewWriteHead + Size;
+	}
+	
+	return Result;
+}
+
+mtlpp::PurgeableState FAGXSubBufferLinear::SetPurgeableState(mtlpp::PurgeableState state)
+{
+	return ParentBuffer.SetPurgeableState(state);
+}
+
+#pragma mark --
+
+FAGXSubBufferMagazine::FAGXSubBufferMagazine(NSUInteger Size, NSUInteger ChunkSize, mtlpp::ResourceOptions Options)
+: MinAlign(ChunkSize)
+, BlockSize(ChunkSize)
+, OutstandingAllocs(0)
+, UsedSize(0)
+{
+	Options = (mtlpp::ResourceOptions)FAGXCommandQueue::GetCompatibleResourceOptions(Options);
+    static bool bSupportsHeaps = GetAGXDeviceContext().SupportsFeature(EAGXFeaturesHeaps);
+    mtlpp::StorageMode Storage = (mtlpp::StorageMode)((Options & mtlpp::ResourceStorageModeMask) >> mtlpp::ResourceStorageModeShift);
+    if (PLATFORM_IOS && bSupportsHeaps && Storage == mtlpp::StorageMode::Private)
+    {
+        MinAlign = GetAGXDeviceContext().GetDevice().HeapBufferSizeAndAlign(BlockSize, Options).Align;
+    }
+    
+    NSUInteger FullSize = Align(Size, MinAlign);
+	METAL_GPUPROFILE(FAGXScopedCPUStats CPUStat(FString::Printf(TEXT("AllocBuffer: %llu, %llu"), FullSize, Options)));
+	
+#if PLATFORM_MAC
+	check(Storage != mtlpp::StorageMode::Managed /* Managed memory cannot be safely suballocated! When you overwrite existing data the GPU buffer is immediately disposed of! */);
+#endif
+
+	if (bSupportsHeaps && Storage == mtlpp::StorageMode::Private)
+	{
+		mtlpp::HeapDescriptor Desc;
+		Desc.SetSize(FullSize);
+		Desc.SetStorageMode(Storage);
+		ParentHeap = GetAGXDeviceContext().GetDevice().NewHeap(Desc);
+		check(ParentHeap.GetPtr());
+		METAL_FATAL_ASSERT(ParentHeap, TEXT("Failed to create heap of size %u and resource options %u"), Size, (uint32)Options);
+#if STATS || ENABLE_LOW_LEVEL_MEM_TRACKER
+		AGXLLM::LogAllocHeap(GetAGXDeviceContext().GetDevice(), ParentHeap);
+#endif
+	}
+	else
+	{
+		ParentBuffer = MTLPP_VALIDATE(mtlpp::Device, GetAGXDeviceContext().GetDevice(), AGXSafeGetRuntimeDebuggingLevel() >= EAGXDebugLevelValidation, NewBuffer(FullSize, Options));
+		check(ParentBuffer.GetPtr());
+		check(ParentBuffer.GetLength() >= FullSize);
+		METAL_FATAL_ASSERT(ParentBuffer, TEXT("Failed to create heap of size %u and resource options %u"), Size, (uint32)Options);
+#if STATS || ENABLE_LOW_LEVEL_MEM_TRACKER
+		AGXLLM::LogAllocBuffer(GetAGXDeviceContext().GetDevice(), ParentBuffer);
+#endif
+		
+		INC_MEMORY_STAT_BY(STAT_AGXBufferUnusedMemory, FullSize);
+		INC_MEMORY_STAT_BY(STAT_AGXMagazineBufferUnusedMemory, FullSize);
+        uint32 BlockCount = FullSize / ChunkSize;
+        Blocks.AddZeroed(BlockCount);
+	}
+}
+
+FAGXSubBufferMagazine::~FAGXSubBufferMagazine()
+{
+	if (ParentHeap)
+	{
+		DEC_MEMORY_STAT_BY(STAT_AGXBufferUnusedMemory, ParentHeap.GetSize());
+		DEC_MEMORY_STAT_BY(STAT_AGXMagazineBufferUnusedMemory, ParentHeap.GetSize());
+	}
+	else
+	{
+		DEC_MEMORY_STAT_BY(STAT_AGXBufferUnusedMemory, ParentBuffer.GetLength());
+		DEC_MEMORY_STAT_BY(STAT_AGXMagazineBufferUnusedMemory, ParentBuffer.GetLength());
+	}
+}
+
+void FAGXSubBufferMagazine::FreeRange(ns::Range const& Range)
+{
+	FPlatformAtomics::InterlockedDecrement(&OutstandingAllocs);
+	if (ParentHeap)
+	{
+		INC_MEMORY_STAT_BY(STAT_AGXBufferUnusedMemory, Range.Length);
+		INC_MEMORY_STAT_BY(STAT_AGXMagazineBufferUnusedMemory, Range.Length);
+		DEC_MEMORY_STAT_BY(STAT_AGXMagazineBufferMemory, Range.Length);
+	}
+	else
+	{
+#if METAL_DEBUG_OPTIONS
+		if (GIsRHIInitialized)
+		{
+			MTLPP_VALIDATE_ONLY(mtlpp::Buffer, ParentBuffer, AGXSafeGetRuntimeDebuggingLevel() >= EAGXDebugLevelValidation, ReleaseRange(Range));
+			FAGXBuffer Buf(ParentBuffer.NewBuffer(Range), false);
+			GetAGXDeviceContext().ValidateIsInactiveBuffer(Buf);
+		}
+#endif
+	
+		uint32 BlockIndex = Range.Location / Range.Length;
+		FPlatformAtomics::AtomicStore(&Blocks[BlockIndex], 0);
+		FPlatformAtomics::InterlockedAdd(&UsedSize, -((int64)Range.Length));
+		
+		INC_MEMORY_STAT_BY(STAT_AGXBufferUnusedMemory, Range.Length);
+		INC_MEMORY_STAT_BY(STAT_AGXMagazineBufferUnusedMemory, Range.Length);
+		DEC_MEMORY_STAT_BY(STAT_AGXMagazineBufferMemory, Range.Length);
+	}
+}
+
+ns::String FAGXSubBufferMagazine::GetLabel() const
+{
+	if (ParentHeap)
+	{
+		return ParentHeap.GetLabel();
+	}
+	else
+	{
+		return ParentBuffer.GetLabel();
+	}
+}
+
+mtlpp::Device FAGXSubBufferMagazine::GetDevice() const
+{
+	if (ParentHeap)
+	{
+		return ParentHeap.GetDevice();
+	}
+	else
+	{
+		return ParentBuffer.GetDevice();
+	}
+}
+
+mtlpp::StorageMode FAGXSubBufferMagazine::GetStorageMode() const
+{
+	if (ParentHeap)
+	{
+		return ParentHeap.GetStorageMode();
+	}
+	else
+	{
+		return ParentBuffer.GetStorageMode();
+	}
+}
+
+mtlpp::CpuCacheMode FAGXSubBufferMagazine::GetCpuCacheMode() const
+{
+	if (ParentHeap)
+	{
+		return ParentHeap.GetCpuCacheMode();
+	}
+	else
+	{
+		return ParentBuffer.GetCpuCacheMode();
+	}
+}
+
+NSUInteger FAGXSubBufferMagazine::GetSize() const
+{
+	if (ParentHeap)
+	{
+		return ParentHeap.GetSize();
+	}
+	else
+	{
+		return ParentBuffer.GetLength();
+	}
+}
+
+NSUInteger FAGXSubBufferMagazine::GetUsedSize() const
+{
+	if (ParentHeap)
+	{
+		return ParentHeap.GetUsedSize();
+	}
+	else
+	{
+		return (NSUInteger)FPlatformAtomics::AtomicRead(&UsedSize);
+	}
+}
+
+NSUInteger FAGXSubBufferMagazine::GetFreeSize() const
+{
+	if (ParentHeap)
+	{
+		return ParentHeap.MaxAvailableSizeWithAlignment(MinAlign);
+	}
+	else
+	{
+		return GetSize() - GetUsedSize();
+	}
+}
+
+int64 FAGXSubBufferMagazine::NumCurrentAllocations() const
+{
+	return OutstandingAllocs;
+}
+
+bool FAGXSubBufferMagazine::CanAllocateSize(NSUInteger Size) const
+{
+	return GetFreeSize() >= Size;
+}
+
+void FAGXSubBufferMagazine::SetLabel(const ns::String& label)
+{
+	if (ParentHeap)
+	{
+		ParentHeap.SetLabel(label);
+	}
+	else
+	{
+		ParentBuffer.SetLabel(label);
+	}
+}
+
+FAGXBuffer FAGXSubBufferMagazine::NewBuffer()
+{
+	NSUInteger Size = BlockSize;
+	FAGXBuffer Result;
+
+	if (ParentHeap)
+	{
+		NSUInteger Storage = (NSUInteger(GetStorageMode()) << mtlpp::ResourceStorageModeShift);
+		NSUInteger Cache = (NSUInteger(GetCpuCacheMode()) << mtlpp::ResourceCpuCacheModeShift);
+		mtlpp::ResourceOptions Opt = mtlpp::ResourceOptions(Storage | Cache);
+		
+		Result = FAGXBuffer(ParentHeap.NewBuffer(Size, Opt), this);
+#if STATS || ENABLE_LOW_LEVEL_MEM_TRACKER
+		AGXLLM::LogAllocBuffer(GetAGXDeviceContext().GetDevice(), Result);
+#endif
+		DEC_MEMORY_STAT_BY(STAT_AGXBufferUnusedMemory, Result.GetLength());
+		DEC_MEMORY_STAT_BY(STAT_AGXMagazineBufferUnusedMemory, Result.GetLength());
+		INC_MEMORY_STAT_BY(STAT_AGXMagazineBufferMemory, Result.GetLength());
+	}
+	else
+	{
+		check(ParentBuffer && ParentBuffer.GetPtr());
+		
+		for (uint32 i = 0; i < Blocks.Num(); i++)
+		{
+			if (FPlatformAtomics::InterlockedCompareExchange(&Blocks[i], 1, 0) == 0)
+			{
+				ns::Range Range(i * BlockSize, BlockSize);
+				FPlatformAtomics::InterlockedAdd(&UsedSize, ((int64)Range.Length));
+				DEC_MEMORY_STAT_BY(STAT_AGXBufferUnusedMemory, Range.Length);
+				DEC_MEMORY_STAT_BY(STAT_AGXMagazineBufferUnusedMemory, Range.Length);
+				INC_MEMORY_STAT_BY(STAT_AGXMagazineBufferMemory, Range.Length);
+				Result = FAGXBuffer(MTLPP_VALIDATE(mtlpp::Buffer, ParentBuffer, AGXSafeGetRuntimeDebuggingLevel() >= EAGXDebugLevelValidation, NewBuffer(Range)), this);
+				break;
+			}
+		}
+	}
+
+	FPlatformAtomics::InterlockedIncrement(&OutstandingAllocs);
+	check(Result && Result.GetPtr());
+	return Result;
+}
+
+mtlpp::PurgeableState FAGXSubBufferMagazine::SetPurgeableState(mtlpp::PurgeableState state)
+{
+	if (ParentHeap)
+	{
+		return ParentHeap.SetPurgeableState(state);
+	}
+	else
+	{
+		return ParentBuffer.SetPurgeableState(state);
+	}
+}
+
+FAGXRingBufferRef::FAGXRingBufferRef(FAGXBuffer Buf)
+: Buffer(Buf)
+, LastRead(Buf.GetLength())
+{
+	Buffer.SetLabel(@"Ring Buffer");
+}
+
+FAGXRingBufferRef::~FAGXRingBufferRef()
+{
+	MTLPP_VALIDATE_ONLY(mtlpp::Buffer, Buffer, AGXSafeGetRuntimeDebuggingLevel() >= EAGXDebugLevelValidation, ReleaseAllRanges());
+	AGXSafeReleaseMetalBuffer(Buffer);
+}
+
+FAGXSubBufferRing::FAGXSubBufferRing(NSUInteger Size, NSUInteger Alignment, mtlpp::ResourceOptions InOptions)
+: LastFrameChange(0)
+, InitialSize(Align(Size, Alignment))
+, MinAlign(Alignment)
+, CommitHead(0)
+, SubmitHead(0)
+, WriteHead(0)
+, BufferSize(0)
+, Options(InOptions)
+, Storage((mtlpp::StorageMode)((Options & mtlpp::ResourceStorageModeMask) >> mtlpp::ResourceStorageModeShift))
+{
+	Options = (mtlpp::ResourceOptions)FAGXCommandQueue::GetCompatibleResourceOptions(Options);
+	check(Storage != mtlpp::StorageMode::Private /* Private memory requires command-buffers and encoders to properly marshal! */);
+	FMemory::Memzero(FrameSize);
+}
+
+FAGXSubBufferRing::~FAGXSubBufferRing()
+{
+}
+
+mtlpp::Device FAGXSubBufferRing::GetDevice() const
+{
+	return Buffer.IsValid() ? Buffer->Buffer.GetDevice() : nil;
+}
+mtlpp::StorageMode FAGXSubBufferRing::GetStorageMode() const
+{
+	return Buffer.IsValid() ? Buffer->Buffer.GetStorageMode() : Storage;
+}
+mtlpp::CpuCacheMode FAGXSubBufferRing::GetCpuCacheMode() const
+{
+	return Buffer.IsValid() ? Buffer->Buffer.GetCpuCacheMode() : ((mtlpp::CpuCacheMode)((Options & mtlpp::ResourceCpuCacheModeMask) >> mtlpp::ResourceCpuCacheModeShift));
+}
+NSUInteger FAGXSubBufferRing::GetSize() const
+{
+	return Buffer.IsValid() ? Buffer->Buffer.GetLength() : InitialSize;
+}
+
+FAGXBuffer FAGXSubBufferRing::NewBuffer(NSUInteger Size, uint32 Alignment)
+{
+	if (Alignment == 0)
+	{
+		Alignment = MinAlign;
+	}
+	else
+	{
+		Alignment = Align(Alignment, MinAlign);
+	}
+	
+	NSUInteger FullSize = Align(Size, Alignment);
+	
+	// Allocate on first use
+	if(!Buffer.IsValid())
+	{
+		Buffer = MakeShared<FAGXRingBufferRef, ESPMode::ThreadSafe>(GetAGXDeviceContext().GetResourceHeap().CreateBuffer(InitialSize, MinAlign, BUF_Dynamic, Options, true));
+		BufferSize = InitialSize;
+	}
+	
+	if(Buffer->LastRead <= WriteHead)
+	{
+		if (WriteHead + FullSize <= Buffer->Buffer.GetLength())
+		{
+			FAGXBuffer NewBuffer(MTLPP_VALIDATE(mtlpp::Buffer, Buffer->Buffer, AGXSafeGetRuntimeDebuggingLevel() >= EAGXDebugLevelValidation, NewBuffer(ns::Range(WriteHead, FullSize))), false);
+			
+			FMemory::Memset(((uint8*)NewBuffer.GetContents()), 0x0, FullSize);
+			
+			WriteHead += FullSize;
+			// NewBuffer.MarkSingleUse();
+			return NewBuffer;
+		}
+#if PLATFORM_MAC
+		else if (Storage == mtlpp::StorageMode::Managed)
+		{
+			Submit();
+			Buffer = MakeShared<FAGXRingBufferRef, ESPMode::ThreadSafe>(GetAGXDeviceContext().GetResourceHeap().CreateBuffer(BufferSize, MinAlign, BUF_Dynamic, Options, true));
+			WriteHead = 0;
+			CommitHead = 0;
+			SubmitHead = 0;
+		}
+#endif
+		else
+		{
+			WriteHead = 0;
+		}
+	}
+	
+	if(WriteHead + FullSize >= Buffer->LastRead || WriteHead + FullSize > BufferSize)
+	{
+		NSUInteger NewBufferSize = AlignArbitrary(BufferSize + Size, Align(BufferSize / 4, MinAlign));
+		
+		UE_LOG(LogAGX, Verbose, TEXT("Reallocating ring-buffer from %d to %d to avoid wrapping write at offset %d into outstanding buffer region %d at frame %lld]"), (uint32)BufferSize, (uint32)NewBufferSize, (uint32)WriteHead, (uint32)Buffer->LastRead, (uint64)GFrameCounter);
+		
+		Submit();
+		
+		Buffer = MakeShared<FAGXRingBufferRef, ESPMode::ThreadSafe>(GetAGXDeviceContext().GetResourceHeap().CreateBuffer(NewBufferSize, MinAlign, BUF_Dynamic, Options, true));
+		BufferSize = NewBufferSize;
+		WriteHead = 0;
+		CommitHead = 0;
+		SubmitHead = 0;
+	}
+	{
+		FAGXBuffer NewBuffer(MTLPP_VALIDATE(mtlpp::Buffer, Buffer->Buffer, AGXSafeGetRuntimeDebuggingLevel() >= EAGXDebugLevelValidation, NewBuffer(ns::Range(WriteHead, FullSize))), false);
+		
+		AllocatedRanges.Add(ns::Range(WriteHead, FullSize));
+		
+		FMemory::Memset(((uint8*)NewBuffer.GetContents()), 0x0, FullSize);
+		
+		WriteHead += FullSize;
+		// NewBuffer.MarkSingleUse();
+		return NewBuffer;
+	}
+}
+
+void FAGXSubBufferRing::Shrink()
+{
+	if(Buffer.IsValid())
+	{
+		NSUInteger FrameMax = 0;
+		for (uint32 i = 0; i < UE_ARRAY_COUNT(FrameSize); i++)
+		{
+			FrameMax = FMath::Max(FrameMax, FrameSize[i]);
+		}
+		
+		NSUInteger NecessarySize = FMath::Max(FrameMax, InitialSize);
+		NSUInteger ThreeQuarterSize = Align((BufferSize / 4) * 3, MinAlign);
+		
+		if ((GFrameNumberRenderThread - LastFrameChange) >= 120 && NecessarySize < ThreeQuarterSize && NecessarySize < BufferSize)
+		{
+			Submit();
+			
+			UE_LOG(LogAGX, Verbose, TEXT("Shrinking RingBuffer from %u to %u as max. usage is %u at frame %lld]"), (uint32)Buffer->Buffer.GetLength(), (uint32)ThreeQuarterSize, (uint32)FrameMax, GFrameNumberRenderThread);
+			
+			Buffer = MakeShared<FAGXRingBufferRef, ESPMode::ThreadSafe>(GetAGXDeviceContext().GetResourceHeap().CreateBuffer(ThreeQuarterSize, MinAlign, BUF_Dynamic, Options, true));
+			BufferSize = ThreeQuarterSize;
+			WriteHead = 0;
+			CommitHead = 0;
+			SubmitHead = 0;
+			LastFrameChange = GFrameNumberRenderThread;
+		}
+		
+		FrameSize[GFrameNumberRenderThread % UE_ARRAY_COUNT(FrameSize)] = 0;
+	}
+}
+
+void FAGXSubBufferRing::Submit()
+{
+	if (Buffer.IsValid() && WriteHead != SubmitHead)
+	{
+#if PLATFORM_MAC
+		if (Storage == mtlpp::StorageMode::Managed)
+		{
+			check(SubmitHead < WriteHead);
+			ns::Range ModifiedRange(SubmitHead, Align(WriteHead - SubmitHead, MinAlign));
+			Buffer->Buffer.DidModify(ModifiedRange);
+		}
+#endif
+
+		SubmitHead = WriteHead;
+	}
+}
+
+void FAGXSubBufferRing::Commit(mtlpp::CommandBuffer& CmdBuf)
+{
+	if (Buffer.IsValid() && WriteHead != CommitHead)
+	{
+#if PLATFORM_MAC
+		check(Storage != mtlpp::StorageMode::Managed || CommitHead < WriteHead);
+#endif
+		Submit();
+		
+		NSUInteger BytesWritten = 0;
+		if (CommitHead <= WriteHead)
+		{
+			BytesWritten = WriteHead - CommitHead;
+		}
+		else
+		{
+			NSUInteger TrailLen = GetSize() - CommitHead;
+			BytesWritten = TrailLen + WriteHead;
+		}
+		
+		FrameSize[GFrameNumberRenderThread % UE_ARRAY_COUNT(FrameSize)] += Align(BytesWritten, MinAlign);
+		
+		TSharedPtr<FAGXRingBufferRef, ESPMode::ThreadSafe> CmdBufferRingBuffer = Buffer;
+		FPlatformMisc::MemoryBarrier();
+		
+		NSUInteger CommitOffset = CommitHead;
+		NSUInteger WriteOffset = WriteHead;
+		
+		CommitHead = WriteHead;
+		
+		TArray<ns::Range> Ranges = MoveTemp(AllocatedRanges);
+		
+		mtlpp::CommandBufferHandler Handler = [CmdBufferRingBuffer, CommitOffset, WriteOffset, Ranges](mtlpp::CommandBuffer const& InBuffer)
+		{
+#if METAL_DEBUG_OPTIONS
+			if (GAGXBufferScribble && CommitOffset != WriteOffset)
+			{
+				if (CommitOffset < WriteOffset)
+				{
+					FMemory::Memset(((uint8*)CmdBufferRingBuffer->Buffer.GetContents()) + CommitOffset, 0xCD, WriteOffset - CommitOffset);
+				}
+				else
+				{
+					uint32 TrailLen = CmdBufferRingBuffer->Buffer.GetLength() - CommitOffset;
+					FMemory::Memset(((uint8*)CmdBufferRingBuffer->Buffer.GetContents()) + CommitOffset, 0xCD, TrailLen);
+					FMemory::Memset(((uint8*)CmdBufferRingBuffer->Buffer.GetContents()), 0xCD, WriteOffset);
+				}
+			}
+			
+#if MTLPP_CONFIG_VALIDATE
+			for (ns::Range const& Range : Ranges)
+			{
+				MTLPP_VALIDATE_ONLY(mtlpp::Buffer, CmdBufferRingBuffer->Buffer, AGXSafeGetRuntimeDebuggingLevel() >= EAGXDebugLevelValidation, ReleaseRange(Range));
+			}
+#endif
+#endif
+			CmdBufferRingBuffer->SetLastRead(WriteOffset);
+		};
+		CmdBuf.AddCompletedHandler(Handler);
+	}
+}
+
+uint32 FAGXBufferPoolPolicyData::GetPoolBucketIndex(CreationArguments Args)
+{
+	uint32 Size = Args.Size;
+	
+	unsigned long Lower = 0;
+	unsigned long Upper = NumPoolBucketSizes;
+	unsigned long Middle;
+	
+	do
+	{
+		Middle = ( Upper + Lower ) >> 1;
+		if( Size <= BucketSizes[Middle-1] )
+		{
+			Upper = Middle;
+		}
+		else
+		{
+			Lower = Middle;
+		}
+	}
+	while( Upper - Lower > 1 );
+	
+	check( Size <= BucketSizes[Lower] );
+	check( (Lower == 0 ) || ( Size > BucketSizes[Lower-1] ) );
+	
+	return Lower;
+}
+
+uint32 FAGXBufferPoolPolicyData::GetPoolBucketSize(uint32 Bucket)
+{
+	check(Bucket < NumPoolBuckets);
+	uint32 Index = Bucket;
+	checkf(Index < NumPoolBucketSizes, TEXT("%d %d"), Index, NumPoolBucketSizes);
+	return BucketSizes[Index];
+}
+
+FAGXBuffer FAGXBufferPoolPolicyData::CreateResource(CreationArguments Args)
+{
+	check(Args.Device);	
+	uint32 BufferSize = GetPoolBucketSize(GetPoolBucketIndex(Args));
+	METAL_GPUPROFILE(FAGXScopedCPUStats CPUStat(FString::Printf(TEXT("AllocBuffer: %llu, %llu"), BufferSize, mtlpp::ResourceOptions(BUFFER_CACHE_MODE | ((NSUInteger)Args.Storage << mtlpp::ResourceStorageModeShift)))));
+	FAGXBuffer NewBuf(MTLPP_VALIDATE(mtlpp::Device, Args.Device, AGXSafeGetRuntimeDebuggingLevel() >= EAGXDebugLevelValidation, NewBuffer(BufferSize, FAGXCommandQueue::GetCompatibleResourceOptions(mtlpp::ResourceOptions(BUFFER_CACHE_MODE | mtlpp::ResourceOptions::HazardTrackingModeUntracked | ((NSUInteger)Args.Storage << mtlpp::ResourceStorageModeShift))))), true);
+#if STATS || ENABLE_LOW_LEVEL_MEM_TRACKER
+	AGXLLM::LogAllocBuffer(Args.Device, NewBuf);
+#endif
+	INC_MEMORY_STAT_BY(STAT_AGXBufferUnusedMemory, NewBuf.GetLength());
+	INC_MEMORY_STAT_BY(STAT_AGXPooledBufferUnusedMemory, NewBuf.GetLength());
+	return NewBuf;
+}
+
+FAGXBufferPoolPolicyData::CreationArguments FAGXBufferPoolPolicyData::GetCreationArguments(FAGXBuffer const& Resource)
+{
+	return FAGXBufferPoolPolicyData::CreationArguments(Resource.GetDevice(), Resource.GetLength(), 0, Resource.GetStorageMode());
+}
+
+void FAGXBufferPoolPolicyData::FreeResource(FAGXBuffer& Resource)
+{
+	DEC_MEMORY_STAT_BY(STAT_AGXBufferUnusedMemory, Resource.GetLength());
+	DEC_MEMORY_STAT_BY(STAT_AGXPooledBufferUnusedMemory, Resource.GetLength());
+	Resource = nil;
+}
+
+FAGXTexturePool::FAGXTexturePool(FCriticalSection& InPoolMutex)
+: PoolMutex(InPoolMutex)
+{
+}
+
+FAGXTexturePool::~FAGXTexturePool()
+{
+}
+
+FAGXTexture FAGXTexturePool::CreateTexture(mtlpp::Device Device, mtlpp::TextureDescriptor Desc)
+{
+	FAGXTexturePool::Descriptor Descriptor;
+	Descriptor.textureType = (NSUInteger)Desc.GetTextureType();
+	Descriptor.pixelFormat = (NSUInteger)Desc.GetPixelFormat();
+	Descriptor.width = Desc.GetWidth();
+	Descriptor.height = Desc.GetHeight();
+	Descriptor.depth = Desc.GetDepth();
+	Descriptor.mipmapLevelCount = Desc.GetMipmapLevelCount();
+	Descriptor.sampleCount = Desc.GetSampleCount();
+	Descriptor.arrayLength = Desc.GetArrayLength();
+	Descriptor.resourceOptions = Desc.GetResourceOptions();
+	Descriptor.usage = Desc.GetUsage();
+	if (Descriptor.usage == mtlpp::TextureUsage::Unknown)
+	{
+		Descriptor.usage = (mtlpp::TextureUsage)(mtlpp::TextureUsage::ShaderRead | mtlpp::TextureUsage::ShaderWrite | mtlpp::TextureUsage::RenderTarget | mtlpp::TextureUsage::PixelFormatView);
+	}
+	Descriptor.freedFrame = 0;
+
+	FScopeLock Lock(&PoolMutex);
+	FAGXTexture* Tex = Pool.Find(Descriptor);
+	FAGXTexture Texture;
+	if (Tex)
+	{
+		Texture = *Tex;
+		Pool.Remove(Descriptor);
+		if (GAGXResourcePurgeInPool)
+		{
+            Texture.SetPurgeableState(mtlpp::PurgeableState::NonVolatile);
+        }
+	}
+	else
+	{
+		METAL_GPUPROFILE(FAGXScopedCPUStats CPUStat(FString::Printf(TEXT("AllocTexture: %s"), TEXT("")/**FString([Desc.GetPtr() description])*/)));
+		Texture = MTLPP_VALIDATE(mtlpp::Device, Device, AGXSafeGetRuntimeDebuggingLevel() >= EAGXDebugLevelValidation, NewTexture(Desc));
+#if STATS || ENABLE_LOW_LEVEL_MEM_TRACKER
+		AGXLLM::LogAllocTexture(Device, Desc, Texture);
+#endif
+	}
+	return Texture;
+}
+
+void FAGXTexturePool::ReleaseTexture(FAGXTexture& Texture)
+{
+	FAGXTexturePool::Descriptor Descriptor;
+	Descriptor.textureType = (NSUInteger)Texture.GetTextureType();
+	Descriptor.pixelFormat = (NSUInteger)Texture.GetPixelFormat();
+	Descriptor.width = Texture.GetWidth();
+	Descriptor.height = Texture.GetHeight();
+	Descriptor.depth = Texture.GetDepth();
+	Descriptor.mipmapLevelCount = Texture.GetMipmapLevelCount();
+	Descriptor.sampleCount = Texture.GetSampleCount();
+	Descriptor.arrayLength = Texture.GetArrayLength();
+	Descriptor.resourceOptions = ((NSUInteger)Texture.GetStorageMode() << mtlpp::ResourceStorageModeShift) | ((NSUInteger)Texture.GetCpuCacheMode() << mtlpp::ResourceCpuCacheModeShift);
+	Descriptor.usage = Texture.GetUsage();
+	Descriptor.freedFrame = GFrameNumberRenderThread;
+	
+    if (GAGXResourcePurgeInPool && Texture.SetPurgeableState(mtlpp::PurgeableState::KeepCurrent) == mtlpp::PurgeableState::NonVolatile)
+    {
+        Texture.SetPurgeableState(mtlpp::PurgeableState::Volatile);
+    }
+	
+	FScopeLock Lock(&PoolMutex);
+	Pool.Add(Descriptor, Texture);
+}
+	
+void FAGXTexturePool::Drain(bool const bForce)
+{
+	FScopeLock Lock(&PoolMutex);
+	if (bForce)
+	{
+		Pool.Empty();
+	}
+	else
+	{
+		for (auto It = Pool.CreateIterator(); It; ++It)
+		{
+			if ((GFrameNumberRenderThread - It->Key.freedFrame) >= CullAfterNumFrames)
+			{
+				It.RemoveCurrent();
+			}
+            else
+            {
+				if (GAGXResourcePurgeInPool && (GFrameNumberRenderThread - It->Key.freedFrame) >= PurgeAfterNumFrames)
+				{
+					It->Value.SetPurgeableState(mtlpp::PurgeableState::Empty);
+				}
+            }
+        }
+    }
+}
+
+FAGXResourceHeap::FAGXResourceHeap(void)
+: Queue(nullptr)
+, TexturePool(Mutex)
+, TargetPool(Mutex)
+{
+}
+
+FAGXResourceHeap::~FAGXResourceHeap()
+{
+	Compact(nullptr, true);
+}
+
+void FAGXResourceHeap::Init(FAGXCommandQueue& InQueue)
+{
+	Queue = &InQueue;
+}
+
+uint32 FAGXResourceHeap::GetMagazineIndex(uint32 Size)
+{
+	unsigned long Lower = 0;
+	unsigned long Upper = NumMagazineSizes;
+	unsigned long Middle;
+	
+	do
+	{
+		Middle = ( Upper + Lower ) >> 1;
+		if( Size <= MagazineSizes[Middle-1] )
+		{
+			Upper = Middle;
+		}
+		else
+		{
+			Lower = Middle;
+		}
+	}
+	while( Upper - Lower > 1 );
+	
+	check( Size <= MagazineSizes[Lower] );
+	check( (Lower == 0 ) || ( Size > MagazineSizes[Lower-1] ) );;
+	
+	return Lower;
+}
+
+uint32 FAGXResourceHeap::GetHeapIndex(uint32 Size)
+{
+	unsigned long Lower = 0;
+	unsigned long Upper = NumHeapSizes;
+	unsigned long Middle;
+	
+	do
+	{
+		Middle = ( Upper + Lower ) >> 1;
+		if( Size <= HeapSizes[Middle-1] )
+		{
+			Upper = Middle;
+		}
+		else
+		{
+			Lower = Middle;
+		}
+	}
+	while( Upper - Lower > 1 );
+	
+	check( Size <= HeapSizes[Lower] );
+	check( (Lower == 0 ) || ( Size > HeapSizes[Lower-1] ) );;
+	
+	return Lower;
+}
+
+FAGXResourceHeap::TextureHeapSize FAGXResourceHeap::TextureSizeToIndex(uint32 Size)
+{
+	unsigned long Lower = 0;
+	unsigned long Upper = NumTextureHeapSizes;
+	unsigned long Middle;
+	
+	do
+	{
+		Middle = ( Upper + Lower ) >> 1;
+		if( Size <= (HeapTextureHeapSizes[Middle-1] / MinTexturesPerHeap) )
+		{
+			Upper = Middle;
+		}
+		else
+		{
+			Lower = Middle;
+		}
+	}
+	while( Upper - Lower > 1 );
+	
+	check( Size <= (HeapTextureHeapSizes[Lower] / MinTexturesPerHeap) );
+	check( (Lower == 0 ) || ( Size > (HeapTextureHeapSizes[Lower-1] / MinTexturesPerHeap) ) );
+	
+	return (TextureHeapSize)Lower;
+}
+
+mtlpp::Heap FAGXResourceHeap::GetTextureHeap(mtlpp::TextureDescriptor Desc, mtlpp::SizeAndAlign Size)
+{
+	mtlpp::Heap Result;
+	static bool bTextureHeaps = FParse::Param(FCommandLine::Get(),TEXT("agxtextureheaps"));
+	if (FAGXCommandQueue::SupportsFeature(EAGXFeaturesHeaps) && bTextureHeaps && Size.Size <= HeapTextureHeapSizes[MaxTextureSize])
+	{
+		FAGXResourceHeap::TextureHeapSize HeapIndex = TextureSizeToIndex(Size.Size);
+
+		EAGXHeapTextureUsage UsageIndex = EAGXHeapTextureUsageNum;
+		mtlpp::StorageMode StorageMode = Desc.GetStorageMode();
+		mtlpp::CpuCacheMode CPUMode = Desc.GetCpuCacheMode();
+		if ((Desc.GetUsage() & mtlpp::TextureUsage::RenderTarget) && StorageMode == mtlpp::StorageMode::Private && CPUMode == mtlpp::CpuCacheMode::DefaultCache)
+		{
+			UsageIndex = PLATFORM_MAC ? EAGXHeapTextureUsageNum : EAGXHeapTextureUsageRenderTarget;
+		}
+		else if (StorageMode == mtlpp::StorageMode::Private && CPUMode == mtlpp::CpuCacheMode::WriteCombined)
+		{
+			UsageIndex = EAGXHeapTextureUsageResource;
+		}
+		
+		if (UsageIndex < EAGXHeapTextureUsageNum)
+		{
+			for (mtlpp::Heap& Heap : TextureHeaps[UsageIndex][HeapIndex])
+			{
+				if (Heap.MaxAvailableSizeWithAlignment(Size.Align) >= Size.Size)
+				{
+					Result = Heap;
+					break;
+				}
+			}
+			if (!Result)
+			{
+				mtlpp::HeapDescriptor HeapDesc;
+				HeapDesc.SetSize(HeapTextureHeapSizes[HeapIndex]);
+				HeapDesc.SetStorageMode(Desc.GetStorageMode());
+				HeapDesc.SetCpuCacheMode(Desc.GetCpuCacheMode());
+				Result = Queue->GetDevice().NewHeap(HeapDesc);
+#if STATS || ENABLE_LOW_LEVEL_MEM_TRACKER
+				AGXLLM::LogAllocHeap(Queue->GetDevice(), Result);
+#endif
+				TextureHeaps[UsageIndex][HeapIndex].Add(Result);
+			}
+			check(Result);
+		}
+	}
+	return Result;
+}
+
+FAGXBuffer FAGXResourceHeap::CreateBuffer(uint32 Size, uint32 Alignment, uint32 Flags, mtlpp::ResourceOptions Options, bool bForceUnique)
+{
+	LLM_SCOPE_METAL(ELLMTagAGX::Buffers);
+	LLM_PLATFORM_SCOPE_METAL(ELLMTagAGX::Buffers);
+	
+	static bool bSupportsHeaps = GetAGXDeviceContext().SupportsFeature(EAGXFeaturesHeaps);
+	static bool bSupportsBufferSubAllocation = FAGXCommandQueue::SupportsFeature(EAGXFeaturesBufferSubAllocation);
+	bForceUnique |= (!bSupportsBufferSubAllocation && !bSupportsHeaps);
+	
+	uint32 Usage = (Flags & BUF_Static) ? UsageStatic : UsageDynamic;
+	
+	FAGXBuffer Buffer;
+	uint32 BlockSize = Align(Size, Alignment);
+	mtlpp::StorageMode StorageMode = (mtlpp::StorageMode)(((NSUInteger)Options & mtlpp::ResourceStorageModeMask) >> mtlpp::ResourceStorageModeShift);
+	if (BlockSize <= 33554432)
+	{
+		switch (StorageMode)
+		{
+	#if PLATFORM_MAC
+			case mtlpp::StorageMode::Managed:
+			{
+				// TextureBuffers must be 1024 aligned.
+				check(Alignment == 256 || Alignment == 1024);
+				FScopeLock Lock(&Mutex);
+
+				// Disabled Managed sub-allocation as it seems inexplicably slow on the GPU				
+				 if (!bForceUnique && BlockSize <= HeapSizes[NumHeapSizes - 1])
+				 {
+				 	FAGXSubBufferLinear* Found = nullptr;
+				 	for (FAGXSubBufferLinear* Heap : ManagedSubHeaps)
+				 	{
+				 		if (Heap->CanAllocateSize(BlockSize))
+				 		{
+				 			Found = Heap;
+				 			break;
+				 		}
+				 	}
+				 	if (!Found)
+				 	{
+				 		Found = new FAGXSubBufferLinear(HeapAllocSizes[NumHeapSizes - 1], BufferOffsetAlignment, mtlpp::ResourceOptions((NSUInteger)Options & (mtlpp::ResourceStorageModeMask|mtlpp::ResourceHazardTrackingModeMask)), Mutex);
+				 		ManagedSubHeaps.Add(Found);
+				 	}
+				 	check(Found);
+					
+				 	return Found->NewBuffer(BlockSize);
+				 }
+				 else
+				 {
+                    Buffer = ManagedBuffers.CreatePooledResource(FAGXPooledBufferArgs(Queue->GetDevice(), BlockSize, Flags, StorageMode));
+                     if (GAGXResourcePurgeInPool)
+                     {
+                         Buffer.SetPurgeableState(mtlpp::PurgeableState::NonVolatile);
+                     }
+					DEC_MEMORY_STAT_BY(STAT_AGXBufferUnusedMemory, Buffer.GetLength());
+					DEC_MEMORY_STAT_BY(STAT_AGXPooledBufferUnusedMemory, Buffer.GetLength());
+					INC_MEMORY_STAT_BY(STAT_AGXPooledBufferMemory, Buffer.GetLength());
+				}
+				break;
+			}
+	#endif
+			case mtlpp::StorageMode::Private:
+			case mtlpp::StorageMode::Shared:
+			{
+				AllocTypes Storage = StorageMode != mtlpp::StorageMode::Private ? AllocShared : AllocPrivate;
+				check(Alignment == 16 || Alignment == 64 || Alignment == 256 || Alignment == 1024);
+
+				static bool bSupportsPrivateBufferSubAllocation = FAGXCommandQueue::SupportsFeature(EAGXFeaturesPrivateBufferSubAllocation);
+				if (!bForceUnique && BlockSize <= MagazineSizes[NumMagazineSizes - 1] && (Storage == AllocShared || bSupportsPrivateBufferSubAllocation))
+				{
+					FScopeLock Lock(&Mutex);
+					
+					uint32 i = GetMagazineIndex(BlockSize);
+					TArray<FAGXSubBufferMagazine*>& Heaps = SmallBuffers[Usage][Storage][i];
+					
+					FAGXSubBufferMagazine* Found = nullptr;
+					for (FAGXSubBufferMagazine* Heap : Heaps)
+					{
+						if (Heap->CanAllocateSize(BlockSize))
+						{
+							Found = Heap;
+							break;
+						}
+					}
+					
+					if (!Found)
+					{
+						Found = new FAGXSubBufferMagazine(MagazineAllocSizes[i], MagazineSizes[i], mtlpp::ResourceOptions((NSUInteger)Options & (mtlpp::ResourceStorageModeMask|mtlpp::ResourceHazardTrackingModeMask)));
+						SmallBuffers[Usage][Storage][i].Add(Found);
+					}
+					check(Found);
+					
+					Buffer = Found->NewBuffer();
+					check(Buffer && Buffer.GetPtr());
+				}
+				else if (!bForceUnique && BlockSize <= HeapSizes[NumHeapSizes - 1] && (Storage == AllocShared || bSupportsPrivateBufferSubAllocation))
+				{
+					FScopeLock Lock(&Mutex);
+					
+					uint32 i = GetHeapIndex(BlockSize);
+					TArray<FAGXSubBufferHeap*>& Heaps = BufferHeaps[Usage][Storage][i];
+					
+					FAGXSubBufferHeap* Found = nullptr;
+					for (FAGXSubBufferHeap* Heap : Heaps)
+					{
+						if (Heap->CanAllocateSize(BlockSize))
+						{
+							Found = Heap;
+							break;
+						}
+					}
+					
+					if (!Found)
+					{
+						uint32 MinAlign = PLATFORM_MAC ? 1024 : 64;
+						Found = new FAGXSubBufferHeap(HeapAllocSizes[i], MinAlign, mtlpp::ResourceOptions((NSUInteger)Options & (mtlpp::ResourceStorageModeMask|mtlpp::ResourceHazardTrackingModeMask)), Mutex);
+						BufferHeaps[Usage][Storage][i].Add(Found);
+					}
+					check(Found);
+					
+					Buffer = Found->NewBuffer(BlockSize);
+					check(Buffer && Buffer.GetPtr());
+				}
+				else
+				{
+					FScopeLock Lock(&Mutex);
+                    Buffer = Buffers[Storage].CreatePooledResource(FAGXPooledBufferArgs(Queue->GetDevice(), BlockSize, Flags, StorageMode));
+					if (GAGXResourcePurgeInPool)
+					{
+                   		Buffer.SetPurgeableState(mtlpp::PurgeableState::NonVolatile);
+					}
+					DEC_MEMORY_STAT_BY(STAT_AGXBufferUnusedMemory, Buffer.GetLength());
+					DEC_MEMORY_STAT_BY(STAT_AGXPooledBufferUnusedMemory, Buffer.GetLength());
+					INC_MEMORY_STAT_BY(STAT_AGXPooledBufferMemory, Buffer.GetLength());
+				}
+				break;
+			}
+			default:
+			{
+				check(false);
+				break;
+			}
+		}
+	}
+	else
+	{
+		METAL_GPUPROFILE(FAGXScopedCPUStats CPUStat(FString::Printf(TEXT("AllocBuffer: %llu, %llu"), BlockSize, Options)));
+		Buffer = FAGXBuffer(MTLPP_VALIDATE(mtlpp::Device, Queue->GetDevice(), AGXSafeGetRuntimeDebuggingLevel() >= EAGXDebugLevelValidation, NewBuffer(BlockSize, Options)), false);
+#if STATS || ENABLE_LOW_LEVEL_MEM_TRACKER
+		AGXLLM::LogAllocBuffer(Queue->GetDevice(), Buffer);
+#endif
+		INC_MEMORY_STAT_BY(STAT_AGXDeviceBufferMemory, Buffer.GetLength());
+	}
+	
+	if (GAGXBufferZeroFill && Buffer.GetStorageMode() != mtlpp::StorageMode::Private)
+	{
+		FMemory::Memset(((uint8*)Buffer.GetContents()), 0, Buffer.GetLength());
+	}
+	
+    METAL_DEBUG_OPTION(GetAGXDeviceContext().ValidateIsInactiveBuffer(Buffer));
+	METAL_FATAL_ASSERT(Buffer, TEXT("Failed to create buffer of size %u and resource options %u"), Size, (uint32)Options);
+	return Buffer;
+}
+
+void FAGXResourceHeap::ReleaseBuffer(FAGXBuffer& Buffer)
+{
+	mtlpp::StorageMode StorageMode = Buffer.GetStorageMode();
+	if (Buffer.IsPooled())
+	{
+		FScopeLock Lock(&Mutex);
+		
+		INC_MEMORY_STAT_BY(STAT_AGXBufferUnusedMemory, Buffer.GetLength());
+		INC_MEMORY_STAT_BY(STAT_AGXPooledBufferUnusedMemory, Buffer.GetLength());
+        DEC_MEMORY_STAT_BY(STAT_AGXPooledBufferMemory, Buffer.GetLength());
+		
+		if (GAGXResourcePurgeInPool)
+		{
+       		Buffer.SetPurgeableState(mtlpp::PurgeableState::Volatile);
+		}
+        
+		switch (StorageMode)
+		{
+	#if PLATFORM_MAC
+			case mtlpp::StorageMode::Managed:
+			{
+				ManagedBuffers.ReleasePooledResource(Buffer);
+				break;
+			}
+	#endif
+			case mtlpp::StorageMode::Private:
+			{
+				Buffers[AllocPrivate].ReleasePooledResource(Buffer);
+				break;
+			}
+			case mtlpp::StorageMode::Shared:
+			{
+				Buffers[AllocShared].ReleasePooledResource(Buffer);
+				break;
+			}
+			default:
+			{
+				check(false);
+				break;
+			}
+		}
+	}
+	else
+	{
+		DEC_MEMORY_STAT_BY(STAT_AGXDeviceBufferMemory, Buffer.GetLength());
+		Buffer.Release();
+	}
+}
+
+FAGXTexture FAGXResourceHeap::CreateTexture(mtlpp::TextureDescriptor Desc, FAGXSurface* Surface)
+{
+	LLM_SCOPE_METAL(ELLMTagAGX::Textures);
+	LLM_PLATFORM_SCOPE_METAL(ELLMTagAGX::Textures);
+	
+	mtlpp::SizeAndAlign Res = Queue->GetDevice().HeapTextureSizeAndAlign(Desc);
+	mtlpp::Heap Heap = GetTextureHeap(Desc, Res);
+	if (Heap)
+	{
+		METAL_GPUPROFILE(FAGXScopedCPUStats CPUStat(FString::Printf(TEXT("AllocTexture: %s"), TEXT("")/**FString([Desc.GetPtr() description])*/)));
+		FAGXTexture Texture = Heap.NewTexture(Desc);
+#if STATS || ENABLE_LOW_LEVEL_MEM_TRACKER
+		AGXLLM::LogAllocTexture(Queue->GetDevice(), Desc, Texture);
+#endif
+		return Texture;
+	}
+	else if (Desc.GetUsage() & mtlpp::TextureUsage::RenderTarget)
+	{
+		LLM_PLATFORM_SCOPE_METAL(ELLMTagAGX::RenderTargets);
+		return TargetPool.CreateTexture(Queue->GetDevice(), Desc);
+	}
+	else
+	{
+		return TexturePool.CreateTexture(Queue->GetDevice(), Desc);
+	}
+}
+
+void FAGXResourceHeap::ReleaseTexture(FAGXSurface* Surface, FAGXTexture& Texture)
+{
+	if (Texture && !Texture.GetBuffer() && !Texture.GetParentTexture() && !Texture.GetHeap())
+	{
+        if (Texture.GetUsage() & mtlpp::TextureUsage::RenderTarget)
+        {
+           	TargetPool.ReleaseTexture(Texture);
+        }
+        else
+        {
+            TexturePool.ReleaseTexture(Texture);
+        }
+	}
+}
+
+void FAGXResourceHeap::Compact(FAGXRenderPass* Pass, bool const bForce)
+{
+	FScopeLock Lock(&Mutex);
+    for (uint32 u = 0; u < NumUsageTypes; u++)
+    {
+        for (uint32 t = 0; t < NumAllocTypes; t++)
+        {
+            for (uint32 i = 0; i < NumMagazineSizes; i++)
+            {
+                for (auto It = SmallBuffers[u][t][i].CreateIterator(); It; ++It)
+                {
+                    FAGXSubBufferMagazine* Data = *It;
+                    if (Data->NumCurrentAllocations() == 0 || bForce)
+                    {
+                        It.RemoveCurrent();
+                        delete Data;
+                    }
+                }
+            }
+            
+            uint32 BytesCompacted = 0;
+            uint32 const BytesToCompact = GAGXHeapBufferBytesToCompact;
+            
+            for (uint32 i = 0; i < NumHeapSizes; i++)
+            {
+                for (auto It = BufferHeaps[u][t][i].CreateIterator(); It; ++It)
+                {
+                    FAGXSubBufferHeap* Data = *It;
+                    if (Data->NumCurrentAllocations() == 0 || bForce)
+                    {
+                        It.RemoveCurrent();
+                        delete Data;
+                    }
+                }
+            }
+        }
+    }
+
+	Buffers[AllocShared].DrainPool(bForce);
+	Buffers[AllocPrivate].DrainPool(bForce);
+#if PLATFORM_MAC
+	ManagedBuffers.DrainPool(bForce);
+	for (auto It = ManagedSubHeaps.CreateIterator(); It; ++It)
+	{
+		FAGXSubBufferLinear* Data = *It;
+		if (Data->GetUsedSize() == 0 || bForce)
+		{
+			It.RemoveCurrent();
+			delete Data;
+		}
+	}
+#endif
+	TexturePool.Drain(bForce);
+	TargetPool.Drain(bForce);
+}
+	
+uint32 FAGXBufferPoolPolicyData::BucketSizes[FAGXBufferPoolPolicyData::NumPoolBucketSizes] = {
+	256,
+	512,
+	1024,
+	2048,
+	4096,
+	8192,
+	16384,
+	32768,
+	65536,
+	131072,
+	262144,
+	524288,
+	1048576,
+	2097152,
+	4194304,
+	8388608,
+	12582912,
+	16777216,
+	25165824,
+	33554432,
+};
+
+FAGXBufferPool::~FAGXBufferPool()
+{
+}
+
+uint32 FAGXResourceHeap::MagazineSizes[FAGXResourceHeap::NumMagazineSizes] = {
+	16,
+	32,
+	64,
+	128,
+	256,
+	512,
+	1024,
+	2048,
+	4096,
+	8192,
+};
+
+uint32 FAGXResourceHeap::HeapSizes[FAGXResourceHeap::NumHeapSizes] = {
+	1048576,
+	2097152,
+};
+
+uint32 FAGXResourceHeap::MagazineAllocSizes[FAGXResourceHeap::NumMagazineSizes] = {
+	4096,
+	4096,
+	4096,
+	8192,
+	8192,
+	8192,
+	16384,
+	16384,
+	16384,
+	32768,
+};
+
+uint32 FAGXResourceHeap::HeapAllocSizes[FAGXResourceHeap::NumHeapSizes] = {
+	2097152,
+	4194304,
+};
+
+uint32 FAGXResourceHeap::HeapTextureHeapSizes[FAGXResourceHeap::NumTextureHeapSizes] = {
+	4194304,
+	8388608,
+	16777216,
+	33554432,
+	67108864,
+	134217728,
+	268435456
+};
diff --git a/Engine/Source/Runtime/Apple/AGXRHI/Private/AGXBuffer.h b/Engine/Source/Runtime/Apple/AGXRHI/Private/AGXBuffer.h
new file mode 100644
--- /dev/null
+++ b/Engine/Source/Runtime/Apple/AGXRHI/Private/AGXBuffer.h
@@ -0,0 +1,469 @@
+// Copyright Epic Games, Inc. All Rights Reserved.
+
+
+#pragma once
+
+#include "AGXRHIPrivate.h"
+#include "device.hpp"
+#include "buffer.hpp"
+#include "Containers/LockFreeList.h"
+#include "ResourcePool.h"
+
+struct FAGXPooledBufferArgs
+{
+    FAGXPooledBufferArgs() : Device(nil), Size(0), Flags(0), Storage(mtlpp::StorageMode::Shared) {}
+	
+    FAGXPooledBufferArgs(mtlpp::Device InDevice, uint32 InSize, uint32 InFlags, mtlpp::StorageMode InStorage, mtlpp::CpuCacheMode InCpuCacheMode = mtlpp::CpuCacheMode::DefaultCache)
+	: Device(InDevice)
+	, Size(InSize)
+    , Flags(InFlags)
+	, Storage(InStorage)
+	, CpuCacheMode(InCpuCacheMode)
+	{
+	}
+	
+	mtlpp::Device Device;
+	uint32 Size;
+    uint32 Flags;
+	mtlpp::StorageMode Storage;
+	mtlpp::CpuCacheMode CpuCacheMode;
+};
+
+class FAGXSubBufferHeap
+{
+    friend class FAGXResourceHeap;
+    
+public:
+	FAGXSubBufferHeap(NSUInteger Size, NSUInteger Alignment, mtlpp::ResourceOptions, FCriticalSection& PoolMutex);
+	~FAGXSubBufferHeap();
+	
+	ns::String   GetLabel() const;
+    mtlpp::Device       GetDevice() const;
+    mtlpp::StorageMode  GetStorageMode() const;
+    mtlpp::CpuCacheMode GetCpuCacheMode() const;
+    NSUInteger     GetSize() const;
+    NSUInteger     GetUsedSize() const;
+	NSUInteger	 MaxAvailableSize() const;
+	int64     NumCurrentAllocations() const;
+    bool     CanAllocateSize(NSUInteger Size) const;
+
+    void SetLabel(const ns::String& label);
+	
+    FAGXBuffer NewBuffer(NSUInteger length);
+    mtlpp::PurgeableState SetPurgeableState(mtlpp::PurgeableState state);
+	void FreeRange(ns::Range const& Range);
+
+    void SetOwner(ns::Range const& Range, FAGXRHIBuffer* Owner, bool bIsSwap);
+
+private:
+    struct Allocation
+    {
+        ns::Range Range;
+        mtlpp::Buffer::Type Resource;
+        FAGXRHIBuffer* Owner;
+    };
+    
+	FCriticalSection& PoolMutex;
+	int64 volatile OutstandingAllocs;
+	NSUInteger MinAlign;
+	NSUInteger UsedSize;
+	mtlpp::Buffer ParentBuffer;
+	mutable mtlpp::Heap ParentHeap;
+	TArray<ns::Range> FreeRanges;
+    TArray<Allocation> AllocRanges;
+};
+
+class FAGXSubBufferLinear
+{
+public:
+	FAGXSubBufferLinear(NSUInteger Size, NSUInteger Alignment, mtlpp::ResourceOptions, FCriticalSection& PoolMutex);
+	~FAGXSubBufferLinear();
+	
+	ns::String   GetLabel() const;
+	mtlpp::Device       GetDevice() const;
+	mtlpp::StorageMode  GetStorageMode() const;
+	mtlpp::CpuCacheMode GetCpuCacheMode() const;
+	NSUInteger     GetSize() const;
+	NSUInteger     GetUsedSize() const;
+	bool	 CanAllocateSize(NSUInteger Size) const;
+
+	void SetLabel(const ns::String& label);
+	
+	FAGXBuffer NewBuffer(NSUInteger length);
+	mtlpp::PurgeableState SetPurgeableState(mtlpp::PurgeableState state);
+	void FreeRange(ns::Range const& Range);
+	
+private:
+	FCriticalSection& PoolMutex;
+	NSUInteger MinAlign;
+	NSUInteger WriteHead;
+	NSUInteger UsedSize;
+	NSUInteger FreedSize;
+	mtlpp::Buffer ParentBuffer;
+};
+
+class FAGXSubBufferMagazine
+{
+public:
+	FAGXSubBufferMagazine(NSUInteger Size, NSUInteger ChunkSize, mtlpp::ResourceOptions);
+	~FAGXSubBufferMagazine();
+	
+	ns::String   GetLabel() const;
+    mtlpp::Device       GetDevice() const;
+    mtlpp::StorageMode  GetStorageMode() const;
+    mtlpp::CpuCacheMode GetCpuCacheMode() const;
+    NSUInteger     GetSize() const;
+    NSUInteger     GetUsedSize() const;
+	NSUInteger	 GetFreeSize() const;
+	int64     NumCurrentAllocations() const;
+    bool     CanAllocateSize(NSUInteger Size) const;
+
+    void SetLabel(const ns::String& label);
+	void FreeRange(ns::Range const& Range);
+
+    FAGXBuffer NewBuffer();
+    mtlpp::PurgeableState SetPurgeableState(mtlpp::PurgeableState state);
+
+private:
+	NSUInteger MinAlign;
+    NSUInteger BlockSize;
+	int64 volatile OutstandingAllocs;
+	int64 volatile UsedSize;
+	mtlpp::Buffer ParentBuffer;
+	mutable mtlpp::Heap ParentHeap;
+	TArray<int8> Blocks;
+};
+
+struct FAGXRingBufferRef
+{
+	FAGXRingBufferRef(FAGXBuffer Buf);
+	~FAGXRingBufferRef();
+	
+	void SetLastRead(uint64 Read) { FPlatformAtomics::InterlockedExchange((int64*)&LastRead, Read); }
+	
+	FAGXBuffer Buffer;
+	uint64 LastRead;
+};
+
+class FAGXResourceHeap;
+
+class FAGXSubBufferRing
+{
+public:
+	FAGXSubBufferRing(NSUInteger Size, NSUInteger Alignment, mtlpp::ResourceOptions Options);
+	~FAGXSubBufferRing();
+	
+	mtlpp::Device       GetDevice() const;
+	mtlpp::StorageMode  GetStorageMode() const;
+	mtlpp::CpuCacheMode GetCpuCacheMode() const;
+	NSUInteger     GetSize() const;
+	
+	FAGXBuffer NewBuffer(NSUInteger Size, uint32 Alignment);
+	
+	/** Tries to shrink the ring-buffer back toward its initial size, but not smaller. */
+	void Shrink();
+	
+	/** Submits all outstanding writes to the GPU, coalescing the updates into a single contiguous range. */
+	void Submit();
+	
+	/** Commits a completion handler to the cmd-buffer to release the processed range */
+	void Commit(mtlpp::CommandBuffer& CmdBuffer);
+	
+private:
+	NSUInteger FrameSize[10];
+	NSUInteger LastFrameChange;
+	NSUInteger InitialSize;
+	NSUInteger MinAlign;
+	NSUInteger CommitHead;
+	NSUInteger SubmitHead;
+	NSUInteger WriteHead;
+	NSUInteger BufferSize;
+	mtlpp::ResourceOptions Options;
+	mtlpp::StorageMode Storage;
+	TSharedPtr<FAGXRingBufferRef, ESPMode::ThreadSafe> Buffer;
+	TArray<ns::Range> AllocatedRanges;
+};
+
+class FAGXBufferPoolPolicyData
+{
+	enum BucketSizes
+	{
+		// These sizes are required for ring-buffers and esp. Managed Memory which is a Mac-only feature
+		BucketSize256,
+		BucketSize512,
+		BucketSize1k,
+		BucketSize2k,
+		BucketSize4k,
+		BucketSize8k,
+		BucketSize16k,
+		BucketSize32k,
+		BucketSize64k,
+		BucketSize128k,
+		BucketSize256k,
+		BucketSize512k,
+		BucketSize1Mb,
+		BucketSize2Mb,
+		BucketSize4Mb,
+		// These sizes are the ones typically used by buffer allocations
+		BucketSize8Mb,
+		BucketSize12Mb,
+		BucketSize16Mb,
+		BucketSize24Mb,
+		BucketSize32Mb,
+		NumBucketSizes
+	};
+public:
+	/** Buffers are created with a simple byte size */
+	typedef FAGXPooledBufferArgs CreationArguments;
+	enum
+	{
+		NumSafeFrames = 1, /** Number of frames to leave buffers before reclaiming/reusing */
+		NumPoolBucketSizes = NumBucketSizes, /** Number of pool bucket sizes */
+		NumPoolBuckets = NumPoolBucketSizes, /** Number of pool bucket sizes - all entries must use consistent ResourceOptions */
+		NumToDrainPerFrame = 65536, /** Max. number of resources to cull in a single frame */
+		CullAfterFramesNum = 30 /** Resources are culled if unused for more frames than this */
+	};
+	
+	/** Get the pool bucket index from the size
+	 * @param Size the number of bytes for the resource
+	 * @returns The bucket index.
+	 */
+	uint32 GetPoolBucketIndex(CreationArguments Args);
+	
+	/** Get the pool bucket size from the index
+	 * @param Bucket the bucket index
+	 * @returns The bucket size.
+	 */
+	uint32 GetPoolBucketSize(uint32 Bucket);
+	
+	/** Creates the resource
+	 * @param Args The buffer size in bytes.
+	 * @returns A suitably sized buffer or NULL on failure.
+	 */
+	FAGXBuffer CreateResource(CreationArguments Args);
+	
+	/** Gets the arguments used to create resource
+	 * @param Resource The buffer to get data for.
+	 * @returns The arguments used to create the buffer.
+	 */
+	CreationArguments GetCreationArguments(FAGXBuffer const& Resource);
+	
+	/** Frees the resource
+	 * @param Resource The buffer to prepare for release from the pool permanently.
+	 */
+	void FreeResource(FAGXBuffer& Resource);
+	
+private:
+	/** The bucket sizes */
+	static uint32 BucketSizes[NumPoolBucketSizes];
+};
+
+/** A pool for metal buffers with consistent usage, bucketed for efficiency. */
+class FAGXBufferPool : public TResourcePool<FAGXBuffer, FAGXBufferPoolPolicyData, FAGXBufferPoolPolicyData::CreationArguments>
+{
+public:
+	/** Destructor */
+	virtual ~FAGXBufferPool();
+};
+
+class FAGXTexturePool
+{
+	enum
+	{
+        PurgeAfterNumFrames = 2, /* Textures must be reused fairly rapidly but after this number of frames we reclaim the memory, even though the object persists */
+		CullAfterNumFrames = 3, /* Textures must be reused fairly rapidly or we bin them as they are much larger than buffers */
+	};
+public:
+	struct Descriptor
+	{
+		friend uint32 GetTypeHash(Descriptor const& Other)
+		{
+			uint32 Hash = GetTypeHash((uint64)Other.textureType);
+			Hash = HashCombine(Hash, GetTypeHash((uint64)Other.pixelFormat));
+			Hash = HashCombine(Hash, GetTypeHash((uint64)Other.usage));
+			Hash = HashCombine(Hash, GetTypeHash((uint64)Other.width));
+			Hash = HashCombine(Hash, GetTypeHash((uint64)Other.height));
+			Hash = HashCombine(Hash, GetTypeHash((uint64)Other.depth));
+			Hash = HashCombine(Hash, GetTypeHash((uint64)Other.mipmapLevelCount));
+			Hash = HashCombine(Hash, GetTypeHash((uint64)Other.sampleCount));
+			Hash = HashCombine(Hash, GetTypeHash((uint64)Other.arrayLength));
+			Hash = HashCombine(Hash, GetTypeHash((uint64)Other.resourceOptions));
+			return Hash;
+		}
+		
+		bool operator<(Descriptor const& Other) const
+		{
+			if (this != &Other)
+			{
+				return (textureType < Other.textureType ||
+						pixelFormat < Other.pixelFormat ||
+						width < Other.width ||
+						height < Other.height ||
+						depth < Other.depth ||
+						mipmapLevelCount < Other.mipmapLevelCount ||
+						sampleCount < Other.sampleCount ||
+						arrayLength < Other.arrayLength ||
+						resourceOptions < Other.resourceOptions ||
+						usage < Other.usage);
+			}
+			return false;
+		}
+		
+		bool operator==(Descriptor const& Other) const
+		{
+			if (this != &Other)
+			{
+				return (textureType == Other.textureType &&
+				pixelFormat == Other.pixelFormat &&
+				width == Other.width &&
+				height == Other.height &&
+				depth == Other.depth &&
+				mipmapLevelCount == Other.mipmapLevelCount &&
+				sampleCount == Other.sampleCount &&
+				arrayLength == Other.arrayLength &&
+				resourceOptions == Other.resourceOptions &&
+				usage == Other.usage);
+			}
+			return true;
+		}
+		
+		NSUInteger textureType;
+		NSUInteger pixelFormat;
+		NSUInteger width;
+		NSUInteger height;
+		NSUInteger depth;
+		NSUInteger mipmapLevelCount;
+		NSUInteger sampleCount;
+		NSUInteger arrayLength;
+		NSUInteger resourceOptions;
+		NSUInteger usage;
+		NSUInteger freedFrame;
+	};
+	
+	FAGXTexturePool(FCriticalSection& PoolMutex);
+	~FAGXTexturePool();
+	
+	FAGXTexture CreateTexture(mtlpp::Device Device, mtlpp::TextureDescriptor Desc);
+	void ReleaseTexture(FAGXTexture& Texture);
+	
+	void Drain(bool const bForce);
+
+private:
+	FCriticalSection& PoolMutex;
+	TMap<Descriptor, FAGXTexture> Pool;
+};
+
+class FAGXResourceHeap
+{
+	enum MagazineSize
+	{
+		Size16,
+		Size32,
+		Size64,
+		Size128,
+		Size256,
+		Size512,
+		Size1024,
+		Size2048,
+		Size4096,
+        Size8192,
+		NumMagazineSizes
+	};
+	
+	enum HeapSize
+	{
+		Size1Mb,
+		Size2Mb,
+		NumHeapSizes
+	};
+
+	enum TextureHeapSize
+	{
+		Size4Mb,
+		Size8Mb,
+		Size16Mb,
+		Size32Mb,
+		Size64Mb,
+		Size128Mb,
+		Size256Mb,
+		NumTextureHeapSizes,
+		MinTexturesPerHeap = 4,
+		MaxTextureSize = Size64Mb,
+	};
+	
+	enum AllocTypes
+	{
+		AllocShared,
+		AllocPrivate,
+		NumAllocTypes = 2
+	};
+	
+	enum EAGXHeapTextureUsage
+	{
+		/** Regular texture resource */
+		EAGXHeapTextureUsageResource = 0,
+		/** Render target or UAV that can be aliased */
+		EAGXHeapTextureUsageRenderTarget = 1,
+		/** Number of texture usage types */
+		EAGXHeapTextureUsageNum = 2
+	};
+    
+    enum UsageTypes
+    {
+        UsageStatic,
+        UsageDynamic,
+        NumUsageTypes = 2
+    };
+    
+public:
+	FAGXResourceHeap(void);
+	~FAGXResourceHeap();
+	
+	void Init(FAGXCommandQueue& Queue);
+	
+    FAGXBuffer CreateBuffer(uint32 Size, uint32 Alignment, uint32 Flags, mtlpp::ResourceOptions Options, bool bForceUnique = false);
+	FAGXTexture CreateTexture(mtlpp::TextureDescriptor Desc, FAGXSurface* Surface);
+	
+	void ReleaseBuffer(FAGXBuffer& Buffer);
+	void ReleaseTexture(FAGXSurface* Surface, FAGXTexture& Texture);
+	
+	void Compact(class FAGXRenderPass* Pass, bool const bForce);
+	
+private:
+	uint32 GetMagazineIndex(uint32 Size);
+	uint32 GetHeapIndex(uint32 Size);
+	TextureHeapSize TextureSizeToIndex(uint32 Size);
+	
+	mtlpp::Heap GetTextureHeap(mtlpp::TextureDescriptor Desc, mtlpp::SizeAndAlign Size);
+	
+private:
+	static uint32 MagazineSizes[NumMagazineSizes];
+	static uint32 HeapSizes[NumHeapSizes];
+	static uint32 MagazineAllocSizes[NumMagazineSizes];
+	static uint32 HeapAllocSizes[NumHeapSizes];
+	static uint32 HeapTextureHeapSizes[NumTextureHeapSizes];
+
+	FCriticalSection Mutex;
+	FAGXCommandQueue* Queue;
+	
+	/** Small allocations (<= 4KB) are made from magazine allocators that use sub-ranges of a buffer */
+    TArray<FAGXSubBufferMagazine*> SmallBuffers[NumUsageTypes][NumAllocTypes][NumMagazineSizes];
+
+	/** Typical allocations (4KB - 4MB) are made from heap allocators that use sub-ranges of a buffer */
+	/** There are two alignment categories for heaps - 16b for Vertes/Index data and 256b for constant data (macOS-only) */
+    TArray<FAGXSubBufferHeap*> BufferHeaps[NumUsageTypes][NumAllocTypes][NumHeapSizes];
+	
+	/** Larger buffers (up-to 32MB) that are subject to bucketing & pooling rather than sub-allocation */
+	FAGXBufferPool Buffers[NumAllocTypes];
+#if PLATFORM_MAC // All managed buffers are bucketed & pooled rather than sub-allocated to avoid memory consistency complexities
+	FAGXBufferPool ManagedBuffers;
+	TArray<FAGXSubBufferLinear*> ManagedSubHeaps;
+#endif
+	/** Anything else is just allocated directly from the device! */
+	
+	/** We can reuse texture allocations as well, to minimize their performance impact */
+	FAGXTexturePool TexturePool;
+	FAGXTexturePool TargetPool;
+	
+	TArray<mtlpp::Heap> TextureHeaps[EAGXHeapTextureUsageNum][NumTextureHeapSizes];
+};
diff --git a/Engine/Source/Runtime/Apple/AGXRHI/Private/AGXCaptureManager.cpp b/Engine/Source/Runtime/Apple/AGXRHI/Private/AGXCaptureManager.cpp
new file mode 100644
--- /dev/null
+++ b/Engine/Source/Runtime/Apple/AGXRHI/Private/AGXCaptureManager.cpp
@@ -0,0 +1,94 @@
+// Copyright Epic Games, Inc. All Rights Reserved.
+
+#include "AGXRHIPrivate.h"
+#include "AGXCaptureManager.h"
+#include "AGXCommandQueue.h"
+#include "capture_manager.hpp"
+
+bool GAGXSupportsCaptureManager = false;
+
+FAGXCaptureManager::FAGXCaptureManager(mtlpp::Device InDevice, FAGXCommandQueue& InQueue)
+: Device(InDevice)
+, Queue(InQueue)
+{
+	MTLPP_IF_AVAILABLE(10.13, 11.0, 11.0)
+	{
+		GAGXSupportsCaptureManager = true;
+		
+		mtlpp::CaptureManager Manager = mtlpp::CaptureManager::SharedCaptureManager();
+		Manager.SetDefaultCaptureScope(Manager.NewCaptureScopeWithDevice(Device));
+		Manager.GetDefaultCaptureScope().SetLabel(@"1 Frame");
+		
+		FAGXCaptureScope DefaultScope;
+		DefaultScope.MTLScope = Manager.GetDefaultCaptureScope();
+		DefaultScope.Type = EAGXCaptureTypePresent;
+		DefaultScope.StepCount = 1;
+		DefaultScope.LastTrigger = 0;
+		ActiveScopes.Add(DefaultScope);
+		DefaultScope.MTLScope.BeginScope();
+		
+		uint32 PresentStepCounts[] = {2, 5, 10, 15, 30, 60, 90, 120};
+		for (uint32 i = 0; i < (sizeof(PresentStepCounts) / sizeof(uint32)); i++)
+		{
+			FAGXCaptureScope Scope;
+			Scope.MTLScope = Manager.NewCaptureScopeWithDevice(Device);
+			Scope.MTLScope.SetLabel(FString::Printf(TEXT("%u Frames"), PresentStepCounts[i]).GetNSString());
+			Scope.Type = EAGXCaptureTypePresent;
+			Scope.StepCount = PresentStepCounts[i];
+			Scope.LastTrigger = 0;
+			ActiveScopes.Add(Scope);
+			Scope.MTLScope.BeginScope();
+		}
+	}
+}
+
+FAGXCaptureManager::~FAGXCaptureManager()
+{
+}
+
+void FAGXCaptureManager::PresentFrame(uint32 FrameNumber)
+{
+	if (GAGXSupportsCaptureManager)
+	{
+		for (FAGXCaptureScope& Scope : ActiveScopes)
+		{
+			uint32 Diff = 0;
+			if (FrameNumber > Scope.LastTrigger)
+			{
+				Diff = FrameNumber - Scope.LastTrigger;
+			}
+			else
+			{
+				Diff = (UINT32_MAX - Scope.LastTrigger) + FrameNumber;
+			}
+			
+			if (Diff >= Scope.StepCount)
+			{
+				Scope.MTLScope.EndScope();
+				Scope.MTLScope.BeginScope();
+				Scope.LastTrigger = FrameNumber;
+			}
+		}
+	}
+	else
+	{
+		Queue.InsertDebugCaptureBoundary();
+	}
+}
+
+void FAGXCaptureManager::BeginCapture(void)
+{
+	if (GAGXSupportsCaptureManager)
+	{
+		mtlpp::CaptureManager::SharedCaptureManager().StartCaptureWithDevice(Device);
+	}
+}
+
+void FAGXCaptureManager::EndCapture(void)
+{
+	if (GAGXSupportsCaptureManager)
+	{
+		mtlpp::CaptureManager::SharedCaptureManager().StopCapture();
+	}
+}
+
diff --git a/Engine/Source/Runtime/Apple/AGXRHI/Private/AGXCaptureManager.h b/Engine/Source/Runtime/Apple/AGXRHI/Private/AGXCaptureManager.h
new file mode 100644
--- /dev/null
+++ b/Engine/Source/Runtime/Apple/AGXRHI/Private/AGXCaptureManager.h
@@ -0,0 +1,48 @@
+// Copyright Epic Games, Inc. All Rights Reserved.
+
+#pragma once
+
+#include <Metal/Metal.h>
+#include "capture_scope.hpp"
+#include "device.hpp"
+
+class FAGXCommandQueue;
+
+class FAGXCaptureManager
+{
+public:
+	FAGXCaptureManager(mtlpp::Device Device, FAGXCommandQueue& Queue);
+	~FAGXCaptureManager();
+	
+	// Called by the AGXRHI code to trigger the provided capture scopes visible in Xcode.
+	void PresentFrame(uint32 FrameNumber);
+	
+	// Programmatic captures without an Xcode capture scope.
+	// Use them to instrument the code manually to debug issues.
+	void BeginCapture(void);
+	void EndCapture(void);
+	
+private:
+	mtlpp::Device Device;
+	FAGXCommandQueue& Queue;
+	bool bSupportsCaptureManager;
+	
+private:
+	enum EAGXCaptureType
+	{
+		EAGXCaptureTypeUnknown,
+		EAGXCaptureTypeFrame, // (BeginFrame-EndFrame) * StepCount
+		EAGXCaptureTypePresent, // (Present-Present) * StepCount
+		EAGXCaptureTypeViewport, // (Present-Present) * Viewports * StepCount
+	};
+
+	struct FAGXCaptureScope
+	{
+		EAGXCaptureType Type;
+		uint32 StepCount;
+		uint32 LastTrigger;
+		mtlpp::CaptureScope MTLScope;
+	};
+	
+	TArray<FAGXCaptureScope> ActiveScopes;
+};
diff --git a/Engine/Source/Runtime/Apple/AGXRHI/Private/AGXCommandBuffer.cpp b/Engine/Source/Runtime/Apple/AGXRHI/Private/AGXCommandBuffer.cpp
new file mode 100644
--- /dev/null
+++ b/Engine/Source/Runtime/Apple/AGXRHI/Private/AGXCommandBuffer.cpp
@@ -0,0 +1,274 @@
+// Copyright Epic Games, Inc. All Rights Reserved.
+
+/*=============================================================================
+	AGXCommandBuffer.cpp: AGX RHI command buffer wrapper.
+=============================================================================*/
+
+#include "AGXRHIPrivate.h"
+
+#include "AGXCommandBuffer.h"
+#include "AGXRenderCommandEncoder.h"
+#include "AGXBlitCommandEncoder.h"
+#include "AGXComputeCommandEncoder.h"
+#include <objc/runtime.h>
+
+NSString* GAGXDebugCommandTypeNames[EAGXDebugCommandTypeInvalid] = {
+	@"RenderEncoder",
+	@"ComputeEncoder",
+	@"BlitEncoder",
+    @"EndEncoder",
+    @"Pipeline",
+	@"Draw",
+	@"Dispatch",
+	@"Blit",
+	@"Signpost",
+	@"PushGroup",
+	@"PopGroup"
+};
+
+extern int32 GAGXRuntimeDebugLevel;
+
+uint32 AGXSafeGetRuntimeDebuggingLevel()
+{
+	return GIsRHIInitialized ? GetAGXDeviceContext().GetCommandQueue().GetRuntimeDebuggingLevel() : GAGXRuntimeDebugLevel;
+}
+
+#if MTLPP_CONFIG_VALIDATE && METAL_DEBUG_OPTIONS
+
+@implementation FAGXDebugCommandBuffer
+
+APPLE_PLATFORM_OBJECT_ALLOC_OVERRIDES(FAGXDebugCommandBuffer)
+
+-(id)initWithCommandBuffer:(id<MTLCommandBuffer>)Buffer
+{
+	id Self = [super init];
+	if (Self)
+	{
+        DebugLevel = (EAGXDebugLevel)GAGXRuntimeDebugLevel;
+		InnerBuffer = Buffer;
+		DebugGroup = [NSMutableArray new];
+		ActiveEncoder = nil;
+		DebugInfoBuffer = nil;
+		if (DebugLevel >= EAGXDebugLevelValidation)
+		{
+			DebugInfoBuffer = [Buffer.device newBufferWithLength:BufferOffsetAlignment options:0];
+		}
+	}
+	return Self;
+}
+
+-(void)dealloc
+{
+	for (FAGXDebugCommand* Command : DebugCommands)
+	{
+		[Command->Label release];
+		[Command->PassDesc release];
+		delete Command;
+	}
+	
+	[DebugGroup release];
+	[DebugInfoBuffer release];
+	DebugInfoBuffer = nil;
+	
+	[super dealloc];
+}
+
+@end
+
+FAGXCommandBufferDebugging FAGXCommandBufferDebugging::Get(mtlpp::CommandBuffer& Buffer)
+{
+	return Buffer.GetAssociatedObject<FAGXCommandBufferDebugging>((void const*)&FAGXCommandBufferDebugging::Get);
+}
+
+FAGXCommandBufferDebugging::FAGXCommandBufferDebugging()
+: ns::Object<FAGXDebugCommandBuffer*>(nullptr)
+{
+	
+}
+FAGXCommandBufferDebugging::FAGXCommandBufferDebugging(mtlpp::CommandBuffer& Buffer)
+: ns::Object<FAGXDebugCommandBuffer*>([[FAGXDebugCommandBuffer alloc] initWithCommandBuffer:Buffer.GetPtr()], ns::Ownership::Assign)
+{
+	Buffer.SetAssociatedObject((void const*)&FAGXCommandBufferDebugging::Get, *this);
+}
+FAGXCommandBufferDebugging::FAGXCommandBufferDebugging(FAGXDebugCommandBuffer* handle)
+: ns::Object<FAGXDebugCommandBuffer*>(handle)
+{
+	
+}
+
+ns::AutoReleased<ns::String> FAGXCommandBufferDebugging::GetDescription()
+{
+	NSMutableString* String = [[NSMutableString new] autorelease];
+	NSString* Label = m_ptr->InnerBuffer.label ? m_ptr->InnerBuffer.label : @"Unknown";
+	[String appendFormat:@"Command Buffer %p %@:", m_ptr->InnerBuffer, Label];
+	return ns::AutoReleased<ns::String>(String);
+}
+
+ns::AutoReleased<ns::String> FAGXCommandBufferDebugging::GetDebugDescription()
+{
+	NSMutableString* String = [[NSMutableString new] autorelease];
+	NSString* Label = m_ptr->InnerBuffer.label ? m_ptr->InnerBuffer.label : @"Unknown";
+	[String appendFormat:@"Command Buffer %p %@:", m_ptr->InnerBuffer, Label];
+	
+	uint32 Index = 0;
+	if (m_ptr->DebugInfoBuffer)
+	{
+		Index = *((uint32*)m_ptr->DebugInfoBuffer.contents);
+	}
+	
+	uint32 Count = 1;
+	for (FAGXDebugCommand* Command : m_ptr->DebugCommands)
+	{
+		if (Index == Count++)
+		{
+			[String appendFormat:@"\n\t--> %@: %@", GAGXDebugCommandTypeNames[Command->Type], Command->Label];
+		}
+		else
+		{
+			[String appendFormat:@"\n\t%@: %@", GAGXDebugCommandTypeNames[Command->Type], Command->Label];
+		}
+	}
+	
+	return ns::AutoReleased<ns::String>(String);
+}
+
+void FAGXCommandBufferDebugging::BeginRenderCommandEncoder(ns::String const& Label, mtlpp::RenderPassDescriptor const& Desc)
+{
+	if (m_ptr->DebugLevel >= EAGXDebugLevelValidation)
+	{
+		if (m_ptr->DebugLevel >= EAGXDebugLevelLogOperations)
+		{
+			check(!m_ptr->ActiveEncoder);
+			m_ptr->ActiveEncoder = [Label.GetPtr() retain];
+			FAGXDebugCommand* Command = new FAGXDebugCommand;
+			Command->Type = EAGXDebugCommandTypeRenderEncoder;
+			Command->Label = [Label.GetPtr() retain];
+			Command->PassDesc = [Desc.GetPtr() retain];
+			m_ptr->DebugCommands.Add(Command);
+		}
+	}
+}
+void FAGXCommandBufferDebugging::BeginComputeCommandEncoder(ns::String const& Label)
+{
+	if (m_ptr->DebugLevel >= EAGXDebugLevelLogOperations)
+	{
+		check(!m_ptr->ActiveEncoder);
+		m_ptr->ActiveEncoder = [Label.GetPtr() retain];
+		FAGXDebugCommand* Command = new FAGXDebugCommand;
+		Command->Type = EAGXDebugCommandTypeComputeEncoder;
+		Command->Label = [m_ptr->ActiveEncoder retain];
+		Command->PassDesc = nil;
+		m_ptr->DebugCommands.Add(Command);
+	}
+}
+void FAGXCommandBufferDebugging::BeginBlitCommandEncoder(ns::String const& Label)
+{
+	if (m_ptr->DebugLevel >= EAGXDebugLevelLogOperations)
+	{
+		check(!m_ptr->ActiveEncoder);
+		m_ptr->ActiveEncoder = [Label.GetPtr() retain];
+		FAGXDebugCommand* Command = new FAGXDebugCommand;
+		Command->Type = EAGXDebugCommandTypeBlitEncoder;
+		Command->Label = [m_ptr->ActiveEncoder retain];
+		Command->PassDesc = nil;
+		m_ptr->DebugCommands.Add(Command);
+	}
+}
+void FAGXCommandBufferDebugging::EndCommandEncoder()
+{
+	if (m_ptr->DebugLevel >= EAGXDebugLevelLogOperations)
+	{
+		check(m_ptr->ActiveEncoder);
+		FAGXDebugCommand* Command = new FAGXDebugCommand;
+		Command->Type = EAGXDebugCommandTypeEndEncoder;
+		Command->Label = m_ptr->ActiveEncoder;
+		Command->PassDesc = nil;
+		m_ptr->DebugCommands.Add(Command);
+		m_ptr->ActiveEncoder = nil;
+	}
+}
+
+void FAGXCommandBufferDebugging::SetPipeline(ns::String const& Desc)
+{
+	if (m_ptr->DebugLevel >= EAGXDebugLevelLogOperations)
+	{
+		FAGXDebugCommand* Command = new FAGXDebugCommand;
+		Command->Type = EAGXDebugCommandTypePipeline;
+		Command->Label = [Desc.GetPtr() retain];
+		Command->PassDesc = nil;
+		m_ptr->DebugCommands.Add(Command);
+	}
+}
+void FAGXCommandBufferDebugging::Draw(ns::String const& Desc)
+{
+	if (m_ptr->DebugLevel >= EAGXDebugLevelLogOperations)
+	{
+		FAGXDebugCommand* Command = new FAGXDebugCommand;
+		Command->Type = EAGXDebugCommandTypeDraw;
+		Command->Label = [Desc.GetPtr() retain];
+		Command->PassDesc = nil;
+		m_ptr->DebugCommands.Add(Command);
+	}
+}
+void FAGXCommandBufferDebugging::Dispatch(ns::String const& Desc)
+{
+	if (m_ptr->DebugLevel >= EAGXDebugLevelLogOperations)
+	{
+		FAGXDebugCommand* Command = new FAGXDebugCommand;
+		Command->Type = EAGXDebugCommandTypeDispatch;
+		Command->Label = [Desc.GetPtr() retain];
+		Command->PassDesc = nil;
+		m_ptr->DebugCommands.Add(Command);
+	}
+}
+void FAGXCommandBufferDebugging::Blit(ns::String const& Desc)
+{
+	if (m_ptr->DebugLevel >= EAGXDebugLevelLogOperations)
+	{
+		FAGXDebugCommand* Command = new FAGXDebugCommand;
+		Command->Type = EAGXDebugCommandTypeBlit;
+		Command->Label = [Desc.GetPtr() retain];
+		Command->PassDesc = nil;
+		m_ptr->DebugCommands.Add(Command);
+	}
+}
+
+void FAGXCommandBufferDebugging::InsertDebugSignpost(ns::String const& Label)
+{
+	if (m_ptr->DebugLevel >= EAGXDebugLevelLogOperations)
+	{
+		FAGXDebugCommand* Command = new FAGXDebugCommand;
+		Command->Type = EAGXDebugCommandTypeSignpost;
+		Command->Label = [Label.GetPtr() retain];
+		Command->PassDesc = nil;
+		m_ptr->DebugCommands.Add(Command);
+	}
+}
+void FAGXCommandBufferDebugging::PushDebugGroup(ns::String const& Group)
+{
+	if (m_ptr->DebugLevel >= EAGXDebugLevelLogOperations)
+	{
+		[m_ptr->DebugGroup addObject:Group.GetPtr()];
+		FAGXDebugCommand* Command = new FAGXDebugCommand;
+		Command->Type = EAGXDebugCommandTypePushGroup;
+		Command->Label = [Group.GetPtr() retain];
+		Command->PassDesc = nil;
+		m_ptr->DebugCommands.Add(Command);
+	}
+}
+void FAGXCommandBufferDebugging::PopDebugGroup()
+{
+	if (m_ptr->DebugLevel >= EAGXDebugLevelLogOperations)
+	{
+		if (m_ptr->DebugGroup.lastObject)
+		{
+			FAGXDebugCommand* Command = new FAGXDebugCommand;
+			Command->Type = EAGXDebugCommandTypePopGroup;
+			Command->Label = [m_ptr->DebugGroup.lastObject retain];
+			Command->PassDesc = nil;
+			m_ptr->DebugCommands.Add(Command);
+		}
+	}
+}
+
+#endif
