Initial import of AGXRHI

diff --git a/Engine/Source/Runtime/Apple/AGXRHI/Private/AGXCommandBuffer.h b/Engine/Source/Runtime/Apple/AGXRHI/Private/AGXCommandBuffer.h
new file mode 100644
--- /dev/null
+++ b/Engine/Source/Runtime/Apple/AGXRHI/Private/AGXCommandBuffer.h
@@ -0,0 +1,104 @@
+// Copyright Epic Games, Inc. All Rights Reserved.
+
+#pragma once
+
+#include <Metal/Metal.h>
+#include "AGXResources.h"
+#include "MetalShaderResources.h"
+
+/**
+ * EAGXDebugCommandType: Types of command recorded in our debug command-buffer wrapper.
+ */
+enum EAGXDebugCommandType
+{
+	EAGXDebugCommandTypeRenderEncoder,
+	EAGXDebugCommandTypeComputeEncoder,
+	EAGXDebugCommandTypeBlitEncoder,
+	EAGXDebugCommandTypeEndEncoder,
+    EAGXDebugCommandTypePipeline,
+	EAGXDebugCommandTypeDraw,
+	EAGXDebugCommandTypeDispatch,
+	EAGXDebugCommandTypeBlit,
+	EAGXDebugCommandTypeSignpost,
+	EAGXDebugCommandTypePushGroup,
+	EAGXDebugCommandTypePopGroup,
+	EAGXDebugCommandTypeInvalid
+};
+
+/**
+ * EAGXDebugLevel: Level of AGX RHI debug features to be enabled.
+ */
+enum EAGXDebugLevel
+{
+	EAGXDebugLevelOff,
+	EAGXDebugLevelFastValidation,
+	EAGXDebugLevelResetOnBind,
+	EAGXDebugLevelConditionalSubmit,
+	EAGXDebugLevelValidation,
+	EAGXDebugLevelLogOperations,
+	EAGXDebugLevelWaitForComplete,
+};
+
+NS_ASSUME_NONNULL_BEGIN
+/**
+ * FAGXDebugCommand: The data recorded for each command in the debug command-buffer wrapper.
+ */
+struct FAGXDebugCommand
+{
+	NSString* Label;
+	EAGXDebugCommandType Type;
+	MTLRenderPassDescriptor* PassDesc;
+};
+
+#if MTLPP_CONFIG_VALIDATE && METAL_DEBUG_OPTIONS
+
+/**
+ * FAGXDebugCommandBuffer: Wrapper around id<MTLCommandBuffer> that records information about commands.
+ * This allows reporting of substantially more information in debug modes which can be especially helpful
+ * when debugging GPU command-buffer failures.
+ */
+@interface FAGXDebugCommandBuffer : FApplePlatformObject
+{
+@public
+	NSMutableArray<NSString*>* DebugGroup;
+	NSString* ActiveEncoder;
+	id<MTLCommandBuffer> InnerBuffer;
+	TArray<FAGXDebugCommand*> DebugCommands;
+	EAGXDebugLevel DebugLevel;
+	id<MTLBuffer> DebugInfoBuffer;
+};
+
+/** Initialise the wrapper with the provided command-buffer. */
+-(id)initWithCommandBuffer:(id<MTLCommandBuffer>)Buffer;
+
+@end
+
+class FAGXCommandBufferDebugging : public ns::Object<FAGXDebugCommandBuffer*>
+{
+public:
+	FAGXCommandBufferDebugging();
+	FAGXCommandBufferDebugging(mtlpp::CommandBuffer& Buffer);
+	FAGXCommandBufferDebugging(FAGXDebugCommandBuffer* handle);
+	
+	static FAGXCommandBufferDebugging Get(mtlpp::CommandBuffer& Buffer);
+	ns::AutoReleased<ns::String> GetDescription();
+	ns::AutoReleased<ns::String> GetDebugDescription();
+	
+	void BeginRenderCommandEncoder(ns::String const& Label, mtlpp::RenderPassDescriptor const& Desc);
+	void BeginComputeCommandEncoder(ns::String const& Label);
+	void BeginBlitCommandEncoder(ns::String const& Label);
+	void EndCommandEncoder();
+	
+	void SetPipeline(ns::String const& Desc);
+	void Draw(ns::String const& Desc);
+	void Dispatch(ns::String const& Desc);
+	void Blit(ns::String const& Desc);
+	
+	void InsertDebugSignpost(ns::String const& Label);
+	void PushDebugGroup(ns::String const& Group);
+	void PopDebugGroup();
+};
+
+#endif
+
+NS_ASSUME_NONNULL_END
diff --git a/Engine/Source/Runtime/Apple/AGXRHI/Private/AGXCommandBufferFence.cpp b/Engine/Source/Runtime/Apple/AGXRHI/Private/AGXCommandBufferFence.cpp
new file mode 100644
--- /dev/null
+++ b/Engine/Source/Runtime/Apple/AGXRHI/Private/AGXCommandBufferFence.cpp
@@ -0,0 +1,30 @@
+// Copyright Epic Games, Inc. All Rights Reserved.
+
+/*=============================================================================
+	AGXCommandBufferFence.cpp: AGX RHI Command Buffer Fence Implementation.
+=============================================================================*/
+
+#include "AGXRHIPrivate.h"
+#include "AGXCommandBufferFence.h"
+
+
+//------------------------------------------------------------------------------
+
+#pragma mark - AGX RHI Command Buffer Fence Routines - 
+
+
+bool FAGXCommandBufferFence::Wait(uint64 Millis)
+{
+	@autoreleasepool {
+		if (CommandBufferFence)
+		{
+			bool bFinished = CommandBufferFence.Wait(Millis);
+			FPlatformMisc::MemoryBarrier();
+			return bFinished;
+		}
+		else
+		{
+			return true;
+		}
+	}
+}
diff --git a/Engine/Source/Runtime/Apple/AGXRHI/Private/AGXCommandBufferFence.h b/Engine/Source/Runtime/Apple/AGXRHI/Private/AGXCommandBufferFence.h
new file mode 100644
--- /dev/null
+++ b/Engine/Source/Runtime/Apple/AGXRHI/Private/AGXCommandBufferFence.h
@@ -0,0 +1,14 @@
+// Copyright Epic Games, Inc. All Rights Reserved.
+
+/*=============================================================================
+	AGXCommandBufferFence.h: AGX RHI Command Buffer Fence Definition.
+=============================================================================*/
+
+#pragma once
+
+struct FAGXCommandBufferFence
+{
+	bool Wait(uint64 Millis);
+
+	mtlpp::CommandBufferFence CommandBufferFence;
+};
diff --git a/Engine/Source/Runtime/Apple/AGXRHI/Private/AGXCommandEncoder.cpp b/Engine/Source/Runtime/Apple/AGXRHI/Private/AGXCommandEncoder.cpp
new file mode 100644
--- /dev/null
+++ b/Engine/Source/Runtime/Apple/AGXRHI/Private/AGXCommandEncoder.cpp
@@ -0,0 +1,1891 @@
+// Copyright Epic Games, Inc. All Rights Reserved.
+
+/*=============================================================================
+	AGXCommandEncoder.cpp: AGX RHI command encoder wrapper.
+=============================================================================*/
+
+#include "AGXRHIPrivate.h"
+#include "AGXShaderTypes.h"
+#include "AGXGraphicsPipelineState.h"
+#include "AGXCommandBufferFence.h"
+#include "AGXCommandEncoder.h"
+#include "AGXCommandBuffer.h"
+#include "AGXComputeCommandEncoder.h"
+#include "AGXRenderCommandEncoder.h"
+#include "AGXProfiler.h"
+#include "MetalShaderResources.h"
+
+const uint32 EncoderRingBufferSize = 1024 * 1024;
+
+#if METAL_DEBUG_OPTIONS
+extern int32 GAGXBufferScribble;
+#endif
+
+static TCHAR const* const GAGXCommandDataTypeName[] = {
+	TEXT("DrawPrimitive"),
+	TEXT("DrawPrimitiveIndexed"),
+	TEXT("DrawPrimitivePatch"),
+	TEXT("DrawPrimitiveIndirect"),
+	TEXT("DrawPrimitiveIndexedIndirect"),
+	TEXT("Dispatch"),
+	TEXT("DispatchIndirect"),
+};
+
+
+FString FAGXCommandData::ToString() const
+{
+	FString Result;
+	if ((uint32)CommandType < (uint32)FAGXCommandData::Type::Num)
+	{
+		Result = GAGXCommandDataTypeName[(uint32)CommandType];
+		switch(CommandType)
+		{
+			case FAGXCommandData::Type::DrawPrimitive:
+				Result += FString::Printf(TEXT(" BaseInstance: %u InstanceCount: %u VertexCount: %u VertexStart: %u"), Draw.BaseInstance, Draw.InstanceCount, Draw.VertexCount, Draw.VertexStart);
+				break;
+			case FAGXCommandData::Type::DrawPrimitiveIndexed:
+				Result += FString::Printf(TEXT(" BaseInstance: %u BaseVertex: %u IndexCount: %u IndexStart: %u InstanceCount: %u"), DrawIndexed.BaseInstance, DrawIndexed.BaseVertex, DrawIndexed.IndexCount, DrawIndexed.IndexStart, DrawIndexed.InstanceCount);
+				break;
+			case FAGXCommandData::Type::DrawPrimitivePatch:
+				Result += FString::Printf(TEXT(" BaseInstance: %u InstanceCount: %u PatchCount: %u PatchStart: %u"), DrawPatch.BaseInstance, DrawPatch.InstanceCount, DrawPatch.PatchCount, DrawPatch.PatchStart);
+				break;
+			case FAGXCommandData::Type::Dispatch:
+				Result += FString::Printf(TEXT(" X: %u Y: %u Z: %u"), (uint32)Dispatch.threadgroupsPerGrid[0], (uint32)Dispatch.threadgroupsPerGrid[1], (uint32)Dispatch.threadgroupsPerGrid[2]);
+				break;
+			case FAGXCommandData::Type::DispatchIndirect:
+				Result += FString::Printf(TEXT(" Buffer: %p Offset: %u"), (void*)DispatchIndirect.ArgumentBuffer, (uint32)DispatchIndirect.ArgumentOffset);
+				break;
+			case FAGXCommandData::Type::DrawPrimitiveIndirect:
+			case FAGXCommandData::Type::DrawPrimitiveIndexedIndirect:
+			case FAGXCommandData::Type::Num:
+			default:
+				break;
+		}
+	}
+	return Result;
+};
+
+struct FAGXCommandContextDebug
+{
+	TArray<FAGXCommandDebug> Commands;
+	TSet<TRefCountPtr<FAGXGraphicsPipelineState>> PSOs;
+	TSet<TRefCountPtr<FAGXComputeShader>> ComputeShaders;
+	FAGXBuffer DebugBuffer;
+};
+
+@interface FAGXCommandBufferDebug : FApplePlatformObject
+{
+	@public
+	TArray<FAGXCommandContextDebug> Contexts;
+	uint32 Index;
+}
+@end
+@implementation FAGXCommandBufferDebug
+APPLE_PLATFORM_OBJECT_ALLOC_OVERRIDES(FAGXCommandBufferDebug)
+- (instancetype)init
+{
+	id Self = [super init];
+	if (Self)
+	{
+		Index = ~0u;
+	}
+	return Self;
+}
+- (void)dealloc
+{
+	Contexts.Empty();
+	[super dealloc];
+}
+@end
+
+char const* FAGXCommandBufferMarkers::kTableAssociationKey = "FAGXCommandBufferMarkers::kTableAssociationKey";
+
+FAGXCommandBufferMarkers::FAGXCommandBufferMarkers(void)
+: ns::Object<FAGXCommandBufferDebug*, ns::CallingConvention::ObjectiveC>(nil)
+{
+	
+}
+
+FAGXCommandBufferMarkers::FAGXCommandBufferMarkers(mtlpp::CommandBuffer& CmdBuf)
+: ns::Object<FAGXCommandBufferDebug*, ns::CallingConvention::ObjectiveC>([FAGXCommandBufferDebug new], ns::Ownership::Assign)
+{
+	CmdBuf.SetAssociatedObject<FAGXCommandBufferMarkers>(FAGXCommandBufferMarkers::kTableAssociationKey, *this);
+	m_ptr->Contexts.SetNum(1);
+}
+
+
+FAGXCommandBufferMarkers::FAGXCommandBufferMarkers(FAGXCommandBufferDebug* CmdBuf)
+: ns::Object<FAGXCommandBufferDebug*, ns::CallingConvention::ObjectiveC>(CmdBuf)
+{
+	
+}
+
+void FAGXCommandBufferMarkers::AllocateContexts(uint32 NumContexts)
+{
+	if (m_ptr && m_ptr->Contexts.Num() < NumContexts)
+	{
+		m_ptr->Contexts.SetNum(NumContexts);
+	}
+}
+
+uint32 FAGXCommandBufferMarkers::AddCommand(uint32 CmdBufIndex, uint32 Encoder, uint32 ContextIndex, FAGXBuffer& DebugBuffer, FAGXGraphicsPipelineState* PSO, FAGXComputeShader* ComputeShader, FAGXCommandData& Data)
+{
+	uint32 Num = 0;
+	if (m_ptr)
+	{
+		if (m_ptr->Index == ~0u)
+		{
+			m_ptr->Index = CmdBufIndex;
+		}
+		
+		FAGXCommandContextDebug& Context = m_ptr->Contexts[ContextIndex];
+		if (Context.DebugBuffer != DebugBuffer)
+		{
+			Context.DebugBuffer = DebugBuffer;
+		}
+		
+		if (PSO)
+			Context.PSOs.Add(PSO);
+		if (ComputeShader)
+			Context.ComputeShaders.Add(ComputeShader);
+		
+		Num = Context.Commands.Num();
+		FAGXCommandDebug Command;
+        Command.CmdBufIndex = CmdBufIndex;
+		Command.Encoder = Encoder;
+		Command.Index = Num;
+		Command.PSO = PSO;
+		Command.ComputeShader = ComputeShader;
+		Command.Data = Data;
+		Context.Commands.Add(Command);
+	}
+	return Num;
+}
+
+TArray<FAGXCommandDebug>* FAGXCommandBufferMarkers::GetCommands(uint32 ContextIndex)
+{
+	TArray<FAGXCommandDebug>* Result = nullptr;
+	if (m_ptr)
+	{
+		FAGXCommandContextDebug& Context = m_ptr->Contexts[ContextIndex];
+		Result = &Context.Commands;
+	}
+	return Result;
+}
+
+ns::AutoReleased<FAGXBuffer> FAGXCommandBufferMarkers::GetDebugBuffer(uint32 ContextIndex)
+{
+	ns::AutoReleased<FAGXBuffer> Buffer;
+	if (m_ptr)
+	{
+		FAGXCommandContextDebug& Context = m_ptr->Contexts[ContextIndex];
+		Buffer = Context.DebugBuffer;
+	}
+	return Buffer;
+}
+
+uint32 FAGXCommandBufferMarkers::NumContexts() const
+{
+	uint32 Num = 0;
+	if (m_ptr)
+	{
+		Num = m_ptr->Contexts.Num();
+	}
+	return Num;
+}
+
+uint32 FAGXCommandBufferMarkers::GetIndex() const
+{
+	uint32 Num = 0;
+	if (m_ptr)
+	{
+		Num = m_ptr->Index;
+	}
+	return Num;
+}
+
+FAGXCommandBufferMarkers FAGXCommandBufferMarkers::Get(mtlpp::CommandBuffer const& CmdBuf)
+{
+	return CmdBuf.GetAssociatedObject<FAGXCommandBufferMarkers>(FAGXCommandBufferMarkers::kTableAssociationKey);
+}
+
+#pragma mark - Public C++ Boilerplate -
+
+FAGXCommandEncoder::FAGXCommandEncoder(FAGXCommandList& CmdList, EAGXCommandEncoderType InType)
+: CommandList(CmdList)
+, bSupportsMetalFeaturesSetBytes(CmdList.GetCommandQueue().SupportsFeature(EAGXFeaturesSetBytes))
+, RingBuffer(EncoderRingBufferSize, BufferOffsetAlignment, FAGXCommandQueue::GetCompatibleResourceOptions((mtlpp::ResourceOptions)(mtlpp::ResourceOptions::HazardTrackingModeUntracked | BUFFER_RESOURCE_STORAGE_MANAGED)))
+, RenderPassDesc(nil)
+, EncoderFence(nil)
+#if ENABLE_METAL_GPUPROFILE
+, CommandBufferStats(nullptr)
+#endif
+#if METAL_DEBUG_OPTIONS
+, WaitCount(0)
+, UpdateCount(0)
+#endif
+, DebugGroups([NSMutableArray new])
+, FenceStage(mtlpp::RenderStages::Fragment)
+, EncoderNum(0)
+, CmdBufIndex(0)
+, Type(InType)
+{
+	for (uint32 Frequency = 0; Frequency < uint32(mtlpp::FunctionType::Kernel)+1; Frequency++)
+	{
+		FMemory::Memzero(ShaderBuffers[Frequency].Bytes);
+		FMemory::Memzero(ShaderBuffers[Frequency].Offsets);
+		FMemory::Memzero(ShaderBuffers[Frequency].Lengths);
+		FMemory::Memzero(ShaderBuffers[Frequency].Usage);
+		ShaderBuffers[Frequency].SideTable = [[FAGXBufferData alloc] init];
+		ShaderBuffers[Frequency].SideTable->Data = (uint8*)(&ShaderBuffers[Frequency].Lengths[0]);
+		ShaderBuffers[Frequency].SideTable->Len = sizeof(ShaderBuffers[Frequency].Lengths);
+		ShaderBuffers[Frequency].Bound = 0;
+	}
+	
+	for (uint32 i = 0; i < MaxSimultaneousRenderTargets; i++)
+	{
+		ColorStoreActions[i] = mtlpp::StoreAction::Unknown;
+	}
+	DepthStoreAction = mtlpp::StoreAction::Unknown;
+	StencilStoreAction = mtlpp::StoreAction::Unknown;
+}
+
+FAGXCommandEncoder::~FAGXCommandEncoder(void)
+{
+	if(CommandBuffer)
+	{
+		EndEncoding();
+		CommitCommandBuffer(false);
+	}
+	
+	check(!IsRenderCommandEncoderActive());
+	check(!IsComputeCommandEncoderActive());
+	check(!IsBlitCommandEncoderActive());
+	
+	AGXSafeReleaseMetalRenderPassDescriptor(RenderPassDesc);
+	RenderPassDesc = nil;
+
+	if(DebugGroups)
+	{
+		[DebugGroups release];
+	}
+	
+	for (uint32 Frequency = 0; Frequency < uint32(mtlpp::FunctionType::Kernel)+1; Frequency++)
+	{
+		for (uint32 i = 0; i < ML_MaxBuffers; i++)
+		{
+			ShaderBuffers[Frequency].Buffers[i] = nil;
+		}
+		FMemory::Memzero(ShaderBuffers[Frequency].Bytes);
+		FMemory::Memzero(ShaderBuffers[Frequency].Offsets);
+		FMemory::Memzero(ShaderBuffers[Frequency].Lengths);
+		FMemory::Memzero(ShaderBuffers[Frequency].Usage);
+		ShaderBuffers[Frequency].SideTable->Data = nullptr;
+		[ShaderBuffers[Frequency].SideTable release];
+		ShaderBuffers[Frequency].SideTable = nil;
+		ShaderBuffers[Frequency].Bound = 0;
+	}
+}
+
+void FAGXCommandEncoder::Reset(void)
+{
+    check(IsRenderCommandEncoderActive() == false
+          && IsComputeCommandEncoderActive() == false
+          && IsBlitCommandEncoderActive() == false);
+	
+	if(RenderPassDesc)
+	{
+		AGXSafeReleaseMetalRenderPassDescriptor(RenderPassDesc);
+		RenderPassDesc = nil;
+	}
+	
+	{
+		for (uint32 i = 0; i < MaxSimultaneousRenderTargets; i++)
+		{
+			ColorStoreActions[i] = mtlpp::StoreAction::Unknown;
+		}
+		DepthStoreAction = mtlpp::StoreAction::Unknown;
+		StencilStoreAction = mtlpp::StoreAction::Unknown;
+	}
+	
+	for (uint32 Frequency = 0; Frequency < uint32(mtlpp::FunctionType::Kernel)+1; Frequency++)
+	{
+		for (uint32 i = 0; i < ML_MaxBuffers; i++)
+		{
+			ShaderBuffers[Frequency].Buffers[i] = nil;
+		}
+    	FMemory::Memzero(ShaderBuffers[Frequency].Bytes);
+		FMemory::Memzero(ShaderBuffers[Frequency].Offsets);
+		FMemory::Memzero(ShaderBuffers[Frequency].Lengths);
+		FMemory::Memzero(ShaderBuffers[Frequency].Usage);
+		ShaderBuffers[Frequency].Bound = 0;
+	}
+	
+	[DebugGroups removeAllObjects];
+}
+
+void FAGXCommandEncoder::ResetLive(void)
+{
+	for (uint32 Frequency = 0; Frequency < uint32(mtlpp::FunctionType::Kernel)+1; Frequency++)
+	{
+		for (uint32 i = 0; i < ML_MaxBuffers; i++)
+		{
+			ShaderBuffers[Frequency].Buffers[i] = nil;
+		}
+		FMemory::Memzero(ShaderBuffers[Frequency].Bytes);
+		FMemory::Memzero(ShaderBuffers[Frequency].Offsets);
+		FMemory::Memzero(ShaderBuffers[Frequency].Lengths);
+		ShaderBuffers[Frequency].Bound = 0;
+	}
+	
+	if (IsRenderCommandEncoderActive())
+	{
+		for (uint32 i = 0; i < ML_MaxBuffers; i++)
+		{
+			RenderCommandEncoder.SetVertexBuffer(nil, 0, i);
+			RenderCommandEncoder.SetFragmentBuffer(nil, 0, i);
+		}
+		
+		for (uint32 i = 0; i < ML_MaxTextures; i++)
+		{
+			RenderCommandEncoder.SetVertexTexture(nil, i);
+			RenderCommandEncoder.SetFragmentTexture(nil, i);
+		}
+	}
+	else if (IsComputeCommandEncoderActive())
+	{
+		for (uint32 i = 0; i < ML_MaxBuffers; i++)
+		{
+			ComputeCommandEncoder.SetBuffer(nil, 0, i);
+		}
+		
+		for (uint32 i = 0; i < ML_MaxTextures; i++)
+		{
+			ComputeCommandEncoder.SetTexture(nil, i);
+		}
+	}
+}
+
+#pragma mark - Public Command Buffer Mutators -
+
+void FAGXCommandEncoder::StartCommandBuffer(void)
+{
+	check(!CommandBuffer || EncoderNum == 0);
+	check(IsRenderCommandEncoderActive() == false
+          && IsComputeCommandEncoderActive() == false
+          && IsBlitCommandEncoderActive() == false);
+
+	if (!CommandBuffer)
+	{
+		CmdBufIndex++;
+		CommandBuffer = CommandList.GetCommandQueue().CreateCommandBuffer();
+		METAL_DEBUG_LAYER(EAGXDebugLevelFastValidation, CommandBufferDebug = FAGXCommandBufferDebugging::Get(CommandBuffer));
+		
+		if (GAGXCommandBufferDebuggingEnabled)
+		{
+			CommandBufferMarkers = FAGXCommandBufferMarkers(CommandBuffer);
+		}
+		
+		if ([DebugGroups count] > 0)
+		{
+			CommandBuffer.SetLabel([DebugGroups lastObject]);
+		}
+		
+	#if ENABLE_METAL_GPUPROFILE
+		FAGXProfiler* Profiler = FAGXProfiler::GetProfiler();
+		if (Profiler)
+		{
+			CommandBufferStats = Profiler->AllocateCommandBuffer(CommandBuffer, 0);
+		}
+	#endif
+	}
+}
+	
+void FAGXCommandEncoder::CommitCommandBuffer(uint32 const Flags)
+{
+	check(CommandBuffer);
+	check(IsRenderCommandEncoderActive() == false
+          && IsComputeCommandEncoderActive() == false
+          && IsBlitCommandEncoderActive() == false);
+
+	bool const bWait = (Flags & EAGXSubmitFlagsWaitOnCommandBuffer);
+	bool const bIsLastCommandBuffer = (Flags & EAGXSubmitFlagsLastCommandBuffer);
+	
+	if (EncoderNum == 0 && !bWait && !(Flags & EAGXSubmitFlagsForce))
+	{
+		return;
+	}
+	
+	if(CommandBuffer.GetLabel() == nil && [DebugGroups count] > 0)
+	{
+		CommandBuffer.SetLabel([DebugGroups lastObject]);
+	}
+	
+	if (!(Flags & EAGXSubmitFlagsBreakCommandBuffer))
+	{
+		RingBuffer.Commit(CommandBuffer);
+	}
+	else
+	{
+		RingBuffer.Submit();
+	}
+    
+#if METAL_DEBUG_OPTIONS
+    if(CommandList.GetCommandQueue().GetRuntimeDebuggingLevel() >= EAGXDebugLevelValidation)
+    {
+        for (FAGXBuffer const& Buffer : ActiveBuffers)
+        {
+            GetAGXDeviceContext().AddActiveBuffer(Buffer);
+        }
+        
+        TSet<ns::AutoReleased<FAGXBuffer>> NewActiveBuffers = MoveTemp(ActiveBuffers);
+        AddCompletionHandler([NewActiveBuffers](mtlpp::CommandBuffer const&)
+                             {
+                                 for (FAGXBuffer const& Buffer : NewActiveBuffers)
+                                 {
+                                     GetAGXDeviceContext().RemoveActiveBuffer(Buffer);
+                                 }
+                             });
+    }
+#endif
+#if ENABLE_METAL_GPUPROFILE
+	CommandBufferStats->End(CommandBuffer);
+	CommandBufferStats = nullptr;
+#endif
+
+	CommandList.Commit(CommandBuffer, MoveTemp(CompletionHandlers), bWait, bIsLastCommandBuffer);
+	
+	CommandBuffer = nil;
+	if (Flags & EAGXSubmitFlagsCreateCommandBuffer)
+	{
+		StartCommandBuffer();
+		check(CommandBuffer);
+	}
+	
+	BufferBindingHistory.Empty();
+	TextureBindingHistory.Empty();
+	
+	EncoderNum = 0;
+}
+
+#pragma mark - Public Command Encoder Accessors -
+	
+bool FAGXCommandEncoder::IsParallelRenderCommandEncoderActive(void) const
+{
+	return ParallelRenderCommandEncoder.GetPtr() != nil;
+}
+	
+bool FAGXCommandEncoder::IsRenderCommandEncoderActive(void) const
+{
+	return RenderCommandEncoder.GetPtr() != nil || ParallelRenderCommandEncoder.GetPtr() != nil;
+}
+
+bool FAGXCommandEncoder::IsComputeCommandEncoderActive(void) const
+{
+	return ComputeCommandEncoder.GetPtr() != nil;
+}
+
+bool FAGXCommandEncoder::IsBlitCommandEncoderActive(void) const
+{
+	return BlitCommandEncoder.GetPtr() != nil;
+}
+
+bool FAGXCommandEncoder::IsImmediate(void) const
+{
+	return CommandList.IsImmediate();
+}
+
+bool FAGXCommandEncoder::IsParallel(void) const
+{
+	return CommandList.IsParallel() && (Type == EAGXCommandEncoderCurrent);
+}
+
+bool FAGXCommandEncoder::IsRenderPassDescriptorValid(void) const
+{
+	return (RenderPassDesc != nil);
+}
+
+mtlpp::RenderPassDescriptor const& FAGXCommandEncoder::GetRenderPassDescriptor(void) const
+{
+	return RenderPassDesc;
+}
+
+mtlpp::ParallelRenderCommandEncoder& FAGXCommandEncoder::GetParallelRenderCommandEncoder(void)
+{
+	return ParallelRenderCommandEncoder;
+}
+
+mtlpp::RenderCommandEncoder& FAGXCommandEncoder::GetChildRenderCommandEncoder(uint32 Index)
+{
+	check(IsParallelRenderCommandEncoderActive() && Index < ChildRenderCommandEncoders.Num());
+	return ChildRenderCommandEncoders[Index];
+}
+
+mtlpp::RenderCommandEncoder& FAGXCommandEncoder::GetRenderCommandEncoder(void)
+{
+	check(IsRenderCommandEncoderActive() && RenderCommandEncoder);
+	return RenderCommandEncoder;
+}
+
+mtlpp::ComputeCommandEncoder& FAGXCommandEncoder::GetComputeCommandEncoder(void)
+{
+	check(IsComputeCommandEncoderActive());
+	return ComputeCommandEncoder;
+}
+
+mtlpp::BlitCommandEncoder& FAGXCommandEncoder::GetBlitCommandEncoder(void)
+{
+	check(IsBlitCommandEncoderActive());
+	return BlitCommandEncoder;
+}
+
+TRefCountPtr<FAGXFence> const& FAGXCommandEncoder::GetEncoderFence(void) const
+{
+	return EncoderFence;
+}
+	
+#pragma mark - Public Command Encoder Mutators -
+
+void FAGXCommandEncoder::BeginParallelRenderCommandEncoding(uint32 NumChildren)
+{
+	check(IsImmediate());
+	check(RenderPassDesc);
+	check(CommandBuffer);
+	check(IsRenderCommandEncoderActive() == false && IsComputeCommandEncoderActive() == false && IsBlitCommandEncoderActive() == false);
+	
+	FenceResources.Append(TransitionedResources);
+	
+	ParallelRenderCommandEncoder = MTLPP_VALIDATE(mtlpp::CommandBuffer, CommandBuffer, AGXSafeGetRuntimeDebuggingLevel() >= EAGXDebugLevelValidation, ParallelRenderCommandEncoder(RenderPassDesc));
+	METAL_DEBUG_LAYER(EAGXDebugLevelFastValidation, ParallelEncoderDebug = FAGXParallelRenderCommandEncoderDebugging(ParallelRenderCommandEncoder, RenderPassDesc, CommandBufferDebug));
+	
+	EncoderNum++;
+	
+	check(!EncoderFence);
+	NSString* Label = nil;
+	
+	if(GetEmitDrawEvents())
+	{
+		Label = [NSString stringWithFormat:@"ParallelRenderCommandEncoder: %@", [DebugGroups count] > 0 ? [DebugGroups lastObject] : (NSString*)CFSTR("InitialPass")];
+		ParallelRenderCommandEncoder.SetLabel(Label);
+		
+		if([DebugGroups count])
+		{
+			for (NSString* Group in DebugGroups)
+			{
+				ParallelRenderCommandEncoder.PushDebugGroup(Group);
+				METAL_DEBUG_LAYER(EAGXDebugLevelFastValidation, RenderEncoderDebug.PushDebugGroup(Group));
+			}
+		}
+	}
+	
+	for (uint32 i = 0; i < NumChildren; i++)
+	{
+		mtlpp::RenderCommandEncoder CommandEncoder = MTLPP_VALIDATE(mtlpp::ParallelRenderCommandEncoder, ParallelRenderCommandEncoder, AGXSafeGetRuntimeDebuggingLevel() >= EAGXDebugLevelValidation, GetRenderCommandEncoder());
+		ChildRenderCommandEncoders.Add(CommandEncoder);
+		METAL_DEBUG_LAYER(EAGXDebugLevelFastValidation, ParallelEncoderDebug.GetRenderCommandEncoderDebugger(CommandEncoder));
+	}
+}
+
+void FAGXCommandEncoder::BeginRenderCommandEncoding(void)
+{
+	check(RenderPassDesc);
+	check(CommandList.IsParallel() || CommandBuffer);
+	check(IsRenderCommandEncoderActive() == false && IsComputeCommandEncoderActive() == false && IsBlitCommandEncoderActive() == false);
+	
+	FenceResources.Append(TransitionedResources);
+	
+	if (!CommandList.IsParallel() || Type == EAGXCommandEncoderPrologue)
+	{
+		RenderCommandEncoder = MTLPP_VALIDATE(mtlpp::CommandBuffer, CommandBuffer, AGXSafeGetRuntimeDebuggingLevel() >= EAGXDebugLevelValidation, RenderCommandEncoder(RenderPassDesc));
+		METAL_DEBUG_LAYER(EAGXDebugLevelFastValidation, RenderEncoderDebug = FAGXRenderCommandEncoderDebugging(RenderCommandEncoder, RenderPassDesc, CommandBufferDebug));
+		EncoderNum++;	
+	}
+	else
+	{
+		RenderCommandEncoder = GetAGXDeviceContext().GetParallelRenderCommandEncoder(CommandList.GetParallelIndex(), ParallelRenderCommandEncoder, CommandBuffer);
+		METAL_DEBUG_LAYER(EAGXDebugLevelFastValidation, RenderEncoderDebug = FAGXRenderCommandEncoderDebugging::Get(RenderCommandEncoder));
+	}
+	
+	check(!EncoderFence);
+	NSString* Label = nil;
+	
+	if(GetEmitDrawEvents())
+	{
+		Label = [NSString stringWithFormat:@"RenderEncoder: %@", [DebugGroups count] > 0 ? [DebugGroups lastObject] : (NSString*)CFSTR("InitialPass")];
+		RenderCommandEncoder.SetLabel(Label);
+		
+		if([DebugGroups count])
+		{
+			for (NSString* Group in DebugGroups)
+			{
+				RenderCommandEncoder.PushDebugGroup(Group);
+				METAL_DEBUG_LAYER(EAGXDebugLevelFastValidation, RenderEncoderDebug.PushDebugGroup(Group));
+			}
+		}
+	}
+	
+	if (CommandList.IsImmediate())
+	{
+		EncoderFence = CommandList.GetCommandQueue().CreateFence(Label);
+	}
+}
+
+void FAGXCommandEncoder::BeginComputeCommandEncoding(mtlpp::DispatchType DispatchType)
+{
+	check(CommandBuffer);
+	check(IsRenderCommandEncoderActive() == false && IsComputeCommandEncoderActive() == false && IsBlitCommandEncoderActive() == false);
+	
+	FenceResources.Append(TransitionedResources);
+	TransitionedResources.Empty();
+	
+	if (DispatchType == mtlpp::DispatchType::Serial)
+	{
+		ComputeCommandEncoder = MTLPP_VALIDATE(mtlpp::CommandBuffer, CommandBuffer, AGXSafeGetRuntimeDebuggingLevel() >= EAGXDebugLevelValidation, ComputeCommandEncoder());
+	}
+	else
+	{
+		ComputeCommandEncoder = MTLPP_VALIDATE(mtlpp::CommandBuffer, CommandBuffer, AGXSafeGetRuntimeDebuggingLevel() >= EAGXDebugLevelValidation, ComputeCommandEncoder(DispatchType));
+	}
+	METAL_DEBUG_LAYER(EAGXDebugLevelFastValidation, ComputeEncoderDebug = FAGXComputeCommandEncoderDebugging(ComputeCommandEncoder, CommandBufferDebug));
+
+	EncoderNum++;
+	
+	check(!EncoderFence);
+	NSString* Label = nil;
+	
+	if(GetEmitDrawEvents())
+	{
+		Label = [NSString stringWithFormat:@"ComputeEncoder: %@", [DebugGroups count] > 0 ? [DebugGroups lastObject] : (NSString*)CFSTR("InitialPass")];
+		ComputeCommandEncoder.SetLabel(Label);
+		
+		if([DebugGroups count])
+		{
+			for (NSString* Group in DebugGroups)
+			{
+				ComputeCommandEncoder.PushDebugGroup(Group);
+				METAL_DEBUG_LAYER(EAGXDebugLevelFastValidation, ComputeEncoderDebug.PushDebugGroup(Group));
+			}
+		}
+	}
+	
+	EncoderFence = CommandList.GetCommandQueue().CreateFence(Label);
+}
+
+void FAGXCommandEncoder::BeginBlitCommandEncoding(void)
+{
+	check(CommandBuffer);
+	check(IsRenderCommandEncoderActive() == false && IsComputeCommandEncoderActive() == false && IsBlitCommandEncoderActive() == false);
+	
+	FenceResources.Append(TransitionedResources);
+	TransitionedResources.Empty();
+	
+	BlitCommandEncoder = MTLPP_VALIDATE(mtlpp::CommandBuffer, CommandBuffer, AGXSafeGetRuntimeDebuggingLevel() >= EAGXDebugLevelValidation, BlitCommandEncoder());
+	METAL_DEBUG_LAYER(EAGXDebugLevelFastValidation, BlitEncoderDebug = FAGXBlitCommandEncoderDebugging(BlitCommandEncoder, CommandBufferDebug));
+	
+	EncoderNum++;
+	
+	check(!EncoderFence);
+	NSString* Label = nil;
+	
+	if(GetEmitDrawEvents())
+	{
+		Label = [NSString stringWithFormat:@"BlitEncoder: %@", [DebugGroups count] > 0 ? [DebugGroups lastObject] : (NSString*)CFSTR("InitialPass")];
+		BlitCommandEncoder.SetLabel(Label);
+		
+		if([DebugGroups count])
+		{
+			for (NSString* Group in DebugGroups)
+			{
+				BlitCommandEncoder.PushDebugGroup(Group);
+				METAL_DEBUG_LAYER(EAGXDebugLevelFastValidation, BlitEncoderDebug.PushDebugGroup(Group));
+			}
+		}
+	}
+	
+	EncoderFence = CommandList.GetCommandQueue().CreateFence(Label);
+}
+
+TRefCountPtr<FAGXFence> FAGXCommandEncoder::EndEncoding(void)
+{
+	static bool bSupportsFences = CommandList.GetCommandQueue().SupportsFeature(EAGXFeaturesFences);
+	TRefCountPtr<FAGXFence> Fence = nullptr;
+	@autoreleasepool
+	{
+		if(IsRenderCommandEncoderActive())
+		{
+			if (RenderCommandEncoder)
+			{
+				check(!bSupportsFences || EncoderFence || !CommandList.IsImmediate());
+				if (ParallelRenderCommandEncoder.GetPtr() == nil)
+				{
+					check(RenderPassDesc);
+					
+					ns::Array<mtlpp::RenderPassColorAttachmentDescriptor> ColorAttachments = RenderPassDesc.GetColorAttachments();
+					for (uint32 i = 0; i < MaxSimultaneousRenderTargets; i++)
+					{
+						if (ColorAttachments[i].GetTexture() && ColorAttachments[i].GetStoreAction() == mtlpp::StoreAction::Unknown)
+						{
+							mtlpp::StoreAction Action = ColorStoreActions[i];
+							check(Action != mtlpp::StoreAction::Unknown);
+							RenderCommandEncoder.SetColorStoreAction((mtlpp::StoreAction)Action, i);
+						}
+					}
+					if (RenderPassDesc.GetDepthAttachment().GetTexture() && RenderPassDesc.GetDepthAttachment().GetStoreAction() == mtlpp::StoreAction::Unknown)
+					{
+						mtlpp::StoreAction Action = DepthStoreAction;
+						check(Action != mtlpp::StoreAction::Unknown);
+						RenderCommandEncoder.SetDepthStoreAction((mtlpp::StoreAction)Action);
+					}
+					if (RenderPassDesc.GetStencilAttachment().GetTexture() && RenderPassDesc.GetStencilAttachment().GetStoreAction() == mtlpp::StoreAction::Unknown)
+					{
+						mtlpp::StoreAction Action = StencilStoreAction;
+						check(Action != mtlpp::StoreAction::Unknown);
+						RenderCommandEncoder.SetStencilStoreAction((mtlpp::StoreAction)Action);
+					}
+				}
+				
+				for (FAGXFence* FragFence : FragmentFences)
+				{
+					if (FragFence->NeedsWait(mtlpp::RenderStages::Fragment))
+					{
+						mtlpp::Fence FragmentFence = FragFence->Get(mtlpp::RenderStages::Fragment);
+						mtlpp::Fence FragInnerFence = METAL_DEBUG_OPTION(CommandList.GetCommandQueue().GetRuntimeDebuggingLevel() >= EAGXDebugLevelValidation ? mtlpp::Fence(((FAGXDebugFence*)FragmentFence.GetPtr()).Inner) :) FragmentFence;
+						
+						RenderCommandEncoder.WaitForFence(FragInnerFence, FenceStage);
+						METAL_DEBUG_LAYER(EAGXDebugLevelFastValidation, RenderEncoderDebug.AddWaitFence(FragmentFence));
+						FragFence->Wait(mtlpp::RenderStages::Fragment);
+					}
+				}
+				FragmentFences.Empty();
+				
+				if (FenceStage == mtlpp::RenderStages::Vertex)
+				{
+					FenceResources.Empty();
+					FenceStage = mtlpp::RenderStages::Fragment;
+				}
+				
+				if (EncoderFence && EncoderFence->NeedsWrite(mtlpp::RenderStages::Fragment))
+				{
+					Fence = EncoderFence;
+				}
+				UpdateFence(EncoderFence);
+				
+#if METAL_DEBUG_OPTIONS
+				if (bSupportsFences && AGXSafeGetRuntimeDebuggingLevel() >= EAGXDebugLevelFastValidation && (!WaitCount || !UpdateCount))
+				{
+					UE_LOG(LogAGX, Error, TEXT("%s has incorrect fence waits (%u) vs. updates (%u)."), *FString(RenderCommandEncoder.GetLabel()), WaitCount, UpdateCount);
+					
+				}
+				WaitCount = 0;
+				UpdateCount = 0;
+#endif
+				RenderCommandEncoder.EndEncoding();
+				METAL_DEBUG_LAYER(EAGXDebugLevelFastValidation, RenderEncoderDebug.EndEncoder());
+				RenderCommandEncoder = nil;
+				EncoderFence = nullptr;
+			}
+			
+			if (ParallelRenderCommandEncoder && IsParallel())
+			{
+				RingBuffer.Commit(CommandBuffer);
+				
+			#if METAL_DEBUG_OPTIONS
+				if(CommandList.GetCommandQueue().GetRuntimeDebuggingLevel() >= EAGXDebugLevelValidation)
+				{
+					for (FAGXBuffer const& Buffer : ActiveBuffers)
+					{
+						GetAGXDeviceContext().AddActiveBuffer(Buffer);
+					}
+					
+					TSet<ns::AutoReleased<FAGXBuffer>> NewActiveBuffers = MoveTemp(ActiveBuffers);
+					AddCompletionHandler([NewActiveBuffers](mtlpp::CommandBuffer const&)
+										{
+											for (FAGXBuffer const& Buffer : NewActiveBuffers)
+											{
+												GetAGXDeviceContext().RemoveActiveBuffer(Buffer);
+											}
+										});
+				}
+			#endif
+
+				BufferBindingHistory.Empty();
+				TextureBindingHistory.Empty();
+				
+				EncoderNum = 0;
+
+				CommandBuffer = nil;
+				
+				ParallelRenderCommandEncoder = nil;
+			}
+
+			if (ParallelRenderCommandEncoder && IsImmediate())
+			{
+				{
+					check(RenderPassDesc);
+					
+					ns::Array<mtlpp::RenderPassColorAttachmentDescriptor> ColorAttachments = RenderPassDesc.GetColorAttachments();
+					for (uint32 i = 0; i < MaxSimultaneousRenderTargets; i++)
+					{
+						if (ColorAttachments[i].GetTexture() && ColorAttachments[i].GetStoreAction() == mtlpp::StoreAction::Unknown)
+						{
+							mtlpp::StoreAction Action = ColorStoreActions[i];
+							check(Action != mtlpp::StoreAction::Unknown);
+							ParallelRenderCommandEncoder.SetColorStoreAction((mtlpp::StoreAction)Action, i);
+						}
+					}
+					if (RenderPassDesc.GetDepthAttachment().GetTexture() && RenderPassDesc.GetDepthAttachment().GetStoreAction() == mtlpp::StoreAction::Unknown)
+					{
+						mtlpp::StoreAction Action = DepthStoreAction;
+						check(Action != mtlpp::StoreAction::Unknown);
+						ParallelRenderCommandEncoder.SetDepthStoreAction((mtlpp::StoreAction)Action);
+					}
+					if (RenderPassDesc.GetStencilAttachment().GetTexture() && RenderPassDesc.GetStencilAttachment().GetStoreAction() == mtlpp::StoreAction::Unknown)
+					{
+						mtlpp::StoreAction Action = StencilStoreAction;
+						check(Action != mtlpp::StoreAction::Unknown);
+						ParallelRenderCommandEncoder.SetStencilStoreAction((mtlpp::StoreAction)Action);
+					}
+				}
+
+				ParallelRenderCommandEncoder.EndEncoding();
+				METAL_DEBUG_LAYER(EAGXDebugLevelFastValidation, ParallelEncoderDebug.EndEncoder());
+				ParallelRenderCommandEncoder = nil;
+
+				ChildRenderCommandEncoders.Empty();
+			}
+		}
+		else if(IsComputeCommandEncoderActive())
+		{
+			check(!bSupportsFences || EncoderFence);
+			
+			for (FAGXFence* FragFence : FragmentFences)
+			{
+				if (FragFence->NeedsWait(mtlpp::RenderStages::Fragment))
+				{
+					mtlpp::Fence FragmentFence = FragFence->Get(mtlpp::RenderStages::Fragment);
+					mtlpp::Fence FragInnerFence = METAL_DEBUG_OPTION(CommandList.GetCommandQueue().GetRuntimeDebuggingLevel() >= EAGXDebugLevelValidation ? mtlpp::Fence(((FAGXDebugFence*)FragmentFence.GetPtr()).Inner) :) FragmentFence;
+					
+					ComputeCommandEncoder.WaitForFence(FragInnerFence);
+					METAL_DEBUG_LAYER(EAGXDebugLevelFastValidation, ComputeEncoderDebug.AddWaitFence(FragmentFence));
+					FragFence->Wait(mtlpp::RenderStages::Fragment);
+				}
+			}
+			FragmentFences.Empty();
+			FenceResources.Empty();
+			FenceStage = mtlpp::RenderStages::Fragment;
+			
+			if (EncoderFence && EncoderFence->NeedsWrite(mtlpp::RenderStages::Fragment))
+			{
+				Fence = EncoderFence;
+			}
+			UpdateFence(EncoderFence);
+			
+#if METAL_DEBUG_OPTIONS
+			if (bSupportsFences && AGXSafeGetRuntimeDebuggingLevel() >= EAGXDebugLevelFastValidation && (!WaitCount || !UpdateCount))
+			{
+				UE_LOG(LogAGX, Error, TEXT("%s has incorrect fence waits (%u) vs. updates (%u)."), *FString(ComputeCommandEncoder.GetLabel()), WaitCount, UpdateCount);
+				
+			}
+			WaitCount = 0;
+			UpdateCount = 0;
+#endif
+			ComputeCommandEncoder.EndEncoding();
+			METAL_DEBUG_LAYER(EAGXDebugLevelFastValidation, ComputeEncoderDebug.EndEncoder());
+			ComputeCommandEncoder = nil;
+			EncoderFence = nullptr;
+		}
+		else if(IsBlitCommandEncoderActive())
+		{
+			// check(!bSupportsFences || EncoderFence);
+			
+			for (FAGXFence* FragFence : FragmentFences)
+			{
+				if (FragFence->NeedsWait(mtlpp::RenderStages::Fragment))
+				{
+					mtlpp::Fence FragmentFence = FragFence->Get(mtlpp::RenderStages::Fragment);
+					mtlpp::Fence FragInnerFence = METAL_DEBUG_OPTION(CommandList.GetCommandQueue().GetRuntimeDebuggingLevel() >= EAGXDebugLevelValidation ? mtlpp::Fence(((FAGXDebugFence*)FragmentFence.GetPtr()).Inner) :) FragmentFence;
+					
+					BlitCommandEncoder.WaitForFence(FragInnerFence);
+					METAL_DEBUG_LAYER(EAGXDebugLevelFastValidation, BlitEncoderDebug.AddWaitFence(FragmentFence));
+					FragFence->Wait(mtlpp::RenderStages::Fragment);
+				}
+			}
+			FragmentFences.Empty();
+			FenceResources.Empty();
+			FenceStage = mtlpp::RenderStages::Fragment;
+			
+			if (EncoderFence && EncoderFence->NeedsWrite(mtlpp::RenderStages::Fragment))
+			{
+				Fence = EncoderFence;
+			}
+			UpdateFence(EncoderFence);
+			
+#if METAL_DEBUG_OPTIONS
+			if (bSupportsFences && AGXSafeGetRuntimeDebuggingLevel() >= EAGXDebugLevelFastValidation && (!WaitCount || !UpdateCount))
+			{
+				UE_LOG(LogAGX, Error, TEXT("%s has incorrect fence waits (%u) vs. updates (%u)."), *FString(BlitCommandEncoder.GetLabel()), WaitCount, UpdateCount);
+				
+			}
+			WaitCount = 0;
+			UpdateCount = 0;
+#endif
+			BlitCommandEncoder.EndEncoding();
+			METAL_DEBUG_LAYER(EAGXDebugLevelFastValidation, BlitEncoderDebug.EndEncoder());
+			BlitCommandEncoder = nil;
+			EncoderFence = nullptr;
+		}
+	}
+	
+	for (uint32 Frequency = 0; Frequency < uint32(mtlpp::FunctionType::Kernel)+1; Frequency++)
+	{
+		for (uint32 i = 0; i < ML_MaxBuffers; i++)
+		{
+			ShaderBuffers[Frequency].Buffers[i] = nil;
+		}
+		FMemory::Memzero(ShaderBuffers[Frequency].Bytes);
+		FMemory::Memzero(ShaderBuffers[Frequency].Offsets);
+		FMemory::Memzero(ShaderBuffers[Frequency].Lengths);
+		FMemory::Memzero(ShaderBuffers[Frequency].Usage);
+		ShaderBuffers[Frequency].Bound = 0;
+	}
+    return Fence;
+}
+
+void FAGXCommandEncoder::InsertCommandBufferFence(FAGXCommandBufferFence& Fence, mtlpp::CommandBufferHandler Handler)
+{
+	check(CommandBuffer);
+	
+	Fence.CommandBufferFence = CommandBuffer.GetCompletionFence();
+	
+	if (Handler)
+	{
+		AddCompletionHandler(Handler);
+	}
+}
+
+void FAGXCommandEncoder::AddCompletionHandler(mtlpp::CommandBufferHandler Handler)
+{
+	check(Handler);
+	
+	mtlpp::CommandBufferHandler HeapHandler = Block_copy(Handler);
+	CompletionHandlers.Add(HeapHandler);
+	Block_release(HeapHandler);
+}
+
+void FAGXCommandEncoder::UpdateFence(FAGXFence* Fence)
+{
+	check(IsRenderCommandEncoderActive() || IsComputeCommandEncoderActive() || IsBlitCommandEncoderActive());
+	static bool bSupportsFences = CommandList.GetCommandQueue().SupportsFeature(EAGXFeaturesFences);
+	if ((bSupportsFences METAL_DEBUG_OPTION(|| CommandList.GetCommandQueue().GetRuntimeDebuggingLevel() >= EAGXDebugLevelValidation)) && Fence)
+	{
+		mtlpp::Fence VertexFence = Fence->Get(mtlpp::RenderStages::Vertex);
+		mtlpp::Fence InnerFence = METAL_DEBUG_OPTION(CommandList.GetCommandQueue().GetRuntimeDebuggingLevel() >= EAGXDebugLevelValidation ? mtlpp::Fence(((FAGXDebugFence*)VertexFence.GetPtr()).Inner) :) VertexFence;
+		if (RenderCommandEncoder)
+		{
+			mtlpp::Fence FragmentFence = Fence->Get(mtlpp::RenderStages::Fragment);
+			mtlpp::Fence FragInnerFence = METAL_DEBUG_OPTION(CommandList.GetCommandQueue().GetRuntimeDebuggingLevel() >= EAGXDebugLevelValidation ? mtlpp::Fence(((FAGXDebugFence*)FragmentFence.GetPtr()).Inner) :) FragmentFence;
+			
+			if (Fence->NeedsWrite(mtlpp::RenderStages::Vertex))
+			{
+				RenderCommandEncoder.UpdateFence(InnerFence, (mtlpp::RenderStages)(mtlpp::RenderStages::Vertex));
+				METAL_DEBUG_LAYER(EAGXDebugLevelFastValidation, RenderEncoderDebug.AddUpdateFence(VertexFence));
+				Fence->Write(mtlpp::RenderStages::Vertex);
+				METAL_DEBUG_LAYER(EAGXDebugLevelFastValidation, UpdateCount++);
+			}
+
+			if (Fence->NeedsWrite(mtlpp::RenderStages::Fragment))
+			{
+				RenderCommandEncoder.UpdateFence(FragInnerFence, (mtlpp::RenderStages)(mtlpp::RenderStages::Fragment));
+				METAL_DEBUG_LAYER(EAGXDebugLevelFastValidation, RenderEncoderDebug.AddUpdateFence(FragmentFence));
+				Fence->Write(mtlpp::RenderStages::Fragment);
+				METAL_DEBUG_LAYER(EAGXDebugLevelFastValidation, UpdateCount++);
+			}
+		}
+		else if (ComputeCommandEncoder && Fence->NeedsWrite(mtlpp::RenderStages::Vertex))
+		{
+			ComputeCommandEncoder.UpdateFence(InnerFence);
+			METAL_DEBUG_LAYER(EAGXDebugLevelFastValidation, ComputeEncoderDebug.AddUpdateFence(VertexFence));
+			Fence->Write(mtlpp::RenderStages::Vertex);
+			METAL_DEBUG_LAYER(EAGXDebugLevelFastValidation, UpdateCount++);
+		}
+		else if (BlitCommandEncoder && Fence->NeedsWrite(mtlpp::RenderStages::Vertex))
+		{
+			BlitCommandEncoder.UpdateFence(InnerFence);
+			METAL_DEBUG_LAYER(EAGXDebugLevelFastValidation, BlitEncoderDebug.AddUpdateFence(VertexFence));
+			Fence->Write(mtlpp::RenderStages::Vertex);
+			METAL_DEBUG_LAYER(EAGXDebugLevelFastValidation, UpdateCount++);
+		}
+	}
+}
+
+void FAGXCommandEncoder::WaitForFence(FAGXFence* Fence)
+{
+	check(IsRenderCommandEncoderActive() || IsComputeCommandEncoderActive() || IsBlitCommandEncoderActive());
+	static bool bSupportsFences = CommandList.GetCommandQueue().SupportsFeature(EAGXFeaturesFences);
+	if ((bSupportsFences METAL_DEBUG_OPTION(|| CommandList.GetCommandQueue().GetRuntimeDebuggingLevel() >= EAGXDebugLevelValidation)) && Fence)
+	{
+		if (Fence->NeedsWait(mtlpp::RenderStages::Vertex))
+		{
+			METAL_DEBUG_LAYER(EAGXDebugLevelFastValidation, WaitCount++);
+			
+			mtlpp::Fence VertexFence = Fence->Get(mtlpp::RenderStages::Vertex);
+			mtlpp::Fence InnerFence = METAL_DEBUG_OPTION(CommandList.GetCommandQueue().GetRuntimeDebuggingLevel() >= EAGXDebugLevelValidation ? mtlpp::Fence(((FAGXDebugFence*)VertexFence.GetPtr()).Inner) :) VertexFence;
+			if (RenderCommandEncoder)
+			{
+				RenderCommandEncoder.WaitForFence(InnerFence, (mtlpp::RenderStages)(mtlpp::RenderStages::Vertex));
+				METAL_DEBUG_LAYER(EAGXDebugLevelFastValidation, RenderEncoderDebug.AddWaitFence(VertexFence));
+				Fence->Wait(mtlpp::RenderStages::Vertex);
+			}
+			else if (ComputeCommandEncoder)
+			{
+				ComputeCommandEncoder.WaitForFence(InnerFence);
+				METAL_DEBUG_LAYER(EAGXDebugLevelFastValidation, ComputeEncoderDebug.AddWaitFence(VertexFence));
+				Fence->Wait(mtlpp::RenderStages::Vertex);
+			}
+			else if (BlitCommandEncoder)
+			{
+				BlitCommandEncoder.WaitForFence(InnerFence);
+				METAL_DEBUG_LAYER(EAGXDebugLevelFastValidation, BlitEncoderDebug.AddWaitFence(VertexFence));
+				Fence->Wait(mtlpp::RenderStages::Vertex);
+			}
+		}
+		if (Fence->NeedsWait(mtlpp::RenderStages::Fragment))
+		{
+			if (FenceStage == mtlpp::RenderStages::Vertex || BlitCommandEncoder)
+			{
+				mtlpp::Fence FragmentFence = Fence->Get(mtlpp::RenderStages::Fragment);
+				mtlpp::Fence FragInnerFence = METAL_DEBUG_OPTION(CommandList.GetCommandQueue().GetRuntimeDebuggingLevel() >= EAGXDebugLevelValidation ? mtlpp::Fence(((FAGXDebugFence*)FragmentFence.GetPtr()).Inner) :) FragmentFence;
+				if (RenderCommandEncoder)
+				{
+					RenderCommandEncoder.WaitForFence(FragInnerFence, (mtlpp::RenderStages)(mtlpp::RenderStages::Vertex));
+					METAL_DEBUG_LAYER(EAGXDebugLevelFastValidation, RenderEncoderDebug.AddWaitFence(FragmentFence));
+					Fence->Wait(mtlpp::RenderStages::Fragment);
+				}
+				else if (ComputeCommandEncoder)
+				{
+					ComputeCommandEncoder.WaitForFence(FragInnerFence);
+					METAL_DEBUG_LAYER(EAGXDebugLevelFastValidation, ComputeEncoderDebug.AddWaitFence(FragmentFence));
+					Fence->Wait(mtlpp::RenderStages::Fragment);
+				}
+				else if (BlitCommandEncoder)
+				{
+					BlitCommandEncoder.WaitForFence(FragInnerFence);
+					METAL_DEBUG_LAYER(EAGXDebugLevelFastValidation, BlitEncoderDebug.AddWaitFence(FragmentFence));
+					Fence->Wait(mtlpp::RenderStages::Fragment);
+				}
+				METAL_DEBUG_LAYER(EAGXDebugLevelFastValidation, WaitCount++);
+			}
+			else
+			{
+				METAL_DEBUG_LAYER(EAGXDebugLevelFastValidation, WaitCount++);
+				FragmentFences.Add(Fence);
+			}
+		}
+	}
+}
+
+void FAGXCommandEncoder::WaitAndUpdateFence(FAGXFence* Fence)
+{
+	check(IsRenderCommandEncoderActive() || IsComputeCommandEncoderActive() || IsBlitCommandEncoderActive());
+	static bool bSupportsFences = CommandList.GetCommandQueue().SupportsFeature(EAGXFeaturesFences);
+	if ((bSupportsFences METAL_DEBUG_OPTION(|| CommandList.GetCommandQueue().GetRuntimeDebuggingLevel() >= EAGXDebugLevelValidation)) && Fence)
+	{
+		METAL_DEBUG_LAYER(EAGXDebugLevelFastValidation, WaitCount++);
+		METAL_DEBUG_LAYER(EAGXDebugLevelFastValidation, UpdateCount++);
+		
+		mtlpp::Fence VertexFence = Fence->Get(mtlpp::RenderStages::Vertex);
+		mtlpp::Fence InnerFence = METAL_DEBUG_OPTION(CommandList.GetCommandQueue().GetRuntimeDebuggingLevel() >= EAGXDebugLevelValidation ? mtlpp::Fence(((FAGXDebugFence*)VertexFence.GetPtr()).Inner) :) VertexFence;
+		if (RenderCommandEncoder)
+		{
+			mtlpp::Fence FragmentFence = Fence->Get(mtlpp::RenderStages::Fragment);
+			mtlpp::Fence FragInnerFence = METAL_DEBUG_OPTION(CommandList.GetCommandQueue().GetRuntimeDebuggingLevel() >= EAGXDebugLevelValidation ? mtlpp::Fence(((FAGXDebugFence*)FragmentFence.GetPtr()).Inner) :) FragmentFence;
+			
+			METAL_DEBUG_LAYER(EAGXDebugLevelFastValidation, WaitCount++);
+			METAL_DEBUG_LAYER(EAGXDebugLevelFastValidation, UpdateCount++);
+			
+			RenderCommandEncoder.WaitForFence(InnerFence, (mtlpp::RenderStages)(mtlpp::RenderStages::Vertex));
+			METAL_DEBUG_LAYER(EAGXDebugLevelFastValidation, RenderEncoderDebug.AddWaitFence(VertexFence));
+			Fence->Wait(mtlpp::RenderStages::Vertex);
+			
+			RenderCommandEncoder.WaitForFence(FragInnerFence, (mtlpp::RenderStages)(mtlpp::RenderStages::Fragment));
+			METAL_DEBUG_LAYER(EAGXDebugLevelFastValidation, RenderEncoderDebug.AddWaitFence(FragmentFence));
+			Fence->Wait(mtlpp::RenderStages::Fragment);
+			
+			RenderCommandEncoder.UpdateFence(InnerFence, (mtlpp::RenderStages)(mtlpp::RenderStages::Fragment));
+			METAL_DEBUG_LAYER(EAGXDebugLevelFastValidation, RenderEncoderDebug.AddUpdateFence(VertexFence));
+			Fence->Write(mtlpp::RenderStages::Vertex);
+			
+			RenderCommandEncoder.UpdateFence(FragInnerFence, (mtlpp::RenderStages)(mtlpp::RenderStages::Vertex));
+			METAL_DEBUG_LAYER(EAGXDebugLevelFastValidation, RenderEncoderDebug.AddUpdateFence(FragmentFence));
+			Fence->Write(mtlpp::RenderStages::Fragment);
+		}
+		else if (ComputeCommandEncoder)
+		{
+			ComputeCommandEncoder.WaitForFence(InnerFence);
+			METAL_DEBUG_LAYER(EAGXDebugLevelFastValidation, ComputeEncoderDebug.AddWaitFence(VertexFence));
+			Fence->Wait(mtlpp::RenderStages::Vertex);
+			
+			ComputeCommandEncoder.UpdateFence(InnerFence);
+			METAL_DEBUG_LAYER(EAGXDebugLevelFastValidation, ComputeEncoderDebug.AddUpdateFence(VertexFence));
+			Fence->Write(mtlpp::RenderStages::Vertex);
+		}
+		else if (BlitCommandEncoder)
+		{
+			BlitCommandEncoder.WaitForFence(InnerFence);
+			METAL_DEBUG_LAYER(EAGXDebugLevelFastValidation, BlitEncoderDebug.AddWaitFence(VertexFence));
+			Fence->Wait(mtlpp::RenderStages::Vertex);
+			
+			BlitCommandEncoder.UpdateFence(InnerFence);
+			METAL_DEBUG_LAYER(EAGXDebugLevelFastValidation, BlitEncoderDebug.AddUpdateFence(VertexFence));
+			Fence->Write(mtlpp::RenderStages::Vertex);
+		}
+	}
+}
+
+#pragma mark - Public Debug Support -
+
+void FAGXCommandEncoder::InsertDebugSignpost(ns::String const& String)
+{
+	if (String)
+	{
+		if (RenderCommandEncoder)
+		{
+			RenderCommandEncoder.InsertDebugSignpost(String);
+			METAL_DEBUG_LAYER(EAGXDebugLevelFastValidation, RenderEncoderDebug.InsertDebugSignpost(String));
+		}
+		else if (ParallelRenderCommandEncoder && !IsParallel())
+		{
+			ParallelRenderCommandEncoder.InsertDebugSignpost(String);
+			METAL_DEBUG_LAYER(EAGXDebugLevelFastValidation, ParallelRenderCommandEncoder.InsertDebugSignpost(String));
+		}
+		else if (ComputeCommandEncoder)
+		{
+			ComputeCommandEncoder.InsertDebugSignpost(String);
+			METAL_DEBUG_LAYER(EAGXDebugLevelFastValidation, ComputeEncoderDebug.InsertDebugSignpost(String));
+		}
+		else if (BlitCommandEncoder)
+		{
+			BlitCommandEncoder.InsertDebugSignpost(String);
+			METAL_DEBUG_LAYER(EAGXDebugLevelFastValidation, BlitEncoderDebug.InsertDebugSignpost(String));
+		}
+	}
+}
+
+void FAGXCommandEncoder::PushDebugGroup(ns::String const& String)
+{
+	if (String)
+	{
+		[DebugGroups addObject:String];
+		if (RenderCommandEncoder)
+		{
+			RenderCommandEncoder.PushDebugGroup(String);
+			METAL_DEBUG_LAYER(EAGXDebugLevelFastValidation, RenderEncoderDebug.PushDebugGroup(String));
+		}
+		else if (ParallelRenderCommandEncoder && !IsParallel())
+		{
+			ParallelRenderCommandEncoder.PushDebugGroup(String);
+			METAL_DEBUG_LAYER(EAGXDebugLevelFastValidation, ParallelRenderCommandEncoder.PushDebugGroup(String));
+		}
+		else if (ComputeCommandEncoder)
+		{
+			ComputeCommandEncoder.PushDebugGroup(String);
+			METAL_DEBUG_LAYER(EAGXDebugLevelFastValidation, ComputeEncoderDebug.PushDebugGroup(String));
+		}
+		else if (BlitCommandEncoder)
+		{
+			BlitCommandEncoder.PushDebugGroup(String);
+			METAL_DEBUG_LAYER(EAGXDebugLevelFastValidation, BlitEncoderDebug.PushDebugGroup(String));
+		}
+	}
+}
+
+void FAGXCommandEncoder::PopDebugGroup(void)
+{
+	if (DebugGroups.count > 0)
+	{
+		[DebugGroups removeLastObject];
+		if (RenderCommandEncoder)
+		{
+			RenderCommandEncoder.PopDebugGroup();
+			METAL_DEBUG_LAYER(EAGXDebugLevelFastValidation, RenderEncoderDebug.PopDebugGroup());
+		}
+		else if (ParallelRenderCommandEncoder && !IsParallel())
+		{
+			ParallelRenderCommandEncoder.PopDebugGroup();
+			METAL_DEBUG_LAYER(EAGXDebugLevelFastValidation, ParallelRenderCommandEncoder.PopDebugGroup());
+		}
+		else if (ComputeCommandEncoder)
+		{
+			ComputeCommandEncoder.PopDebugGroup();
+			METAL_DEBUG_LAYER(EAGXDebugLevelFastValidation, ComputeEncoderDebug.PopDebugGroup());
+		}
+		else if (BlitCommandEncoder)
+		{
+			BlitCommandEncoder.PopDebugGroup();
+			METAL_DEBUG_LAYER(EAGXDebugLevelFastValidation, BlitEncoderDebug.PopDebugGroup());
+		}
+	}
+}
+
+FAGXCommandBufferMarkers& FAGXCommandEncoder::GetMarkers(void)
+{
+	return CommandBufferMarkers;
+}
+
+#if ENABLE_METAL_GPUPROFILE
+FAGXCommandBufferStats* FAGXCommandEncoder::GetCommandBufferStats(void)
+{
+	return CommandBufferStats;
+}
+#endif
+
+#pragma mark - Public Render State Mutators -
+
+void FAGXCommandEncoder::SetRenderPassDescriptor(mtlpp::RenderPassDescriptor RenderPass)
+{
+	check(IsRenderCommandEncoderActive() == false && IsComputeCommandEncoderActive() == false && IsBlitCommandEncoderActive() == false);
+	check(RenderPass);
+	
+	if(RenderPass.GetPtr() != RenderPassDesc.GetPtr())
+	{
+		AGXSafeReleaseMetalRenderPassDescriptor(RenderPassDesc);
+		RenderPassDesc = RenderPass;
+		{
+			for (uint32 i = 0; i < MaxSimultaneousRenderTargets; i++)
+			{
+				ColorStoreActions[i] = mtlpp::StoreAction::Unknown;
+			}
+			DepthStoreAction = mtlpp::StoreAction::Unknown;
+			StencilStoreAction = mtlpp::StoreAction::Unknown;
+		}
+	}
+	check(RenderPassDesc);
+	
+	for (uint32 Frequency = 0; Frequency < uint32(mtlpp::FunctionType::Kernel)+1; Frequency++)
+	{
+		for (uint32 i = 0; i < ML_MaxBuffers; i++)
+		{
+			ShaderBuffers[Frequency].Buffers[i] = nil;
+		}
+		FMemory::Memzero(ShaderBuffers[Frequency].Bytes);
+		FMemory::Memzero(ShaderBuffers[Frequency].Offsets);
+		FMemory::Memzero(ShaderBuffers[Frequency].Lengths);
+		FMemory::Memzero(ShaderBuffers[Frequency].Usage);
+		ShaderBuffers[Frequency].Bound = 0;
+	}
+}
+
+void FAGXCommandEncoder::SetRenderPassStoreActions(mtlpp::StoreAction const* const ColorStore, mtlpp::StoreAction const DepthStore, mtlpp::StoreAction const StencilStore)
+{
+	check(RenderPassDesc);
+	{
+		for (uint32 i = 0; i < MaxSimultaneousRenderTargets; i++)
+		{
+			ColorStoreActions[i] = ColorStore[i];
+		}
+		DepthStoreAction = DepthStore;
+		StencilStoreAction = StencilStore;
+	}
+}
+
+void FAGXCommandEncoder::SetRenderPipelineState(FAGXShaderPipeline* PipelineState)
+{
+	check (RenderCommandEncoder);
+	{
+		METAL_DEBUG_LAYER(EAGXDebugLevelFastValidation, RenderEncoderDebug.SetPipeline(PipelineState));
+		RenderCommandEncoder.SetRenderPipelineState(PipelineState->RenderPipelineState);
+	}
+}
+
+void FAGXCommandEncoder::SetViewport(mtlpp::Viewport const Viewport[], uint32 NumActive)
+{
+	check(RenderCommandEncoder);
+	check(NumActive >= 1 && NumActive < ML_MaxViewports);
+	if (NumActive == 1)
+	{
+		RenderCommandEncoder.SetViewport(Viewport[0]);
+	}
+#if PLATFORM_MAC
+	else
+	{
+		check(FAGXCommandQueue::SupportsFeature(EAGXFeaturesMultipleViewports));
+		RenderCommandEncoder.SetViewports(Viewport, NumActive);
+	}
+#endif
+}
+
+void FAGXCommandEncoder::SetFrontFacingWinding(mtlpp::Winding const InFrontFacingWinding)
+{
+    check (RenderCommandEncoder);
+	{
+		RenderCommandEncoder.SetFrontFacingWinding(InFrontFacingWinding);
+	}
+}
+
+void FAGXCommandEncoder::SetCullMode(mtlpp::CullMode const InCullMode)
+{
+    check (RenderCommandEncoder);
+	{
+		RenderCommandEncoder.SetCullMode(InCullMode);
+	}
+}
+
+void FAGXCommandEncoder::SetDepthBias(float const InDepthBias, float const InSlopeScale, float const InClamp)
+{
+    check (RenderCommandEncoder);
+	{
+		RenderCommandEncoder.SetDepthBias(InDepthBias, InSlopeScale, InClamp);
+	}
+}
+
+void FAGXCommandEncoder::SetScissorRect(mtlpp::ScissorRect const Rect[], uint32 NumActive)
+{
+    check(RenderCommandEncoder);
+	check(NumActive >= 1 && NumActive < ML_MaxViewports);
+	if (NumActive == 1)
+	{
+		RenderCommandEncoder.SetScissorRect(Rect[0]);
+	}
+#if PLATFORM_MAC
+	else
+	{
+		check(FAGXCommandQueue::SupportsFeature(EAGXFeaturesMultipleViewports));
+		RenderCommandEncoder.SetScissorRects(Rect, NumActive);
+	}
+#endif
+}
+
+void FAGXCommandEncoder::SetTriangleFillMode(mtlpp::TriangleFillMode const InFillMode)
+{
+    check(RenderCommandEncoder);
+	{
+		RenderCommandEncoder.SetTriangleFillMode(InFillMode);
+	}
+}
+
+void FAGXCommandEncoder::SetBlendColor(float const Red, float const Green, float const Blue, float const Alpha)
+{
+	check(RenderCommandEncoder);
+	{
+		RenderCommandEncoder.SetBlendColor(Red, Green, Blue, Alpha);
+	}
+}
+
+void FAGXCommandEncoder::SetDepthStencilState(mtlpp::DepthStencilState const& InDepthStencilState)
+{
+    check (RenderCommandEncoder);
+	{
+		RenderCommandEncoder.SetDepthStencilState(InDepthStencilState);
+		METAL_DEBUG_LAYER(EAGXDebugLevelFastValidation, RenderEncoderDebug.SetDepthStencilState(InDepthStencilState));
+	}
+}
+
+void FAGXCommandEncoder::SetStencilReferenceValue(uint32 const ReferenceValue)
+{
+    check (RenderCommandEncoder);
+	{
+		RenderCommandEncoder.SetStencilReferenceValue(ReferenceValue);
+	}
+}
+
+void FAGXCommandEncoder::SetVisibilityResultMode(mtlpp::VisibilityResultMode const Mode, NSUInteger const Offset)
+{
+    check (RenderCommandEncoder);
+	{
+		check(Mode == mtlpp::VisibilityResultMode::Disabled || RenderPassDesc.GetVisibilityResultBuffer());
+		RenderCommandEncoder.SetVisibilityResultMode(Mode, Offset);
+	}
+}
+	
+#pragma mark - Public Shader Resource Mutators -
+
+void FAGXCommandEncoder::SetShaderBuffer(mtlpp::FunctionType const FunctionType, FAGXBuffer const& Buffer, NSUInteger const Offset, NSUInteger const Length, NSUInteger index, mtlpp::ResourceUsage const Usage, EPixelFormat const Format)
+{
+	check(index < ML_MaxBuffers);
+    if(GetAGXDeviceContext().SupportsFeature(EAGXFeaturesSetBufferOffset) && Buffer && (ShaderBuffers[uint32(FunctionType)].Bound & (1 << index)) && ShaderBuffers[uint32(FunctionType)].Buffers[index] == Buffer)
+    {
+		if (FunctionType == mtlpp::FunctionType::Vertex || FunctionType == mtlpp::FunctionType::Kernel)
+		{
+			FenceResource(Buffer);
+		}
+		SetShaderBufferOffset(FunctionType, Offset, Length, index);
+		ShaderBuffers[uint32(FunctionType)].Lengths[(index*2)+1] = GAGXBufferFormats[Format].DataFormat;
+		ShaderBuffers[uint32(FunctionType)].Usage[index] = Usage;
+	}
+    else
+    {
+		if(Buffer)
+		{
+			ShaderBuffers[uint32(FunctionType)].Bound |= (1 << index);
+		}
+		else
+		{
+			ShaderBuffers[uint32(FunctionType)].Bound &= ~(1 << index);
+		}
+		ShaderBuffers[uint32(FunctionType)].Buffers[index] = Buffer;
+		ShaderBuffers[uint32(FunctionType)].Bytes[index] = nil;
+		ShaderBuffers[uint32(FunctionType)].Offsets[index] = Offset;
+		ShaderBuffers[uint32(FunctionType)].Usage[index] = Usage;
+		ShaderBuffers[uint32(FunctionType)].Lengths[index*2] = Length;
+		ShaderBuffers[uint32(FunctionType)].Lengths[(index*2)+1] = GAGXBufferFormats[Format].DataFormat;
+		
+		SetShaderBufferInternal(FunctionType, index);
+    }
+    
+	if (Buffer)
+	{
+		BufferBindingHistory.Add(ns::AutoReleased<FAGXBuffer>(Buffer));
+	}
+}
+
+void FAGXCommandEncoder::SetShaderData(mtlpp::FunctionType const FunctionType, FAGXBufferData* Data, NSUInteger const Offset, NSUInteger const Index, EPixelFormat const Format)
+{
+	check(Index < ML_MaxBuffers);
+	
+	if(Data)
+	{
+		ShaderBuffers[uint32(FunctionType)].Bound |= (1 << Index);
+	}
+	else
+	{
+		ShaderBuffers[uint32(FunctionType)].Bound &= ~(1 << Index);
+	}
+	
+	ShaderBuffers[uint32(FunctionType)].Buffers[Index] = nil;
+	ShaderBuffers[uint32(FunctionType)].Bytes[Index] = Data;
+	ShaderBuffers[uint32(FunctionType)].Offsets[Index] = Offset;
+	ShaderBuffers[uint32(FunctionType)].Usage[Index] = mtlpp::ResourceUsage::Read;
+	ShaderBuffers[uint32(FunctionType)].Lengths[Index*2] = Data ? (Data->Len - Offset) : 0;
+	ShaderBuffers[uint32(FunctionType)].Lengths[(Index*2)+1] = GAGXBufferFormats[Format].DataFormat;
+	
+	SetShaderBufferInternal(FunctionType, Index);
+}
+
+void FAGXCommandEncoder::SetShaderBytes(mtlpp::FunctionType const FunctionType, uint8 const* Bytes, NSUInteger const Length, NSUInteger const Index)
+{
+	check(Index < ML_MaxBuffers);
+	
+	if(Bytes && Length)
+	{
+		ShaderBuffers[uint32(FunctionType)].Bound |= (1 << Index);
+
+		if (bSupportsMetalFeaturesSetBytes)
+		{
+			switch (FunctionType)
+			{
+				case mtlpp::FunctionType::Vertex:
+					check(RenderCommandEncoder);
+					METAL_DEBUG_LAYER(EAGXDebugLevelFastValidation, RenderEncoderDebug.SetBytes(EAGXShaderVertex, Bytes, Length, Index));
+					RenderCommandEncoder.SetVertexData(Bytes, Length, Index);
+					break;
+				case mtlpp::FunctionType::Fragment:
+					check(RenderCommandEncoder);
+					METAL_DEBUG_LAYER(EAGXDebugLevelFastValidation, RenderEncoderDebug.SetBytes(EAGXShaderFragment, Bytes, Length, Index));
+					RenderCommandEncoder.SetFragmentData(Bytes, Length, Index);
+					break;
+				case mtlpp::FunctionType::Kernel:
+					check(ComputeCommandEncoder);
+					METAL_DEBUG_LAYER(EAGXDebugLevelFastValidation, ComputeEncoderDebug.SetBytes(Bytes, Length, Index));
+					ComputeCommandEncoder.SetBytes(Bytes, Length, Index);
+					break;
+				default:
+					check(false);
+					break;
+			}
+			
+			ShaderBuffers[uint32(FunctionType)].Buffers[Index] = nil;
+		}
+		else
+		{
+			FAGXBuffer Buffer = RingBuffer.NewBuffer(Length, BufferOffsetAlignment);
+			FMemory::Memcpy(((uint8*)Buffer.GetContents()), Bytes, Length);
+			ShaderBuffers[uint32(FunctionType)].Buffers[Index] = Buffer;
+		}
+		ShaderBuffers[uint32(FunctionType)].Bytes[Index] = nil;
+		ShaderBuffers[uint32(FunctionType)].Offsets[Index] = 0;
+		ShaderBuffers[uint32(FunctionType)].Usage[Index] = mtlpp::ResourceUsage::Read;
+		ShaderBuffers[uint32(FunctionType)].Lengths[Index*2] = Length;
+		ShaderBuffers[uint32(FunctionType)].Lengths[(Index*2)+1] = GAGXBufferFormats[PF_Unknown].DataFormat;
+	}
+	else
+	{
+		ShaderBuffers[uint32(FunctionType)].Bound &= ~(1 << Index);
+		
+		ShaderBuffers[uint32(FunctionType)].Buffers[Index] = nil;
+		ShaderBuffers[uint32(FunctionType)].Bytes[Index] = nil;
+		ShaderBuffers[uint32(FunctionType)].Offsets[Index] = 0;
+		ShaderBuffers[uint32(FunctionType)].Usage[Index] = mtlpp::ResourceUsage(0);
+		ShaderBuffers[uint32(FunctionType)].Lengths[Index*2] = 0;
+		ShaderBuffers[uint32(FunctionType)].Lengths[(Index*2)+1] = GAGXBufferFormats[PF_Unknown].DataFormat;
+	}
+	
+	SetShaderBufferInternal(FunctionType, Index);
+}
+
+void FAGXCommandEncoder::SetShaderBufferOffset(mtlpp::FunctionType FunctionType, NSUInteger const Offset, NSUInteger const Length, NSUInteger const index)
+{
+	check(index < ML_MaxBuffers);
+    checkf(ShaderBuffers[uint32(FunctionType)].Buffers[index] && (ShaderBuffers[uint32(FunctionType)].Bound & (1 << index)), TEXT("Buffer must already be bound"));
+	check(GetAGXDeviceContext().SupportsFeature(EAGXFeaturesSetBufferOffset));
+	ShaderBuffers[uint32(FunctionType)].Offsets[index] = Offset;
+	ShaderBuffers[uint32(FunctionType)].Lengths[index*2] = Length;
+	ShaderBuffers[uint32(FunctionType)].Lengths[(index*2)+1] = GAGXBufferFormats[PF_Unknown].DataFormat;
+	switch (FunctionType)
+	{
+		case mtlpp::FunctionType::Vertex:
+			check (RenderCommandEncoder);
+			RenderCommandEncoder.SetVertexBufferOffset(Offset + ShaderBuffers[uint32(FunctionType)].Buffers[index].GetOffset(), index);
+			METAL_DEBUG_LAYER(EAGXDebugLevelFastValidation, RenderEncoderDebug.SetBufferOffset(EAGXShaderVertex, Offset + ShaderBuffers[uint32(FunctionType)].Buffers[index].GetOffset(), index));
+			break;
+		case mtlpp::FunctionType::Fragment:
+			check(RenderCommandEncoder);
+			RenderCommandEncoder.SetFragmentBufferOffset(Offset + ShaderBuffers[uint32(FunctionType)].Buffers[index].GetOffset(), index);
+			METAL_DEBUG_LAYER(EAGXDebugLevelFastValidation, RenderEncoderDebug.SetBufferOffset(EAGXShaderFragment, Offset + ShaderBuffers[uint32(FunctionType)].Buffers[index].GetOffset(), index));
+			break;
+		case mtlpp::FunctionType::Kernel:
+			check (ComputeCommandEncoder);
+			ComputeCommandEncoder.SetBufferOffset(Offset + ShaderBuffers[uint32(FunctionType)].Buffers[index].GetOffset(), index);
+			METAL_DEBUG_LAYER(EAGXDebugLevelFastValidation, ComputeEncoderDebug.SetBufferOffset(Offset + ShaderBuffers[uint32(FunctionType)].Buffers[index].GetOffset(), index));
+			break;
+		default:
+			check(false);
+			break;
+	}
+}
+
+void FAGXCommandEncoder::SetShaderTexture(mtlpp::FunctionType FunctionType, FAGXTexture const& Texture, NSUInteger index, mtlpp::ResourceUsage Usage)
+{
+	check(index < ML_MaxTextures);
+	switch (FunctionType)
+	{
+		case mtlpp::FunctionType::Vertex:
+			check (RenderCommandEncoder);
+			FenceResource(Texture);
+			// MTLPP_VALIDATE(mtlpp::RenderCommandEncoder, RenderCommandEncoder, AGXSafeGetRuntimeDebuggingLevel() >= EAGXDebugLevelValidation, UseResource(Texture, mtlpp::ResourceUsage::Read));
+			METAL_DEBUG_LAYER(EAGXDebugLevelFastValidation, RenderEncoderDebug.SetTexture(EAGXShaderVertex, Texture, index));
+			RenderCommandEncoder.SetVertexTexture(Texture, index);
+			break;
+		case mtlpp::FunctionType::Fragment:
+			check(RenderCommandEncoder);
+			// MTLPP_VALIDATE(mtlpp::RenderCommandEncoder, RenderCommandEncoder, AGXSafeGetRuntimeDebuggingLevel() >= EAGXDebugLevelValidation, UseResource(Texture, mtlpp::ResourceUsage::Read));
+			METAL_DEBUG_LAYER(EAGXDebugLevelFastValidation, RenderEncoderDebug.SetTexture(EAGXShaderFragment, Texture, index));
+			RenderCommandEncoder.SetFragmentTexture(Texture, index);
+			break;
+		case mtlpp::FunctionType::Kernel:
+			check (ComputeCommandEncoder);
+			FenceResource(Texture);
+			// MTLPP_VALIDATE(mtlpp::ComputeCommandEncoder, ComputeCommandEncoder, AGXSafeGetRuntimeDebuggingLevel() >= EAGXDebugLevelValidation, UseResource(Texture, mtlpp::ResourceUsage::Read));
+			METAL_DEBUG_LAYER(EAGXDebugLevelFastValidation, ComputeEncoderDebug.SetTexture(Texture, index));
+			ComputeCommandEncoder.SetTexture(Texture, index);
+			break;
+		default:
+			check(false);
+			break;
+	}
+	
+	if (Texture)
+	{
+		uint8 Swizzle[4] = {0,0,0,0};
+		assert(sizeof(Swizzle) == sizeof(uint32));
+		if (Texture.GetPixelFormat() == mtlpp::PixelFormat::X32_Stencil8
+#if PLATFORM_MAC
+		 ||	Texture.GetPixelFormat() == mtlpp::PixelFormat::X24_Stencil8
+#endif
+		)
+		{
+			Swizzle[0] = Swizzle[1] = Swizzle[2] = Swizzle[3] = 1;
+		}
+		FMemory::Memcpy(&ShaderBuffers[uint32(FunctionType)].Lengths[(ML_MaxBuffers*2)+(index*2)], Swizzle, sizeof(Swizzle));
+		ShaderBuffers[uint32(FunctionType)].Lengths[(ML_MaxBuffers*2)+(index*2)+1] = 0;
+		TextureBindingHistory.Add(ns::AutoReleased<FAGXTexture>(Texture));
+	}
+}
+
+void FAGXCommandEncoder::SetShaderSamplerState(mtlpp::FunctionType FunctionType, mtlpp::SamplerState const& Sampler, NSUInteger index)
+{
+	check(index < ML_MaxSamplers);
+	switch (FunctionType)
+	{
+		case mtlpp::FunctionType::Vertex:
+       		check (RenderCommandEncoder);
+			METAL_DEBUG_LAYER(EAGXDebugLevelFastValidation, RenderEncoderDebug.SetSamplerState(EAGXShaderVertex, Sampler, index));
+			RenderCommandEncoder.SetVertexSamplerState(Sampler, index);
+			break;
+		case mtlpp::FunctionType::Fragment:
+			check (RenderCommandEncoder);
+			METAL_DEBUG_LAYER(EAGXDebugLevelFastValidation, RenderEncoderDebug.SetSamplerState(EAGXShaderFragment, Sampler, index));
+			RenderCommandEncoder.SetFragmentSamplerState(Sampler, index);
+			break;
+		case mtlpp::FunctionType::Kernel:
+			check (ComputeCommandEncoder);
+			METAL_DEBUG_LAYER(EAGXDebugLevelFastValidation, ComputeEncoderDebug.SetSamplerState(Sampler, index));
+			ComputeCommandEncoder.SetSamplerState(Sampler, index);
+			break;
+		default:
+			check(false);
+			break;
+	}
+}
+
+void FAGXCommandEncoder::SetShaderSideTable(mtlpp::FunctionType const FunctionType, NSUInteger const Index)
+{
+	if (Index < ML_MaxBuffers)
+	{
+		SetShaderData(FunctionType, ShaderBuffers[uint32(FunctionType)].SideTable, 0, Index);
+	}
+}
+
+void FAGXCommandEncoder::UseIndirectArgumentResource(FAGXTexture const& Texture, mtlpp::ResourceUsage const Usage)
+{
+	FenceResource(Texture);
+	UseResource(Texture, Usage);
+	TextureBindingHistory.Add(ns::AutoReleased<FAGXTexture>(Texture));
+}
+
+void FAGXCommandEncoder::UseIndirectArgumentResource(FAGXBuffer const& Buffer, mtlpp::ResourceUsage const Usage)
+{
+	FenceResource(Buffer);
+	UseResource(Buffer, Usage);
+	BufferBindingHistory.Add(ns::AutoReleased<FAGXBuffer>(Buffer));
+}
+
+void FAGXCommandEncoder::TransitionResources(mtlpp::Resource const& Resource)
+{
+	TransitionedResources.Add(Resource.GetPtr());
+}
+
+#pragma mark - Public Compute State Mutators -
+
+void FAGXCommandEncoder::SetComputePipelineState(FAGXShaderPipeline* State)
+{
+	check (ComputeCommandEncoder);
+	{
+		METAL_DEBUG_LAYER(EAGXDebugLevelFastValidation, ComputeEncoderDebug.SetPipeline(State));
+		ComputeCommandEncoder.SetComputePipelineState(State->ComputePipelineState);
+	}
+}
+
+#pragma mark - Public Ring-Buffer Accessor -
+	
+FAGXSubBufferRing& FAGXCommandEncoder::GetRingBuffer(void)
+{
+	return RingBuffer;
+}
+
+#pragma mark - Public Resource query Access -
+
+bool FAGXCommandEncoder::HasTextureBindingHistory(FAGXTexture const& Texture) const
+{
+	return TextureBindingHistory.Contains(ns::AutoReleased<FAGXTexture>(Texture));
+}
+
+bool FAGXCommandEncoder::HasBufferBindingHistory(FAGXBuffer const& Buffer) const
+{
+	return BufferBindingHistory.Contains(ns::AutoReleased<FAGXBuffer>(Buffer));
+}
+
+#pragma mark - Private Functions -
+
+void FAGXCommandEncoder::FenceResource(mtlpp::Texture const& Resource)
+{
+	mtlpp::Resource::Type Res = Resource.GetPtr();
+	ns::AutoReleased<mtlpp::Texture> Parent = Resource.GetParentTexture();
+	ns::AutoReleased<mtlpp::Buffer> Buffer = Resource.GetBuffer();
+	if (Parent)
+	{
+		Res = Parent.GetPtr();
+	}
+	else if (Buffer)
+	{
+		Res = Buffer.GetPtr();
+	}
+	if (FenceStage == mtlpp::RenderStages::Vertex || FenceResources.Contains(Res))
+	{
+		FenceStage = mtlpp::RenderStages::Vertex;
+		
+		for (FAGXFence* FragFence : FragmentFences)
+		{
+			if (FragFence->NeedsWait(mtlpp::RenderStages::Fragment))
+			{
+				mtlpp::Fence FragmentFence = FragFence->Get(mtlpp::RenderStages::Fragment);
+				mtlpp::Fence FragInnerFence = METAL_DEBUG_OPTION(CommandList.GetCommandQueue().GetRuntimeDebuggingLevel() >= EAGXDebugLevelValidation ? mtlpp::Fence(((FAGXDebugFence*)FragmentFence.GetPtr()).Inner) :) FragmentFence;
+				
+				if (RenderCommandEncoder)
+				{
+					RenderCommandEncoder.WaitForFence(FragInnerFence, (mtlpp::RenderStages)(mtlpp::RenderStages::Vertex));
+					METAL_DEBUG_LAYER(EAGXDebugLevelFastValidation, RenderEncoderDebug.AddWaitFence(FragmentFence));
+					FragFence->Wait(mtlpp::RenderStages::Fragment);
+				}
+				else if (ComputeCommandEncoder)
+				{
+					ComputeCommandEncoder.WaitForFence(FragInnerFence);
+					METAL_DEBUG_LAYER(EAGXDebugLevelFastValidation, ComputeEncoderDebug.AddWaitFence(FragmentFence));
+					FragFence->Wait(mtlpp::RenderStages::Fragment);
+				}
+				else if (BlitCommandEncoder)
+				{
+					BlitCommandEncoder.WaitForFence(FragInnerFence);
+					METAL_DEBUG_LAYER(EAGXDebugLevelFastValidation, BlitEncoderDebug.AddWaitFence(FragmentFence));
+					FragFence->Wait(mtlpp::RenderStages::Fragment);
+				}
+				METAL_DEBUG_LAYER(EAGXDebugLevelFastValidation, WaitCount++);
+			}
+		}
+		FragmentFences.Empty();
+	}
+}
+
+void FAGXCommandEncoder::FenceResource(mtlpp::Buffer const& Resource)
+{
+	mtlpp::Resource::Type Res = Resource.GetPtr();
+	if (FenceStage == mtlpp::RenderStages::Vertex || FenceResources.Contains(Res))
+	{
+		FenceStage = mtlpp::RenderStages::Vertex;
+		
+		for (FAGXFence* FragFence : FragmentFences)
+		{
+			if (FragFence->NeedsWait(mtlpp::RenderStages::Fragment))
+			{
+				mtlpp::Fence FragmentFence = FragFence->Get(mtlpp::RenderStages::Fragment);
+				mtlpp::Fence FragInnerFence = METAL_DEBUG_OPTION(CommandList.GetCommandQueue().GetRuntimeDebuggingLevel() >= EAGXDebugLevelValidation ? mtlpp::Fence(((FAGXDebugFence*)FragmentFence.GetPtr()).Inner) :) FragmentFence;
+				
+				if (RenderCommandEncoder)
+				{
+					RenderCommandEncoder.WaitForFence(FragInnerFence, (mtlpp::RenderStages)(mtlpp::RenderStages::Vertex));
+					METAL_DEBUG_LAYER(EAGXDebugLevelFastValidation, RenderEncoderDebug.AddWaitFence(FragmentFence));
+					FragFence->Wait(mtlpp::RenderStages::Fragment);
+				}
+				else if (ComputeCommandEncoder)
+				{
+					ComputeCommandEncoder.WaitForFence(FragInnerFence);
+					METAL_DEBUG_LAYER(EAGXDebugLevelFastValidation, ComputeEncoderDebug.AddWaitFence(FragmentFence));
+					FragFence->Wait(mtlpp::RenderStages::Fragment);
+				}
+				else if (BlitCommandEncoder)
+				{
+					BlitCommandEncoder.WaitForFence(FragInnerFence);
+					METAL_DEBUG_LAYER(EAGXDebugLevelFastValidation, BlitEncoderDebug.AddWaitFence(FragmentFence));
+					FragFence->Wait(mtlpp::RenderStages::Fragment);
+				}
+			}
+		}
+		FragmentFences.Empty();
+	}
+}
+
+void FAGXCommandEncoder::UseResource(mtlpp::Resource const& Resource, mtlpp::ResourceUsage const Usage)
+{
+	static bool UseResourceAvailable = FAGXCommandQueue::SupportsFeature(EAGXFeaturesIABs);
+	if (UseResourceAvailable || AGXSafeGetRuntimeDebuggingLevel() >= EAGXDebugLevelValidation)
+	{
+		mtlpp::ResourceUsage Current = ResourceUsage.FindRef(Resource.GetPtr());
+		if (Current != Usage)
+		{
+			ResourceUsage.Add(Resource.GetPtr(), Usage);
+			if (RenderCommandEncoder)
+			{
+				MTLPP_VALIDATE(mtlpp::RenderCommandEncoder, RenderCommandEncoder, AGXSafeGetRuntimeDebuggingLevel() >= EAGXDebugLevelValidation, UseResource(Resource, Usage));
+			}
+			else if (ComputeCommandEncoder)
+			{
+				MTLPP_VALIDATE(mtlpp::ComputeCommandEncoder, ComputeCommandEncoder, AGXSafeGetRuntimeDebuggingLevel() >= EAGXDebugLevelValidation, UseResource(Resource, Usage));
+			}
+		}
+	}
+}
+
+void FAGXCommandEncoder::SetShaderBufferInternal(mtlpp::FunctionType Function, uint32 Index)
+{
+	NSUInteger Offset = ShaderBuffers[uint32(Function)].Offsets[Index];
+	mtlpp::ResourceUsage Usage = ShaderBuffers[uint32(Function)].Usage[Index];
+	bool bBufferHasBytes = ShaderBuffers[uint32(Function)].Bytes[Index] != nil;
+	if (!ShaderBuffers[uint32(Function)].Buffers[Index] && bBufferHasBytes && !bSupportsMetalFeaturesSetBytes)
+	{
+		uint8 const* Bytes = (((uint8 const*)ShaderBuffers[uint32(Function)].Bytes[Index]->Data) + ShaderBuffers[uint32(Function)].Offsets[Index]);
+		uint32 Len = ShaderBuffers[uint32(Function)].Bytes[Index]->Len - ShaderBuffers[uint32(Function)].Offsets[Index];
+		
+		Offset = 0;
+		ShaderBuffers[uint32(Function)].Buffers[Index] = RingBuffer.NewBuffer(Len, BufferOffsetAlignment);
+		
+		FMemory::Memcpy(((uint8*)ShaderBuffers[uint32(Function)].Buffers[Index].GetContents()) + Offset, Bytes, Len);
+	}
+	
+	ns::AutoReleased<FAGXBuffer>& Buffer = ShaderBuffers[uint32(Function)].Buffers[Index];
+	if (Buffer)
+	{
+#if METAL_DEBUG_OPTIONS
+		if(CommandList.GetCommandQueue().GetRuntimeDebuggingLevel() >= EAGXDebugLevelValidation)
+		{
+			ActiveBuffers.Add(Buffer);
+		}
+#endif
+        
+		switch (Function)
+		{
+			case mtlpp::FunctionType::Vertex:
+				ShaderBuffers[uint32(Function)].Bound |= (1 << Index);
+				check(RenderCommandEncoder);
+				FenceResource(Buffer);
+				// MTLPP_VALIDATE(mtlpp::RenderCommandEncoder, RenderCommandEncoder, AGXSafeGetRuntimeDebuggingLevel() >= EAGXDebugLevelValidation, UseResource(Buffer, mtlpp::ResourceUsage::Read));
+				METAL_DEBUG_LAYER(EAGXDebugLevelFastValidation, RenderEncoderDebug.SetBuffer(EAGXShaderVertex, Buffer, Offset, Index));
+				RenderCommandEncoder.SetVertexBuffer(Buffer, Offset, Index);
+				break;
+			case mtlpp::FunctionType::Fragment:
+				ShaderBuffers[uint32(Function)].Bound |= (1 << Index);
+				check(RenderCommandEncoder);
+				// MTLPP_VALIDATE(mtlpp::RenderCommandEncoder, RenderCommandEncoder, AGXSafeGetRuntimeDebuggingLevel() >= EAGXDebugLevelValidation, UseResource(Buffer, mtlpp::ResourceUsage::Read));
+				METAL_DEBUG_LAYER(EAGXDebugLevelFastValidation, RenderEncoderDebug.SetBuffer(EAGXShaderFragment, Buffer, Offset, Index));
+				RenderCommandEncoder.SetFragmentBuffer(Buffer, Offset, Index);
+				break;
+			case mtlpp::FunctionType::Kernel:
+				ShaderBuffers[uint32(Function)].Bound |= (1 << Index);
+				check(ComputeCommandEncoder);
+				FenceResource(Buffer);
+				// MTLPP_VALIDATE(mtlpp::ComputeCommandEncoder, ComputeCommandEncoder, AGXSafeGetRuntimeDebuggingLevel() >= EAGXDebugLevelValidation, UseResource(Buffer, mtlpp::ResourceUsage::Read));
+				METAL_DEBUG_LAYER(EAGXDebugLevelFastValidation, ComputeEncoderDebug.SetBuffer(Buffer, Offset, Index));
+				ComputeCommandEncoder.SetBuffer(Buffer, Offset, Index);
+				break;
+			default:
+				check(false);
+				break;
+		}
+		
+		if (Buffer.IsSingleUse())
+		{
+			ShaderBuffers[uint32(Function)].Usage[Index] = mtlpp::ResourceUsage(0);
+			ShaderBuffers[uint32(Function)].Offsets[Index] = 0;
+			ShaderBuffers[uint32(Function)].Buffers[Index] = nil;
+			ShaderBuffers[uint32(Function)].Bound &= ~(1 << Index);
+		}
+	}
+	else if (bBufferHasBytes && bSupportsMetalFeaturesSetBytes)
+	{
+		uint8 const* Bytes = (((uint8 const*)ShaderBuffers[uint32(Function)].Bytes[Index]->Data) + ShaderBuffers[uint32(Function)].Offsets[Index]);
+		uint32 Len = ShaderBuffers[uint32(Function)].Bytes[Index]->Len - ShaderBuffers[uint32(Function)].Offsets[Index];
+		
+		switch (Function)
+		{
+			case mtlpp::FunctionType::Vertex:
+				ShaderBuffers[uint32(Function)].Bound |= (1 << Index);
+				check(RenderCommandEncoder);
+				METAL_DEBUG_LAYER(EAGXDebugLevelFastValidation, RenderEncoderDebug.SetBytes(EAGXShaderVertex, Bytes, Len, Index));
+				RenderCommandEncoder.SetVertexData(Bytes, Len, Index);
+				break;
+			case mtlpp::FunctionType::Fragment:
+				ShaderBuffers[uint32(Function)].Bound |= (1 << Index);
+				check(RenderCommandEncoder);
+				METAL_DEBUG_LAYER(EAGXDebugLevelFastValidation, RenderEncoderDebug.SetBytes(EAGXShaderFragment, Bytes, Len, Index));
+				RenderCommandEncoder.SetFragmentData(Bytes, Len, Index);
+				break;
+			case mtlpp::FunctionType::Kernel:
+				ShaderBuffers[uint32(Function)].Bound |= (1 << Index);
+				check(ComputeCommandEncoder);
+				METAL_DEBUG_LAYER(EAGXDebugLevelFastValidation, ComputeEncoderDebug.SetBytes(Bytes, Len, Index));
+				ComputeCommandEncoder.SetBytes(Bytes, Len, Index);
+				break;
+			default:
+				check(false);
+				break;
+		}
+	}
+}
diff --git a/Engine/Source/Runtime/Apple/AGXRHI/Private/AGXCommandEncoder.h b/Engine/Source/Runtime/Apple/AGXRHI/Private/AGXCommandEncoder.h
new file mode 100644
--- /dev/null
+++ b/Engine/Source/Runtime/Apple/AGXRHI/Private/AGXCommandEncoder.h
@@ -0,0 +1,588 @@
+// Copyright Epic Games, Inc. All Rights Reserved.
+
+#pragma once
+
+#include <Metal/Metal.h>
+#include "AGXBuffer.h"
+#include "AGXDebugCommandEncoder.h"
+#include "AGXBlitCommandEncoder.h"
+#include "AGXComputeCommandEncoder.h"
+#include "AGXRenderCommandEncoder.h"
+#include "AGXFence.h"
+#include "AGXProfiler.h"
+#include "command_buffer.hpp"
+
+class FAGXCommandList;
+class FAGXCommandQueue;
+class FAGXGraphicsPipelineState;
+struct FAGXCommandBufferFence;
+
+/**
+ * Enumeration for submission hints to avoid unclear bool values.
+ */
+enum EAGXSubmitFlags
+{
+	/** No submission flags. */
+	EAGXSubmitFlagsNone = 0,
+	/** Create the next command buffer. */
+	EAGXSubmitFlagsCreateCommandBuffer = 1 << 0,
+	/** Wait on the submitted command buffer. */
+	EAGXSubmitFlagsWaitOnCommandBuffer = 1 << 1,
+	/** Break a single logical command-buffer into parts to keep the GPU active. */
+	EAGXSubmitFlagsBreakCommandBuffer = 1 << 2,
+	/** Submit the prologue command-buffer only, leave the current command-buffer active.  */
+	EAGXSubmitFlagsAsyncCommandBuffer = 1 << 3,
+	/** Submit and reset all the cached encoder state.  */
+	EAGXSubmitFlagsResetState = 1 << 4,
+	/** Force submission even if the command-buffer is empty. */
+	EAGXSubmitFlagsForce = 1 << 5,
+	/** Indicates this is the final command buffer in a frame. */
+	EAGXSubmitFlagsLastCommandBuffer = 1 << 6,
+};
+
+@class FAGXCommandBufferDebug;
+
+struct FAGXCommandData
+{
+	enum class Type
+	{
+		DrawPrimitive,
+		DrawPrimitiveIndexed,
+		DrawPrimitivePatch,
+		DrawPrimitiveIndirect,
+		DrawPrimitiveIndexedIndirect,
+		Dispatch,
+		DispatchIndirect,
+		Num,
+	};
+	struct DispatchIndirectArgs
+	{
+		id ArgumentBuffer;
+		NSUInteger ArgumentOffset;
+	};
+	Type CommandType;
+	union
+	{
+		mtlpp::DrawPrimitivesIndirectArguments Draw;
+		mtlpp::DrawIndexedPrimitivesIndirectArguments DrawIndexed;
+		mtlpp::DrawPatchIndirectArguments DrawPatch;
+		MTLDispatchThreadgroupsIndirectArguments Dispatch;
+		DispatchIndirectArgs DispatchIndirect;
+	};
+	FString ToString() const;
+};
+
+struct FAGXCommandDebug
+{
+    uint32 CmdBufIndex;
+	uint32 Encoder;
+	uint32 Index;
+	FAGXGraphicsPipelineState* PSO;
+	FAGXComputeShader* ComputeShader;
+	FAGXCommandData Data;
+};
+
+class FAGXCommandBufferMarkers : public ns::Object<FAGXCommandBufferDebug*, ns::CallingConvention::ObjectiveC>
+{
+public:
+	FAGXCommandBufferMarkers(void);
+	FAGXCommandBufferMarkers(mtlpp::CommandBuffer& CmdBuf);
+	FAGXCommandBufferMarkers(FAGXCommandBufferDebug* CmdBuf);
+	
+	void AllocateContexts(uint32 NumContexts);
+	uint32 AddCommand(uint32 CmdBufIndex, uint32 Encoder, uint32 ContextIndex, FAGXBuffer& DebugBuffer, FAGXGraphicsPipelineState* PSO, FAGXComputeShader* ComputeShader, FAGXCommandData& Data);
+	TArray<FAGXCommandDebug>* GetCommands(uint32 Context);
+	ns::AutoReleased<FAGXBuffer> GetDebugBuffer(uint32 ContextIndex);
+	uint32 NumContexts() const;
+	uint32 GetIndex() const;
+	
+	static FAGXCommandBufferMarkers Get(mtlpp::CommandBuffer const& CmdBuf);
+	
+	static char const* kTableAssociationKey;
+};
+
+/**
+ * EAGXCommandEncoderType:
+ *   EAGXCommandEncoderCurrent: The primary encoder that is used for draw calls & dispatches
+ *   EAGXCommandEncoderPrologue: A secondary encoder that is used for blits & dispatches that setup resources & state for the current encoder.
+ */
+enum EAGXCommandEncoderType
+{
+	EAGXCommandEncoderCurrent,
+	EAGXCommandEncoderPrologue
+};
+
+/**
+ * FAGXCommandEncoder:
+ *	Wraps the details of switching between different command encoders on the command-buffer, allowing for restoration of the render encoder if needed.
+ * 	UE expects the API to serialise commands in-order, but Metal expects applications to work with command-buffers directly so we need to implement
+ *	the RHI semantics by switching between encoder types. This class hides the ugly details. Eventually it might be possible to move some of the operations
+ *	into pre- & post- command-buffers so that we avoid encoder switches but that will take changes to the RHI and high-level code too, so it won't happen soon.
+ */
+class FAGXCommandEncoder
+{
+public:
+#pragma mark - Public C++ Boilerplate -
+
+	/** Default constructor */
+	FAGXCommandEncoder(FAGXCommandList& CmdList, EAGXCommandEncoderType Type);
+	
+	/** Destructor */
+	~FAGXCommandEncoder(void);
+	
+	/** Reset cached state for reuse */
+	void Reset(void);
+	
+	/** Reset cached state for reuse while in rendering */
+	void ResetLive(void);
+	
+#pragma mark - Public Command Buffer Mutators -
+
+	/**
+	 * Start encoding to CommandBuffer. It is an error to call this with any outstanding command encoders or current command buffer.
+	 * Instead call EndEncoding & CommitCommandBuffer before calling this.
+	 */
+	void StartCommandBuffer(void);
+	
+	/**
+	 * Commit the existing command buffer if there is one & optionally waiting for completion, if there isn't a current command buffer this is a no-op.
+	 * @param Flags Flags to control commit behaviour.
+ 	 */
+	void CommitCommandBuffer(uint32 const Flags);
+
+#pragma mark - Public Command Buffer Accessors -
+	
+	/** @returns the current command buffer */
+	mtlpp::CommandBuffer& GetCommandBuffer() { return CommandBuffer; }
+
+	/** @returns the current command buffer */
+	mtlpp::CommandBuffer const& GetCommandBuffer() const { return CommandBuffer; }
+	
+	/** @returns the monotonically incremented command buffer index */
+	uint32 GetCommandBufferIndex() const { return CmdBufIndex; }
+
+#pragma mark - Public Command Encoder Accessors -
+	
+	/** @returns True if and only if there is an active parallel render command encoder, otherwise false. */
+	bool IsParallelRenderCommandEncoderActive(void) const;
+	
+	/** @returns True if and only if there is an active render command encoder, otherwise false. */
+	bool IsRenderCommandEncoderActive(void) const;
+	
+	/** @returns True if and only if there is an active compute command encoder, otherwise false. */
+	bool IsComputeCommandEncoderActive(void) const;
+	
+	/** @returns True if and only if there is an active blit command encoder, otherwise false. */
+	bool IsBlitCommandEncoderActive(void) const;
+	
+	/**
+	 * True iff the command-encoder submits immediately to the command-queue, false if it performs any buffering.
+	 * @returns True iff the command-list submits immediately to the command-queue, false if it performs any buffering.
+	 */
+	bool IsImmediate(void) const;
+	
+	/**
+	 * True iff the command-encoder encodes only to a child of a parallel render command encoder, false if it is standalone.
+	 * @returns True iff the command-encoder encodes only to a child of a parallel render command encoder, false if it is standalone.
+	 */
+	bool IsParallel(void) const;
+
+	/** @returns True if and only if there is valid render pass descriptor set on the encoder, otherwise false. */
+	bool IsRenderPassDescriptorValid(void) const;
+	
+	/** @returns The current render pass descriptor. */
+	mtlpp::RenderPassDescriptor const& GetRenderPassDescriptor(void) const;
+	
+	/** @returns The active render command encoder or nil if there isn't one. */
+	mtlpp::ParallelRenderCommandEncoder& GetParallelRenderCommandEncoder(void);
+	
+	/** @returns The child render command encoder of the current parallel render encoder for Index. */
+	mtlpp::RenderCommandEncoder& GetChildRenderCommandEncoder(uint32 Index);
+
+	/** @returns The active render command encoder or nil if there isn't one. */
+	mtlpp::RenderCommandEncoder& GetRenderCommandEncoder(void);
+	
+	/** @returns The active compute command encoder or nil if there isn't one. */
+	mtlpp::ComputeCommandEncoder& GetComputeCommandEncoder(void);
+	
+	/** @returns The active blit command encoder or nil if there isn't one. */
+	mtlpp::BlitCommandEncoder& GetBlitCommandEncoder(void);
+	
+	/** @returns The MTLFence for the current encoder or nil if there isn't one. */
+	TRefCountPtr<FAGXFence> const& GetEncoderFence(void) const;
+	
+	/** @returns The number of encoded passes in the command buffer. */
+	uint32 NumEncodedPasses(void) const { return EncoderNum; }
+	
+#pragma mark - Public Command Encoder Mutators -
+
+	/**
+ 	 * Begins encoding rendering commands into the current command buffer. No other encoder may be active & the mtlpp::RenderPassDescriptor must previously have been set.
+	 * @param NumChildren The number of child render-encoders to create. 
+	 */
+	void BeginParallelRenderCommandEncoding(uint32 NumChildren);
+
+	/**
+ 	 * Begins encoding rendering commands into the current command buffer. No other encoder may be active & the mtlpp::RenderPassDescriptor must previously have been set.
+	 */
+	void BeginRenderCommandEncoding(void);
+	
+	/** Begins encoding compute commands into the current command buffer. No other encoder may be active. */
+	void BeginComputeCommandEncoding(mtlpp::DispatchType Type = mtlpp::DispatchType::Serial);
+	
+	/** Begins encoding blit commands into the current command buffer. No other encoder may be active. */
+	void BeginBlitCommandEncoding(void);
+	
+	/** Declare that all command generation from this encoder is complete, and detach from the MTLCommandBuffer if there is an encoder active or does nothing if there isn't. */
+	TRefCountPtr<FAGXFence> EndEncoding(void);
+	
+	/** Initialises a fence for the current command-buffer optionally adding a command-buffer completion handler to the command-buffer */
+	void InsertCommandBufferFence(FAGXCommandBufferFence& Fence, mtlpp::CommandBufferHandler Handler);
+	
+	/** Adds a command-buffer completion handler to the command-buffer */
+	void AddCompletionHandler(mtlpp::CommandBufferHandler Handler);
+	
+	/** Update the event to capture all GPU work so far enqueued by this encoder. */
+	void UpdateFence(FAGXFence* Fence);
+	
+	/** Prevent further GPU work until the event is reached. */
+	void WaitForFence(FAGXFence* Fence);
+	
+	/**
+	 * Prevent further GPU work until the event is reached and then kick the fence again so future encoders may wait on it.
+	 * Unlike the functions above this is not guarded and the caller must ensure proper matching waits & updates.
+	 */
+	void WaitAndUpdateFence(FAGXFence* Fence);
+
+#pragma mark - Public Debug Support -
+	
+	/*
+	 * Inserts a debug string into the command buffer.  This does not change any API behavior, but can be useful when debugging.
+	 * @param string The name of the signpost. 
+	 */
+	void InsertDebugSignpost(ns::String const& String);
+	
+	/*
+	 * Push a new named string onto a stack of string labels.
+	 * @param string The name of the debug group. 
+	 */
+	void PushDebugGroup(ns::String const& String);
+	
+	/* Pop the latest named string off of the stack. */
+	void PopDebugGroup(void);
+	
+	/** Returns the object that records debug command markers into the command-buffer. */
+	FAGXCommandBufferMarkers& GetMarkers(void);
+	
+#if ENABLE_METAL_GPUPROFILE
+	/* Get the command-buffer stats object. */
+	FAGXCommandBufferStats* GetCommandBufferStats(void);
+#endif
+
+#if MTLPP_CONFIG_VALIDATE && METAL_DEBUG_OPTIONS
+	/** @returns The active render command encoder or nil if there isn't one. */
+	FAGXCommandBufferDebugging& GetCommandBufferDebugging(void) { return CommandBufferDebug; } 
+	
+	/** @returns The active render command encoder or nil if there isn't one. */
+	FAGXRenderCommandEncoderDebugging& GetRenderCommandEncoderDebugging(void) { return RenderEncoderDebug; } 
+	
+	/** @returns The active compute command encoder or nil if there isn't one. */
+	FAGXComputeCommandEncoderDebugging& GetComputeCommandEncoderDebugging(void) { return ComputeEncoderDebug; }
+	
+	/** @returns The active blit command encoder or nil if there isn't one. */
+	FAGXBlitCommandEncoderDebugging& GetBlitCommandEncoderDebugging(void) { return BlitEncoderDebug; }
+	
+	/** @returns The active blit command encoder or nil if there isn't one. */
+	FAGXParallelRenderCommandEncoderDebugging& GetParallelRenderCommandEncoderDebugging(void) { return ParallelEncoderDebug; }
+#endif
+	
+#pragma mark - Public Render State Mutators -
+
+	/**
+	 * Set the render pass descriptor - no encoder may be active when this function is called.
+	 * @param RenderPass The render pass descriptor to set. May be nil.
+	 */
+	void SetRenderPassDescriptor(mtlpp::RenderPassDescriptor RenderPass);
+	
+	/**
+	 * Set the render pass store actions, call after SetRenderPassDescriptor but before EndEncoding.
+	 * @param ColorStore The store actions for color targets.
+	 * @param DepthStore The store actions for the depth buffer - use mtlpp::StoreAction::Unknown if no depth-buffer bound.
+	 * @param StencilStore The store actions for the stencil buffer - use mtlpp::StoreAction::Unknown if no stencil-buffer bound.
+	 */
+	void SetRenderPassStoreActions(mtlpp::StoreAction const* const ColorStore, mtlpp::StoreAction const DepthStore, mtlpp::StoreAction const StencilStore);
+	
+	/*
+	 * Sets the current render pipeline state object.
+	 * @param PipelineState The pipeline state to set. Must not be nil.
+	 */
+	void SetRenderPipelineState(FAGXShaderPipeline* const PipelineState);
+	
+	/*
+	 * Set the viewport, which is used to transform vertexes from normalized device coordinates to window coordinates.  Fragments that lie outside of the viewport are clipped, and optionally clamped for fragments outside of znear/zfar.
+	 * @param Viewport The array of viewport dimensions to use.
+	 * @param NumActive The number of active viewport dimensions to use.
+	 */
+	void SetViewport(mtlpp::Viewport const Viewport[], uint32 NumActive);
+	
+	/*
+	 * The winding order of front-facing primitives.
+	 * @param FrontFacingWinding The front face winding.
+	 */
+	void SetFrontFacingWinding(mtlpp::Winding const FrontFacingWinding);
+	
+	/*
+	 * Controls if primitives are culled when front facing, back facing, or not culled at all.
+	 * @param CullMode The cull mode.
+	 */
+	void SetCullMode(mtlpp::CullMode const CullMode);
+	
+	/*
+	 * Depth Bias.
+	 * @param DepthBias The depth-bias value.
+	 * @param SlopeScale The slope-scale to apply.
+	 * @param Clamp The value to clamp to.
+	 */
+	void SetDepthBias(float const DepthBias, float const SlopeScale, float const Clamp);
+	
+	/*
+	 * Specifies a rectangle for a fragment scissor test.  All fragments outside of this rectangle are discarded.
+	 * @param Rect The array of scissor rect dimensions.
+	 * @param NumActive The number of active scissor rect dimensions.
+	 */
+	void SetScissorRect(mtlpp::ScissorRect const Rect[], uint32 NumActive);
+	
+	/*
+	 * Set how to rasterize triangle and triangle strip primitives.
+	 * @param FillMode The fill mode.
+	 */
+	void SetTriangleFillMode(mtlpp::TriangleFillMode const FillMode);
+	
+	/*
+	 * Set the constant blend color used across all blending on all render targets
+	 * @param Red The value for the red channel in 0-1.
+	 * @param Green The value for the green channel in 0-1.
+	 * @param Blue The value for the blue channel in 0-1.
+	 * @param Alpha The value for the alpha channel in 0-1.
+	 */
+	void SetBlendColor(float const Red, float const Green, float const Blue, float const Alpha);
+	
+	/*
+	 * Set the DepthStencil state object.
+	 * @param DepthStencilState The depth-stencil state, must not be nil.
+	 */
+	void SetDepthStencilState(mtlpp::DepthStencilState const& DepthStencilState);
+	
+	/*
+	 * Set the stencil reference value for both the back and front stencil buffers.
+	 * @param ReferenceValue The stencil ref value to use.
+	 */
+	void SetStencilReferenceValue(uint32 const ReferenceValue);
+	
+	/*
+	 * Monitor if samples pass the depth and stencil tests.
+	 * @param Mode Controls if the counter is disabled or moniters passing samples.
+	 * @param Offset The offset relative to the occlusion query buffer provided when the command encoder was created.  offset must be a multiple of 8.
+	 */
+	void SetVisibilityResultMode(mtlpp::VisibilityResultMode const Mode, NSUInteger const Offset);
+	
+#pragma mark - Public Shader Resource Mutators -
+	
+	/*
+	 * Set a global buffer for the specified shader frequency at the given bind point index.
+	 * @param FunctionType The shader function to modify.
+	 * @param Buffer The buffer to bind or nil to clear.
+	 * @param Offset The offset in the buffer or 0 when Buffer is nil.
+	 * @param Length The data length - caller is responsible for accounting for non-zero Offset.
+	 * @param Index The index to modify.
+	 * @param Usage The resource usage mask.
+	 * @param Format The Pixel format to reinterpret the resource as.
+	 */
+	void SetShaderBuffer(mtlpp::FunctionType const FunctionType, FAGXBuffer const& Buffer, NSUInteger const Offset, NSUInteger const Length, NSUInteger const Index, mtlpp::ResourceUsage const Usage, EPixelFormat const Format = PF_Unknown);
+	
+	/*
+	 * Set an FAGXBufferData to the specified shader frequency at the given bind point index.
+	 * @param FunctionType The shader function to modify.
+	 * @param Data The data to bind or nullptr to clear.
+	 * @param Offset The offset in the buffer or 0 when Buffer is nil.
+	 * @param Index The index to modify.
+	 * @param Format The pixel format to reinterpret the resource as.
+	 */
+	void SetShaderData(mtlpp::FunctionType const FunctionType, FAGXBufferData* Data, NSUInteger const Offset, NSUInteger const Index, EPixelFormat const Format = PF_Unknown);
+	
+	/*
+	 * Set bytes to the specified shader frequency at the given bind point index.
+	 * @param FunctionType The shader function to modify.
+	 * @param Bytes The data to bind or nullptr to clear.
+	 * @param Length The length of the buffer or 0 when Bytes is nil.
+	 * @param Index The index to modify.
+	 */
+	void SetShaderBytes(mtlpp::FunctionType const FunctionType, uint8 const* Bytes, NSUInteger const Length, NSUInteger const Index);
+	
+	/*
+	 * Set a global texture for the specified shader frequency at the given bind point index.
+	 * @param FunctionType The shader function to modify.
+	 * @param Texture The texture to bind or nil to clear.
+	 * @param Index The index to modify.
+	 * @param Usage The resource usage mask.
+	 */
+	void SetShaderTexture(mtlpp::FunctionType const FunctionType, FAGXTexture const& Texture, NSUInteger const Index, mtlpp::ResourceUsage const Usage);
+	
+	/*
+	 * Set a global sampler for the specified shader frequency at the given bind point index.
+	 * @param FunctionType The shader function to modify.
+	 * @param Sampler The sampler state to bind or nil to clear.
+	 * @param Index The index to modify.
+	 */
+	void SetShaderSamplerState(mtlpp::FunctionType const FunctionType, mtlpp::SamplerState const& Sampler, NSUInteger const Index);
+	
+	/*
+	 * Set the shader side-table data for FunctionType at Index.
+	 * @param FunctionType The shader function to modify.
+	 * @param Index The index to bind data to.
+	 */
+	void SetShaderSideTable(mtlpp::FunctionType const FunctionType, NSUInteger const Index);
+	
+    /*
+     * Inform the driver that an indirect argument resource will be used and what the usage mode will be.
+     * @param Texture The texture that will be used as an indirect argument.
+     * @param Usage The usage mode for the texture.
+     */
+	void UseIndirectArgumentResource(FAGXTexture const& Texture, mtlpp::ResourceUsage const Usage);
+    
+    /*
+     * Inform the driver that an indirect argument resource will be used and what the usage mode will be.
+     * @param Buffer The buffer that will be used as an indirect argument.
+     * @param Usage The usage mode for the texture.
+     */
+	void UseIndirectArgumentResource(FAGXBuffer const& Buffer, mtlpp::ResourceUsage const Usage);
+
+	/*
+	 * Transition resource so that we can barrier fragment->vertex stages.
+	 * @param Resource The resource we are going to make readable having been written.
+	 */
+	void TransitionResources(mtlpp::Resource const& Resource);
+	
+#pragma mark - Public Compute State Mutators -
+	
+	/*
+	 * Set the compute pipeline state that will be used.
+	 * @param State The state to set - must not be nil.
+	 */
+	void SetComputePipelineState(FAGXShaderPipeline* State);
+
+#pragma mark - Public Ring-Buffer Accessor -
+	
+	/*
+	 * Get the internal ring-buffer used for temporary allocations.
+	 * @returns The temporary allocation buffer for this command-encoder.
+	 */
+	FAGXSubBufferRing& GetRingBuffer(void);
+	
+#pragma mark - Public Resource query Access -
+
+	/*
+	 * Returns True if the Resource has been bound to a command encoder, otherwise false.  History will be cleared after a commit operation
+	 * @returns True if the Resource has been bound to a command encoder, otherwise false.
+	 */
+	bool HasTextureBindingHistory(FAGXTexture const& Texture) const;
+	bool HasBufferBindingHistory(FAGXBuffer const& Buffer) const;
+	
+private:
+#pragma mark - Private Functions -
+	/*
+	 * Set the offset for the buffer bound on the specified shader frequency at the given bind point index.
+	 * @param FunctionType The shader function to modify.
+	 * @param Offset The offset in the buffer or 0 when Buffer is nil.
+	 * @param Length The data length - caller is responsible for accounting for non-zero Offset.
+	 * @param Index The index to modify.
+	 */
+	void SetShaderBufferOffset(mtlpp::FunctionType const FunctionType, NSUInteger const Offset, NSUInteger const Length, NSUInteger const Index);
+	
+	void SetShaderBufferInternal(mtlpp::FunctionType Function, uint32 Index);
+	
+	void FenceResource(mtlpp::Texture const& Resource);
+	void FenceResource(mtlpp::Buffer const& Resource);
+
+	void UseResource(mtlpp::Resource const& Resource, mtlpp::ResourceUsage const Usage);
+	
+#pragma mark - Private Type Declarations -
+private:
+    /** A structure of arrays for the current buffer binding settings. */
+    struct FAGXBufferBindings
+    {
+		/** Side-table wrapper object to allow us to use Set*Bytes. */
+		FAGXBufferData* SideTable;
+        /** The bound buffers or nil. */
+		ns::AutoReleased<FAGXBuffer> Buffers[ML_MaxBuffers];
+        /** The bound buffers or nil. */
+        FAGXBufferData* Bytes[ML_MaxBuffers];
+        /** The bound buffer offsets or 0. */
+        NSUInteger Offsets[ML_MaxBuffers];
+		/** The usage mask for the bound resource or 0 */
+		mtlpp::ResourceUsage Usage[ML_MaxBuffers];
+		/** The bound buffer lengths */
+		uint32 Lengths[(ML_MaxBuffers*2) + (ML_MaxTextures*2)];
+        /** A bitmask for which buffers were bound by the application where a bit value of 1 is bound and 0 is unbound. */
+        uint32 Bound;
+	};
+	
+#pragma mark - Private Member Variables -
+	FAGXCommandList& CommandList;
+
+    // Cache Queue feature
+    bool bSupportsMetalFeaturesSetBytes;
+    
+	FAGXBufferBindings ShaderBuffers[int(mtlpp::FunctionType::Kernel)+1];
+	
+	mtlpp::StoreAction ColorStoreActions[MaxSimultaneousRenderTargets];
+	mtlpp::StoreAction DepthStoreAction;
+	mtlpp::StoreAction StencilStoreAction;
+	
+	FAGXSubBufferRing RingBuffer;
+	
+	mtlpp::RenderPassDescriptor RenderPassDesc;
+	
+	mtlpp::CommandBuffer CommandBuffer;
+	mtlpp::ParallelRenderCommandEncoder ParallelRenderCommandEncoder;
+	mtlpp::RenderCommandEncoder RenderCommandEncoder;
+	mtlpp::ComputeCommandEncoder ComputeCommandEncoder;
+	mtlpp::BlitCommandEncoder BlitCommandEncoder;
+	TArray<mtlpp::RenderCommandEncoder> ChildRenderCommandEncoders;
+	FAGXCommandBufferMarkers CommandBufferMarkers;
+	
+	METAL_DEBUG_ONLY(FAGXCommandBufferDebugging CommandBufferDebug);
+	METAL_DEBUG_ONLY(FAGXRenderCommandEncoderDebugging RenderEncoderDebug);
+	METAL_DEBUG_ONLY(FAGXComputeCommandEncoderDebugging ComputeEncoderDebug);
+	METAL_DEBUG_ONLY(FAGXBlitCommandEncoderDebugging BlitEncoderDebug);
+	METAL_DEBUG_ONLY(FAGXParallelRenderCommandEncoderDebugging ParallelEncoderDebug);
+	
+	TRefCountPtr<FAGXFence> EncoderFence;
+#if ENABLE_METAL_GPUPROFILE
+	FAGXCommandBufferStats* CommandBufferStats;
+#endif
+#if METAL_DEBUG_OPTIONS
+	uint8 WaitCount;
+	uint8 UpdateCount;
+#endif
+
+	TArray<ns::Object<mtlpp::CommandBufferHandler>> CompletionHandlers;
+	NSMutableArray* DebugGroups;
+    
+    TSet<ns::AutoReleased<FAGXBuffer>> ActiveBuffers;
+    
+	TSet<ns::AutoReleased<FAGXBuffer>> BufferBindingHistory;
+	TSet<ns::AutoReleased<FAGXTexture>> TextureBindingHistory;
+	
+	TMap<mtlpp::Resource::Type, mtlpp::ResourceUsage> ResourceUsage;
+	
+	TSet<mtlpp::Resource::Type> TransitionedResources;
+	TSet<mtlpp::Resource::Type> FenceResources;
+	
+	TSet<TRefCountPtr<FAGXFence>> FragmentFences;
+	
+	mtlpp::RenderStages FenceStage;
+	uint32 EncoderNum;
+	uint32 CmdBufIndex;
+	EAGXCommandEncoderType Type;
+};
