Initial import of AGXRHI

diff --git a/Engine/Source/Runtime/Apple/AGXRHI/Private/AGXCommandList.cpp b/Engine/Source/Runtime/Apple/AGXRHI/Private/AGXCommandList.cpp
new file mode 100644
--- /dev/null
+++ b/Engine/Source/Runtime/Apple/AGXRHI/Private/AGXCommandList.cpp
@@ -0,0 +1,374 @@
+// Copyright Epic Games, Inc. All Rights Reserved.
+
+/*=============================================================================
+	AGXCommandList.cpp: AGX RHI command buffer list wrapper.
+=============================================================================*/
+
+#include "AGXRHIPrivate.h"
+#include "AGXShaderTypes.h"
+#include "AGXGraphicsPipelineState.h"
+#include "AGXCommandList.h"
+#include "AGXCommandQueue.h"
+#include "AGXProfiler.h"
+#include "AGXCommandBuffer.h"
+#include "AGXCommandEncoder.h"
+#include "ns.hpp"
+
+#pragma mark - Public C++ Boilerplate -
+
+#if PLATFORM_IOS
+extern bool GIsSuspended;
+#endif
+extern int32 GAGXDebugOpsCount;
+
+FAGXCommandList::FAGXCommandList(FAGXCommandQueue& InCommandQueue, bool const bInImmediate)
+: CommandQueue(InCommandQueue)
+, Index(0)
+, Num(0)
+, bImmediate(bInImmediate)
+{
+}
+
+FAGXCommandList::~FAGXCommandList(void)
+{
+}
+	
+#pragma mark - Public Command List Mutators -
+
+extern CORE_API bool GIsGPUCrashed;
+static void ReportMetalCommandBufferFailure(mtlpp::CommandBuffer const& CompletedBuffer, TCHAR const* ErrorType, bool bDoCheck=true)
+{
+	GIsGPUCrashed = true;
+	
+	NSString* Label = CompletedBuffer.GetLabel();
+	int32 Code = CompletedBuffer.GetError().GetCode();
+	NSString* Domain = CompletedBuffer.GetError().GetDomain();
+	NSString* ErrorDesc = CompletedBuffer.GetError().GetLocalizedDescription();
+	NSString* FailureDesc = CompletedBuffer.GetError().GetLocalizedFailureReason();
+	NSString* RecoveryDesc = CompletedBuffer.GetError().GetLocalizedRecoverySuggestion();
+	
+	FString LabelString = Label ? FString(Label) : FString(TEXT("Unknown"));
+	FString DomainString = Domain ? FString(Domain) : FString(TEXT("Unknown"));
+	FString ErrorString = ErrorDesc ? FString(ErrorDesc) : FString(TEXT("Unknown"));
+	FString FailureString = FailureDesc ? FString(FailureDesc) : FString(TEXT("Unknown"));
+	FString RecoveryString = RecoveryDesc ? FString(RecoveryDesc) : FString(TEXT("Unknown"));
+	
+	NSString* Desc = CompletedBuffer.GetPtr().debugDescription;
+	UE_LOG(LogAGX, Warning, TEXT("%s"), *FString(Desc));
+	
+#if PLATFORM_IOS
+    if (bDoCheck && !GIsSuspended && !GIsRenderingThreadSuspended)
+#endif
+    {
+		FAGXCommandBufferMarkers Markers = FAGXCommandBufferMarkers::Get(CompletedBuffer);
+		if (Markers)
+		{
+			uint32 CommandBufferIndex = Markers.GetIndex();
+			FAGXDebugInfo BrokenDraw;
+			BrokenDraw.EncoderIndex = UINT32_MAX;
+			BrokenDraw.CommandIndex = UINT32_MAX;
+			BrokenDraw.CmdBuffIndex = UINT32_MAX;
+			BrokenDraw.ContextIndex = UINT32_MAX;
+            BrokenDraw.CommandBuffer = 0;
+			uint32 BrokenContext = 0;
+			bool bFoundBrokenDraw = false;
+			for (uint32 i = 0; (CommandBufferIndex != ~0u) && !bFoundBrokenDraw && i < Markers.NumContexts(); i++)
+			{
+				ns::AutoReleased<FAGXBuffer> DebugBuffer = Markers.GetDebugBuffer(i);
+				TArray<FAGXCommandDebug>* Commands = Markers.GetCommands(i);
+				if (DebugBuffer && Commands && Commands->Num())
+				{
+					uint32 DebugLength = DebugBuffer.GetLength();
+					uint32 DebugCount = DebugLength / sizeof(FAGXDebugInfo);
+					check(DebugCount >= 1);
+					
+					FAGXDebugInfo* DebugArray = (FAGXDebugInfo*)DebugBuffer.GetContents();
+					FAGXDebugInfo CurrentDraw = DebugArray[0];
+					
+					// On TBDR we find the disjoint where one tile has progressed further than another
+					// We find the earliest failure in the probably vain hope that this is the tile that actually failed
+					// There's actually no guarantee of that, it might be one of the later tiles that exploded
+					for (uint32 j = 0; j < DebugCount; j++)
+					{
+						if (CommandBufferIndex == DebugArray[j].CmdBuffIndex && DebugArray[j].EncoderIndex < CurrentDraw.EncoderIndex && DebugArray[j].CommandIndex < CurrentDraw.CommandIndex)
+						{
+							CurrentDraw = DebugArray[j];
+							break;
+						}
+					}
+
+					// Find the first command to fail - which depends on the order of encoders - parallel contexts make this more complicated
+					if (CurrentDraw.CmdBuffIndex == BrokenDraw.CmdBuffIndex && CurrentDraw.CommandIndex < Commands->Num() && CurrentDraw.EncoderIndex < BrokenDraw.EncoderIndex)
+					{
+						BrokenDraw = CurrentDraw;
+                        BrokenContext = i;
+                        bFoundBrokenDraw = true;
+					}
+                    else if(Commands->Num() > 1)
+                    {
+						BrokenDraw.EncoderIndex = 0;
+						BrokenDraw.CommandIndex = 0;
+						BrokenDraw.CmdBuffIndex = CommandBufferIndex;
+						BrokenDraw.ContextIndex = 0;
+                        BrokenContext = i;
+						BrokenDraw.CommandBuffer = (uintptr_t)CompletedBuffer.GetPtr();
+                        BrokenDraw = DebugArray[1];
+                        bFoundBrokenDraw = true;
+                    }
+				}
+			}
+			if (bFoundBrokenDraw)
+			{
+				id<MTLCommandBuffer> Ptr = reinterpret_cast<id<MTLCommandBuffer>>(BrokenDraw.CommandBuffer);
+				UE_LOG(LogAGX, Error, TEXT("GPU last wrote Command Buffer: %u (%llx) Encoder Index: %d Context Index: %d Draw Index: %d PSO: VS: %u_%u, PS: %u_%u."), BrokenDraw.CmdBuffIndex, BrokenDraw.CommandBuffer, BrokenDraw.EncoderIndex, BrokenDraw.ContextIndex, BrokenDraw.CommandIndex, BrokenDraw.PSOSignature[0], BrokenDraw.PSOSignature[1], BrokenDraw.PSOSignature[2], BrokenDraw.PSOSignature[3]);
+			
+				ns::AutoReleased<FAGXBuffer> DebugBuffer = Markers.GetDebugBuffer(BrokenContext);
+				TArray<FAGXCommandDebug>* Commands = Markers.GetCommands(BrokenContext);
+				if (Commands->Num())
+				{
+					UE_LOG(LogAGX, Error, TEXT("Failed executing following commands:"));
+					uint32 StartIdx = BrokenDraw.CommandIndex;
+					for (uint32 CmdIdx = StartIdx; CmdIdx < Commands->Num(); CmdIdx++)
+					{
+						FAGXCommandDebug& Command = (*Commands)[CmdIdx];
+						FString VSHash = Command.PSO->VertexShader->GetHash().ToString();
+						uint32 VSSig[2] = { Command.PSO->VertexShader->SourceLen, Command.PSO->VertexShader->SourceCRC };
+						
+						FString PSHash;
+						uint32 PSSig[2] = { 0, 0 };
+						if (Command.PSO->PixelShader.IsValid())
+						{
+							PSHash = Command.PSO->PixelShader->GetHash().ToString();
+							PSSig[0] = Command.PSO->PixelShader->SourceLen;
+							PSSig[1] = Command.PSO->PixelShader->SourceCRC;
+						}
+						
+						UE_LOG(LogAGX, Error, TEXT("Command Buffer: %d (%p) Encoder: %d Command: %d: %s PSO: VS: %s (%u_%u), PS: %s (%u_%u)"), Command.CmdBufIndex, CompletedBuffer.GetPtr(), Command.Encoder, Command.Index, *Command.Data.ToString(), *VSHash, VSSig[0], VSSig[1], *PSHash, PSSig[0], PSSig[1]);
+					}
+				}
+                if (DebugBuffer)
+                {
+                    uint32 DebugLength = DebugBuffer.GetLength();
+                    uint32 DebugCount = DebugLength / sizeof(FAGXDebugInfo);
+                    check(DebugCount >= 1);
+                    
+                    FAGXDebugInfo* DebugArray = (FAGXDebugInfo*)DebugBuffer.GetContents();
+                    for (uint32 i = 0; i < DebugCount; i++)
+                    {
+                        FAGXDebugInfo& Command = DebugArray[i];
+						
+						// Stop when nothing has been written into the buffer
+						if (Command.CmdBuffIndex == 0 && Command.CommandBuffer == 0 && Command.EncoderIndex == 0 && Command.CommandIndex == 0 && Command.PSOSignature[0] == 0 && Command.PSOSignature[1] == 0 && Command.PSOSignature[2] == 0 && Command.PSOSignature[3] == 0)
+						{
+							break;
+						}
+						
+                        UE_LOG(LogAGX, Error, TEXT("Command Buffer: %d (%p) Debug Buffer: %p Tile: %u Context: %d Encoder: %d Command: %d PSO: VS: %u_%u, PS: %u_%u"), Command.CmdBuffIndex, Command.CommandBuffer, DebugBuffer.GetPtr(), i, Command.ContextIndex, Command.EncoderIndex, Command.CommandIndex, Command.PSOSignature[0], Command.PSOSignature[1], Command.PSOSignature[2], Command.PSOSignature[3]);
+                    }
+                }
+			}
+		}
+		
+#if PLATFORM_IOS
+        UE_LOG(LogAGX, Warning, TEXT("Command Buffer %s Failed with %s Error! Error Domain: %s Code: %d Description %s %s %s"), *LabelString, ErrorType, *DomainString, Code, *ErrorString, *FailureString, *RecoveryString);
+        FIOSPlatformMisc::GPUAssert();
+#else
+		UE_LOG(LogAGX, Fatal, TEXT("Command Buffer %s Failed with %s Error! Error Domain: %s Code: %d Description %s %s %s"), *LabelString, ErrorType, *DomainString, Code, *ErrorString, *FailureString, *RecoveryString);
+#endif
+    }
+}
+
+static __attribute__ ((optnone)) void MetalCommandBufferFailureInternal(mtlpp::CommandBuffer const& CompletedBuffer)
+{
+	ReportMetalCommandBufferFailure(CompletedBuffer, TEXT("Internal"));
+}
+
+static __attribute__ ((optnone)) void MetalCommandBufferFailureTimeout(mtlpp::CommandBuffer const& CompletedBuffer)
+{
+    ReportMetalCommandBufferFailure(CompletedBuffer, TEXT("Timeout"), PLATFORM_IOS);
+}
+
+static __attribute__ ((optnone)) void MetalCommandBufferFailurePageFault(mtlpp::CommandBuffer const& CompletedBuffer)
+{
+	ReportMetalCommandBufferFailure(CompletedBuffer, TEXT("PageFault"));
+}
+
+static __attribute__ ((optnone)) void MetalCommandBufferFailureBlacklisted(mtlpp::CommandBuffer const& CompletedBuffer)
+{
+	ReportMetalCommandBufferFailure(CompletedBuffer, TEXT("Blacklisted"));
+}
+
+static __attribute__ ((optnone)) void MetalCommandBufferFailureNotPermitted(mtlpp::CommandBuffer const& CompletedBuffer)
+{
+	// when iOS goes into the background, it can get a delayed NotPermitted error, so we can't crash in this case, just allow it to not be submitted
+	ReportMetalCommandBufferFailure(CompletedBuffer, TEXT("NotPermitted"), !PLATFORM_IOS);
+}
+
+static __attribute__ ((optnone)) void MetalCommandBufferFailureOutOfMemory(mtlpp::CommandBuffer const& CompletedBuffer)
+{
+	ReportMetalCommandBufferFailure(CompletedBuffer, TEXT("OutOfMemory"));
+}
+
+static __attribute__ ((optnone)) void MetalCommandBufferFailureInvalidResource(mtlpp::CommandBuffer const& CompletedBuffer)
+{
+	ReportMetalCommandBufferFailure(CompletedBuffer, TEXT("InvalidResource"));
+}
+
+static void HandleMetalCommandBufferError(mtlpp::CommandBuffer const& CompletedBuffer)
+{
+	MTLCommandBufferError Code = (MTLCommandBufferError)CompletedBuffer.GetError().GetCode();
+	switch(Code)
+	{
+		case MTLCommandBufferErrorInternal:
+			MetalCommandBufferFailureInternal(CompletedBuffer);
+			break;
+		case MTLCommandBufferErrorTimeout:
+			MetalCommandBufferFailureTimeout(CompletedBuffer);
+			break;
+		case MTLCommandBufferErrorPageFault:
+			MetalCommandBufferFailurePageFault(CompletedBuffer);
+			break;
+		case MTLCommandBufferErrorBlacklisted:
+			MetalCommandBufferFailureBlacklisted(CompletedBuffer);
+			break;
+		case MTLCommandBufferErrorNotPermitted:
+			MetalCommandBufferFailureNotPermitted(CompletedBuffer);
+			break;
+		case MTLCommandBufferErrorOutOfMemory:
+			MetalCommandBufferFailureOutOfMemory(CompletedBuffer);
+			break;
+		case MTLCommandBufferErrorInvalidResource:
+			MetalCommandBufferFailureInvalidResource(CompletedBuffer);
+			break;
+		case MTLCommandBufferErrorNone:
+			// No error
+			break;
+		default:
+			ReportMetalCommandBufferFailure(CompletedBuffer, TEXT("Unknown"));
+			break;
+	}
+}
+
+static __attribute__ ((optnone)) void HandleAMDMetalCommandBufferError(mtlpp::CommandBuffer const& CompletedBuffer)
+{
+	HandleMetalCommandBufferError(CompletedBuffer);
+}
+
+static __attribute__ ((optnone)) void HandleNVIDIAMetalCommandBufferError(mtlpp::CommandBuffer const& CompletedBuffer)
+{
+	HandleMetalCommandBufferError(CompletedBuffer);
+}
+
+static __attribute__ ((optnone)) void HandleIntelMetalCommandBufferError(mtlpp::CommandBuffer const& CompletedBuffer)
+{
+	HandleMetalCommandBufferError(CompletedBuffer);
+}
+
+void FAGXCommandList::HandleMetalCommandBufferFailure(mtlpp::CommandBuffer const& CompletedBuffer)
+{
+	if (CompletedBuffer.GetError().GetDomain() == MTLCommandBufferErrorDomain || [CompletedBuffer.GetError().GetDomain() isEqualToString:MTLCommandBufferErrorDomain])
+	{
+		if (GRHIVendorId && IsRHIDeviceAMD())
+		{
+			HandleAMDMetalCommandBufferError(CompletedBuffer);
+		}
+		else if (GRHIVendorId && IsRHIDeviceNVIDIA())
+		{
+			HandleNVIDIAMetalCommandBufferError(CompletedBuffer);
+		}
+		else if (GRHIVendorId && IsRHIDeviceIntel())
+		{
+			HandleIntelMetalCommandBufferError(CompletedBuffer);
+		}
+		else
+		{
+			HandleMetalCommandBufferError(CompletedBuffer);
+		}
+	}
+	else
+	{
+		ReportMetalCommandBufferFailure(CompletedBuffer, TEXT("Unknown"));
+	}
+}
+
+void FAGXCommandList::SetParallelIndex(uint32 InIndex, uint32 InNum)
+{
+	if (!IsImmediate())
+	{
+		Index = InIndex;
+		Num = InNum;
+	}
+}
+
+void FAGXCommandList::Commit(mtlpp::CommandBuffer& Buffer, TArray<ns::Object<mtlpp::CommandBufferHandler>> CompletionHandlers, bool const bWait, bool const bIsLastCommandBuffer)
+{
+	check(Buffer);
+
+	// The lifetime of this array is per frame
+	if (!FrameCommitedBufferTimings.IsValid())
+	{
+		FrameCommitedBufferTimings = MakeShared<TArray<FAGXCommandBufferTiming>, ESPMode::ThreadSafe>();
+	}
+
+	// The lifetime of this should be for the entire game
+	if (!LastCompletedBufferTiming.IsValid())
+	{
+		LastCompletedBufferTiming = MakeShared<FAGXCommandBufferTiming, ESPMode::ThreadSafe>();
+	}
+
+	Buffer.AddCompletedHandler([CompletionHandlers, FrameCommitedBufferTimingsLocal = FrameCommitedBufferTimings, LastCompletedBufferTimingLocal = LastCompletedBufferTiming](mtlpp::CommandBuffer const& CompletedBuffer)
+	{
+		if (CompletedBuffer.GetStatus() == mtlpp::CommandBufferStatus::Error)
+		{
+			HandleMetalCommandBufferFailure(CompletedBuffer);
+		}
+		if (CompletionHandlers.Num())
+		{
+			for (ns::Object<mtlpp::CommandBufferHandler> Handler : CompletionHandlers)
+			{
+				Handler.GetPtr()(CompletedBuffer);
+			}
+		}
+
+		if (CompletedBuffer.GetStatus() == mtlpp::CommandBufferStatus::Completed)
+		{
+			FrameCommitedBufferTimingsLocal->Add({CompletedBuffer.GetGpuStartTime(), CompletedBuffer.GetGpuEndTime()});
+		}
+
+		// If this is the last reference, then it is the last command buffer to return, so record the frame
+		if (FrameCommitedBufferTimingsLocal.IsUnique())
+		{
+			FAGXGPUProfiler::RecordFrame(*FrameCommitedBufferTimingsLocal, *LastCompletedBufferTimingLocal);
+		}
+	});
+
+	// If bIsLastCommandBuffer is set then this is the end of the "frame".
+	if (bIsLastCommandBuffer)
+	{
+		FrameCommitedBufferTimings = MakeShared<TArray<FAGXCommandBufferTiming>, ESPMode::ThreadSafe>();
+	}
+
+	if (bImmediate)
+	{
+		CommandQueue.CommitCommandBuffer(Buffer);
+		if (bWait)
+		{
+			Buffer.WaitUntilCompleted();
+		}
+	}
+	else
+	{
+		check(!bWait);
+		SubmittedBuffers.Add(Buffer);
+	}
+}
+
+void FAGXCommandList::Submit(uint32 InIndex, uint32 Count)
+{
+	// Only deferred contexts should call Submit, the immediate context commits directly to the command-queue.
+	check(!bImmediate);
+
+	// Command queue takes ownership of the array
+	CommandQueue.SubmitCommandBuffers(SubmittedBuffers, InIndex, Count);
+	SubmittedBuffers.Empty();
+}
diff --git a/Engine/Source/Runtime/Apple/AGXRHI/Private/AGXCommandList.h b/Engine/Source/Runtime/Apple/AGXRHI/Private/AGXCommandList.h
new file mode 100644
--- /dev/null
+++ b/Engine/Source/Runtime/Apple/AGXRHI/Private/AGXCommandList.h
@@ -0,0 +1,101 @@
+// Copyright Epic Games, Inc. All Rights Reserved.
+
+#pragma once
+
+#include <Metal/Metal.h>
+
+class FAGXCommandQueue;
+
+/**
+ * FAGXCommandList:
+ * Encapsulates multiple command-buffers into an ordered list for submission. 
+ * For the immediate context this is irrelevant and is merely a pass-through into the CommandQueue, but
+ * for deferred/parallel contexts it is required as they must queue their command buffers until they can 
+ * be committed to the command-queue in the proper order which is only known at the end of parallel encoding.
+ */
+class FAGXCommandList
+{
+public:
+#pragma mark - Public C++ Boilerplate -
+
+	/**
+	 * Constructor
+	 * @param InCommandQueue The command-queue to which the command-list's buffers will be submitted.
+	 */
+	FAGXCommandList(FAGXCommandQueue& InCommandQueue, bool const bInImmediate);
+	
+	/** Destructor */
+	~FAGXCommandList(void);
+	
+	/**
+	 * Command buffer failure reporting function.
+	 * @param CompletedBuffer The buffer to check for failure.
+	 */
+	static void HandleMetalCommandBufferFailure(mtlpp::CommandBuffer const& CompletedBuffer);
+	
+#pragma mark - Public Command List Mutators -
+	
+	/** 
+	 * Set the number of parallel command-lists and the index of this command-list within the pass.
+	 * @param Index The index of this command-list within the parallel pass.
+	 * @param Num The number of command-lists within the parallel pass.
+	 */
+	void SetParallelIndex(uint32 Index, uint32 Num);
+
+	/** 
+	 * Commits the provided buffer to the command-list for execution. When parallel encoding this will be submitted later.
+	 * @param Buffer The buffer to submit to the command-list.
+	 * @param CompletionHandlers The completion handlers that should be attached to this command-buffer.
+	 * @param bWait Whether to wait for the command buffer to complete - it is an error to set this to true on a deferred command-list.
+	 * @param bIsLastCommandBuffer True if this is the final command buffer in a frame.
+	 */
+	void Commit(mtlpp::CommandBuffer& Buffer, TArray<ns::Object<mtlpp::CommandBufferHandler>> CompletionHandlers, bool const bWait, bool const bIsLastCommandBuffer);
+	
+	/**
+	 * Submits all outstanding command-buffers in the proper commit order to the command-queue.
+	 * When more than one command-list is active the command-queue will buffer the command-lists until all are committed to guarantee order of submission to the GPU.
+	 * @param Index The command-list's intended index in the command-queue.
+	 * @param Count The number of command-lists that will be committed to the command-queue.
+	 */
+	void Submit(uint32 Index, uint32 Count);
+	
+#pragma mark - Public Command List Accessors -
+	
+	/**
+	 * True iff the command-list submits immediately to the command-queue, false if it performs any buffering.
+	 * @returns True iff the command-list submits immediately to the command-queue, false if it performs any buffering.
+	 */
+	bool IsImmediate(void) const { return bImmediate; }
+	
+	/**
+	 * True iff the command-list is part of an MTLParallelRenderCommandEncoder pass.
+	 * @returns True iff the command-list is part of an MTLParallelRenderCommandEncoder pass, false for immediate and parallel-command-buffer contexts.
+	 */
+	bool IsParallel(void) const { return !bImmediate && Num > 0 && FAGXCommandQueue::SupportsFeature(EAGXFeaturesParallelRenderEncoders); }
+	
+
+	/**
+	 * The index of this command-list within the parallel pass.
+	 * @returns The index of this command-list within the parallel pass, 0 when IsImmediate() is true.
+	 */
+	uint32 GetParallelIndex(void) const { return Index; }
+	
+	/**
+	 * The number of command-lists within the parallel pass.
+	 * @returns The number of command-lists within the parallel pass, 0 when IsImmediate() is true.
+	 */
+	uint32 GetParallelNum(void) const { return Num; }
+
+	/** @returns The command queue to which this command-list submits command-buffers. */
+	FAGXCommandQueue& GetCommandQueue(void) const { return CommandQueue; }
+	
+private:
+#pragma mark - Private Member Variables -
+	FAGXCommandQueue& CommandQueue;
+	TSharedPtr<TArray<FAGXCommandBufferTiming>, ESPMode::ThreadSafe> FrameCommitedBufferTimings;
+	TSharedPtr<FAGXCommandBufferTiming, ESPMode::ThreadSafe> LastCompletedBufferTiming;
+	TArray<mtlpp::CommandBuffer> SubmittedBuffers;
+	uint32 Index;
+	uint32 Num;
+	bool bImmediate;
+};
diff --git a/Engine/Source/Runtime/Apple/AGXRHI/Private/AGXCommandQueue.cpp b/Engine/Source/Runtime/Apple/AGXRHI/Private/AGXCommandQueue.cpp
new file mode 100644
--- /dev/null
+++ b/Engine/Source/Runtime/Apple/AGXRHI/Private/AGXCommandQueue.cpp
@@ -0,0 +1,484 @@
+// Copyright Epic Games, Inc. All Rights Reserved.
+
+/*=============================================================================
+	AGXCommandQueue.cpp: AGX RHI command queue wrapper.
+=============================================================================*/
+
+#include "AGXRHIPrivate.h"
+
+#include "AGXCommandQueue.h"
+#include "AGXCommandBuffer.h"
+#include "AGXCommandList.h"
+#include "AGXProfiler.h"
+#include "Misc/ConfigCacheIni.h"
+#include "command_buffer.hpp"
+
+#pragma mark - Private C++ Statics -
+NSUInteger FAGXCommandQueue::PermittedOptions = 0;
+uint64 FAGXCommandQueue::Features = 0;
+extern mtlpp::VertexFormat GAGXFColorVertexFormat;
+bool GAGXCommandBufferDebuggingEnabled = 0;
+
+#pragma mark - Public C++ Boilerplate -
+
+FAGXCommandQueue::FAGXCommandQueue(mtlpp::Device InDevice, uint32 const MaxNumCommandBuffers /* = 0 */)
+: Device(InDevice)
+, ParallelCommandLists(0)
+, RuntimeDebuggingLevel(EAGXDebugLevelOff)
+{
+	int32 MaxShaderVersion = 0;
+	int32 IndirectArgumentTier = 0;
+#if PLATFORM_MAC
+	int32 DefaultMaxShaderVersion = 5; // MSL v2.2
+	int32 MinShaderVersion = 4; // MSL v2.1
+    const TCHAR* const Settings = TEXT("/Script/MacTargetPlatform.MacTargetSettings");
+#else
+	int32 DefaultMaxShaderVersion = 2;
+	int32 MinShaderVersion = 2;
+    const TCHAR* const Settings = TEXT("/Script/IOSRuntimeSettings.IOSRuntimeSettings");
+#endif
+    if(!GConfig->GetInt(Settings, TEXT("MaxShaderLanguageVersion"), MaxShaderVersion, GEngineIni))
+    {
+        MaxShaderVersion = DefaultMaxShaderVersion;
+    }
+	if(!GConfig->GetInt(Settings, TEXT("IndirectArgumentTier"), IndirectArgumentTier, GEngineIni))
+	{
+		IndirectArgumentTier = 0;
+	}
+	MaxShaderVersion = FMath::Max(MinShaderVersion, MaxShaderVersion);
+	AGXValidateVersion(MaxShaderVersion);
+
+	if(MaxNumCommandBuffers == 0)
+	{
+		CommandQueue = Device.NewCommandQueue();
+	}
+	else
+	{
+		CommandQueue = Device.NewCommandQueue(MaxNumCommandBuffers);
+	}
+	check(CommandQueue);
+#if PLATFORM_IOS
+	NSOperatingSystemVersion Vers = [[NSProcessInfo processInfo] operatingSystemVersion];
+	if(Vers.majorVersion >= 9)
+	{
+		Features = EAGXFeaturesSetBufferOffset | EAGXFeaturesSetBytes;
+
+#if PLATFORM_TVOS
+        Features &= ~(EAGXFeaturesSetBytes);
+		
+		if(Device.SupportsFeatureSet(mtlpp::FeatureSet::tvOS_GPUFamily2_v1))
+		{
+			Features |= EAGXFeaturesCountingQueries | EAGXFeaturesBaseVertexInstance | EAGXFeaturesIndirectBuffer | EAGXFeaturesMSAADepthResolve | EAGXFeaturesMSAAStoreAndResolve;
+		}
+		
+		if(Vers.majorVersion > 10)
+		{
+			Features |= EAGXFeaturesPrivateBufferSubAllocation;
+			
+			if(Vers.majorVersion >= 11)
+			{
+				Features |= EAGXFeaturesGPUCaptureManager | EAGXFeaturesBufferSubAllocation | EAGXFeaturesParallelRenderEncoders | EAGXFeaturesPipelineBufferMutability;
+				
+				if (MaxShaderVersion >= 3)
+				{
+					GAGXFColorVertexFormat = mtlpp::VertexFormat::UChar4Normalized_BGRA;
+				}
+
+				if (Vers.majorVersion >= 12)
+				{
+					Features |= EAGXFeaturesMaxThreadsPerThreadgroup;
+					
+					if (FParse::Param(FCommandLine::Get(),TEXT("metalfence")))
+					{
+					Features |= EAGXFeaturesFences;
+					}
+					
+					if (FParse::Param(FCommandLine::Get(),TEXT("metalheap")))
+					{
+					Features |= EAGXFeaturesHeaps;
+					}
+					
+					if (MaxShaderVersion >= 4)
+					{
+						Features |= EAGXFeaturesTextureBuffers;
+					}
+				}
+			}
+		}
+#else
+		if (Device.SupportsFeatureSet(mtlpp::FeatureSet::iOS_GPUFamily3_v1))
+		{
+			Features |= EAGXFeaturesCountingQueries | EAGXFeaturesBaseVertexInstance | EAGXFeaturesIndirectBuffer | EAGXFeaturesMSAADepthResolve;
+		}
+		
+		if(Device.SupportsFeatureSet(mtlpp::FeatureSet::iOS_GPUFamily3_v2) || Device.SupportsFeatureSet(mtlpp::FeatureSet::iOS_GPUFamily2_v3) || Device.SupportsFeatureSet(mtlpp::FeatureSet::iOS_GPUFamily1_v3))
+		{
+			if (FParse::Param(FCommandLine::Get(),TEXT("metalfence")))
+			{
+				Features |= EAGXFeaturesFences;
+			}
+			
+			if (FParse::Param(FCommandLine::Get(),TEXT("metalheap")))
+			{
+				Features |= EAGXFeaturesHeaps;
+			}
+		}
+		
+		if(Device.SupportsFeatureSet(mtlpp::FeatureSet::iOS_GPUFamily3_v2))
+		{
+			Features |= EAGXFeaturesMSAAStoreAndResolve;
+		}
+		
+		if(Vers.majorVersion > 10 || (Vers.majorVersion == 10 && Vers.minorVersion >= 3))
+        {
+			// Turning the below option on will allocate more buffer memory which isn't generally desirable on iOS
+			// Features |= EAGXFeaturesEfficientBufferBlits;
+			
+			// These options are fine however as thye just change how we allocate small buffers
+            Features |= EAGXFeaturesBufferSubAllocation;
+			Features |= EAGXFeaturesPrivateBufferSubAllocation;
+			
+			if(Vers.majorVersion >= 11)
+			{
+				if (MaxShaderVersion >= 3)
+				{
+					GAGXFColorVertexFormat = mtlpp::VertexFormat::UChar4Normalized_BGRA;
+				}
+				
+				Features |= EAGXFeaturesPresentMinDuration | EAGXFeaturesGPUCaptureManager | EAGXFeaturesBufferSubAllocation | EAGXFeaturesParallelRenderEncoders | EAGXFeaturesPipelineBufferMutability;
+                
+				// Turn on Texture Buffers! These are faster on the GPU as we don't need to do out-of-bounds tests but require Metal 2.1 and macOS 10.14
+				if (Vers.majorVersion >= 12)
+				{
+					Features |= EAGXFeaturesMaxThreadsPerThreadgroup;
+                    if (!FParse::Param(FCommandLine::Get(),TEXT("nometalfence")))
+                    {
+                        Features |= EAGXFeaturesFences;
+                    }
+                    
+                    if (!FParse::Param(FCommandLine::Get(),TEXT("nometalheap")))
+                    {
+                        Features |= EAGXFeaturesHeaps;
+                    }
+					
+					if (MaxShaderVersion >= 4)
+					{
+						Features |= EAGXFeaturesTextureBuffers;
+					}
+                    
+                    if(Device.SupportsFeatureSet(mtlpp::FeatureSet::iOS_GPUFamily4_v1))
+                    {
+                        Features |= EAGXFeaturesTileShaders;
+                        
+                        // The below implies tile shaders which are necessary to order the draw calls and generate a buffer that shows what PSOs/draws ran on each tile.
+                        IConsoleVariable* GPUCrashDebuggingCVar = IConsoleManager::Get().FindConsoleVariable(TEXT("r.GPUCrashDebugging"));
+                        GAGXCommandBufferDebuggingEnabled = (GPUCrashDebuggingCVar && GPUCrashDebuggingCVar->GetInt() != 0) || FParse::Param(FCommandLine::Get(),TEXT("metalgpudebug"));
+                    }
+                    
+					if (Device.SupportsFeatureSet(mtlpp::FeatureSet::iOS_GPUFamily5_v1))
+					{
+						Features |= EAGXFeaturesLayeredRendering;
+					}
+				}
+			}
+        }
+#endif
+	}
+	else if(Vers.majorVersion == 8 && Vers.minorVersion >= 3)
+	{
+		Features = EAGXFeaturesSetBufferOffset;
+	}
+#else // Assume that Mac & other platforms all support these from the start. They can diverge later.
+	const bool bIsNVIDIA = [Device.GetName().GetPtr() rangeOfString:@"Nvidia" options:NSCaseInsensitiveSearch].location != NSNotFound;
+	Features = EAGXFeaturesCountingQueries | EAGXFeaturesBaseVertexInstance | EAGXFeaturesIndirectBuffer | EAGXFeaturesLayeredRendering | EAGXFeaturesCubemapArrays;
+	if (!bIsNVIDIA)
+	{
+		Features |= EAGXFeaturesSetBufferOffset;
+	}
+	if (Device.SupportsFeatureSet(mtlpp::FeatureSet::macOS_GPUFamily1_v2))
+    {
+        Features |= EAGXFeaturesMSAADepthResolve | EAGXFeaturesMSAAStoreAndResolve;
+        
+        // Assume that set*Bytes only works on macOS Sierra and above as no-one has tested it anywhere else.
+		Features |= EAGXFeaturesSetBytes;
+		
+		FString DeviceName(Device.GetName());
+		// On earlier OS versions Intel Broadwell couldn't suballocate properly
+		if (!(DeviceName.Contains(TEXT("Intel")) && (DeviceName.Contains(TEXT("5300")) || DeviceName.Contains(TEXT("6000")) || DeviceName.Contains(TEXT("6100")))) || FPlatformMisc::MacOSXVersionCompare(10,14,0) >= 0)
+		{
+			// Using Private Memory & BlitEncoders for Vertex & Index data should be *much* faster.
+        	Features |= EAGXFeaturesEfficientBufferBlits;
+        	
+			Features |= EAGXFeaturesBufferSubAllocation;
+					
+	        // On earlier OS versions Vega didn't like non-zero blit offsets
+	        if (!DeviceName.Contains(TEXT("Vega")) || FPlatformMisc::MacOSXVersionCompare(10,13,5) >= 0)
+	        {
+				Features |= EAGXFeaturesPrivateBufferSubAllocation;
+			}
+		}
+		
+		GAGXFColorVertexFormat = mtlpp::VertexFormat::UChar4Normalized_BGRA;
+		
+		// On 10.13.5+ we can use MTLParallelRenderEncoder
+		if (FPlatformMisc::MacOSXVersionCompare(10,13,5) >= 0)
+		{
+			// Except on Nvidia for the moment
+			if ([Device.GetName().GetPtr() rangeOfString:@"Nvidia" options:NSCaseInsensitiveSearch].location == NSNotFound && !FParse::Param(FCommandLine::Get(),TEXT("nometalparallelencoder")))
+			{
+				Features |= EAGXFeaturesParallelRenderEncoders;
+			}
+		}
+
+		// Turn on Texture Buffers! These are faster on the GPU as we don't need to do out-of-bounds tests but require Metal 2.1 and macOS 10.14
+		if (FPlatformMisc::MacOSXVersionCompare(10,14,0) >= 0)
+		{
+			Features |= EAGXFeaturesMaxThreadsPerThreadgroup;
+			if (MaxShaderVersion >= 4)
+			{
+				Features |= EAGXFeaturesTextureBuffers;
+            }
+            if (IndirectArgumentTier >= 1)
+            {
+                Features |= EAGXFeaturesIABs;
+				
+				if (IndirectArgumentTier >= 2)
+				{
+					Features |= EAGXFeaturesTier2IABs;
+				}
+            }
+            
+            IConsoleVariable* GPUCrashDebuggingCVar = IConsoleManager::Get().FindConsoleVariable(TEXT("r.GPUCrashDebugging"));
+            GAGXCommandBufferDebuggingEnabled = (GPUCrashDebuggingCVar && GPUCrashDebuggingCVar->GetInt() != 0) || FParse::Param(FCommandLine::Get(),TEXT("metalgpudebug"));
+            
+            // The editor spawns so many viewports and preview icons that we can run out of hardware fences!
+			// Need to figure out a way to safely flush the rendering and reuse the fences when that happens.
+#if WITH_EDITORONLY_DATA
+			if (!GIsEditor)
+#endif
+			{
+				if (FParse::Param(FCommandLine::Get(),TEXT("metalfence")))
+				{
+					Features |= EAGXFeaturesFences;
+				}
+				
+				// There are still too many driver bugs to use MTLHeap on macOS - nothing works without causing random, undebuggable GPU hangs that completely deadlock the Mac and don't generate any validation errors or command-buffer failures
+				if (FParse::Param(FCommandLine::Get(),TEXT("forcemetalheap")))
+				{
+					Features |= EAGXFeaturesHeaps;
+				}
+			}
+		}
+    }
+    else if ([Device.GetName().GetPtr() rangeOfString:@"Nvidia" options:NSCaseInsensitiveSearch].location != NSNotFound)
+    {
+		// Using set*Bytes fixes bugs on Nvidia for 10.11 so we should use it...
+    	Features |= EAGXFeaturesSetBytes;
+    }
+    
+    if(Device.SupportsFeatureSet(mtlpp::FeatureSet::macOS_GPUFamily1_v3) && FPlatformMisc::MacOSXVersionCompare(10,13,0) >= 0)
+    {
+        Features |= EAGXFeaturesMultipleViewports | EAGXFeaturesPipelineBufferMutability | EAGXFeaturesGPUCaptureManager;
+		
+		if (FParse::Param(FCommandLine::Get(),TEXT("metalfence")))
+		{
+			Features |= EAGXFeaturesFences;
+		}
+		
+		if (FParse::Param(FCommandLine::Get(),TEXT("metalheap")))
+		{
+			Features |= EAGXFeaturesHeaps;
+		}
+		
+		if (FParse::Param(FCommandLine::Get(),TEXT("metaliabs")))
+		{
+			Features |= EAGXFeaturesIABs;
+		}
+    }
+#endif
+	
+#if !UE_BUILD_SHIPPING
+	Class MTLDebugDevice = NSClassFromString(@"MTLDebugDevice");
+	if ([Device isKindOfClass:MTLDebugDevice])
+	{
+		Features |= EAGXFeaturesValidation;
+	}
+#endif
+
+	static const auto CVar = IConsoleManager::Get().FindConsoleVariable(TEXT("r.Shaders.Optimize"));
+	if (CVar->GetInt() == 0 || FParse::Param(FCommandLine::Get(),TEXT("metalshaderdebug")))
+	{
+		Features |= EAGXFeaturesGPUTrace;
+	}
+
+	PermittedOptions = 0;
+	PermittedOptions |= mtlpp::ResourceOptions::CpuCacheModeDefaultCache;
+	PermittedOptions |= mtlpp::ResourceOptions::CpuCacheModeWriteCombined;
+	{
+		PermittedOptions |= mtlpp::ResourceOptions::StorageModeShared;
+		PermittedOptions |= mtlpp::ResourceOptions::StorageModePrivate;
+#if PLATFORM_MAC
+		PermittedOptions |= mtlpp::ResourceOptions::StorageModeManaged;
+#else
+		PermittedOptions |= mtlpp::ResourceOptions::StorageModeMemoryless;
+#endif
+		// You can't use HazardUntracked under the validation layer due to bugs in the layer when trying to create linear-textures/texture-buffers
+		if ((Features & EAGXFeaturesFences) && !(Features & EAGXFeaturesValidation))
+		{
+			PermittedOptions |= mtlpp::ResourceOptions::HazardTrackingModeUntracked;
+		}
+	}
+}
+
+FAGXCommandQueue::~FAGXCommandQueue(void)
+{
+	// void
+}
+	
+#pragma mark - Public Command Buffer Mutators -
+
+mtlpp::CommandBuffer FAGXCommandQueue::CreateCommandBuffer(void)
+{
+#if PLATFORM_MAC
+	static bool bUnretainedRefs = FParse::Param(FCommandLine::Get(),TEXT("metalunretained"))
+	|| (!FParse::Param(FCommandLine::Get(),TEXT("metalretainrefs"))
+		&& ([Device.GetName() rangeOfString:@"Nvidia" options:NSCaseInsensitiveSearch].location == NSNotFound)
+		&& ([Device.GetName() rangeOfString:@"Intel" options:NSCaseInsensitiveSearch].location == NSNotFound));
+#else
+	static bool bUnretainedRefs = !FParse::Param(FCommandLine::Get(),TEXT("metalretainrefs"));
+#endif
+	
+	mtlpp::CommandBuffer CmdBuffer;
+	@autoreleasepool
+	{
+		CmdBuffer = bUnretainedRefs ? MTLPP_VALIDATE(mtlpp::CommandQueue, CommandQueue, AGXSafeGetRuntimeDebuggingLevel() >= EAGXDebugLevelValidation, CommandBufferWithUnretainedReferences()) : MTLPP_VALIDATE(mtlpp::CommandQueue, CommandQueue, AGXSafeGetRuntimeDebuggingLevel() >= EAGXDebugLevelValidation, CommandBuffer());
+		
+		if (RuntimeDebuggingLevel > EAGXDebugLevelOff)
+		{			
+			METAL_DEBUG_ONLY(FAGXCommandBufferDebugging AddDebugging(CmdBuffer));
+			MTLPP_VALIDATION(mtlpp::CommandBufferValidationTable ValidatedCommandBuffer(CmdBuffer));
+		}
+	}
+	CommandBufferFences.Push(new mtlpp::CommandBufferFence(CmdBuffer.GetCompletionFence()));
+	INC_DWORD_STAT(STAT_AGXCommandBufferCreatedPerFrame);
+	return CmdBuffer;
+}
+
+void FAGXCommandQueue::CommitCommandBuffer(mtlpp::CommandBuffer& CommandBuffer)
+{
+	check(CommandBuffer);
+	INC_DWORD_STAT(STAT_AGXCommandBufferCommittedPerFrame);
+	
+	MTLPP_VALIDATE(mtlpp::CommandBuffer, CommandBuffer, AGXSafeGetRuntimeDebuggingLevel() >= EAGXDebugLevelValidation, Commit());
+	
+	// Wait for completion when debugging command-buffers.
+	if (RuntimeDebuggingLevel >= EAGXDebugLevelWaitForComplete)
+	{
+		CommandBuffer.WaitUntilCompleted();
+	}
+}
+
+void FAGXCommandQueue::SubmitCommandBuffers(TArray<mtlpp::CommandBuffer> BufferList, uint32 Index, uint32 Count)
+{
+	CommandBuffers.SetNumZeroed(Count);
+	CommandBuffers[Index] = BufferList;
+	ParallelCommandLists |= (1 << Index);
+	if (ParallelCommandLists == ((1 << Count) - 1))
+	{
+		for (uint32 i = 0; i < Count; i++)
+		{
+			TArray<mtlpp::CommandBuffer>& CmdBuffers = CommandBuffers[i];
+			for (mtlpp::CommandBuffer Buffer : CmdBuffers)
+			{
+				check(Buffer);
+				CommitCommandBuffer(Buffer);
+			}
+			CommandBuffers[i].Empty();
+		}
+		
+		ParallelCommandLists = 0;
+	}
+}
+
+FAGXFence* FAGXCommandQueue::CreateFence(ns::String const& Label) const
+{
+	if ((Features & EAGXFeaturesFences) != 0)
+	{
+		FAGXFence* InternalFence = FAGXFencePool::Get().AllocateFence();
+		for (uint32 i = mtlpp::RenderStages::Vertex; InternalFence && i <= mtlpp::RenderStages::Fragment; i++)
+		{
+			mtlpp::Fence InnerFence = InternalFence->Get((mtlpp::RenderStages)i);
+			NSString* String = nil;
+			if (GetEmitDrawEvents())
+			{
+				String = [NSString stringWithFormat:@"%u %p: %@", i, InnerFence.GetPtr(), Label.GetPtr()];
+			}
+	#if METAL_DEBUG_OPTIONS
+			if (RuntimeDebuggingLevel >= EAGXDebugLevelValidation)
+			{
+				FAGXDebugFence* Fence = (FAGXDebugFence*)InnerFence.GetPtr();
+				Fence.label = String;
+			}
+			else
+	#endif
+			if(InnerFence && String)
+			{
+				InnerFence.SetLabel(String);
+			}
+		}
+		return InternalFence;
+	}
+	else
+	{
+		return nullptr;
+	}
+}
+
+void FAGXCommandQueue::GetCommittedCommandBufferFences(TArray<mtlpp::CommandBufferFence>& Fences)
+{
+	TArray<mtlpp::CommandBufferFence*> Temp;
+	CommandBufferFences.PopAll(Temp);
+	for (mtlpp::CommandBufferFence* Fence : Temp)
+	{
+		Fences.Add(*Fence);
+		delete Fence;
+	}
+}
+
+#pragma mark - Public Command Queue Accessors -
+	
+mtlpp::Device& FAGXCommandQueue::GetDevice(void)
+{
+	return Device;
+}
+
+mtlpp::ResourceOptions FAGXCommandQueue::GetCompatibleResourceOptions(mtlpp::ResourceOptions Options)
+{
+	NSUInteger NewOptions = (Options & PermittedOptions);
+#if PLATFORM_IOS // Swizzle Managed to Shared for iOS - we can do this as they are equivalent, unlike Shared -> Managed on Mac.
+	if ((Options & (1 /*mtlpp::StorageMode::Managed*/ << mtlpp::ResourceStorageModeShift)))
+	{
+		NewOptions |= mtlpp::ResourceOptions::StorageModeShared;
+	}
+#endif
+	return (mtlpp::ResourceOptions)NewOptions;
+}
+
+#pragma mark - Public Debug Support -
+
+void FAGXCommandQueue::InsertDebugCaptureBoundary(void)
+{
+	PRAGMA_DISABLE_DEPRECATION_WARNINGS
+	[CommandQueue insertDebugCaptureBoundary];
+	PRAGMA_ENABLE_DEPRECATION_WARNINGS
+}
+
+void FAGXCommandQueue::SetRuntimeDebuggingLevel(int32 const Level)
+{
+	RuntimeDebuggingLevel = Level;
+}
+
+int32 FAGXCommandQueue::GetRuntimeDebuggingLevel(void) const
+{
+	return RuntimeDebuggingLevel;
+}
diff --git a/Engine/Source/Runtime/Apple/AGXRHI/Private/AGXCommandQueue.h b/Engine/Source/Runtime/Apple/AGXRHI/Private/AGXCommandQueue.h
new file mode 100644
--- /dev/null
+++ b/Engine/Source/Runtime/Apple/AGXRHI/Private/AGXCommandQueue.h
@@ -0,0 +1,169 @@
+// Copyright Epic Games, Inc. All Rights Reserved.
+
+#pragma once
+
+#include <Metal/Metal.h>
+#include "device.hpp"
+#include "command_queue.hpp"
+#include "command_buffer.hpp"
+
+#include "Containers/LockFreeList.h"
+
+class FAGXCommandList;
+
+/**
+ * Enumeration of features which are present only on some OS/device combinations.
+ * These have to be checked at runtime as well as compile time to ensure backward compatibility.
+ */
+typedef NS_OPTIONS(uint64, EAGXFeatures)
+{
+	/** Support for specifying an update to the buffer offset only */
+	EAGXFeaturesSetBufferOffset = 1 << 0,
+	/** Supports NSUInteger counting visibility queries */
+	EAGXFeaturesCountingQueries = 1 << 1,
+	/** Supports base vertex/instance for draw calls */
+	EAGXFeaturesBaseVertexInstance = 1 << 2,
+	/** Supports indirect buffers for draw calls */
+	EAGXFeaturesIndirectBuffer = 1 << 3,
+	/** Supports layered rendering */
+	EAGXFeaturesLayeredRendering = 1 << 4,
+	/** Support for specifying small buffers as byte arrays */
+	EAGXFeaturesSetBytes = 1 << 5,
+	/** Unused Reserved Bit */
+	EAGXFeaturesUnusedReservedBit6 = 1 << 6, // was EAGXFeaturesTessellation
+	/** Supports framework-level validation */
+	EAGXFeaturesValidation = 1 << 7,
+	/** Supports detailed statistics */
+	EAGXFeaturesStatistics = 1 << 8,
+	/** Supports the explicit MTLHeap APIs */
+	EAGXFeaturesHeaps = 1 << 9,
+	/** Supports the explicit MTLFence APIs */
+	EAGXFeaturesFences = 1 << 10,
+	/** Supports MSAA Depth Resolves */
+	EAGXFeaturesMSAADepthResolve = 1 << 11,
+	/** Supports Store & Resolve in a single store action */
+	EAGXFeaturesMSAAStoreAndResolve = 1 << 12,
+	/** Supports framework GPU frame capture */
+	EAGXFeaturesGPUTrace = 1 << 13,
+	/** Supports the use of cubemap arrays */
+	EAGXFeaturesCubemapArrays = 1 << 14,
+	/** Supports the specification of multiple viewports and scissor rects */
+	EAGXFeaturesMultipleViewports = 1 << 15,
+    /** Supports minimum on-glass duration for drawables */
+    EAGXFeaturesPresentMinDuration = 1llu << 16llu,
+    /** Supports programmatic frame capture API */
+    EAGXFeaturesGPUCaptureManager = 1llu << 17llu,
+	/** Supports efficient buffer-blits */
+	EAGXFeaturesEfficientBufferBlits = 1llu << 18llu,
+	/** Supports any kind of buffer sub-allocation */
+	EAGXFeaturesBufferSubAllocation = 1llu << 19llu,
+	/** Supports private buffer sub-allocation */
+	EAGXFeaturesPrivateBufferSubAllocation = 1llu << 20llu,
+	/** Supports texture buffers */
+	EAGXFeaturesTextureBuffers = 1llu << 21llu,
+	/** Supports max. compute threads per threadgroup */
+	EAGXFeaturesMaxThreadsPerThreadgroup = 1llu << 22llu,
+	/** Supports parallel render encoders */
+	EAGXFeaturesParallelRenderEncoders = 1llu << 23llu,
+	/** Supports indirect argument buffers */
+	EAGXFeaturesIABs = 1llu << 24llu,
+	/** Supports specifying the mutability of buffers bound to PSOs */
+    EAGXFeaturesPipelineBufferMutability = 1llu << 25llu,
+    /** Supports tile shaders */
+    EAGXFeaturesTileShaders = 1llu << 26llu,
+	/** Unused Reserved Bit */
+	EAGXFeaturesUnusedReservedBit27 = 1llu << 27llu, // was EAGXFeaturesSeparateTessellation
+	/** Supports indirect argument buffers Tier 2 */
+	EAGXFeaturesTier2IABs = 1llu << 28llu,
+};
+
+/**
+ * FAGXCommandQueue:
+ */
+class FAGXCommandQueue
+{
+public:
+#pragma mark - Public C++ Boilerplate -
+
+	/**
+	 * Constructor
+	 * @param Device The Metal device to create on.
+	 * @param MaxNumCommandBuffers The maximum number of incomplete command-buffers, defaults to 0 which implies the system default.
+	 */
+	FAGXCommandQueue(mtlpp::Device Device, uint32 const MaxNumCommandBuffers = 0);
+	
+	/** Destructor */
+	~FAGXCommandQueue(void);
+	
+#pragma mark - Public Command Buffer Mutators -
+
+	/**
+	 * Start encoding to CommandBuffer. It is an error to call this with any outstanding command encoders or current command buffer.
+	 * Instead call EndEncoding & CommitCommandBuffer before calling this.
+	 * @param CommandBuffer The new command buffer to begin encoding to.
+	 */
+	mtlpp::CommandBuffer CreateCommandBuffer(void);
+	
+	/**
+	 * Commit the supplied command buffer immediately.
+	 * @param CommandBuffer The command buffer to commit, must be non-nil.
+ 	 */
+	void CommitCommandBuffer(mtlpp::CommandBuffer& CommandBuffer);
+	
+	/**
+	 * Deferred contexts submit their internal lists of command-buffers out of order, the command-queue takes ownership and handles reordering them & lazily commits them once all command-buffer lists are submitted.
+	 * @param BufferList The list of buffers to enqueue into the command-queue at the given index.
+	 * @param Index The 0-based index to commit BufferList's contents into relative to other active deferred contexts.
+	 * @param Count The total number of deferred contexts that will submit - only once all are submitted can any command-buffer be committed.
+	 */
+	void SubmitCommandBuffers(TArray<mtlpp::CommandBuffer> BufferList, uint32 Index, uint32 Count);
+
+	/** @returns Creates a new MTLFence or nil if this is unsupported */
+	FAGXFence* CreateFence(ns::String const& Label) const;
+	
+	/** @params Fences An array of command-buffer fences for the committed command-buffers */
+	void GetCommittedCommandBufferFences(TArray<mtlpp::CommandBufferFence>& Fences);
+	
+#pragma mark - Public Command Queue Accessors -
+	
+	/** @returns The command queue's native device. */
+	mtlpp::Device& GetDevice(void);
+	
+	/** Converts a Metal v1.1+ resource option to something valid on the current version. */
+	static mtlpp::ResourceOptions GetCompatibleResourceOptions(mtlpp::ResourceOptions Options);
+	
+	/**
+	 * @param InFeature A specific Metal feature to check for.
+	 * @returns True if the requested feature is supported, else false.
+	 */
+	static inline bool SupportsFeature(EAGXFeatures InFeature) { return ((Features & InFeature) != 0); }
+
+	/**
+	* @param InFeature A specific Metal feature to check for.
+	* @returns True if RHISupportsSeparateMSAAAndResolveTextures will be true.  
+	* Currently Mac only.
+	*/
+	static inline bool SupportsSeparateMSAAAndResolveTarget() { return (PLATFORM_MAC != 0 || GMaxRHIFeatureLevel >= ERHIFeatureLevel::SM5); }
+
+#pragma mark - Public Debug Support -
+
+	/** Inserts a boundary that marks the end of a frame for the debug capture tool. */
+	void InsertDebugCaptureBoundary(void);
+	
+	/** Enable or disable runtime debugging features. */
+	void SetRuntimeDebuggingLevel(int32 const Level);
+	
+	/** @returns The level of runtime debugging features enabled. */
+	int32 GetRuntimeDebuggingLevel(void) const;
+
+private:
+#pragma mark - Private Member Variables -
+	mtlpp::Device Device;
+	mtlpp::CommandQueue CommandQueue;
+	TArray<TArray<mtlpp::CommandBuffer>> CommandBuffers;
+	TLockFreePointerListLIFO<mtlpp::CommandBufferFence> CommandBufferFences;
+	uint64 ParallelCommandLists;
+	int32 RuntimeDebuggingLevel;
+	static NSUInteger PermittedOptions;
+	static uint64 Features;
+};
diff --git a/Engine/Source/Runtime/Apple/AGXRHI/Private/AGXCommands.cpp b/Engine/Source/Runtime/Apple/AGXRHI/Private/AGXCommands.cpp
new file mode 100644
--- /dev/null
+++ b/Engine/Source/Runtime/Apple/AGXRHI/Private/AGXCommands.cpp
@@ -0,0 +1,732 @@
+// Copyright Epic Games, Inc. All Rights Reserved.
+
+/*=============================================================================
+	AGXCommands.cpp: AGX RHI commands implementation.
+=============================================================================*/
+
+#include "AGXRHIPrivate.h"
+#include "AGXShaderTypes.h"
+#include "AGXGraphicsPipelineState.h"
+#include "AGXComputePipelineState.h"
+#include "GlobalShader.h"
+#include "OneColorShader.h"
+#include "RHICommandList.h"
+#include "RHIStaticStates.h"
+#include "ShaderParameterUtils.h"
+#include "SceneUtils.h"
+#include "AGXProfiler.h"
+#include "AGXCommandBuffer.h"
+#include "StaticBoundShaderState.h"
+#include "EngineGlobals.h"
+#include "PipelineStateCache.h"
+
+static const bool GUsesInvertedZ = true;
+
+/** Vertex declaration for just one FVector4 position. */
+class FVector4VertexDeclaration : public FRenderResource
+{
+public:
+	FVertexDeclarationRHIRef VertexDeclarationRHI;
+	virtual void InitRHI() override
+	{
+		FVertexDeclarationElementList Elements;
+		Elements.Add(FVertexElement(0, 0, VET_Float4, 0, sizeof(FVector4)));
+		VertexDeclarationRHI = PipelineStateCache::GetOrCreateVertexDeclaration(Elements);
+	}
+	virtual void ReleaseRHI() override
+	{
+		VertexDeclarationRHI.SafeRelease();
+	}
+};
+static TGlobalResource<FVector4VertexDeclaration> GVector4VertexDeclaration;
+
+mtlpp::PrimitiveType AGXTranslatePrimitiveType(uint32 PrimitiveType)
+{
+	switch (PrimitiveType)
+	{
+		case PT_TriangleList:	return mtlpp::PrimitiveType::Triangle;
+		case PT_TriangleStrip:	return mtlpp::PrimitiveType::TriangleStrip;
+		case PT_LineList:		return mtlpp::PrimitiveType::Line;
+		case PT_PointList:		return mtlpp::PrimitiveType::Point;
+		// Metal doesn't actually actually draw in control-point patch-lists because of the way the compute shader stage works - it can handle any arbitrary patch size and will output triangles.
+		case PT_1_ControlPointPatchList:
+		case PT_2_ControlPointPatchList:
+		case PT_3_ControlPointPatchList:
+		case PT_4_ControlPointPatchList:
+		case PT_5_ControlPointPatchList:
+		case PT_6_ControlPointPatchList:
+		case PT_7_ControlPointPatchList:
+		case PT_8_ControlPointPatchList:
+		case PT_9_ControlPointPatchList:
+		case PT_10_ControlPointPatchList:
+		case PT_11_ControlPointPatchList:
+		case PT_12_ControlPointPatchList:
+		case PT_13_ControlPointPatchList:
+		case PT_14_ControlPointPatchList:
+		case PT_15_ControlPointPatchList:
+		case PT_16_ControlPointPatchList:
+		case PT_17_ControlPointPatchList:
+		case PT_18_ControlPointPatchList:
+		case PT_19_ControlPointPatchList:
+		case PT_20_ControlPointPatchList:
+		case PT_21_ControlPointPatchList:
+		case PT_22_ControlPointPatchList:
+		case PT_23_ControlPointPatchList:
+		case PT_24_ControlPointPatchList:
+		case PT_25_ControlPointPatchList:
+		case PT_26_ControlPointPatchList:
+		case PT_27_ControlPointPatchList:
+		case PT_28_ControlPointPatchList:
+		case PT_29_ControlPointPatchList:
+		case PT_30_ControlPointPatchList:
+		case PT_31_ControlPointPatchList:
+		case PT_32_ControlPointPatchList:
+		{
+			return mtlpp::PrimitiveType::Triangle;
+		}
+		default:
+			METAL_FATAL_ERROR(TEXT("Unsupported primitive type %d"), (int32)PrimitiveType);
+			return mtlpp::PrimitiveType::Triangle;
+	}
+}
+
+static FORCEINLINE EAGXShaderStages GetShaderStage(FRHIGraphicsShader* ShaderRHI)
+{
+	EAGXShaderStages Stage = EAGXShaderStages::Num;
+	switch (ShaderRHI->GetFrequency())
+	{
+	case SF_Vertex:		Stage = EAGXShaderStages::Vertex; break;
+	case SF_Pixel:		Stage = EAGXShaderStages::Pixel; break;
+	default:
+		checkf(0, TEXT("FRHIShader Type %d is invalid or unsupported!"), (int32)ShaderRHI->GetFrequency());
+		NOT_SUPPORTED("RHIShaderStage");
+		break;
+	}
+
+	return Stage;
+}
+
+void FAGXRHICommandContext::RHISetStreamSource(uint32 StreamIndex, FRHIBuffer* VertexBufferRHI,uint32 Offset)
+{
+	@autoreleasepool {
+		FAGXVertexBuffer* VertexBuffer = ResourceCast(VertexBufferRHI);
+		
+		FAGXBuffer TheBuffer = nil;
+		if(VertexBuffer && !VertexBuffer->Data)
+		{
+			TheBuffer = VertexBuffer->GetCurrentBuffer();
+		}
+		
+		Context->GetCurrentState().SetVertexStream(StreamIndex, VertexBuffer ? &TheBuffer : nil, VertexBuffer ? VertexBuffer->Data : nil, Offset, VertexBuffer ? VertexBuffer->GetSize() : 0);
+	}
+}
+
+void FAGXRHICommandContext::RHISetComputeShader(FRHIComputeShader* ComputeShaderRHI)
+{
+	@autoreleasepool {
+	FAGXComputeShader* ComputeShader = ResourceCast(ComputeShaderRHI);
+
+	// cache this for Dispatch
+	// sets this compute shader pipeline as the current (this resets all state, so we need to set all resources after calling this)
+	Context->GetCurrentState().SetComputeShader(ComputeShader);
+
+	ApplyStaticUniformBuffers(ComputeShader);
+	}
+}
+
+void FAGXRHICommandContext::RHISetComputePipelineState(FRHIComputePipelineState* ComputePipelineState)
+{
+	@autoreleasepool {
+	FAGXComputePipelineState* ComputePipeline = ResourceCast(ComputePipelineState);
+	
+	// cache this for Dispatch
+	// sets this compute shader pipeline as the current (this resets all state, so we need to set all resources after calling this)
+	Context->GetCurrentState().SetComputeShader(ComputePipeline->GetComputeShader());
+
+	ApplyStaticUniformBuffers(ComputePipeline->GetComputeShader());
+	}
+}
+
+void FAGXRHICommandContext::RHIDispatchComputeShader(uint32 ThreadGroupCountX, uint32 ThreadGroupCountY, uint32 ThreadGroupCountZ)
+{
+	@autoreleasepool {
+	ThreadGroupCountX = FMath::Max(ThreadGroupCountX, 1u);
+	ThreadGroupCountY = FMath::Max(ThreadGroupCountY, 1u);
+	ThreadGroupCountZ = FMath::Max(ThreadGroupCountZ, 1u);
+	
+	Context->Dispatch(ThreadGroupCountX, ThreadGroupCountY, ThreadGroupCountZ);
+	}
+}
+
+void FAGXRHICommandContext::RHIDispatchIndirectComputeShader(FRHIBuffer* ArgumentBufferRHI, uint32 ArgumentOffset)
+{
+	@autoreleasepool {
+	if (GetAGXDeviceContext().SupportsFeature(EAGXFeaturesIndirectBuffer))
+	{
+		FAGXVertexBuffer* VertexBuffer = ResourceCast(ArgumentBufferRHI);
+		
+		Context->DispatchIndirect(VertexBuffer, ArgumentOffset);
+	}
+	else
+	{
+		NOT_SUPPORTED("RHIDispatchIndirectComputeShader");
+	}
+	}
+}
+
+void FAGXRHICommandContext::RHISetViewport(float MinX, float MinY,float MinZ, float MaxX, float MaxY,float MaxZ)
+{
+	@autoreleasepool {
+	mtlpp::Viewport Viewport;
+	Viewport.originX = MinX;
+	Viewport.originY = MinY;
+	Viewport.width = MaxX - MinX;
+	Viewport.height = MaxY - MinY;
+	Viewport.znear = MinZ;
+	Viewport.zfar = MaxZ;
+	
+	Context->GetCurrentState().SetViewport(Viewport);
+	}
+}
+
+void FAGXRHICommandContext::RHISetStereoViewport(float LeftMinX, float RightMinX, float LeftMinY, float RightMinY, float MinZ, float LeftMaxX, float RightMaxX, float LeftMaxY, float RightMaxY, float MaxZ)
+{
+	if (FAGXCommandQueue::SupportsFeature(EAGXFeaturesMultipleViewports))
+	{
+		@autoreleasepool {
+		mtlpp::Viewport Viewport[2];
+		
+		Viewport[0].originX = LeftMinX;
+		Viewport[0].originY = LeftMinY;
+		Viewport[0].width = LeftMaxX - LeftMinX;
+		Viewport[0].height = LeftMaxY - LeftMinY;
+		Viewport[0].znear = MinZ;
+		Viewport[0].zfar = MaxZ;
+		
+		Viewport[1].originX = RightMinX;
+		Viewport[1].originY = RightMinY;
+		Viewport[1].width = RightMaxX - RightMinX;
+		Viewport[1].height = RightMaxY - RightMinY;
+		Viewport[1].znear = MinZ;
+		Viewport[1].zfar = MaxZ;
+		
+		Context->GetCurrentState().SetViewports(Viewport, 2);
+		}
+	}
+	else
+	{
+		NOT_SUPPORTED("RHISetStereoViewport");
+	}
+}
+
+void FAGXRHICommandContext::RHISetMultipleViewports(uint32 Count, const FViewportBounds* Data)
+{ 
+	NOT_SUPPORTED("RHISetMultipleViewports");
+}
+
+void FAGXRHICommandContext::RHISetScissorRect(bool bEnable,uint32 MinX,uint32 MinY,uint32 MaxX,uint32 MaxY)
+{
+	@autoreleasepool {
+	mtlpp::ScissorRect Scissor;
+	Scissor.x = MinX;
+	Scissor.y = MinY;
+	Scissor.width = MaxX - MinX;
+	Scissor.height = MaxY - MinY;
+
+	// metal doesn't support 0 sized scissor rect
+	if (bEnable == false || Scissor.width == 0 || Scissor.height == 0)
+	{
+		mtlpp::Viewport const& Viewport = Context->GetCurrentState().GetViewport(0);
+		CGSize FBSize = Context->GetCurrentState().GetFrameBufferSize();
+		
+		Scissor.x = Viewport.originX;
+		Scissor.y = Viewport.originY;
+		Scissor.width = (Viewport.originX + Viewport.width <= FBSize.width) ? Viewport.width : FBSize.width - Viewport.originX;
+		Scissor.height = (Viewport.originY + Viewport.height <= FBSize.height) ? Viewport.height : FBSize.height - Viewport.originY;
+	}
+	Context->GetCurrentState().SetScissorRect(bEnable, Scissor);
+	}
+}
+
+void FAGXRHICommandContext::RHISetGraphicsPipelineState(FRHIGraphicsPipelineState* GraphicsState, bool bApplyAdditionalState)
+{
+	@autoreleasepool {
+		FAGXGraphicsPipelineState* PipelineState = ResourceCast(GraphicsState);
+		if (AGXSafeGetRuntimeDebuggingLevel() >= EAGXDebugLevelResetOnBind && Context->GetCurrentState().GetGraphicsPSO() != PipelineState)
+		{
+			Context->GetCurrentRenderPass().GetCurrentCommandEncoder().ResetLive();
+		}
+		Context->GetCurrentState().SetGraphicsPipelineState(PipelineState);
+
+		RHISetStencilRef(0);
+		RHISetBlendFactor(FLinearColor(1.0f, 1.0f, 1.0f));
+
+		if (bApplyAdditionalState)
+		{
+			ApplyStaticUniformBuffers(PipelineState->VertexShader.GetReference());
+			ApplyStaticUniformBuffers(PipelineState->PixelShader.GetReference());
+		}
+	}
+}
+
+void FAGXRHICommandContext::RHISetStaticUniformBuffers(const FUniformBufferStaticBindings& InUniformBuffers)
+{
+	FMemory::Memzero(GlobalUniformBuffers.GetData(), GlobalUniformBuffers.Num() * sizeof(FRHIUniformBuffer*));
+
+	for (int32 Index = 0; Index < InUniformBuffers.GetUniformBufferCount(); ++Index)
+	{
+		GlobalUniformBuffers[InUniformBuffers.GetSlot(Index)] = InUniformBuffers.GetUniformBuffer(Index);
+	}
+}
+
+void FAGXRHICommandContext::RHISetUAVParameter(FRHIPixelShader* PixelShaderRHI, uint32 UAVIndex, FRHIUnorderedAccessView* UAVRHI)
+{
+	@autoreleasepool {
+		FAGXUnorderedAccessView* UAV = ResourceCast(UAVRHI);
+		Context->GetCurrentState().SetShaderUnorderedAccessView(EAGXShaderStages::Pixel, UAVIndex, UAV);
+	}
+}
+
+void FAGXRHICommandContext::RHISetUAVParameter(FRHIComputeShader* ComputeShaderRHI, uint32 UAVIndex, FRHIUnorderedAccessView* UAVRHI)
+{
+	@autoreleasepool {
+	FAGXUnorderedAccessView* UAV = ResourceCast(UAVRHI);
+	Context->GetCurrentState().SetShaderUnorderedAccessView(EAGXShaderStages::Compute, UAVIndex, UAV);
+	}
+}
+
+void FAGXRHICommandContext::RHISetUAVParameter(FRHIComputeShader* ComputeShaderRHI,uint32 UAVIndex, FRHIUnorderedAccessView* UAVRHI, uint32 InitialCount)
+{
+	@autoreleasepool {
+		FAGXUnorderedAccessView* UAV = ResourceCast(UAVRHI);
+		Context->GetCurrentState().SetShaderUnorderedAccessView(EAGXShaderStages::Compute, UAVIndex, UAV);
+	}
+}
+
+
+void FAGXRHICommandContext::RHISetShaderTexture(FRHIGraphicsShader* ShaderRHI, uint32 TextureIndex, FRHITexture* NewTextureRHI)
+{
+	@autoreleasepool {
+	FAGXSurface* Surface = AGXGetMetalSurfaceFromRHITexture(NewTextureRHI);
+	EAGXShaderStages Stage = GetShaderStage(ShaderRHI);
+	if (Surface != nullptr)
+	{
+        if (Surface->Texture || !(Surface->Flags & TexCreate_Presentable))
+        {
+            Context->GetCurrentState().SetShaderTexture(Stage, Surface->Texture, TextureIndex, (mtlpp::ResourceUsage)(mtlpp::ResourceUsage::Read|mtlpp::ResourceUsage::Sample));
+        }
+        else
+        {
+            FAGXTexture Tex = Surface->GetCurrentTexture();
+            Context->GetCurrentState().SetShaderTexture(Stage, Tex, TextureIndex, (mtlpp::ResourceUsage)(mtlpp::ResourceUsage::Read|mtlpp::ResourceUsage::Sample));
+        }
+	}
+	else
+	{
+		Context->GetCurrentState().SetShaderTexture(Stage, nil, TextureIndex, mtlpp::ResourceUsage(0));
+	}
+	}
+}
+
+void FAGXRHICommandContext::RHISetShaderTexture(FRHIComputeShader* ComputeShader, uint32 TextureIndex, FRHITexture* NewTextureRHI)
+{
+	@autoreleasepool {
+	FAGXSurface* Surface = AGXGetMetalSurfaceFromRHITexture(NewTextureRHI);
+	if (Surface != nullptr)
+    {
+        if (Surface->Texture || !(Surface->Flags & TexCreate_Presentable))
+        {
+            Context->GetCurrentState().SetShaderTexture(EAGXShaderStages::Compute, Surface->Texture, TextureIndex, (mtlpp::ResourceUsage)(mtlpp::ResourceUsage::Read|mtlpp::ResourceUsage::Sample));
+        }
+        else
+        {
+            FAGXTexture Tex = Surface->GetCurrentTexture();
+            Context->GetCurrentState().SetShaderTexture(EAGXShaderStages::Compute, Tex, TextureIndex, (mtlpp::ResourceUsage)(mtlpp::ResourceUsage::Read|mtlpp::ResourceUsage::Sample));
+        }
+	}
+	else
+	{
+		Context->GetCurrentState().SetShaderTexture(EAGXShaderStages::Compute, nil, TextureIndex, mtlpp::ResourceUsage(0));
+	}
+	}
+}
+
+
+void FAGXRHICommandContext::RHISetShaderResourceViewParameter(FRHIGraphicsShader* ShaderRHI, uint32 TextureIndex, FRHIShaderResourceView* SRVRHI)
+{
+	@autoreleasepool {
+	FAGXShaderResourceView* SRV = ResourceCast(SRVRHI);
+	EAGXShaderStages Stage = GetShaderStage(ShaderRHI);
+	Context->GetCurrentState().SetShaderResourceView(Context, Stage, TextureIndex, SRV);
+	}
+}
+
+void FAGXRHICommandContext::RHISetShaderResourceViewParameter(FRHIComputeShader* ComputeShaderRHI,uint32 TextureIndex, FRHIShaderResourceView* SRVRHI)
+{
+	@autoreleasepool {
+	FAGXShaderResourceView* SRV = ResourceCast(SRVRHI);
+	Context->GetCurrentState().SetShaderResourceView(Context, EAGXShaderStages::Compute, TextureIndex, SRV);
+	}
+}
+
+
+void FAGXRHICommandContext::RHISetShaderSampler(FRHIGraphicsShader* ShaderRHI, uint32 SamplerIndex, FRHISamplerState* NewStateRHI)
+{
+	@autoreleasepool {
+	FAGXSamplerState* NewState = ResourceCast(NewStateRHI);
+	EAGXShaderStages Stage = GetShaderStage(ShaderRHI);
+	Context->GetCurrentState().SetShaderSamplerState(Stage, NewState, SamplerIndex);
+	}
+}
+
+void FAGXRHICommandContext::RHISetShaderSampler(FRHIComputeShader* ComputeShader, uint32 SamplerIndex, FRHISamplerState* NewStateRHI)
+{
+	@autoreleasepool {
+	FAGXSamplerState* NewState = ResourceCast(NewStateRHI);
+
+	Context->GetCurrentState().SetShaderSamplerState(EAGXShaderStages::Compute, NewState, SamplerIndex);
+	}
+}
+
+void FAGXRHICommandContext::RHISetShaderParameter(FRHIGraphicsShader* ShaderRHI, uint32 BufferIndex, uint32 BaseIndex, uint32 NumBytes, const void* NewValue)
+{
+	@autoreleasepool {
+	EAGXShaderStages Stage = GetShaderStage(ShaderRHI);
+	Context->GetCurrentState().GetShaderParameters(Stage).Set(BufferIndex, BaseIndex, NumBytes, NewValue);
+	}
+}
+
+void FAGXRHICommandContext::RHISetShaderParameter(FRHIComputeShader* ComputeShaderRHI,uint32 BufferIndex, uint32 BaseIndex, uint32 NumBytes, const void* NewValue)
+{
+	@autoreleasepool {
+	Context->GetCurrentState().GetShaderParameters(EAGXShaderStages::Compute).Set(BufferIndex, BaseIndex, NumBytes, NewValue);
+	}
+}
+
+void FAGXRHICommandContext::RHISetStencilRef(uint32 StencilRef)
+{
+	Context->GetCurrentState().SetStencilRef(StencilRef);
+}
+
+void FAGXRHICommandContext::RHISetBlendFactor(const FLinearColor& BlendFactor)
+{
+	Context->GetCurrentState().SetBlendFactor(BlendFactor);
+}
+
+void FAGXRHICommandContext::SetRenderTargets(uint32 NumSimultaneousRenderTargets, const FRHIRenderTargetView* NewRenderTargets,
+	const FRHIDepthRenderTargetView* NewDepthStencilTargetRHI)
+{
+	@autoreleasepool {
+	FAGXContext* Manager = Context;
+	FRHIDepthRenderTargetView DepthView;
+	if (NewDepthStencilTargetRHI)
+	{
+		DepthView = *NewDepthStencilTargetRHI;
+	}
+	else
+	{
+		DepthView = FRHIDepthRenderTargetView(nullptr, ERenderTargetLoadAction::EClear, ERenderTargetStoreAction::ENoAction);
+	}
+
+	FRHISetRenderTargetsInfo Info(NumSimultaneousRenderTargets, NewRenderTargets, DepthView);
+	SetRenderTargetsAndClear(Info);
+	}
+}
+
+void FAGXRHICommandContext::SetRenderTargetsAndClear(const FRHISetRenderTargetsInfo& RenderTargetsInfo)
+{
+	@autoreleasepool {
+		
+	FRHIRenderPassInfo PassInfo;
+	bool bHasTarget = (RenderTargetsInfo.DepthStencilRenderTarget.Texture != nullptr);
+	FAGXContext* Manager = Context;
+	
+	for (uint32 i = 0; i < RenderTargetsInfo.NumColorRenderTargets; i++)
+	{
+		if (RenderTargetsInfo.ColorRenderTarget[i].Texture)
+		{
+			PassInfo.ColorRenderTargets[i].RenderTarget = RenderTargetsInfo.ColorRenderTarget[i].Texture;
+			PassInfo.ColorRenderTargets[i].ArraySlice = RenderTargetsInfo.ColorRenderTarget[i].ArraySliceIndex;
+			PassInfo.ColorRenderTargets[i].MipIndex = RenderTargetsInfo.ColorRenderTarget[i].MipIndex;
+			PassInfo.ColorRenderTargets[i].Action = MakeRenderTargetActions(RenderTargetsInfo.ColorRenderTarget[i].LoadAction, RenderTargetsInfo.ColorRenderTarget[i].StoreAction);
+		bHasTarget = (RenderTargetsInfo.ColorRenderTarget[i].Texture != nullptr);
+			PassInfo.bIsMSAA |= PassInfo.ColorRenderTargets[i].RenderTarget->GetNumSamples() > 1;
+		}
+	}
+		
+	if (RenderTargetsInfo.DepthStencilRenderTarget.Texture)
+	{
+		PassInfo.DepthStencilRenderTarget.DepthStencilTarget = RenderTargetsInfo.DepthStencilRenderTarget.Texture;
+		PassInfo.DepthStencilRenderTarget.ExclusiveDepthStencil = RenderTargetsInfo.DepthStencilRenderTarget.GetDepthStencilAccess();
+		PassInfo.DepthStencilRenderTarget.Action = MakeDepthStencilTargetActions(MakeRenderTargetActions(RenderTargetsInfo.DepthStencilRenderTarget.DepthLoadAction, RenderTargetsInfo.DepthStencilRenderTarget.DepthStoreAction), MakeRenderTargetActions(RenderTargetsInfo.DepthStencilRenderTarget.StencilLoadAction, RenderTargetsInfo.DepthStencilRenderTarget.GetStencilStoreAction()));
+		PassInfo.bIsMSAA |= RenderTargetsInfo.DepthStencilRenderTarget.Texture->GetNumSamples() > 1;
+	}
+		
+	PassInfo.NumOcclusionQueries = UINT16_MAX;
+	PassInfo.bOcclusionQueries = true;
+
+	// Ignore any attempt to "clear" the render-targets as that is senseless with the way AGXRHI has to try and coalesce passes.
+	if (bHasTarget)
+	{
+		Manager->SetRenderPassInfo(PassInfo);
+
+		// Set the viewport to the full size of render target 0.
+		if (RenderTargetsInfo.ColorRenderTarget[0].Texture)
+		{
+			const FRHIRenderTargetView& RenderTargetView = RenderTargetsInfo.ColorRenderTarget[0];
+			FAGXSurface* RenderTarget = AGXGetMetalSurfaceFromRHITexture(RenderTargetView.Texture);
+
+			uint32 Width = FMath::Max((uint32)(RenderTarget->Texture.GetWidth() >> RenderTargetView.MipIndex), (uint32)1);
+			uint32 Height = FMath::Max((uint32)(RenderTarget->Texture.GetHeight() >> RenderTargetView.MipIndex), (uint32)1);
+
+			RHISetViewport(0, 0, 0.0f, Width, Height, 1.0f);
+		}
+	}
+	}
+}
+
+
+void FAGXRHICommandContext::RHIDrawPrimitive(uint32 BaseVertexIndex, uint32 NumPrimitives, uint32 NumInstances)
+{
+	@autoreleasepool {
+	SCOPE_CYCLE_COUNTER(STAT_AGXDrawCallTime);
+		
+	uint32 PrimitiveType = Context->GetCurrentState().GetPrimitiveType();
+	
+		
+	//checkf(NumInstances == 1, TEXT("Currently only 1 instance is supported"));
+	
+	NumInstances = FMath::Max(NumInstances,1u);
+	
+	RHI_DRAW_CALL_STATS(PrimitiveType,NumInstances*NumPrimitives);
+
+	// how many verts to render
+	uint32 NumVertices = GetVertexCountForPrimitiveCount(NumPrimitives, PrimitiveType);
+	uint32 VertexCount = GetVertexCountForPrimitiveCount(NumPrimitives,PrimitiveType);
+	
+	Context->DrawPrimitive(PrimitiveType, BaseVertexIndex, NumPrimitives, NumInstances);
+	}
+}
+
+void FAGXRHICommandContext::RHIDrawPrimitiveIndirect(FRHIBuffer* ArgumentBufferRHI, uint32 ArgumentOffset)
+{
+	@autoreleasepool {
+#if PLATFORM_IOS
+	NOT_SUPPORTED("RHIDrawPrimitiveIndirect");
+#else
+	SCOPE_CYCLE_COUNTER(STAT_AGXDrawCallTime);
+	uint32 PrimitiveType = Context->GetCurrentState().GetPrimitiveType();
+	
+	
+	RHI_DRAW_CALL_STATS(PrimitiveType,1);
+	FAGXResourceMultiBuffer* ArgumentBuffer = ResourceCast(ArgumentBufferRHI);
+	
+	Context->DrawPrimitiveIndirect(PrimitiveType, ArgumentBuffer, ArgumentOffset);
+#endif
+	}
+}
+
+void FAGXRHICommandContext::RHIDrawIndexedPrimitive(FRHIBuffer* IndexBufferRHI, int32 BaseVertexIndex, uint32 FirstInstance,
+	uint32 NumVertices, uint32 StartIndex, uint32 NumPrimitives, uint32 NumInstances)
+{
+	@autoreleasepool {
+	SCOPE_CYCLE_COUNTER(STAT_AGXDrawCallTime);
+	//checkf(NumInstances == 1, TEXT("Currently only 1 instance is supported"));
+	checkf(GRHISupportsBaseVertexIndex || BaseVertexIndex == 0, TEXT("BaseVertexIndex must be 0, see GRHISupportsBaseVertexIndex"));
+	checkf(GRHISupportsFirstInstance || FirstInstance == 0, TEXT("FirstInstance must be 0, see GRHISupportsFirstInstance"));
+	uint32 PrimitiveType = Context->GetCurrentState().GetPrimitiveType();
+	
+		
+	RHI_DRAW_CALL_STATS(PrimitiveType,FMath::Max(NumInstances,1u)*NumPrimitives);
+
+	FAGXResourceMultiBuffer* IndexBuffer = ResourceCast(IndexBufferRHI);
+	Context->DrawIndexedPrimitive(IndexBuffer->GetCurrentBuffer(), IndexBuffer->GetStride(), IndexBuffer->IndexType, PrimitiveType, BaseVertexIndex, FirstInstance, NumVertices, StartIndex, NumPrimitives, NumInstances);
+	}
+}
+
+void FAGXRHICommandContext::RHIDrawIndexedIndirect(FRHIBuffer* IndexBufferRHI, FRHIBuffer* ArgumentsBufferRHI, int32 DrawArgumentsIndex, uint32 NumInstances)
+{
+	@autoreleasepool {
+	if (GetAGXDeviceContext().SupportsFeature(EAGXFeaturesIndirectBuffer))
+	{
+		check(NumInstances > 1);
+		
+		SCOPE_CYCLE_COUNTER(STAT_AGXDrawCallTime);
+		
+		uint32 PrimitiveType = Context->GetCurrentState().GetPrimitiveType();
+		
+
+		RHI_DRAW_CALL_STATS(PrimitiveType,1);
+		FAGXResourceMultiBuffer* IndexBuffer = ResourceCast(IndexBufferRHI);
+		FAGXResourceMultiBuffer* ArgumentsBuffer = ResourceCast(ArgumentsBufferRHI);
+		
+		Context->DrawIndexedIndirect(IndexBuffer, PrimitiveType, ArgumentsBuffer, DrawArgumentsIndex, NumInstances);
+	}
+	else
+	{
+		NOT_SUPPORTED("RHIDrawIndexedIndirect");
+	}
+	}
+}
+
+void FAGXRHICommandContext::RHIDrawIndexedPrimitiveIndirect(FRHIBuffer* IndexBufferRHI, FRHIBuffer* ArgumentBufferRHI, uint32 ArgumentOffset)
+{
+	@autoreleasepool {
+	if (GetAGXDeviceContext().SupportsFeature(EAGXFeaturesIndirectBuffer))
+	{
+		SCOPE_CYCLE_COUNTER(STAT_AGXDrawCallTime);
+		
+		uint32 PrimitiveType = Context->GetCurrentState().GetPrimitiveType();
+		
+
+		RHI_DRAW_CALL_STATS(PrimitiveType,1);
+		FAGXResourceMultiBuffer* IndexBuffer = ResourceCast(IndexBufferRHI);
+		FAGXResourceMultiBuffer* ArgumentsBuffer = ResourceCast(ArgumentBufferRHI);
+		
+		Context->DrawIndexedPrimitiveIndirect(PrimitiveType, IndexBuffer, ArgumentsBuffer, ArgumentOffset);
+	}
+	else
+	{
+		NOT_SUPPORTED("RHIDrawIndexedPrimitiveIndirect");
+	}
+	}
+}
+
+void FAGXDynamicRHI::SetupRecursiveResources()
+{
+    /*
+	@autoreleasepool {
+	static bool bSetupResources = false;
+	if (GRHISupportsRHIThread && !bSetupResources)
+	{
+		FRHICommandListImmediate& RHICmdList = FRHICommandListExecutor::GetImmediateCommandList();
+		extern int32 GCreateShadersOnLoad;
+		TGuardValue<int32> Guard(GCreateShadersOnLoad, 1);
+		auto ShaderMap = GetGlobalShaderMap(GMaxRHIFeatureLevel);
+		TShaderMapRef<TOneColorVS<true> > DefaultVertexShader(ShaderMap);
+		TShaderMapRef<TOneColorVS<true, true> > LayeredVertexShader(ShaderMap);
+		GVector4VertexDeclaration.InitRHI();
+		
+		for (uint32 Instanced = 0; Instanced < 2; Instanced++)
+		{
+			FShader* VertexShader = !Instanced ? (FShader*)*DefaultVertexShader : (FShader*)*LayeredVertexShader;
+			
+			for (int32 NumBuffers = 1; NumBuffers <= MaxSimultaneousRenderTargets; NumBuffers++)
+			{
+				FOneColorPS* PixelShader = NULL;
+				
+				// Set the shader to write to the appropriate number of render targets
+				// On AMD PC hardware, outputting to a color index in the shader without a matching render target set has a significant performance hit
+				if (NumBuffers <= 1)
+				{
+					TShaderMapRef<TOneColorPixelShaderMRT<1> > MRTPixelShader(ShaderMap);
+					PixelShader = *MRTPixelShader;
+				}
+				else if (IsFeatureLevelSupported( GMaxRHIShaderPlatform, ERHIFeatureLevel::SM5 ))
+				{
+					if (NumBuffers == 2)
+					{
+						TShaderMapRef<TOneColorPixelShaderMRT<2> > MRTPixelShader(ShaderMap);
+						PixelShader = *MRTPixelShader;
+					}
+					else if (NumBuffers== 3)
+					{
+						TShaderMapRef<TOneColorPixelShaderMRT<3> > MRTPixelShader(ShaderMap);
+						PixelShader = *MRTPixelShader;
+					}
+					else if (NumBuffers == 4)
+					{
+						TShaderMapRef<TOneColorPixelShaderMRT<4> > MRTPixelShader(ShaderMap);
+						PixelShader = *MRTPixelShader;
+					}
+					else if (NumBuffers == 5)
+					{
+						TShaderMapRef<TOneColorPixelShaderMRT<5> > MRTPixelShader(ShaderMap);
+						PixelShader = *MRTPixelShader;
+					}
+					else if (NumBuffers == 6)
+					{
+						TShaderMapRef<TOneColorPixelShaderMRT<6> > MRTPixelShader(ShaderMap);
+						PixelShader = *MRTPixelShader;
+					}
+					else if (NumBuffers == 7)
+					{
+						TShaderMapRef<TOneColorPixelShaderMRT<7> > MRTPixelShader(ShaderMap);
+						PixelShader = *MRTPixelShader;
+					}
+					else if (NumBuffers == 8)
+					{
+						TShaderMapRef<TOneColorPixelShaderMRT<8> > MRTPixelShader(ShaderMap);
+						PixelShader = *MRTPixelShader;
+					}
+				}
+				
+				// SetGlobalBoundShaderState(RHICmdList, GMaxRHIFeatureLevel, GClearMRTBoundShaderState[NumBuffers - 1][Instanced], GVector4VertexDeclaration.VertexDeclarationRHI, VertexShader, PixelShader);
+			}
+		}
+		
+		bSetupResources = true;
+	}
+	}
+    */
+}
+
+void FAGXRHICommandContext::RHIClearMRT(bool bClearColor,int32 NumClearColors,const FLinearColor* ClearColorArray,bool bClearDepth,float Depth,bool bClearStencil,uint32 Stencil)
+{
+	NOT_SUPPORTED("RHIClearMRT");
+}
+
+void FAGXDynamicRHI::RHIBlockUntilGPUIdle()
+{
+	@autoreleasepool {
+	ImmediateContext.Context->SubmitCommandBufferAndWait();
+	}
+}
+
+uint32 FAGXDynamicRHI::RHIGetGPUFrameCycles(uint32 GPUIndex)
+{
+	check(GPUIndex == 0);
+	return GGPUFrameTime;
+}
+
+void FAGXDynamicRHI::RHIExecuteCommandList(FRHICommandList* RHICmdList)
+{
+	NOT_SUPPORTED("RHIExecuteCommandList");
+}
+
+void FAGXRHICommandContext::RHISetDepthBounds(float MinDepth, float MaxDepth)
+{
+	METAL_IGNORED(FAGXRHICommandContextSetDepthBounds);
+}
+
+void FAGXRHICommandContext::RHISubmitCommandsHint()
+{
+	@autoreleasepool {
+    Context->SubmitCommandsHint();
+	}
+}
+
+void FAGXRHICommandContext::RHIDiscardRenderTargets(bool Depth, bool Stencil, uint32 ColorBitMask)
+{
+	Context->GetCurrentState().DiscardRenderTargets(Depth, Stencil, ColorBitMask);
+}
+
+IRHICommandContext* FAGXDynamicRHI::RHIGetDefaultContext()
+{
+	return &ImmediateContext;
+}
+
+IRHIComputeContext* FAGXDynamicRHI::RHIGetDefaultAsyncComputeContext()
+{
+	@autoreleasepool {
+	IRHIComputeContext* ComputeContext = GSupportsEfficientAsyncCompute && AsyncComputeContext ? AsyncComputeContext : RHIGetDefaultContext();
+	// On platforms that support non-async compute we set this to the normal context.  It won't be async, but the high level
+	// code can be agnostic if it wants to be.
+	return ComputeContext;
+	}
+}
+
+#if PLATFORM_USES_FIXED_RHI_CLASS
+#define INTERNAL_DECORATOR(Method) ((FAGXRHICommandContext&)CmdList.GetContext()).FAGXRHICommandContext::Method
+#include "RHICommandListCommandExecutes.inl"
+#endif
diff --git a/Engine/Source/Runtime/Apple/AGXRHI/Private/AGXComputeCommandEncoder.cpp b/Engine/Source/Runtime/Apple/AGXRHI/Private/AGXComputeCommandEncoder.cpp
new file mode 100644
--- /dev/null
+++ b/Engine/Source/Runtime/Apple/AGXRHI/Private/AGXComputeCommandEncoder.cpp
@@ -0,0 +1,454 @@
+// Copyright Epic Games, Inc. All Rights Reserved.
+
+/*=============================================================================
+	AGXComputeCommandEncoder.cpp: AGX command encoder wrapper.
+=============================================================================*/
+
+#include "AGXRHIPrivate.h"
+
+#include "AGXComputeCommandEncoder.h"
+#include "AGXCommandBuffer.h"
+#include "AGXFence.h"
+#include "AGXPipeline.h"
+
+#if MTLPP_CONFIG_VALIDATE && METAL_DEBUG_OPTIONS
+NS_ASSUME_NONNULL_BEGIN
+
+#if METAL_DEBUG_OPTIONS
+static NSString* GAGXDebugComputeShader = @"#include <metal_stdlib>\n"
+@"using namespace metal;\n"
+@"kernel void WriteCommandIndexCS(constant uint* Input [[ buffer(0) ]], device atomic_uint* Output [[ buffer(1) ]])\n"
+@"{\n"
+@"	atomic_store_explicit(Output, Input[0], memory_order_relaxed);\n"
+@"}\n";
+
+static id <MTLComputePipelineState> GetDebugComputeShaderState(id<MTLDevice> Device)
+{
+	static id<MTLComputePipelineState> State = nil;
+	if (!State)
+	{
+		id<MTLLibrary> Lib = [Device newLibraryWithSource:GAGXDebugComputeShader options:nil error:nullptr];
+		check(Lib);
+		id<MTLFunction> Func = [Lib newFunctionWithName:@"WriteCommandIndexCS"];
+		check(Func);
+		State = [Device newComputePipelineStateWithFunction:Func error:nil];
+		[Func release];
+		[Lib release];
+	}
+	check(State);
+	return State;
+}
+#endif
+
+@interface FAGXDebugComputeCommandEncoder : FAGXDebugCommandEncoder
+{
+@public
+	id<MTLComputeCommandEncoder> Inner;
+	FAGXCommandBufferDebugging Buffer;
+	FAGXShaderPipeline* Pipeline;
+
+#pragma mark - Private Member Variables -
+#if METAL_DEBUG_OPTIONS
+	FAGXDebugShaderResourceMask ResourceMask;
+    FAGXDebugBufferBindings ShaderBuffers;
+    FAGXDebugTextureBindings ShaderTextures;
+    FAGXDebugSamplerBindings ShaderSamplers;
+#endif
+}
+/** Initialise the wrapper with the provided command-buffer. */
+-(id)initWithEncoder:(id<MTLComputeCommandEncoder>)Encoder andCommandBuffer:(FAGXCommandBufferDebugging const&)Buffer;
+@end
+
+@implementation FAGXDebugComputeCommandEncoder
+
+APPLE_PLATFORM_OBJECT_ALLOC_OVERRIDES(FAGXDebugComputeCommandEncoder)
+
+-(id)initWithEncoder:(id<MTLComputeCommandEncoder>)Encoder andCommandBuffer:(FAGXCommandBufferDebugging const&)SourceBuffer
+{
+	id Self = [super init];
+	if (Self)
+	{
+        Inner = Encoder;
+		Buffer = SourceBuffer;
+        Pipeline = nil;
+	}
+	return Self;
+}
+
+-(void)dealloc
+{
+	[Pipeline release];
+	[super dealloc];
+}
+
+@end
+
+#if METAL_DEBUG_OPTIONS
+void FAGXComputeCommandEncoderDebugging::InsertDebugDispatch()
+{
+//	switch (Buffer->DebugLevel)
+//	{
+//		case EAGXDebugLevelConditionalSubmit:
+//		case EAGXDebugLevelWaitForComplete:
+//		case EAGXDebugLevelLogOperations:
+//		case EAGXDebugLevelValidation:
+//		{
+//			uint32 const Index = Buffer->DebugCommands.Num();
+//			[Inner setBytes:&Index length:sizeof(Index) atIndex:0];
+//			[Inner setBuffer:Buffer->DebugInfoBuffer offset:0 atIndex:1];
+//			[Inner setComputePipelineState:GetDebugComputeShaderState(Inner.device)];
+//			
+//			[Inner dispatchThreadgroups:mtlpp::Size const&Make(1, 1, 1) threadsPerThreadgroup:mtlpp::Size const&Make(1, 1, 1)];
+//
+//			if (Pipeline && Pipeline->ComputePipelineState)
+//			{
+//				[Inner setComputePipelineState:Pipeline->ComputePipelineState];
+//			}
+//			
+//			if (ShaderBuffers.Buffers[0])
+//			{
+//				[Inner setBuffer:ShaderBuffers.Buffers[0] offset:ShaderBuffers.Offsets[0] atIndex:0];
+//			}
+//			else if (ShaderBuffers.Bytes[0])
+//			{
+//				[Inner setBytes:ShaderBuffers.Bytes[0] length:ShaderBuffers.Offsets[0] atIndex:0];
+//			}
+//			
+//			if (ShaderBuffers.Buffers[1])
+//			{
+//				[Inner setBuffer:ShaderBuffers.Buffers[1] offset:ShaderBuffers.Offsets[1] atIndex:1];
+//			}
+//			else if (ShaderBuffers.Bytes[1])
+//			{
+//				[Inner setBytes:ShaderBuffers.Bytes[1] length:ShaderBuffers.Offsets[1] atIndex:1];
+//			}
+//		}
+//		default:
+//		{
+//			break;
+//		}
+//	}
+}
+#endif
+
+FAGXComputeCommandEncoderDebugging::FAGXComputeCommandEncoderDebugging()
+{
+	
+}
+FAGXComputeCommandEncoderDebugging::FAGXComputeCommandEncoderDebugging(mtlpp::ComputeCommandEncoder& Encoder, FAGXCommandBufferDebugging& Buffer)
+: FAGXCommandEncoderDebugging((FAGXDebugCommandEncoder*)[[[FAGXDebugComputeCommandEncoder alloc] initWithEncoder:Encoder.GetPtr() andCommandBuffer:Buffer] autorelease])
+{
+	Buffer.BeginComputeCommandEncoder([NSString stringWithFormat:@"Compute: %@", Encoder.GetLabel().GetPtr()]);
+	Encoder.SetAssociatedObject((void const*)&FAGXComputeCommandEncoderDebugging::Get, (FAGXCommandEncoderDebugging const&)*this);
+}
+FAGXComputeCommandEncoderDebugging::FAGXComputeCommandEncoderDebugging(FAGXDebugCommandEncoder* handle)
+: FAGXCommandEncoderDebugging((FAGXDebugCommandEncoder*)handle)
+{
+	
+}
+
+FAGXComputeCommandEncoderDebugging FAGXComputeCommandEncoderDebugging::Get(mtlpp::ComputeCommandEncoder& Encoder)
+{
+	return Encoder.GetAssociatedObject<FAGXComputeCommandEncoderDebugging>((void const*)&FAGXComputeCommandEncoderDebugging::Get);
+}
+
+void FAGXComputeCommandEncoderDebugging::InsertDebugSignpost(ns::String const& Label)
+{
+	((FAGXDebugComputeCommandEncoder*)m_ptr)->Buffer.InsertDebugSignpost(Label);
+}
+void FAGXComputeCommandEncoderDebugging::PushDebugGroup(ns::String const& Group)
+{
+    ((FAGXDebugComputeCommandEncoder*)m_ptr)->Buffer.PushDebugGroup(Group);
+}
+void FAGXComputeCommandEncoderDebugging::PopDebugGroup()
+{
+    ((FAGXDebugComputeCommandEncoder*)m_ptr)->Buffer.PopDebugGroup();
+#if METAL_DEBUG_OPTIONS
+	InsertDebugDispatch();
+#endif
+}
+
+void FAGXComputeCommandEncoderDebugging::EndEncoder()
+{
+	((FAGXDebugComputeCommandEncoder*)m_ptr)->Buffer.EndCommandEncoder();
+}
+
+void FAGXComputeCommandEncoderDebugging::DispatchThreadgroups(mtlpp::Size const& threadgroupsPerGrid, mtlpp::Size const& threadsPerThreadgroup)
+{
+#if METAL_DEBUG_OPTIONS
+	switch(((FAGXDebugComputeCommandEncoder*)m_ptr)->Buffer.GetPtr()->DebugLevel)
+	{
+		case EAGXDebugLevelConditionalSubmit:
+		case EAGXDebugLevelWaitForComplete:
+		case EAGXDebugLevelLogOperations:
+		{
+			((FAGXDebugComputeCommandEncoder*)m_ptr)->Buffer.Dispatch([NSString stringWithFormat:@"%s", __PRETTY_FUNCTION__]);
+		}
+		case EAGXDebugLevelValidation:
+		{
+			Validate();
+		}
+		default:
+		{
+			break;
+		}
+	}
+#endif
+}
+
+void FAGXComputeCommandEncoderDebugging::SetPipeline(FAGXShaderPipeline* Pipeline)
+{
+#if METAL_DEBUG_OPTIONS
+	((FAGXDebugComputeCommandEncoder*)m_ptr)->Pipeline = [Pipeline retain];
+	switch (((FAGXDebugComputeCommandEncoder*)m_ptr)->Buffer.GetPtr()->DebugLevel)
+	{
+		case EAGXDebugLevelConditionalSubmit:
+		case EAGXDebugLevelWaitForComplete:
+		case EAGXDebugLevelLogOperations:
+		{
+			((FAGXDebugComputeCommandEncoder*)m_ptr)->Buffer.SetPipeline(Pipeline->ComputePipelineState.GetLabel());
+		}
+		default:
+		{
+			break;
+		}
+	}
+#endif
+}
+
+void FAGXComputeCommandEncoderDebugging::SetBytes(const void * bytes, NSUInteger length, NSUInteger index)
+{
+#if METAL_DEBUG_OPTIONS
+	switch (((FAGXDebugComputeCommandEncoder*)m_ptr)->Buffer.GetPtr()->DebugLevel)
+	{
+		case EAGXDebugLevelConditionalSubmit:
+		case EAGXDebugLevelWaitForComplete:
+		case EAGXDebugLevelLogOperations:
+		case EAGXDebugLevelValidation:
+		{
+			((FAGXDebugComputeCommandEncoder*)m_ptr)->ShaderBuffers.Buffers[index] = nil;
+			((FAGXDebugComputeCommandEncoder*)m_ptr)->ShaderBuffers.Bytes[index] = bytes;
+			((FAGXDebugComputeCommandEncoder*)m_ptr)->ShaderBuffers.Offsets[index] = length;
+			((FAGXDebugComputeCommandEncoder*)m_ptr)->ResourceMask.BufferMask = bytes ? (((FAGXDebugComputeCommandEncoder*)m_ptr)->ResourceMask.BufferMask | (1 << (index))) : (((FAGXDebugComputeCommandEncoder*)m_ptr)->ResourceMask.BufferMask & ~(1 << (index)));
+		}
+		default:
+		{
+			break;
+		}
+	}
+#endif
+}
+void FAGXComputeCommandEncoderDebugging::SetBuffer( FAGXBuffer const& buffer, NSUInteger offset, NSUInteger index)
+{
+#if METAL_DEBUG_OPTIONS
+	switch (((FAGXDebugComputeCommandEncoder*)m_ptr)->Buffer.GetPtr()->DebugLevel)
+	{
+		case EAGXDebugLevelConditionalSubmit:
+		case EAGXDebugLevelWaitForComplete:
+		case EAGXDebugLevelLogOperations:
+		case EAGXDebugLevelValidation:
+		{
+			((FAGXDebugComputeCommandEncoder*)m_ptr)->ShaderBuffers.Buffers[index] = buffer;
+			((FAGXDebugComputeCommandEncoder*)m_ptr)->ShaderBuffers.Bytes[index] = nil;
+			((FAGXDebugComputeCommandEncoder*)m_ptr)->ShaderBuffers.Offsets[index] = offset;
+			((FAGXDebugComputeCommandEncoder*)m_ptr)->ResourceMask.BufferMask = buffer ? (((FAGXDebugComputeCommandEncoder*)m_ptr)->ResourceMask.BufferMask | (1 << (index))) : (((FAGXDebugComputeCommandEncoder*)m_ptr)->ResourceMask.BufferMask & ~(1 << (index)));
+		}
+		default:
+		{
+			break;
+		}
+	}
+#endif
+}
+void FAGXComputeCommandEncoderDebugging::SetBufferOffset(NSUInteger offset, NSUInteger index)
+{
+#if METAL_DEBUG_OPTIONS
+	switch (((FAGXDebugComputeCommandEncoder*)m_ptr)->Buffer.GetPtr()->DebugLevel)
+	{
+		case EAGXDebugLevelConditionalSubmit:
+		case EAGXDebugLevelWaitForComplete:
+		case EAGXDebugLevelLogOperations:
+		case EAGXDebugLevelValidation:
+		{
+			((FAGXDebugComputeCommandEncoder*)m_ptr)->ShaderBuffers.Offsets[index] = offset;
+			check(((FAGXDebugComputeCommandEncoder*)m_ptr)->ResourceMask.BufferMask & (1 << (index)));
+		}
+		default:
+		{
+			break;
+		}
+	}
+#endif
+}
+
+void FAGXComputeCommandEncoderDebugging::SetTexture( FAGXTexture const& texture, NSUInteger index)
+{
+#if METAL_DEBUG_OPTIONS
+	switch (((FAGXDebugComputeCommandEncoder*)m_ptr)->Buffer.GetPtr()->DebugLevel)
+	{
+		case EAGXDebugLevelConditionalSubmit:
+		case EAGXDebugLevelWaitForComplete:
+		case EAGXDebugLevelLogOperations:
+		case EAGXDebugLevelValidation:
+		{
+			((FAGXDebugComputeCommandEncoder*)m_ptr)->ShaderTextures.Textures[index] = texture;
+			((FAGXDebugComputeCommandEncoder*)m_ptr)->ResourceMask.TextureMask = texture ? (((FAGXDebugComputeCommandEncoder*)m_ptr)->ResourceMask.TextureMask | (FAGXTextureMask(1) << (index))) : (((FAGXDebugComputeCommandEncoder*)m_ptr)->ResourceMask.TextureMask & ~(FAGXTextureMask(1) << (index)));
+		}
+		default:
+		{
+			break;
+		}
+	}
+#endif
+}
+
+void FAGXComputeCommandEncoderDebugging::SetSamplerState( mtlpp::SamplerState const& sampler, NSUInteger index)
+{
+#if METAL_DEBUG_OPTIONS
+	switch (((FAGXDebugComputeCommandEncoder*)m_ptr)->Buffer.GetPtr()->DebugLevel)
+	{
+		case EAGXDebugLevelConditionalSubmit:
+		case EAGXDebugLevelWaitForComplete:
+		case EAGXDebugLevelLogOperations:
+		case EAGXDebugLevelValidation:
+		{
+			((FAGXDebugComputeCommandEncoder*)m_ptr)->ShaderSamplers.Samplers[index] = sampler;
+			((FAGXDebugComputeCommandEncoder*)m_ptr)->ResourceMask.SamplerMask = sampler ? (((FAGXDebugComputeCommandEncoder*)m_ptr)->ResourceMask.SamplerMask | (1 << (index))) : (((FAGXDebugComputeCommandEncoder*)m_ptr)->ResourceMask.SamplerMask & ~(1 << (index)));
+		}
+		default:
+		{
+			break;
+		}
+	}
+#endif
+}
+
+void FAGXComputeCommandEncoderDebugging::SetSamplerState( mtlpp::SamplerState const& sampler, float lodMinClamp, float lodMaxClamp, NSUInteger index)
+{
+#if METAL_DEBUG_OPTIONS
+	switch (((FAGXDebugComputeCommandEncoder*)m_ptr)->Buffer.GetPtr()->DebugLevel)
+	{
+		case EAGXDebugLevelConditionalSubmit:
+		case EAGXDebugLevelWaitForComplete:
+		case EAGXDebugLevelLogOperations:
+		case EAGXDebugLevelValidation:
+		{
+			((FAGXDebugComputeCommandEncoder*)m_ptr)->ShaderSamplers.Samplers[index] = sampler;
+			((FAGXDebugComputeCommandEncoder*)m_ptr)->ResourceMask.SamplerMask = sampler ? (((FAGXDebugComputeCommandEncoder*)m_ptr)->ResourceMask.SamplerMask | (1 << (index))) : (((FAGXDebugComputeCommandEncoder*)m_ptr)->ResourceMask.SamplerMask & ~(1 << (index)));
+		}
+		default:
+		{
+			break;
+		}
+	}
+#endif
+}
+
+void FAGXComputeCommandEncoderDebugging::DispatchThreadgroupsWithIndirectBuffer(FAGXBuffer const& indirectBuffer, NSUInteger indirectBufferOffset, mtlpp::Size const& threadsPerThreadgroup)
+{
+#if METAL_DEBUG_OPTIONS
+	switch(((FAGXDebugComputeCommandEncoder*)m_ptr)->Buffer.GetPtr()->DebugLevel)
+	{
+		case EAGXDebugLevelConditionalSubmit:
+		case EAGXDebugLevelWaitForComplete:
+		case EAGXDebugLevelLogOperations:
+		{
+			((FAGXDebugComputeCommandEncoder*)m_ptr)->Buffer.Dispatch([NSString stringWithFormat:@"%s", __PRETTY_FUNCTION__]);
+		}
+		case EAGXDebugLevelValidation:
+		{
+			Validate();
+		}
+		default:
+		{
+			break;
+		}
+	}
+#endif
+}
+
+void FAGXComputeCommandEncoderDebugging::Validate()
+{
+	bool bOK = true;
+	
+#if METAL_DEBUG_OPTIONS
+	switch(((FAGXDebugComputeCommandEncoder*)m_ptr)->Buffer.GetPtr()->DebugLevel)
+	{
+		case EAGXDebugLevelConditionalSubmit:
+		case EAGXDebugLevelWaitForComplete:
+		case EAGXDebugLevelLogOperations:
+		case EAGXDebugLevelValidation:
+		{
+			check(((FAGXDebugComputeCommandEncoder*)m_ptr)->Pipeline);
+			
+			MTLComputePipelineReflection* Reflection = ((FAGXDebugComputeCommandEncoder*)m_ptr)->Pipeline->ComputePipelineReflection;
+			check(Reflection);
+	
+			NSArray<MTLArgument*>* Arguments = Reflection.arguments;
+			for (uint32 i = 0; i < Arguments.count; i++)
+			{
+				MTLArgument* Arg = [Arguments objectAtIndex:i];
+				check(Arg);
+				switch(Arg.type)
+				{
+					case MTLArgumentTypeBuffer:
+					{
+						checkf(Arg.index < ML_MaxBuffers, TEXT("Metal buffer index exceeded!"));
+						if ((((FAGXDebugComputeCommandEncoder*)m_ptr)->ShaderBuffers.Buffers[Arg.index] == nil && ((FAGXDebugComputeCommandEncoder*)m_ptr)->ShaderBuffers.Bytes[Arg.index] == nil))
+						{
+							UE_LOG(LogAGX, Warning, TEXT("Unbound buffer at Metal index %u which will crash the driver: %s"), (uint32)Arg.index, *FString([Arg description]));
+							bOK = false;
+						}
+						break;
+					}
+					case MTLArgumentTypeThreadgroupMemory:
+					{
+						break;
+					}
+					case MTLArgumentTypeTexture:
+					{
+						checkf(Arg.index < ML_MaxTextures, TEXT("Metal texture index exceeded!"));
+						if (((FAGXDebugComputeCommandEncoder*)m_ptr)->ShaderTextures.Textures[Arg.index] == nil)
+						{
+							UE_LOG(LogAGX, Warning, TEXT("Unbound texture at Metal index %u  which will crash the driver: %s"), (uint32)Arg.index, *FString([Arg description]));
+							bOK = false;
+						}
+						else if (((FAGXDebugComputeCommandEncoder*)m_ptr)->ShaderTextures.Textures[Arg.index].textureType != Arg.textureType)
+						{
+							UE_LOG(LogAGX, Warning, TEXT("Incorrect texture type bound at Metal index %u which will crash the driver: %s\n%s"), (uint32)Arg.index, *FString([Arg description]), *FString([((FAGXDebugComputeCommandEncoder*)m_ptr)->ShaderTextures.Textures[Arg.index] description]));
+							bOK = false;
+						}
+						break;
+					}
+					case MTLArgumentTypeSampler:
+					{
+						checkf(Arg.index < ML_MaxSamplers, TEXT("Metal sampler index exceeded!"));
+						if (((FAGXDebugComputeCommandEncoder*)m_ptr)->ShaderSamplers.Samplers[Arg.index] == nil)
+						{
+							UE_LOG(LogAGX, Warning, TEXT("Unbound sampler at Metal index %u which will crash the driver: %s"), (uint32)Arg.index, *FString([Arg description]));
+							bOK = false;
+						}
+						break;
+					}
+					default:
+						check(false);
+						break;
+				}
+			}
+			break;
+		}
+		default:
+		{
+			break;
+		}
+	}
+	
+    if (!bOK)
+    {
+        UE_LOG(LogAGX, Error, TEXT("Metal Validation failures for compute shader:\n%s"), (((FAGXDebugComputeCommandEncoder*)m_ptr)->Pipeline && ((FAGXDebugComputeCommandEncoder*)m_ptr)->Pipeline->ComputeSource) ? *FString(((FAGXDebugComputeCommandEncoder*)m_ptr)->Pipeline->ComputeSource) : TEXT("nil"));
+    }
+#endif
+}
+
+NS_ASSUME_NONNULL_END
+#endif
diff --git a/Engine/Source/Runtime/Apple/AGXRHI/Private/AGXComputeCommandEncoder.h b/Engine/Source/Runtime/Apple/AGXRHI/Private/AGXComputeCommandEncoder.h
new file mode 100644
--- /dev/null
+++ b/Engine/Source/Runtime/Apple/AGXRHI/Private/AGXComputeCommandEncoder.h
@@ -0,0 +1,42 @@
+// Copyright Epic Games, Inc. All Rights Reserved.
+
+#pragma once
+
+#include "AGXDebugCommandEncoder.h"
+
+#if MTLPP_CONFIG_VALIDATE && METAL_DEBUG_OPTIONS
+NS_ASSUME_NONNULL_BEGIN
+
+@class FAGXDebugComputeCommandEncoder;
+@class FAGXShaderPipeline;
+
+
+class FAGXComputeCommandEncoderDebugging : public FAGXCommandEncoderDebugging
+{
+	void InsertDebugDispatch();
+	void Validate();
+public:
+	FAGXComputeCommandEncoderDebugging();
+	FAGXComputeCommandEncoderDebugging(mtlpp::ComputeCommandEncoder& Encoder, FAGXCommandBufferDebugging& Buffer);
+	FAGXComputeCommandEncoderDebugging(FAGXDebugCommandEncoder* handle);
+	
+	static FAGXComputeCommandEncoderDebugging Get(mtlpp::ComputeCommandEncoder& Buffer);
+	
+	void InsertDebugSignpost(ns::String const& Label);
+	void PushDebugGroup(ns::String const& Group);
+	void PopDebugGroup();
+	void EndEncoder();
+	void DispatchThreadgroups(mtlpp::Size const& threadgroupsPerGrid, mtlpp::Size const& threadsPerThreadgroup);
+	void SetPipeline(FAGXShaderPipeline* Pipeline);
+	void SetBytes(const void * bytes, NSUInteger length, NSUInteger index);
+	void SetBuffer( FAGXBuffer const& buffer, NSUInteger offset, NSUInteger index);
+	void SetBufferOffset(NSUInteger offset, NSUInteger index);
+	void SetTexture( FAGXTexture const& texture, NSUInteger index);
+	void SetSamplerState( mtlpp::SamplerState const& sampler, NSUInteger index);
+	void SetSamplerState( mtlpp::SamplerState const& sampler, float lodMinClamp, float lodMaxClamp, NSUInteger index);
+	
+	void DispatchThreadgroupsWithIndirectBuffer(FAGXBuffer const& indirectBuffer, NSUInteger indirectBufferOffset, mtlpp::Size const& threadsPerThreadgroup);
+};
+
+NS_ASSUME_NONNULL_END
+#endif
diff --git a/Engine/Source/Runtime/Apple/AGXRHI/Private/AGXComputePipelineState.h b/Engine/Source/Runtime/Apple/AGXRHI/Private/AGXComputePipelineState.h
new file mode 100644
--- /dev/null
+++ b/Engine/Source/Runtime/Apple/AGXRHI/Private/AGXComputePipelineState.h
@@ -0,0 +1,30 @@
+// Copyright Epic Games, Inc. All Rights Reserved.
+
+/*=============================================================================
+	AGXComputePipelineState.h: AGX RHI compute pipeline state class.
+=============================================================================*/
+
+#pragma once
+
+class FAGXComputePipelineState : public FRHIComputePipelineState
+{
+public:
+	FAGXComputePipelineState(FAGXComputeShader* InComputeShader)
+		: ComputeShader(InComputeShader)
+	{
+		check(InComputeShader);
+	}
+
+	virtual ~FAGXComputePipelineState()
+	{
+		// void
+	}
+
+	FAGXComputeShader* GetComputeShader()
+	{
+		return ComputeShader;
+	}
+
+private:
+	TRefCountPtr<FAGXComputeShader> ComputeShader;
+};
