Initial import of AGXRHI

diff --git a/Engine/Source/Runtime/Apple/AGXRHI/Private/AGXContext.cpp b/Engine/Source/Runtime/Apple/AGXRHI/Private/AGXContext.cpp
new file mode 100644
--- /dev/null
+++ b/Engine/Source/Runtime/Apple/AGXRHI/Private/AGXContext.cpp
@@ -0,0 +1,1832 @@
+// Copyright Epic Games, Inc. All Rights Reserved.
+
+#include "AGXRHIPrivate.h"
+#include "AGXRHIRenderQuery.h"
+#include "AGXVertexDeclaration.h"
+#include "AGXShaderTypes.h"
+#include "AGXGraphicsPipelineState.h"
+#include "Misc/App.h"
+#if PLATFORM_IOS
+#include "IOS/IOSAppDelegate.h"
+#endif
+#include "Misc/ConfigCacheIni.h"
+#include "HAL/PlatformFramePacer.h"
+#include "Runtime/HeadMountedDisplay/Public/IHeadMountedDisplayModule.h"
+
+#include "AGXContext.h"
+#include "AGXProfiler.h"
+#include "AGXCommandBuffer.h"
+
+#include "AGXFrameAllocator.h"
+
+int32 GAGXSupportsIntermediateBackBuffer = 0;
+static FAutoConsoleVariableRef CVarAGXSupportsIntermediateBackBuffer(
+	TEXT("rhi.AGX.SupportsIntermediateBackBuffer"),
+	GAGXSupportsIntermediateBackBuffer,
+	TEXT("When enabled (> 0) allocate an intermediate texture to use as the back-buffer & blit from there into the actual device back-buffer, this is required if we use the experimental separate presentation thread. (Off by default (0))"), ECVF_ReadOnly);
+
+int32 GAGXSeparatePresentThread = 0;
+static FAutoConsoleVariableRef CVarAGXSeparatePresentThread(
+	TEXT("rhi.AGX.SeparatePresentThread"),
+	GAGXSeparatePresentThread,
+	TEXT("When enabled (> 0) requires rhi.AGX.SupportsIntermediateBackBuffer be enabled and will cause two intermediate back-buffers be allocated so that the presentation of frames to the screen can be run on a separate thread.\n")
+	TEXT("This option uncouples the Render/RHI thread from calls to -[CAMetalLayer nextDrawable] and will run arbitrarily fast by rendering but not waiting to present all frames. This is equivalent to running without V-Sync, but without the screen tearing.\n")
+	TEXT("On iOS/tvOS this is the only way to run without locking the CPU to V-Sync somewhere - this shouldn't be used in a shipping title without understanding the power/heat implications.\n")
+	TEXT("(Off by default (0))"), ECVF_ReadOnly);
+
+int32 GAGXNonBlockingPresent = 0;
+static FAutoConsoleVariableRef CVarAGXNonBlockingPresent(
+	TEXT("rhi.AGX.NonBlockingPresent"),
+	GAGXNonBlockingPresent,
+	TEXT("When enabled (> 0) this will force AGXRHI to query if a back-buffer is available to present and if not will skip the frame. Only functions on macOS, it is ignored on iOS/tvOS.\n")
+	TEXT("(Off by default (0))"));
+
+#if PLATFORM_MAC
+static int32 GAGXCommandQueueSize = 5120; // This number is large due to texture streaming - currently each texture is its own command-buffer.
+// The whole AGXRHI needs to be changed to use MTLHeaps/MTLFences & reworked so that operations with the same synchronisation requirements are collapsed into a single blit command-encoder/buffer.
+#else
+static int32 GAGXCommandQueueSize = 0;
+#endif
+
+static FAutoConsoleVariableRef CVarAGXCommandQueueSize(
+	TEXT("rhi.AGX.CommandQueueSize"),
+	GAGXCommandQueueSize,
+	TEXT("The maximum number of command-buffers that can be allocated from each command-queue. (Default: 5120 Mac, 64 iOS/tvOS)"), ECVF_ReadOnly);
+
+int32 GAGXBufferZeroFill = 0; // Deliberately not static
+static FAutoConsoleVariableRef CVarAGXBufferZeroFill(
+	TEXT("rhi.AGX.BufferZeroFill"),
+	GAGXBufferZeroFill,
+	TEXT("Debug option: when enabled will fill the buffer contents with 0 when allocating buffer objects, or regions thereof. (Default: 0, Off)"));
+
+#if METAL_DEBUG_OPTIONS
+int32 GAGXBufferScribble = 0; // Deliberately not static, see InitFrame_UniformBufferPoolCleanup
+static FAutoConsoleVariableRef CVarAGXBufferScribble(
+	TEXT("rhi.AGX.BufferScribble"),
+	GAGXBufferScribble,
+	TEXT("Debug option: when enabled will scribble over the buffer contents with a single value when releasing buffer objects, or regions thereof. (Default: 0, Off)"));
+
+static int32 GAGXResourcePurgeOnDelete = 0;
+static FAutoConsoleVariableRef CVarAGXResourcePurgeOnDelete(
+	TEXT("rhi.AGX.ResourcePurgeOnDelete"),
+	GAGXResourcePurgeOnDelete,
+	TEXT("Debug option: when enabled all MTLResource objects will have their backing stores purged on release - any subsequent access will be invalid and cause a command-buffer failure. Useful for making intermittent resource lifetime errors more common and easier to track. (Default: 0, Off)"));
+
+static int32 GAGXResourceDeferDeleteNumFrames = 0;
+static FAutoConsoleVariableRef CVarAGXResourceDeferDeleteNumFrames(
+	TEXT("rhi.AGX.ResourceDeferDeleteNumFrames"),
+	GAGXResourcePurgeOnDelete,
+	TEXT("Debug option: set to the number of frames that must have passed before resource free-lists are processed and resources disposed of. (Default: 0, Off)"));
+#endif
+
+#if UE_BUILD_SHIPPING
+int32 GAGXRuntimeDebugLevel = 0;
+#else
+int32 GAGXRuntimeDebugLevel = 1;
+#endif
+static FAutoConsoleVariableRef CVarAGXRuntimeDebugLevel(
+	TEXT("rhi.AGX.RuntimeDebugLevel"),
+	GAGXRuntimeDebugLevel,
+	TEXT("The level of debug validation performed by AGXRHI in addition to the underlying Metal API & validation layer.\n")
+	TEXT("Each subsequent level adds more tests and reporting in addition to the previous level.\n")
+	TEXT("*LEVELS >= 3 ARE IGNORED IN SHIPPING AND TEST BUILDS*. (Default: 1 (Debug, Development), 0 (Test, Shipping))\n")
+	TEXT("\t0: Off,\n")
+	TEXT("\t1: Enable light-weight validation of resource bindings & API usage,\n")
+	TEXT("\t2: Reset resource bindings when binding a PSO/Compute-Shader to simplify GPU debugging,\n")
+	TEXT("\t3: Allow rhi.AGX.CommandBufferCommitThreshold to break command-encoders (except when MSAA is enabled),\n")
+	TEXT("\t4: Enable slower, more extensive validation checks for resource types & encoder usage,\n")
+    TEXT("\t5: Record the draw, blit & dispatch commands issued into a command-buffer and report them on failure,\n")
+    TEXT("\t6: Wait for each command-buffer to complete immediately after submission."));
+
+float GAGXPresentFramePacing = 0.0f;
+#if !PLATFORM_MAC
+static FAutoConsoleVariableRef CVarAGXPresentFramePacing(
+	TEXT("rhi.AGX.PresentFramePacing"),
+	GAGXPresentFramePacing,
+	TEXT("Specify the desired frame rate for presentation (iOS 10.3+ only, default: 0.0f, off"));
+#endif
+
+#if PLATFORM_MAC
+static int32 GAGXDefaultUniformBufferAllocation = 1024*1024;
+#else
+static int32 GAGXDefaultUniformBufferAllocation = 1024*32;
+#endif
+static FAutoConsoleVariableRef CVarAGXDefaultUniformBufferAllocation(
+    TEXT("rhi.AGX.DefaultUniformBufferAllocation"),
+    GAGXDefaultUniformBufferAllocation,
+    TEXT("Default size of a uniform buffer allocation."));
+
+#if PLATFORM_MAC
+static int32 GAGXTargetUniformAllocationLimit = 1024 * 1024 * 50;
+#else
+static int32 GAGXTargetUniformAllocationLimit = 1024 * 1024 * 5;
+#endif
+static FAutoConsoleVariableRef CVarAGXTargetUniformAllocationLimit(
+     TEXT("rhi.AGX.TargetUniformAllocationLimit"),
+     GAGXTargetUniformAllocationLimit,
+     TEXT("Target Allocation limit for the uniform buffer pool."));
+
+#if PLATFORM_MAC
+static int32 GAGXTargetTransferAllocatorLimit = 1024*1024*50;
+#else
+static int32 GAGXTargetTransferAllocatorLimit = 1024*1024*2;
+#endif
+static FAutoConsoleVariableRef CVarAGXTargetTransferAllocationLimit(
+	TEXT("rhi.AGX.TargetTransferAllocationLimit"),
+	GAGXTargetTransferAllocatorLimit,
+	TEXT("Target Allocation limit for the upload staging buffer pool."));
+
+#if PLATFORM_MAC
+static int32 GAGXDefaultTransferAllocation = 1024*1024*10;
+#else
+static int32 GAGXDefaultTransferAllocation = 1024*1024*1;
+#endif
+static FAutoConsoleVariableRef CVarAGXDefaultTransferAllocation(
+	TEXT("rhi.AGX.DefaultTransferAllocation"),
+	GAGXDefaultTransferAllocation,
+	TEXT("Default size of a single entry in the upload pool."));
+
+
+
+#if PLATFORM_MAC
+static ns::AutoReleased<ns::Object<id <NSObject>>> GAGXDeviceObserver;
+static mtlpp::Device GetMTLDevice(uint32& DeviceIndex)
+{
+#if PLATFORM_MAC_ARM64
+    return mtlpp::Device::CreateSystemDefaultDevice();
+#else
+	SCOPED_AUTORELEASE_POOL;
+	
+	DeviceIndex = 0;
+	
+	ns::Array<mtlpp::Device> DeviceList;
+	
+	if (FPlatformMisc::MacOSXVersionCompare(10, 13, 4) >= 0)
+	{
+			DeviceList = mtlpp::Device::CopyAllDevicesWithObserver(GAGXDeviceObserver, ^(const mtlpp::Device & Device, const ns::String & Notification)
+			{
+				if ([Notification.GetPtr() isEqualToString:MTLDeviceWasAddedNotification])
+				{
+					FPlatformMisc::GPUChangeNotification(Device.GetRegistryID(), FPlatformMisc::EMacGPUNotification::Added);
+				}
+				else if ([Notification.GetPtr() isEqualToString:MTLDeviceRemovalRequestedNotification])
+				{
+					FPlatformMisc::GPUChangeNotification(Device.GetRegistryID(), FPlatformMisc::EMacGPUNotification::RemovalRequested);
+				}
+				else if ([Notification.GetPtr() isEqualToString:MTLDeviceWasRemovedNotification])
+				{
+					FPlatformMisc::GPUChangeNotification(Device.GetRegistryID(), FPlatformMisc::EMacGPUNotification::Removed);
+				}
+			});
+	}
+	else
+	{
+		DeviceList = mtlpp::Device::CopyAllDevices();
+	}
+	
+	const int32 NumDevices = DeviceList.GetSize();
+	
+	TArray<FMacPlatformMisc::FGPUDescriptor> const& GPUs = FPlatformMisc::GetGPUDescriptors();
+	check(GPUs.Num() > 0);
+
+	// @TODO  here, GetGraphicsAdapterLuid() is used as a device index (how the function "GetGraphicsAdapter" used to work)
+	//        eventually we want the HMD module to return the MTLDevice's registryID, but we cannot fully handle that until
+	//        we drop support for 10.12
+	//  NOTE: this means any implementation of GetGraphicsAdapterLuid() for Mac should return an index, and use -1 as a 
+	//        sentinel value representing "no device" (instead of 0, which is used in the LUID case)
+	int32 HmdGraphicsAdapter  = IHeadMountedDisplayModule::IsAvailable() ? (int32)IHeadMountedDisplayModule::Get().GetGraphicsAdapterLuid() : -1;
+ 	int32 OverrideRendererId = FPlatformMisc::GetExplicitRendererIndex();
+	
+	int32 ExplicitRendererId = OverrideRendererId >= 0 ? OverrideRendererId : HmdGraphicsAdapter;
+	if(ExplicitRendererId < 0 && GPUs.Num() > 1 && FMacPlatformMisc::MacOSXVersionCompare(10, 11, 5) == 0)
+	{
+		OverrideRendererId = -1;
+		bool bForceExplicitRendererId = false;
+		for(uint32 i = 0; i < GPUs.Num(); i++)
+		{
+			FMacPlatformMisc::FGPUDescriptor const& GPU = GPUs[i];
+			if((GPU.GPUVendorId == 0x10DE))
+			{
+				OverrideRendererId = i;
+				bForceExplicitRendererId = (GPU.GPUMetalBundle && ![GPU.GPUMetalBundle isEqualToString:@"GeForceMTLDriverWeb"]);
+			}
+			else if(!GPU.GPUHeadless && GPU.GPUVendorId != 0x8086)
+			{
+				OverrideRendererId = i;
+			}
+		}
+		if (bForceExplicitRendererId)
+		{
+			ExplicitRendererId = OverrideRendererId;
+		}
+	}
+	
+	mtlpp::Device SelectedDevice;
+	if (ExplicitRendererId >= 0 && ExplicitRendererId < GPUs.Num())
+	{
+		FMacPlatformMisc::FGPUDescriptor const& GPU = GPUs[ExplicitRendererId];
+		TArray<FString> NameComponents;
+		FString(GPU.GPUName).TrimStart().ParseIntoArray(NameComponents, TEXT(" "));	
+		for (uint32 index = 0; index < NumDevices; index++)
+		{
+			mtlpp::Device Device = DeviceList[index];
+			
+			if(MTLPP_CHECK_AVAILABLE(10.13, 11.0, 11.0) && (Device.GetRegistryID() == GPU.RegistryID))
+			{
+				DeviceIndex = ExplicitRendererId;
+				SelectedDevice = Device;
+			}
+			else if(([Device.GetName().GetPtr() rangeOfString:@"Nvidia" options:NSCaseInsensitiveSearch].location != NSNotFound && GPU.GPUVendorId == 0x10DE)
+			   || ([Device.GetName().GetPtr() rangeOfString:@"AMD" options:NSCaseInsensitiveSearch].location != NSNotFound && GPU.GPUVendorId == 0x1002)
+			   || ([Device.GetName().GetPtr() rangeOfString:@"Intel" options:NSCaseInsensitiveSearch].location != NSNotFound && GPU.GPUVendorId == 0x8086))
+			{
+				bool bMatchesName = (NameComponents.Num() > 0);
+				for (FString& Component : NameComponents)
+				{
+					bMatchesName &= FString(Device.GetName().GetPtr()).Contains(Component);
+				}
+				if((Device.IsHeadless() == GPU.GPUHeadless || GPU.GPUVendorId != 0x1002) && bMatchesName)
+                {
+					DeviceIndex = ExplicitRendererId;
+					SelectedDevice = Device;
+					break;
+				}
+			}
+		}
+		if(!SelectedDevice)
+		{
+			UE_LOG(LogAGX, Warning,  TEXT("Couldn't find Metal device to match GPU descriptor (%s) from IORegistry - using default device."), *FString(GPU.GPUName));
+		}
+	}
+	if (SelectedDevice == nil)
+	{
+		TArray<FString> NameComponents;
+		SelectedDevice = mtlpp::Device::CreateSystemDefaultDevice();
+		bool bFoundDefault = false;
+		for (uint32 i = 0; i < GPUs.Num(); i++)
+		{
+			FMacPlatformMisc::FGPUDescriptor const& GPU = GPUs[i];
+			if(MTLPP_CHECK_AVAILABLE(10.13, 11.0, 11.0) && (SelectedDevice.GetRegistryID() == GPU.RegistryID))
+			{
+				DeviceIndex = i;
+				bFoundDefault = true;
+				break;
+			}
+			else if(([SelectedDevice.GetName().GetPtr() rangeOfString:@"Nvidia" options:NSCaseInsensitiveSearch].location != NSNotFound && GPU.GPUVendorId == 0x10DE)
+					|| ([SelectedDevice.GetName().GetPtr() rangeOfString:@"AMD" options:NSCaseInsensitiveSearch].location != NSNotFound && GPU.GPUVendorId == 0x1002)
+					|| ([SelectedDevice.GetName().GetPtr() rangeOfString:@"Intel" options:NSCaseInsensitiveSearch].location != NSNotFound && GPU.GPUVendorId == 0x8086))
+			{
+				NameComponents.Empty();
+				bool bMatchesName = FString(GPU.GPUName).TrimStart().ParseIntoArray(NameComponents, TEXT(" ")) > 0;
+				for (FString& Component : NameComponents)
+				{
+					bMatchesName &= FString(SelectedDevice.GetName().GetPtr()).Contains(Component);
+				}
+				if((SelectedDevice.IsHeadless() == GPU.GPUHeadless || GPU.GPUVendorId != 0x1002) && bMatchesName)
+                {
+					DeviceIndex = i;
+					bFoundDefault = true;
+					break;
+				}
+			}
+		}
+		if(!bFoundDefault)
+		{
+			UE_LOG(LogAGX, Warning,  TEXT("Couldn't find Metal device %s in GPU descriptors from IORegistry - capability reporting may be wrong."), *FString(SelectedDevice.GetName().GetPtr()));
+		}
+	}
+	return SelectedDevice;
+#endif // PLATFORM_MAC_ARM64
+}
+
+mtlpp::PrimitiveTopologyClass TranslatePrimitiveTopology(uint32 PrimitiveType)
+{
+	switch (PrimitiveType)
+	{
+		case PT_TriangleList:
+		case PT_TriangleStrip:
+			return mtlpp::PrimitiveTopologyClass::Triangle;
+		case PT_LineList:
+			return mtlpp::PrimitiveTopologyClass::Line;
+		case PT_PointList:
+			return mtlpp::PrimitiveTopologyClass::Point;
+		case PT_1_ControlPointPatchList:
+		case PT_2_ControlPointPatchList:
+		case PT_3_ControlPointPatchList:
+		case PT_4_ControlPointPatchList:
+		case PT_5_ControlPointPatchList:
+		case PT_6_ControlPointPatchList:
+		case PT_7_ControlPointPatchList:
+		case PT_8_ControlPointPatchList:
+		case PT_9_ControlPointPatchList:
+		case PT_10_ControlPointPatchList:
+		case PT_11_ControlPointPatchList:
+		case PT_12_ControlPointPatchList:
+		case PT_13_ControlPointPatchList:
+		case PT_14_ControlPointPatchList:
+		case PT_15_ControlPointPatchList:
+		case PT_16_ControlPointPatchList:
+		case PT_17_ControlPointPatchList:
+		case PT_18_ControlPointPatchList:
+		case PT_19_ControlPointPatchList:
+		case PT_20_ControlPointPatchList:
+		case PT_21_ControlPointPatchList:
+		case PT_22_ControlPointPatchList:
+		case PT_23_ControlPointPatchList:
+		case PT_24_ControlPointPatchList:
+		case PT_25_ControlPointPatchList:
+		case PT_26_ControlPointPatchList:
+		case PT_27_ControlPointPatchList:
+		case PT_28_ControlPointPatchList:
+		case PT_29_ControlPointPatchList:
+		case PT_30_ControlPointPatchList:
+		case PT_31_ControlPointPatchList:
+		case PT_32_ControlPointPatchList:
+		{
+			return mtlpp::PrimitiveTopologyClass::Triangle;
+		}
+		default:
+			UE_LOG(LogAGX, Fatal, TEXT("Unsupported primitive topology %d"), (int32)PrimitiveType);
+			return mtlpp::PrimitiveTopologyClass::Triangle;
+	}
+}
+#endif
+
+FAGXDeviceContext* FAGXDeviceContext::CreateDeviceContext()
+{
+	uint32 DeviceIndex = 0;
+#if PLATFORM_IOS
+	mtlpp::Device Device = mtlpp::Device([IOSAppDelegate GetDelegate].IOSView->MetalDevice);
+#else
+	mtlpp::Device Device = GetMTLDevice(DeviceIndex);
+	if (!Device)
+	{
+		FPlatformMisc::MessageBoxExt(EAppMsgType::Ok, TEXT("The graphics card in this Mac appears to erroneously report support for Metal graphics technology, which is required to run this application, but failed to create a Metal device. The application will now exit."), TEXT("Failed to initialize Metal"));
+		exit(0);
+	}
+#endif
+	
+	uint32 MetalDebug = GAGXRuntimeDebugLevel;
+	const bool bOverridesMetalDebug = FParse::Value( FCommandLine::Get(), TEXT( "MetalRuntimeDebugLevel=" ), MetalDebug );
+	if (bOverridesMetalDebug)
+	{
+		GAGXRuntimeDebugLevel = MetalDebug;
+	}
+	
+	MTLPP_VALIDATION(mtlpp::ValidatedDevice::Register(Device));
+	
+	FAGXCommandQueue* Queue = new FAGXCommandQueue(Device, GAGXCommandQueueSize);
+	check(Queue);
+	
+	if (FAGXCommandQueue::SupportsFeature(EAGXFeaturesFences))
+	{
+		FAGXFencePool::Get().Initialise(Device);
+	}
+	
+	return new FAGXDeviceContext(Device, DeviceIndex, Queue);
+}
+
+FAGXDeviceContext::FAGXDeviceContext(mtlpp::Device MetalDevice, uint32 InDeviceIndex, FAGXCommandQueue* Queue)
+: FAGXContext(MetalDevice, *Queue, true)
+, DeviceIndex(InDeviceIndex)
+, CaptureManager(MetalDevice.GetPtr(), *Queue)
+, SceneFrameCounter(0)
+, FrameCounter(0)
+, ActiveContexts(1)
+, ActiveParallelContexts(0)
+, PSOManager(0)
+, FrameNumberRHIThread(0)
+{
+	CommandQueue.SetRuntimeDebuggingLevel(GAGXRuntimeDebugLevel);
+	
+	// If the separate present thread is enabled then an intermediate backbuffer is required
+	check(!GAGXSeparatePresentThread || GAGXSupportsIntermediateBackBuffer);
+	
+	// Hook into the ios framepacer, if it's enabled for this platform.
+	FrameReadyEvent = NULL;
+	if( FPlatformRHIFramePacer::IsEnabled() || GAGXSeparatePresentThread )
+	{
+		FrameReadyEvent = FPlatformProcess::GetSynchEventFromPool();
+		FPlatformRHIFramePacer::InitWithEvent( FrameReadyEvent );
+		
+		// A bit dirty - this allows the present frame pacing to match the CPU pacing by default unless you've overridden it with the CVar
+		// In all likelihood the CVar is only useful for debugging.
+		if (GAGXPresentFramePacing <= 0.0f)
+		{
+			FString FrameRateLockAsEnum;
+			GConfig->GetString(TEXT("/Script/IOSRuntimeSettings.IOSRuntimeSettings"), TEXT("FrameRateLock"), FrameRateLockAsEnum, GEngineIni);
+	
+			uint32 FrameRateLock = 0;
+			FParse::Value(*FrameRateLockAsEnum, TEXT("PUFRL_"), FrameRateLock);
+			if (FrameRateLock > 0)
+			{
+				GAGXPresentFramePacing = (float)FrameRateLock;
+			}
+		}
+	}
+	
+	if (FParse::Param(FCommandLine::Get(), TEXT("AGXIntermediateBackBuffer")) || FParse::Param(FCommandLine::Get(), TEXT("AGXOffscreenOnly")))
+	{
+		GAGXSupportsIntermediateBackBuffer = 1;
+	}
+    
+    // initialize uniform allocator
+    UniformBufferAllocator = new FAGXFrameAllocator(MetalDevice.GetPtr());
+    UniformBufferAllocator->SetTargetAllocationLimitInBytes(GAGXTargetUniformAllocationLimit);
+    UniformBufferAllocator->SetDefaultAllocationSizeInBytes(GAGXDefaultUniformBufferAllocation);
+    UniformBufferAllocator->SetStatIds(GET_STATID(STAT_AGXUniformAllocatedMemory), GET_STATID(STAT_AGXUniformMemoryInFlight), GET_STATID(STAT_AGXUniformBytesPerFrame));
+	
+	TransferBufferAllocator = new FAGXFrameAllocator(MetalDevice.GetPtr());
+	TransferBufferAllocator->SetTargetAllocationLimitInBytes(GAGXTargetTransferAllocatorLimit);
+	TransferBufferAllocator->SetDefaultAllocationSizeInBytes(GAGXDefaultTransferAllocation);
+	// We won't set StatIds here so it goes to the default frame allocator stats
+	
+	PSOManager = new FAGXPipelineStateCacheManager();
+	
+	METAL_GPUPROFILE(FAGXProfiler::CreateProfiler(this));
+	
+	InitFrame(true, 0, 0);
+}
+
+FAGXDeviceContext::~FAGXDeviceContext()
+{
+	SubmitCommandsHint(EAGXSubmitFlagsWaitOnCommandBuffer);
+	delete &(GetCommandQueue());
+	
+	delete PSOManager;
+    
+    delete UniformBufferAllocator;
+	
+#if PLATFORM_MAC
+	if (FPlatformMisc::MacOSXVersionCompare(10, 13, 4) >= 0)
+	{
+		mtlpp::Device::RemoveDeviceObserver(GAGXDeviceObserver);
+	}
+#endif
+}
+
+void FAGXDeviceContext::Init(void)
+{
+	Heap.Init(GetCommandQueue());
+}
+
+void FAGXDeviceContext::BeginFrame()
+{
+#if ENABLE_METAL_GPUPROFILE
+	FPlatformTLS::SetTlsValue(CurrentContextTLSSlot, this);
+#endif
+	
+	// Wait for the frame semaphore on the immediate context.
+	dispatch_semaphore_wait(CommandBufferSemaphore, DISPATCH_TIME_FOREVER);
+}
+
+#if METAL_DEBUG_OPTIONS
+void FAGXDeviceContext::ScribbleBuffer(FAGXBuffer& Buffer)
+{
+	static uint8 Fill = 0;
+	if (Buffer.GetStorageMode() != mtlpp::StorageMode::Private)
+	{
+		FMemory::Memset(Buffer.GetContents(), Fill++, Buffer.GetLength());
+#if PLATFORM_MAC
+		if (Buffer.GetStorageMode() == mtlpp::StorageMode::Managed)
+		{
+			Buffer.DidModify(ns::Range(0, Buffer.GetLength()));
+		}
+#endif
+	}
+	else
+	{
+		FillBuffer(Buffer, ns::Range(0, Buffer.GetLength()), Fill++);
+	}
+}
+#endif
+
+void FAGXDeviceContext::ClearFreeList()
+{
+	uint32 Index = 0;
+	while(Index < DelayedFreeLists.Num())
+	{
+		FAGXDelayedFreeList* Pair = DelayedFreeLists[Index];
+		if(METAL_DEBUG_OPTION(Pair->DeferCount-- <= 0 &&) Pair->IsComplete())
+		{
+			for( id Entry : Pair->ObjectFreeList )
+			{
+				[Entry release];
+			}
+			for ( FAGXBuffer& Buffer : Pair->UsedBuffers )
+			{
+#if METAL_DEBUG_OPTIONS
+				if (GAGXBufferScribble)
+				{
+					ScribbleBuffer(Buffer);
+				}
+				if (GAGXResourcePurgeOnDelete && !Buffer.GetHeap() && !Buffer.GetParentBuffer())
+				{
+					Buffer.SetPurgeableState(mtlpp::PurgeableState::Empty);
+				}
+#endif
+				Heap.ReleaseBuffer(Buffer);
+			}
+			for ( FAGXTexture& Texture : Pair->UsedTextures )
+			{
+                if (!Texture.GetBuffer() && !Texture.GetParentTexture())
+				{
+#if METAL_DEBUG_OPTIONS
+					if (GAGXResourcePurgeOnDelete && !Texture.GetHeap())
+					{
+						Texture.SetPurgeableState(mtlpp::PurgeableState::Empty);
+					}
+#endif
+					Heap.ReleaseTexture(nullptr, Texture);
+				}
+			}
+			for ( FAGXFence* Fence : Pair->FenceFreeList )
+			{
+				FAGXFencePool::Get().ReleaseFence(Fence);
+			}
+			delete Pair;
+			DelayedFreeLists.RemoveAt(Index, 1, false);
+		}
+		else
+		{
+			Index++;
+		}
+	}
+}
+
+void FAGXDeviceContext::DrainHeap()
+{
+	Heap.Compact(&RenderPass, false);
+}
+
+void FAGXDeviceContext::EndFrame()
+{
+	check(MetalIsSafeToUseRHIThreadResources());
+	
+	// A 'frame' in this context is from the beginning of encoding on the CPU
+	// to the end of all rendering operations on the GPU. So the semaphore is
+	// signalled when the last command buffer finishes GPU execution.
+	{
+		dispatch_semaphore_t CmdBufferSemaphore = CommandBufferSemaphore;
+		dispatch_retain(CmdBufferSemaphore);
+		
+		RenderPass.AddCompletionHandler(
+		[CmdBufferSemaphore](mtlpp::CommandBuffer const& cmd_buf)
+		{
+			dispatch_semaphore_signal(CmdBufferSemaphore);
+			dispatch_release(CmdBufferSemaphore);
+		});
+	}
+	
+	if (bPresented)
+	{
+		CaptureManager.PresentFrame(FrameCounter++);
+		bPresented = false;
+	}
+	
+	// Force submission so the completion handler that signals CommandBufferSemaphore fires.
+	uint32 SubmitFlags = EAGXSubmitFlagsResetState | EAGXSubmitFlagsForce | EAGXSubmitFlagsLastCommandBuffer;
+#if METAL_DEBUG_OPTIONS
+	// Latched update of whether to use runtime debugging features
+	if (GAGXRuntimeDebugLevel != CommandQueue.GetRuntimeDebuggingLevel())
+	{
+		CommandQueue.SetRuntimeDebuggingLevel(GAGXRuntimeDebugLevel);
+		
+		// After change the debug features level wait on commit
+		SubmitFlags |= EAGXSubmitFlagsWaitOnCommandBuffer;
+	}
+#endif
+    
+	SubmitCommandsHint((uint32)SubmitFlags);
+    
+    // increment the internal frame counter
+    FrameNumberRHIThread++;
+	
+    FlushFreeList();
+    
+    ClearFreeList();
+    
+	DrainHeap();
+    
+	InitFrame(true, 0, 0);
+}
+
+void FAGXDeviceContext::BeginScene()
+{
+#if ENABLE_METAL_GPUPROFILE
+	FPlatformTLS::SetTlsValue(CurrentContextTLSSlot, this);
+#endif
+	
+	// Increment the frame counter. INDEX_NONE is a special value meaning "uninitialized", so if
+	// we hit it just wrap around to zero.
+	SceneFrameCounter++;
+	if (SceneFrameCounter == INDEX_NONE)
+	{
+		SceneFrameCounter++;
+	}
+}
+
+void FAGXDeviceContext::EndScene()
+{
+}
+
+void FAGXDeviceContext::BeginDrawingViewport(FAGXViewport* Viewport)
+{
+#if ENABLE_METAL_GPUPROFILE
+	FPlatformTLS::SetTlsValue(CurrentContextTLSSlot, this);
+#endif
+}
+
+bool FAGXDeviceContext::FAGXDelayedFreeList::IsComplete() const
+{
+	bool bFinished = true;
+	for (mtlpp::CommandBufferFence const& Fence : Fences)
+	{
+		bFinished &= Fence.Wait(0);
+
+		if (!bFinished)
+			break;
+	}
+	return bFinished;
+}
+
+void FAGXDeviceContext::FlushFreeList(bool const bFlushFences)
+{
+	FAGXDelayedFreeList* NewList = new FAGXDelayedFreeList;
+	
+	// Get the committed command buffer fences and clear the array in the command-queue
+	GetCommandQueue().GetCommittedCommandBufferFences(NewList->Fences);
+	
+	METAL_DEBUG_OPTION(NewList->DeferCount = GAGXResourceDeferDeleteNumFrames);
+	FreeListMutex.Lock();
+	NewList->UsedBuffers = MoveTemp(UsedBuffers);
+	NewList->UsedTextures = MoveTemp(UsedTextures);
+	NewList->ObjectFreeList = ObjectFreeList;
+	if (bFlushFences)
+	{
+		TArray<FAGXFence*> Fences;
+		FenceFreeList.PopAll(Fences);
+		for (FAGXFence* Fence : Fences)
+		{
+			if(!UsedFences.Contains(Fence))
+			{
+				UsedFences.Add(Fence);
+			}
+		}
+		NewList->FenceFreeList = MoveTemp(UsedFences);
+	}
+#if METAL_DEBUG_OPTIONS
+	if (FrameFences.Num())
+	{
+		FrameFences.Empty();
+	}
+#endif
+	ObjectFreeList.Empty(ObjectFreeList.Num());
+	FreeListMutex.Unlock();
+	
+	DelayedFreeLists.Add(NewList);
+}
+
+void FAGXDeviceContext::EndDrawingViewport(FAGXViewport* Viewport, bool bPresent, bool bLockToVsync)
+{
+	// enqueue a present if desired
+	static bool const bOffscreenOnly = FParse::Param(FCommandLine::Get(), TEXT("AGXOffscreenOnly"));
+	if (bPresent && !bOffscreenOnly)
+	{
+		
+#if PLATFORM_MAC
+		// Handle custom present
+		FRHICustomPresent* const CustomPresent = Viewport->GetCustomPresent();
+		if (CustomPresent != nullptr)
+		{
+			int32 SyncInterval = 0;
+			{
+				SCOPE_CYCLE_COUNTER(STAT_AGXCustomPresentTime);
+				CustomPresent->Present(SyncInterval);
+			}
+			
+			mtlpp::CommandBuffer CurrentCommandBuffer = GetCurrentCommandBuffer();
+			check(CurrentCommandBuffer);
+			
+			CurrentCommandBuffer.AddScheduledHandler([CustomPresent](mtlpp::CommandBuffer const&) {
+				CustomPresent->PostPresent();
+			});
+		}
+#endif
+		
+		RenderPass.End();
+		
+		SubmitCommandsHint(EAGXSubmitFlagsForce|EAGXSubmitFlagsCreateCommandBuffer);
+		
+		Viewport->Present(GetCommandQueue(), bLockToVsync);
+	}
+	
+	bPresented = bPresent;
+	
+	// We may be limiting our framerate to the display link
+	if( FrameReadyEvent != nullptr && !GAGXSeparatePresentThread )
+	{
+		bool bIgnoreThreadIdleStats = true; // Idle time is already counted by the caller
+		FrameReadyEvent->Wait(MAX_uint32, bIgnoreThreadIdleStats);
+	}
+	
+	Viewport->ReleaseDrawable();
+}
+
+void FAGXDeviceContext::ReleaseObject(id Object)
+{
+	if (GIsAGXInitialized) // @todo zebra: there seems to be some race condition at exit when the framerate is very low
+	{
+		check(Object);
+		FreeListMutex.Lock();
+		if(!ObjectFreeList.Contains(Object))
+        {
+            ObjectFreeList.Add(Object);
+        }
+        else
+        {
+            [Object release];
+        }
+		FreeListMutex.Unlock();
+	}
+}
+
+void FAGXDeviceContext::ReleaseTexture(FAGXSurface* Surface, FAGXTexture& Texture)
+{
+	if (GIsAGXInitialized) // @todo zebra: there seems to be some race condition at exit when the framerate is very low
+	{
+		check(Surface && Texture);
+		ReleaseTexture(Texture);
+	}
+}
+
+void FAGXDeviceContext::ReleaseTexture(FAGXTexture& Texture)
+{
+	if(GIsAGXInitialized)
+	{
+		check(Texture);
+		FreeListMutex.Lock();
+        if (Texture.GetStorageMode() == mtlpp::StorageMode::Private)
+        {
+            Heap.ReleaseTexture(nullptr, Texture);
+			
+			// Ensure that the Objective-C handle can't disappear prior to the GPU being done with it without racing with the above
+			if(!ObjectFreeList.Contains(Texture.GetPtr()))
+			{
+				[Texture.GetPtr() retain];
+				ObjectFreeList.Add(Texture.GetPtr());
+        }
+        }
+		else if(!UsedTextures.Contains(Texture))
+		{
+			UsedTextures.Add(MoveTemp(Texture));
+		}
+		FreeListMutex.Unlock();
+	}
+}
+
+void FAGXDeviceContext::ReleaseFence(FAGXFence* Fence)
+{
+#if METAL_DEBUG_OPTIONS
+	if(GetCommandList().GetCommandQueue().GetRuntimeDebuggingLevel() >= EAGXDebugLevelValidation)
+	{
+		FScopeLock Lock(&FreeListMutex);
+		FrameFences.Add(Fence);
+	}
+#endif
+	
+	if (GIsAGXInitialized) // @todo zebra: there seems to be some race condition at exit when the framerate is very low
+	{
+		check(Fence);
+		FenceFreeList.Push(Fence);
+	}
+}
+
+FAGXTexture FAGXDeviceContext::CreateTexture(FAGXSurface* Surface, mtlpp::TextureDescriptor Descriptor)
+{
+	FAGXTexture Tex = Heap.CreateTexture(Descriptor, Surface);
+#if METAL_DEBUG_OPTIONS
+	if (GAGXResourcePurgeOnDelete && !Tex.GetHeap())
+	{
+		Tex.SetPurgeableState(mtlpp::PurgeableState::NonVolatile);
+	}
+#endif
+	
+	return Tex;
+}
+
+FAGXBuffer FAGXDeviceContext::CreatePooledBuffer(FAGXPooledBufferArgs const& Args)
+{
+	NSUInteger CpuResourceOption = ((NSUInteger)Args.CpuCacheMode) << mtlpp::ResourceCpuCacheModeShift;
+	
+	uint32 RequestedBufferOffsetAlignment = BufferOffsetAlignment;
+	
+	if((Args.Flags & (BUF_UnorderedAccess | BUF_ShaderResource)) != 0)
+	{
+		// Buffer backed linear textures have specific align requirements
+		// We don't know upfront the pixel format that may be requested for an SRV so we can't use minimumLinearTextureAlignmentForPixelFormat:
+		RequestedBufferOffsetAlignment = BufferBackedLinearTextureOffsetAlignment;
+	}
+	
+    FAGXBuffer Buffer = Heap.CreateBuffer(Args.Size, RequestedBufferOffsetAlignment, Args.Flags, FAGXCommandQueue::GetCompatibleResourceOptions((mtlpp::ResourceOptions)(CpuResourceOption | mtlpp::ResourceOptions::HazardTrackingModeUntracked | ((NSUInteger)Args.Storage << mtlpp::ResourceStorageModeShift))));
+	check(Buffer && Buffer.GetPtr());
+#if METAL_DEBUG_OPTIONS
+	if (GAGXResourcePurgeOnDelete && !Buffer.GetHeap())
+	{
+		Buffer.SetPurgeableState(mtlpp::PurgeableState::NonVolatile);
+	}
+#endif
+	
+	return Buffer;
+}
+
+void FAGXDeviceContext::ReleaseBuffer(FAGXBuffer& Buffer)
+{
+	if(GIsAGXInitialized)
+	{
+		check(Buffer);
+		FreeListMutex.Lock();
+		if(!UsedBuffers.Contains(Buffer))
+		{
+			UsedBuffers.Add(MoveTemp(Buffer));
+		}
+		FreeListMutex.Unlock();
+	}
+}
+
+struct FAGXRHICommandUpdateFence final : public FRHICommand<FAGXRHICommandUpdateFence>
+{
+	TRefCountPtr<FAGXFence> Fence;
+	uint32 Num;
+	
+	FORCEINLINE_DEBUGGABLE FAGXRHICommandUpdateFence(TRefCountPtr<FAGXFence>& InFence, uint32 InNum)
+	: Fence(InFence)
+	, Num(InNum)
+	{
+	}
+	
+	void Execute(FRHICommandListBase& CmdList)
+	{
+		GetAGXDeviceContext().SetParallelPassFences(nil, Fence);
+		GetAGXDeviceContext().FinishFrame(true);
+		GetAGXDeviceContext().BeginParallelRenderCommandEncoding(Num);
+	}
+};
+
+FAGXRHICommandContext* FAGXDeviceContext::AcquireContext(int32 NewIndex, int32 NewNum)
+{
+	FAGXRHICommandContext* Context = ParallelContexts.Pop();
+	if (!Context)
+	{
+		FAGXContext* AGXContext = new FAGXContext(GetDevice(), GetCommandQueue(), false);
+		check(AGXContext);
+		
+		FAGXRHICommandContext* CmdContext = static_cast<FAGXRHICommandContext*>(RHIGetDefaultContext());
+		check(CmdContext);
+		
+		Context = new FAGXRHICommandContext(CmdContext->GetProfiler(), AGXContext);
+	}
+	check(Context);
+	
+	if(ParallelFences.Num() < NewNum)
+	{
+		ParallelFences.AddDefaulted(NewNum - ParallelFences.Num());
+	}
+	
+	NSString* StartLabel = nil;
+	NSString* EndLabel = nil;
+#if METAL_DEBUG_OPTIONS
+	StartLabel = [NSString stringWithFormat:@"Start Parallel Context Index %d Num %d", NewIndex, NewNum];
+	EndLabel = [NSString stringWithFormat:@"End Parallel Context Index %d Num %d", NewIndex, NewNum];
+#endif
+	
+	TRefCountPtr<FAGXFence> StartFence = (NewIndex == 0 ? CommandList.GetCommandQueue().CreateFence(StartLabel) : ParallelFences[NewIndex - 1].GetReference());
+	TRefCountPtr<FAGXFence> EndFence = CommandList.GetCommandQueue().CreateFence(EndLabel);
+	ParallelFences[NewIndex] = EndFence;
+	
+	// Give the context the fences so that we can properly order the parallel contexts.
+	Context->GetInternalContext().SetParallelPassFences(StartFence, EndFence);
+	
+	if (NewIndex == 0)
+	{
+		if (FRHICommandListExecutor::GetImmediateCommandList().Bypass() || !IsRunningRHIInSeparateThread())
+		{
+			FAGXRHICommandUpdateFence UpdateCommand(StartFence, NewNum);
+			UpdateCommand.Execute(FRHICommandListExecutor::GetImmediateCommandList());
+		}
+		else
+		{
+			new (FRHICommandListExecutor::GetImmediateCommandList().AllocCommand<FAGXRHICommandUpdateFence>()) FAGXRHICommandUpdateFence(StartFence, NewNum);
+			FRHICommandListExecutor::GetImmediateCommandList().RHIThreadFence(true);
+			FRHICommandListExecutor::GetImmediateCommandList().ImmediateFlush(EImmediateFlushType::DispatchToRHIThread);
+		}
+	}
+	
+	FPlatformAtomics::InterlockedIncrement(&ActiveContexts);
+	return Context;
+}
+
+void FAGXDeviceContext::ReleaseContext(FAGXRHICommandContext* Context)
+{
+	ParallelContexts.Push(Context);
+	FPlatformAtomics::InterlockedDecrement(&ActiveContexts);
+	check(ActiveContexts >= 1);
+}
+
+uint32 FAGXDeviceContext::GetNumActiveContexts(void) const
+{
+	return ActiveContexts;
+}
+
+uint32 FAGXDeviceContext::GetDeviceIndex(void) const
+{
+	return DeviceIndex;
+}
+
+void FAGXDeviceContext::NewLock(FAGXRHIBuffer* Buffer, FAGXFrameAllocator::AllocationEntry& Allocation)
+{
+	check(!OutstandingLocks.Contains(Buffer));
+	OutstandingLocks.Add(Buffer, Allocation);
+}
+
+FAGXFrameAllocator::AllocationEntry FAGXDeviceContext::FetchAndRemoveLock(FAGXRHIBuffer* Buffer)
+{
+	FAGXFrameAllocator::AllocationEntry Backing = OutstandingLocks.FindAndRemoveChecked(Buffer);
+	return Backing;
+}
+
+#if METAL_DEBUG_OPTIONS
+void FAGXDeviceContext::AddActiveBuffer(FAGXBuffer const& Buffer)
+{
+    if(GetCommandList().GetCommandQueue().GetRuntimeDebuggingLevel() >= EAGXDebugLevelValidation)
+    {
+        FScopeLock Lock(&ActiveBuffersMutex);
+        
+        NSRange DestRange = NSMakeRange(Buffer.GetOffset(), Buffer.GetLength());
+        TArray<NSRange>* Ranges = ActiveBuffers.Find(Buffer.GetPtr());
+        if (!Ranges)
+        {
+            ActiveBuffers.Add(Buffer.GetPtr(), TArray<NSRange>());
+            Ranges = ActiveBuffers.Find(Buffer.GetPtr());
+        }
+        Ranges->Add(DestRange);
+    }
+}
+
+static bool operator==(NSRange const& A, NSRange const& B)
+{
+    return NSEqualRanges(A, B);
+}
+
+void FAGXDeviceContext::RemoveActiveBuffer(FAGXBuffer const& Buffer)
+{
+    if(GetCommandList().GetCommandQueue().GetRuntimeDebuggingLevel() >= EAGXDebugLevelValidation)
+    {
+        FScopeLock Lock(&ActiveBuffersMutex);
+        
+        NSRange DestRange = NSMakeRange(Buffer.GetOffset(), Buffer.GetLength());
+        TArray<NSRange>& Ranges = ActiveBuffers.FindChecked(Buffer.GetPtr());
+        int32 i = Ranges.RemoveSingle(DestRange);
+        check(i > 0);
+    }
+}
+
+bool FAGXDeviceContext::ValidateIsInactiveBuffer(FAGXBuffer const& Buffer)
+{
+    if(GetCommandList().GetCommandQueue().GetRuntimeDebuggingLevel() >= EAGXDebugLevelValidation)
+    {
+        FScopeLock Lock(&ActiveBuffersMutex);
+        
+        TArray<NSRange>* Ranges = ActiveBuffers.Find(Buffer.GetPtr());
+        if (Ranges)
+        {
+            NSRange DestRange = NSMakeRange(Buffer.GetOffset(), Buffer.GetLength());
+            for (NSRange Range : *Ranges)
+            {
+                if (NSIntersectionRange(Range, DestRange).length > 0)
+                {
+                    UE_LOG(LogAGX, Error, TEXT("ValidateIsInactiveBuffer failed on overlapping ranges ({%d, %d} vs {%d, %d}) of buffer %p."), (uint32)Range.location, (uint32)Range.length, (uint32)Buffer.GetOffset(), (uint32)Buffer.GetLength(), Buffer.GetPtr());
+                    return false;
+                }
+            }
+        }
+    }
+    return true;
+}
+#endif
+
+
+#if ENABLE_METAL_GPUPROFILE
+uint32 FAGXContext::CurrentContextTLSSlot = FPlatformTLS::AllocTlsSlot();
+#endif
+
+FAGXContext::FAGXContext(mtlpp::Device InDevice, FAGXCommandQueue& Queue, bool const bIsImmediate)
+: Device(InDevice)
+, CommandQueue(Queue)
+, CommandList(Queue, bIsImmediate)
+, StateCache(bIsImmediate)
+, RenderPass(CommandList, StateCache)
+, QueryBuffer(new FAGXQueryBufferPool(this))
+, StartFence(nil)
+, EndFence(nil)
+, NumParallelContextsInPass(0)
+{
+	// create a semaphore for multi-buffering the command buffer
+	CommandBufferSemaphore = dispatch_semaphore_create(FParse::Param(FCommandLine::Get(),TEXT("gpulockstep")) ? 1 : 3);
+}
+
+FAGXContext::~FAGXContext()
+{
+	SubmitCommandsHint(EAGXSubmitFlagsWaitOnCommandBuffer);
+}
+
+mtlpp::Device& FAGXContext::GetDevice()
+{
+	return Device;
+}
+
+FAGXCommandQueue& FAGXContext::GetCommandQueue()
+{
+	return CommandQueue;
+}
+
+FAGXCommandList& FAGXContext::GetCommandList()
+{
+	return CommandList;
+}
+
+mtlpp::CommandBuffer const& FAGXContext::GetCurrentCommandBuffer() const
+{
+	return RenderPass.GetCurrentCommandBuffer();
+}
+
+mtlpp::CommandBuffer& FAGXContext::GetCurrentCommandBuffer()
+{
+	return RenderPass.GetCurrentCommandBuffer();
+}
+
+void FAGXContext::InsertCommandBufferFence(FAGXCommandBufferFence& Fence, mtlpp::CommandBufferHandler Handler)
+{
+	check(GetCurrentCommandBuffer());
+	
+	RenderPass.InsertCommandBufferFence(Fence, Handler);
+}
+
+#if ENABLE_METAL_GPUPROFILE
+FAGXContext* FAGXContext::GetCurrentContext()
+{
+	FAGXContext* Current = (FAGXContext*)FPlatformTLS::GetTlsValue(CurrentContextTLSSlot);
+	
+	if (!Current)
+	{
+		// If we are executing this outside of a pass we'll return the default.
+		// TODO This needs further investigation. We should fix all the cases that call this without
+		// a context set.
+		FAGXRHICommandContext* CmdContext = static_cast<FAGXRHICommandContext*>(RHIGetDefaultContext());
+		check(CmdContext);
+		Current = &CmdContext->GetInternalContext();
+	}
+	
+	check(Current);
+	return Current;
+}
+
+void FAGXContext::MakeCurrent(FAGXContext* Context)
+{
+	FPlatformTLS::SetTlsValue(CurrentContextTLSSlot, Context);
+}
+#endif
+
+void FAGXContext::SetParallelPassFences(FAGXFence* Start, FAGXFence* End)
+{
+	check(!StartFence && !EndFence);
+	StartFence = Start;
+	EndFence = End;
+}
+
+TRefCountPtr<FAGXFence> const& FAGXContext::GetParallelPassStartFence(void) const
+{
+	return StartFence;
+}
+
+TRefCountPtr<FAGXFence> const& FAGXContext::GetParallelPassEndFence(void) const
+{
+	return EndFence;
+}
+
+void FAGXContext::InitFrame(bool const bImmediateContext, uint32 Index, uint32 Num)
+{
+#if ENABLE_METAL_GPUPROFILE
+	FPlatformTLS::SetTlsValue(CurrentContextTLSSlot, this);
+#endif
+	
+	// Reset cached state in the encoder
+	StateCache.Reset();
+
+	// Sets the index of the parallel context within the pass
+	if (!bImmediateContext && FAGXCommandQueue::SupportsFeature(EAGXFeaturesParallelRenderEncoders))
+	{
+		CommandList.SetParallelIndex(Index, Num);
+	}
+	else
+	{
+		CommandList.SetParallelIndex(0, 0);
+	}
+	
+	// Reallocate if necessary to ensure >= 80% usage, otherwise we're just too wasteful
+	RenderPass.ShrinkRingBuffers();
+	
+	// Begin the render pass frame.
+	RenderPass.Begin(StartFence);
+	
+	// Unset the start fence, the render-pass owns it and we can consider it encoded now!
+	StartFence = nullptr;
+	
+	// make sure first SetRenderTarget goes through
+	StateCache.InvalidateRenderTargets();
+}
+
+void FAGXContext::FinishFrame(bool const bImmediateContext)
+{
+	// Ensure that we update the end fence for parallel contexts.
+	RenderPass.Update(EndFence);
+	
+	// Unset the end fence, the render-pass owns it and we can consider it encoded now!
+	EndFence = nullptr;
+	
+	// End the render pass
+	RenderPass.End();
+	
+	// Issue any outstanding commands.
+	SubmitCommandsHint((CommandList.IsParallel() ? EAGXSubmitFlagsAsyncCommandBuffer : EAGXSubmitFlagsNone));
+	
+	// make sure first SetRenderTarget goes through
+	StateCache.InvalidateRenderTargets();
+	
+	if (!bImmediateContext)
+	{
+		StateCache.Reset();
+	}
+
+#if ENABLE_METAL_GPUPROFILE
+	FPlatformTLS::SetTlsValue(CurrentContextTLSSlot, nullptr);
+#endif
+}
+
+void FAGXContext::TransitionResource(FRHIUnorderedAccessView* InResource)
+{
+	FAGXUnorderedAccessView* UAV = ResourceCast(InResource);
+
+	// figure out which one of the resources we need to set
+	FAGXStructuredBuffer* StructuredBuffer = UAV->SourceView->SourceStructuredBuffer.GetReference();
+	FAGXVertexBuffer*     VertexBuffer     = UAV->SourceView->SourceVertexBuffer.GetReference();
+	FAGXIndexBuffer*      IndexBuffer      = UAV->SourceView->SourceIndexBuffer.GetReference();
+	FRHITexture*            Texture          = UAV->SourceView->SourceTexture.GetReference();
+	FAGXSurface*          Surface          = UAV->SourceView->TextureView;
+
+	if (StructuredBuffer)
+	{
+		RenderPass.TransitionResources(StructuredBuffer->GetCurrentBuffer());
+	}
+	else if (VertexBuffer && VertexBuffer->GetCurrentBufferOrNil())
+	{
+		RenderPass.TransitionResources(VertexBuffer->GetCurrentBuffer());
+	}
+	else if (IndexBuffer)
+	{
+		RenderPass.TransitionResources(IndexBuffer->GetCurrentBuffer());
+	}
+	else if (Surface)
+	{
+		RenderPass.TransitionResources(Surface->Texture.GetParentTexture());
+	}
+	else if (Texture)
+	{
+		if (!Surface)
+		{
+			Surface = AGXGetMetalSurfaceFromRHITexture(Texture);
+		}
+		if ((Surface != nullptr) && Surface->Texture)
+		{
+			RenderPass.TransitionResources(Surface->Texture);
+			if (Surface->MSAATexture)
+			{
+				RenderPass.TransitionResources(Surface->MSAATexture);
+			}
+		}
+	}
+}
+
+void FAGXContext::TransitionResource(FRHITexture* InResource)
+{
+	FAGXSurface* Surface = AGXGetMetalSurfaceFromRHITexture(InResource);
+
+	if ((Surface != nullptr) && Surface->Texture)
+	{
+		RenderPass.TransitionResources(Surface->Texture);
+		if (Surface->MSAATexture)
+		{
+			RenderPass.TransitionResources(Surface->MSAATexture);
+		}
+	}
+}
+
+void FAGXContext::SubmitCommandsHint(uint32 const Flags)
+{
+	// When the command-buffer is submitted for a reason other than a break of a logical command-buffer (where one high-level command-sequence becomes more than one command-buffer).
+	if (!(Flags & EAGXSubmitFlagsBreakCommandBuffer))
+	{
+		// Release the current query buffer if there are outstanding writes so that it isn't transitioned by a future encoder that will cause a resource access conflict and lifetime error.
+		GetQueryBufferPool()->ReleaseCurrentQueryBuffer();
+	}
+	
+	RenderPass.Submit((EAGXSubmitFlags)Flags);
+}
+
+void FAGXContext::SubmitCommandBufferAndWait()
+{
+	// kick the whole buffer
+	// Commit to hand the commandbuffer off to the gpu
+	// Wait for completion as requested.
+	SubmitCommandsHint((EAGXSubmitFlagsCreateCommandBuffer | EAGXSubmitFlagsBreakCommandBuffer | EAGXSubmitFlagsWaitOnCommandBuffer));
+}
+
+void FAGXContext::ResetRenderCommandEncoder()
+{
+	SubmitCommandsHint();
+	
+	StateCache.InvalidateRenderTargets();
+	
+	SetRenderPassInfo(StateCache.GetRenderPassInfo(), true);
+}
+
+bool FAGXContext::PrepareToDraw(uint32 PrimitiveType)
+{
+	SCOPE_CYCLE_COUNTER(STAT_AGXPrepareDrawTime);
+	TRefCountPtr<FAGXGraphicsPipelineState> CurrentPSO = StateCache.GetGraphicsPSO();
+	check(IsValidRef(CurrentPSO));
+	
+	// Enforce calls to SetRenderTarget prior to issuing draw calls.
+#if PLATFORM_MAC
+	check(StateCache.GetHasValidRenderTarget());
+#else
+	if (!StateCache.GetHasValidRenderTarget())
+	{
+		return false;
+	}
+#endif
+	
+	FAGXHashedVertexDescriptor const& VertexDesc = CurrentPSO->VertexDeclaration->Layout;
+	
+	// Validate the vertex layout in debug mode, or when the validation layer is enabled for development builds.
+	// Other builds will just crash & burn if it is incorrect.
+#if UE_BUILD_DEBUG || UE_BUILD_DEVELOPMENT
+	if(CommandQueue.GetRuntimeDebuggingLevel() >= EAGXDebugLevelValidation)
+	{
+		MTLVertexDescriptor* Layout = VertexDesc.VertexDesc;
+		
+		if(Layout && Layout.layouts)
+		{
+			for (uint32 i = 0; i < MaxVertexElementCount; i++)
+			{
+				auto Attribute = [Layout.attributes objectAtIndexedSubscript:i];
+				if(Attribute && Attribute.format > MTLVertexFormatInvalid)
+				{
+					auto BufferLayout = [Layout.layouts objectAtIndexedSubscript:Attribute.bufferIndex];
+					uint32 BufferLayoutStride = BufferLayout ? BufferLayout.stride : 0;
+					
+					uint32 BufferIndex = METAL_TO_UNREAL_BUFFER_INDEX(Attribute.bufferIndex);
+					
+					if (CurrentPSO->VertexShader->Bindings.InOutMask.IsFieldEnabled(BufferIndex))
+					{
+						uint64 MetalSize = StateCache.GetVertexBufferSize(BufferIndex);
+						
+						// If the vertex attribute is required and either no Metal buffer is bound or the size of the buffer is smaller than the stride, or the stride is explicitly specified incorrectly then the layouts don't match.
+						if (BufferLayoutStride > 0 && MetalSize < BufferLayoutStride)
+						{
+							FString Report = FString::Printf(TEXT("Vertex Layout Mismatch: Index: %d, Len: %lld, Decl. Stride: %d"), Attribute.bufferIndex, MetalSize, BufferLayoutStride);
+							UE_LOG(LogAGX, Warning, TEXT("%s"), *Report);
+						}
+					}
+				}
+			}
+		}
+	}
+#endif
+	
+	// @todo Handle the editor not setting a depth-stencil target for the material editor's tiles which render to depth even when they shouldn't.
+	bool const bNeedsDepthStencilWrite = (IsValidRef(CurrentPSO->PixelShader) && (CurrentPSO->PixelShader->Bindings.InOutMask.IsFieldEnabled(CrossCompiler::FShaderBindingInOutMask::DepthStencilMaskIndex)));
+	
+	// @todo Improve the way we handle binding a dummy depth/stencil so we can get pure UAV raster operations...
+	bool const bNeedsDepthStencilForUAVRaster = (StateCache.GetRenderPassInfo().GetNumColorRenderTargets() == 0);
+	
+	bool const bBindDepthStencilForWrite = bNeedsDepthStencilWrite && !StateCache.HasValidDepthStencilSurface();
+	bool const bBindDepthStencilForUAVRaster = bNeedsDepthStencilForUAVRaster && !StateCache.HasValidDepthStencilSurface();
+	
+	if (bBindDepthStencilForWrite || bBindDepthStencilForUAVRaster)
+	{
+#if UE_BUILD_DEBUG
+		if (bBindDepthStencilForWrite)
+		{
+			UE_LOG(LogAGX, Warning, TEXT("Binding a temporary depth-stencil surface as the bound shader pipeline writes to depth/stencil but no depth/stencil surface was bound!"));
+		}
+		else
+		{
+			check(bNeedsDepthStencilForUAVRaster);
+			UE_LOG(LogAGX, Warning, TEXT("Binding a temporary depth-stencil surface as the bound shader pipeline needs a texture bound - even when only writing to UAVs!"));
+		}
+#endif
+		check(StateCache.GetRenderTargetArraySize() <= 1);
+		CGSize FBSize;
+		if (bBindDepthStencilForWrite)
+		{
+			check(!bBindDepthStencilForUAVRaster);
+			FBSize = StateCache.GetFrameBufferSize();
+		}
+		else
+		{
+			check(bBindDepthStencilForUAVRaster);
+			FBSize = CGSizeMake(StateCache.GetViewport(0).width, StateCache.GetViewport(0).height);
+		}
+		
+		FRHIRenderPassInfo Info = StateCache.GetRenderPassInfo();
+		
+		FTexture2DRHIRef FallbackDepthStencilSurface = StateCache.CreateFallbackDepthStencilSurface(FBSize.width, FBSize.height);
+		check(IsValidRef(FallbackDepthStencilSurface));
+		
+		if (bBindDepthStencilForWrite)
+		{
+			check(!bBindDepthStencilForUAVRaster);
+			Info.DepthStencilRenderTarget.DepthStencilTarget = FallbackDepthStencilSurface;
+			Info.DepthStencilRenderTarget.ExclusiveDepthStencil = FExclusiveDepthStencil::DepthWrite_StencilWrite;
+			Info.DepthStencilRenderTarget.Action = MakeDepthStencilTargetActions(MakeRenderTargetActions(ERenderTargetLoadAction::ELoad, ERenderTargetStoreAction::EStore), MakeRenderTargetActions(ERenderTargetLoadAction::ELoad, ERenderTargetStoreAction::EStore));
+		}
+		else
+		{
+			check(bBindDepthStencilForUAVRaster);
+			Info.DepthStencilRenderTarget.DepthStencilTarget = FallbackDepthStencilSurface;
+			Info.DepthStencilRenderTarget.ExclusiveDepthStencil = FExclusiveDepthStencil::DepthRead_StencilRead;
+			Info.DepthStencilRenderTarget.Action = MakeDepthStencilTargetActions(MakeRenderTargetActions(ERenderTargetLoadAction::ELoad, ERenderTargetStoreAction::ENoAction), MakeRenderTargetActions(ERenderTargetLoadAction::ELoad, ERenderTargetStoreAction::ENoAction));
+		}
+		
+		// Ensure that we make it a Clear/Store -> Load/Store for the colour targets or we might render incorrectly
+		for (uint32 i = 0; i < Info.GetNumColorRenderTargets(); i++)
+		{
+			if (GetLoadAction(Info.ColorRenderTargets[i].Action) != ERenderTargetLoadAction::ELoad)
+			{
+				check(GetStoreAction(Info.ColorRenderTargets[i].Action) == ERenderTargetStoreAction::EStore || GetStoreAction(Info.ColorRenderTargets[i].Action) == ERenderTargetStoreAction::EMultisampleResolve);
+				Info.ColorRenderTargets[i].Action = MakeRenderTargetActions(ERenderTargetLoadAction::ELoad, GetStoreAction(Info.ColorRenderTargets[i].Action));
+			}
+		}
+		
+		if (StateCache.SetRenderPassInfo(Info, StateCache.GetVisibilityResultsBuffer(), true))
+		{
+			RenderPass.RestartRenderPass(StateCache.GetRenderPassDescriptor());
+		}
+		
+		if (bBindDepthStencilForUAVRaster)
+		{
+			mtlpp::ScissorRect Rect(0, 0, (NSUInteger)FBSize.width, (NSUInteger)FBSize.height);
+			StateCache.SetScissorRect(false, Rect);
+		}
+		
+		check(StateCache.GetHasValidRenderTarget());
+	}
+	else if (!bNeedsDepthStencilWrite && !bNeedsDepthStencilForUAVRaster && StateCache.GetFallbackDepthStencilBound())
+	{
+		FRHIRenderPassInfo Info = StateCache.GetRenderPassInfo();
+		Info.DepthStencilRenderTarget.DepthStencilTarget = nullptr;
+		
+		RenderPass.EndRenderPass();
+		
+		StateCache.SetRenderTargetsActive(false);
+		StateCache.SetRenderPassInfo(Info, StateCache.GetVisibilityResultsBuffer(), true);
+		
+		RenderPass.BeginRenderPass(StateCache.GetRenderPassDescriptor());
+		
+		check(StateCache.GetHasValidRenderTarget());
+	}
+	
+	return true;
+}
+
+void FAGXContext::SetRenderPassInfo(const FRHIRenderPassInfo& RenderTargetsInfo, bool const bRestart)
+{
+	if (CommandList.IsParallel())
+	{
+		GetAGXDeviceContext().SetParallelRenderPassDescriptor(RenderTargetsInfo);
+	}
+
+#if UE_BUILD_DEBUG || UE_BUILD_DEVELOPMENT
+	if (!CommandList.IsParallel() && !CommandList.IsImmediate())
+	{
+		bool bClearInParallelBuffer = false;
+		
+		for (uint32 RenderTargetIndex = 0; RenderTargetIndex < MaxSimultaneousRenderTargets; RenderTargetIndex++)
+		{
+			if (RenderTargetIndex < RenderTargetsInfo.GetNumColorRenderTargets() && RenderTargetsInfo.ColorRenderTargets[RenderTargetIndex].RenderTarget != nullptr)
+			{
+				const FRHIRenderPassInfo::FColorEntry& RenderTargetView = RenderTargetsInfo.ColorRenderTargets[RenderTargetIndex];
+				if(GetLoadAction(RenderTargetView.Action) == ERenderTargetLoadAction::EClear)
+				{
+					bClearInParallelBuffer = true;
+				}
+			}
+		}
+		
+		if (bClearInParallelBuffer)
+		{
+			UE_LOG(LogAGX, Warning, TEXT("One or more render targets bound for clear during parallel encoding: this will not behave as expected because each command-buffer will clear the target of the previous contents."));
+		}
+		
+		if (RenderTargetsInfo.DepthStencilRenderTarget.DepthStencilTarget != nullptr)
+		{
+			if(GetLoadAction(GetDepthActions(RenderTargetsInfo.DepthStencilRenderTarget.Action)) == ERenderTargetLoadAction::EClear)
+			{
+				UE_LOG(LogAGX, Warning, TEXT("Depth-target bound for clear during parallel encoding: this will not behave as expected because each command-buffer will clear the target of the previous contents."));
+			}
+			if(GetLoadAction(GetStencilActions(RenderTargetsInfo.DepthStencilRenderTarget.Action)) == ERenderTargetLoadAction::EClear)
+			{
+				UE_LOG(LogAGX, Warning, TEXT("Stencil-target bound for clear during parallel encoding: this will not behave as expected because each command-buffer will clear the target of the previous contents."));
+			}
+		}
+	}
+#endif
+	
+	bool bSet = false;
+	if (IsFeatureLevelSupported( GMaxRHIShaderPlatform, ERHIFeatureLevel::ES3_1 ))
+	{
+		// @todo Improve the way we handle binding a dummy depth/stencil so we can get pure UAV raster operations...
+		const bool bNeedsDepthStencilForUAVRaster = RenderTargetsInfo.GetNumColorRenderTargets() == 0 && !RenderTargetsInfo.DepthStencilRenderTarget.DepthStencilTarget;
+
+		if (bNeedsDepthStencilForUAVRaster)
+		{
+			FRHIRenderPassInfo Info = RenderTargetsInfo;
+			CGSize FBSize = CGSizeMake(StateCache.GetViewport(0).width, StateCache.GetViewport(0).height);
+			FTexture2DRHIRef FallbackDepthStencilSurface = StateCache.CreateFallbackDepthStencilSurface(FBSize.width, FBSize.height);
+			check(IsValidRef(FallbackDepthStencilSurface));
+
+			Info.DepthStencilRenderTarget.DepthStencilTarget = FallbackDepthStencilSurface;
+			Info.DepthStencilRenderTarget.ResolveTarget = nullptr;
+			Info.DepthStencilRenderTarget.ExclusiveDepthStencil = FExclusiveDepthStencil::DepthRead_StencilRead;
+#if PLATFORM_MAC
+			Info.DepthStencilRenderTarget.Action = MakeDepthStencilTargetActions(MakeRenderTargetActions(ERenderTargetLoadAction::ENoAction, ERenderTargetStoreAction::ENoAction), MakeRenderTargetActions(ERenderTargetLoadAction::ENoAction, ERenderTargetStoreAction::ENoAction));
+#else
+			Info.DepthStencilRenderTarget.Action = MakeDepthStencilTargetActions(MakeRenderTargetActions(ERenderTargetLoadAction::EClear, ERenderTargetStoreAction::ENoAction), MakeRenderTargetActions(ERenderTargetLoadAction::EClear, ERenderTargetStoreAction::ENoAction));
+#endif
+
+			if (QueryBuffer->GetCurrentQueryBuffer() != StateCache.GetVisibilityResultsBuffer())
+			{
+				RenderPass.EndRenderPass();
+			}
+			bSet = StateCache.SetRenderPassInfo(Info, QueryBuffer->GetCurrentQueryBuffer(), bRestart);
+		}
+		else
+		{
+			if (QueryBuffer->GetCurrentQueryBuffer() != StateCache.GetVisibilityResultsBuffer())
+			{
+				RenderPass.EndRenderPass();
+			}
+			bSet = StateCache.SetRenderPassInfo(RenderTargetsInfo, QueryBuffer->GetCurrentQueryBuffer(), bRestart);
+		}
+	}
+	else
+	{
+		if (NULL != StateCache.GetVisibilityResultsBuffer())
+		{
+			RenderPass.EndRenderPass();
+		}
+		bSet = StateCache.SetRenderPassInfo(RenderTargetsInfo, NULL, bRestart);
+	}
+	
+	if (bSet && StateCache.GetHasValidRenderTarget())
+	{
+		RenderPass.EndRenderPass();
+
+		if (NumParallelContextsInPass == 0)
+		{
+			RenderPass.BeginRenderPass(StateCache.GetRenderPassDescriptor());
+		}
+		else
+		{
+			RenderPass.BeginParallelRenderPass(StateCache.GetRenderPassDescriptor(), NumParallelContextsInPass);
+		}
+	}
+}
+
+FAGXBuffer FAGXContext::AllocateFromRingBuffer(uint32 Size, uint32 Alignment)
+{
+	return RenderPass.GetRingBuffer().NewBuffer(Size, Alignment);
+}
+
+void FAGXContext::DrawPrimitive(uint32 PrimitiveType, uint32 BaseVertexIndex, uint32 NumPrimitives, uint32 NumInstances)
+{
+	// finalize any pending state
+	if(!PrepareToDraw(PrimitiveType))
+	{
+		return;
+	}
+	
+	RenderPass.DrawPrimitive(PrimitiveType, BaseVertexIndex, NumPrimitives, NumInstances);
+}
+
+void FAGXContext::DrawPrimitiveIndirect(uint32 PrimitiveType, FAGXVertexBuffer* VertexBuffer, uint32 ArgumentOffset)
+{
+	// finalize any pending state
+	if(!PrepareToDraw(PrimitiveType))
+	{
+		return;
+	}
+	
+	RenderPass.DrawPrimitiveIndirect(PrimitiveType, VertexBuffer, ArgumentOffset);
+}
+
+void FAGXContext::DrawIndexedPrimitive(FAGXBuffer const& IndexBuffer, uint32 IndexStride, mtlpp::IndexType IndexType, uint32 PrimitiveType, int32 BaseVertexIndex, uint32 FirstInstance, uint32 NumVertices, uint32 StartIndex, uint32 NumPrimitives, uint32 NumInstances)
+{
+	// finalize any pending state
+	if(!PrepareToDraw(PrimitiveType))
+	{
+		return;
+	}
+	
+	RenderPass.DrawIndexedPrimitive(IndexBuffer, IndexStride, PrimitiveType, BaseVertexIndex, FirstInstance, NumVertices, StartIndex, NumPrimitives, NumInstances);
+}
+
+void FAGXContext::DrawIndexedIndirect(FAGXIndexBuffer* IndexBuffer, uint32 PrimitiveType, FAGXStructuredBuffer* VertexBuffer, int32 DrawArgumentsIndex, uint32 NumInstances)
+{
+	// finalize any pending state
+	if(!PrepareToDraw(PrimitiveType))
+	{
+		return;
+	}
+	
+	RenderPass.DrawIndexedIndirect(IndexBuffer, PrimitiveType, VertexBuffer, DrawArgumentsIndex, NumInstances);
+}
+
+void FAGXContext::DrawIndexedPrimitiveIndirect(uint32 PrimitiveType,FAGXIndexBuffer* IndexBuffer,FAGXVertexBuffer* VertexBuffer, uint32 ArgumentOffset)
+{	
+	// finalize any pending state
+	if(!PrepareToDraw(PrimitiveType))
+	{
+		return;
+	}
+	
+	RenderPass.DrawIndexedPrimitiveIndirect(PrimitiveType, IndexBuffer, VertexBuffer, ArgumentOffset);
+}
+
+void FAGXContext::CopyFromTextureToBuffer(FAGXTexture const& Texture, uint32 sourceSlice, uint32 sourceLevel, mtlpp::Origin sourceOrigin, mtlpp::Size sourceSize, FAGXBuffer const& toBuffer, uint32 destinationOffset, uint32 destinationBytesPerRow, uint32 destinationBytesPerImage, mtlpp::BlitOption options)
+{
+	RenderPass.CopyFromTextureToBuffer(Texture, sourceSlice, sourceLevel, sourceOrigin, sourceSize, toBuffer, destinationOffset, destinationBytesPerRow, destinationBytesPerImage, options);
+}
+
+void FAGXContext::CopyFromBufferToTexture(FAGXBuffer const& Buffer, uint32 sourceOffset, uint32 sourceBytesPerRow, uint32 sourceBytesPerImage, mtlpp::Size sourceSize, FAGXTexture const& toTexture, uint32 destinationSlice, uint32 destinationLevel, mtlpp::Origin destinationOrigin, mtlpp::BlitOption options)
+{
+	RenderPass.CopyFromBufferToTexture(Buffer, sourceOffset, sourceBytesPerRow, sourceBytesPerImage, sourceSize, toTexture, destinationSlice, destinationLevel, destinationOrigin, options);
+}
+
+void FAGXContext::CopyFromTextureToTexture(FAGXTexture const& Texture, uint32 sourceSlice, uint32 sourceLevel, mtlpp::Origin sourceOrigin, mtlpp::Size sourceSize, FAGXTexture const& toTexture, uint32 destinationSlice, uint32 destinationLevel, mtlpp::Origin destinationOrigin)
+{
+	RenderPass.CopyFromTextureToTexture(Texture, sourceSlice, sourceLevel, sourceOrigin, sourceSize, toTexture, destinationSlice, destinationLevel, destinationOrigin);
+}
+
+void FAGXContext::CopyFromBufferToBuffer(FAGXBuffer const& SourceBuffer, NSUInteger SourceOffset, FAGXBuffer const& DestinationBuffer, NSUInteger DestinationOffset, NSUInteger Size)
+{
+	RenderPass.CopyFromBufferToBuffer(SourceBuffer, SourceOffset, DestinationBuffer, DestinationOffset, Size);
+}
+
+bool FAGXContext::AsyncCopyFromBufferToTexture(FAGXBuffer const& Buffer, uint32 sourceOffset, uint32 sourceBytesPerRow, uint32 sourceBytesPerImage, mtlpp::Size sourceSize, FAGXTexture const& toTexture, uint32 destinationSlice, uint32 destinationLevel, mtlpp::Origin destinationOrigin, mtlpp::BlitOption options)
+{
+	return RenderPass.AsyncCopyFromBufferToTexture(Buffer, sourceOffset, sourceBytesPerRow, sourceBytesPerImage, sourceSize, toTexture, destinationSlice, destinationLevel, destinationOrigin, options);
+}
+
+bool FAGXContext::AsyncCopyFromTextureToTexture(FAGXTexture const& Texture, uint32 sourceSlice, uint32 sourceLevel, mtlpp::Origin sourceOrigin, mtlpp::Size sourceSize, FAGXTexture const& toTexture, uint32 destinationSlice, uint32 destinationLevel, mtlpp::Origin destinationOrigin)
+{
+	return RenderPass.AsyncCopyFromTextureToTexture(Texture, sourceSlice, sourceLevel, sourceOrigin, sourceSize, toTexture, destinationSlice, destinationLevel, destinationOrigin);
+}
+
+bool FAGXContext::CanAsyncCopyToBuffer(FAGXBuffer const& DestinationBuffer)
+{
+	return RenderPass.CanAsyncCopyToBuffer(DestinationBuffer);
+}
+
+void FAGXContext::AsyncCopyFromBufferToBuffer(FAGXBuffer const& SourceBuffer, NSUInteger SourceOffset, FAGXBuffer const& DestinationBuffer, NSUInteger DestinationOffset, NSUInteger Size)
+{
+	RenderPass.AsyncCopyFromBufferToBuffer(SourceBuffer, SourceOffset, DestinationBuffer, DestinationOffset, Size);
+}
+
+void FAGXContext::AsyncGenerateMipmapsForTexture(FAGXTexture const& Texture)
+{
+	RenderPass.AsyncGenerateMipmapsForTexture(Texture);
+}
+
+void FAGXContext::SubmitAsyncCommands(mtlpp::CommandBufferHandler ScheduledHandler, mtlpp::CommandBufferHandler CompletionHandler, bool const bWait)
+{
+	RenderPass.AddAsyncCommandBufferHandlers(ScheduledHandler, CompletionHandler);
+	if (bWait)
+	{
+		SubmitCommandsHint((uint32)(EAGXSubmitFlagsAsyncCommandBuffer|EAGXSubmitFlagsWaitOnCommandBuffer|EAGXSubmitFlagsBreakCommandBuffer));
+	}
+}
+
+void FAGXContext::SynchronizeTexture(FAGXTexture const& Texture, uint32 Slice, uint32 Level)
+{
+	RenderPass.SynchronizeTexture(Texture, Slice, Level);
+}
+
+void FAGXContext::SynchroniseResource(mtlpp::Resource const& Resource)
+{
+	RenderPass.SynchroniseResource(Resource);
+}
+
+void FAGXContext::FillBuffer(FAGXBuffer const& Buffer, ns::Range Range, uint8 Value)
+{
+	RenderPass.FillBuffer(Buffer, Range, Value);
+}
+
+void FAGXContext::Dispatch(uint32 ThreadGroupCountX, uint32 ThreadGroupCountY, uint32 ThreadGroupCountZ)
+{
+	RenderPass.Dispatch(ThreadGroupCountX, ThreadGroupCountY, ThreadGroupCountZ);
+}
+
+void FAGXContext::DispatchIndirect(FAGXVertexBuffer* ArgumentBuffer, uint32 ArgumentOffset)
+{
+	RenderPass.DispatchIndirect(ArgumentBuffer, ArgumentOffset);
+}
+
+void FAGXContext::StartTiming(class FAGXEventNode* EventNode)
+{
+	mtlpp::CommandBufferHandler Handler = nil;
+	
+	bool const bHasCurrentCommandBuffer = GetCurrentCommandBuffer();
+	
+	if(EventNode)
+	{
+		Handler = EventNode->Start();
+		
+		if (bHasCurrentCommandBuffer)
+		{
+			RenderPass.AddCompletionHandler(Handler);
+			Block_release(Handler);
+		}
+	}
+	
+	SubmitCommandsHint(EAGXSubmitFlagsCreateCommandBuffer);
+	
+	if (Handler != nil && !bHasCurrentCommandBuffer)
+	{
+		GetCurrentCommandBuffer().AddScheduledHandler(Handler);
+		Block_release(Handler);
+	}
+}
+
+void FAGXContext::EndTiming(class FAGXEventNode* EventNode)
+{
+	bool const bWait = EventNode->Wait();
+	mtlpp::CommandBufferHandler Handler = EventNode->Stop();
+	RenderPass.AddCompletionHandler(Handler);
+	Block_release(Handler);
+	
+	if (!bWait)
+	{
+		SubmitCommandsHint(EAGXSubmitFlagsCreateCommandBuffer);
+	}
+	else
+	{
+		SubmitCommandBufferAndWait();
+	}
+}
+
+void FAGXDeviceContext::BeginParallelRenderCommandEncoding(uint32 Num)
+{
+	FScopeLock Lock(&FreeListMutex);
+	FPlatformAtomics::AtomicStore(&ActiveParallelContexts, (int32)Num);
+	FPlatformAtomics::AtomicStore(&NumParallelContextsInPass, (int32)Num);
+}
+
+void FAGXDeviceContext::SetParallelRenderPassDescriptor(FRHIRenderPassInfo const& TargetInfo)
+{
+	FScopeLock Lock(&FreeListMutex);
+
+	if (!RenderPass.IsWithinParallelPass())
+	{
+		RenderPass.Begin(EndFence);
+		EndFence = nullptr;
+		StateCache.InvalidateRenderTargets();
+		SetRenderPassInfo(TargetInfo, false);
+	}
+}
+
+mtlpp::RenderCommandEncoder FAGXDeviceContext::GetParallelRenderCommandEncoder(uint32 Index, mtlpp::ParallelRenderCommandEncoder& ParallelEncoder, mtlpp::CommandBuffer& CommandBuffer)
+{
+	FScopeLock Lock(&FreeListMutex);
+	
+	check(RenderPass.IsWithinParallelPass());
+	CommandBuffer = GetCurrentCommandBuffer();
+	return RenderPass.GetParallelRenderCommandEncoder(Index, ParallelEncoder);
+}
+
+void FAGXDeviceContext::EndParallelRenderCommandEncoding(void)
+{
+	FScopeLock Lock(&FreeListMutex);
+
+	if (FPlatformAtomics::InterlockedDecrement(&ActiveParallelContexts) == 0)
+	{
+		RenderPass.EndRenderPass();
+		RenderPass.Begin(StartFence, true);
+		StartFence = nullptr;
+		FPlatformAtomics::AtomicStore(&NumParallelContextsInPass, 0);
+	}
+}
+
+#if METAL_SUPPORTS_PARALLEL_RHI_EXECUTE
+
+class FAGXCommandContextContainer : public IRHICommandContextContainer
+{
+	FAGXRHICommandContext* CmdContext;
+	int32 Index;
+	int32 Num;
+	
+public:
+	
+	/** Custom new/delete with recycling */
+	void* operator new(size_t Size);
+	void operator delete(void *RawMemory);
+	
+	FAGXCommandContextContainer(int32 InIndex, int32 InNum)
+	: CmdContext(nullptr)
+	, Index(InIndex)
+	, Num(InNum)
+	{
+		CmdContext = GetAGXDeviceContext().AcquireContext(Index, Num);
+		check(CmdContext);
+	}
+	
+	virtual ~FAGXCommandContextContainer() override final
+	{
+		check(!CmdContext);
+	}
+	
+	virtual IRHICommandContext* GetContext() override final
+	{
+		check(CmdContext);
+		CmdContext->GetInternalContext().InitFrame(false, Index, Num);
+		return CmdContext;
+	}
+	
+	virtual void FinishContext() override final
+	{
+	}
+
+	virtual void SubmitAndFreeContextContainer(int32 NewIndex, int32 NewNum) override final
+	{
+		if (CmdContext)
+		{
+			check(Index == NewIndex && Num == NewNum);
+			
+			if (Index == (Num - 1))
+			{
+				TRefCountPtr<FAGXFence> Fence = CmdContext->GetInternalContext().GetParallelPassEndFence();
+				GetAGXDeviceContext().SetParallelPassFences(Fence, nil);
+			}
+
+			CmdContext->GetInternalContext().FinishFrame(false);
+			GetAGXDeviceContext().EndParallelRenderCommandEncoding();
+
+			CmdContext->GetInternalContext().GetCommandList().Submit(Index, Num);
+			
+			GetAGXDeviceContext().ReleaseContext(CmdContext);
+			CmdContext = nullptr;
+			check(!CmdContext);
+		}
+		delete this;
+	}
+};
+
+static TLockFreeFixedSizeAllocator<sizeof(FAGXCommandContextContainer), PLATFORM_CACHE_LINE_SIZE, FThreadSafeCounter> FAGXCommandContextContainerAllocator;
+
+void* FAGXCommandContextContainer::operator new(size_t Size)
+{
+	return FMemory::Malloc(Size);
+}
+
+/**
+ * Custom delete
+ */
+void FAGXCommandContextContainer::operator delete(void *RawMemory)
+{
+	FMemory::Free(RawMemory);
+}
+
+IRHICommandContextContainer* FAGXDynamicRHI::RHIGetCommandContextContainer(int32 Index, int32 Num)
+{
+	return new FAGXCommandContextContainer(Index, Num);
+}
+
+#else
+
+IRHICommandContextContainer* FAGXDynamicRHI::RHIGetCommandContextContainer(int32 Index, int32 Num)
+{
+	return nullptr;
+}
+
+#endif
diff --git a/Engine/Source/Runtime/Apple/AGXRHI/Private/AGXContext.h b/Engine/Source/Runtime/Apple/AGXRHI/Private/AGXContext.h
new file mode 100644
--- /dev/null
+++ b/Engine/Source/Runtime/Apple/AGXRHI/Private/AGXContext.h
@@ -0,0 +1,348 @@
+// Copyright Epic Games, Inc. All Rights Reserved.
+
+
+#pragma once
+
+#import <Metal/Metal.h>
+#import <QuartzCore/CAMetalLayer.h>
+#include "AGXViewport.h"
+#include "AGXCommandEncoder.h"
+#include "AGXCommandQueue.h"
+#include "AGXCommandList.h"
+#include "AGXRenderPass.h"
+#include "AGXBuffer.h"
+#include "AGXCaptureManager.h"
+#include "AGXFrameAllocator.h"
+#if PLATFORM_IOS
+#include "IOS/IOSView.h"
+#endif
+#include "Containers/LockFreeList.h"
+#include "device.hpp"
+
+#define NUM_SAFE_FRAMES 4
+
+class FAGXRHICommandContext;
+class FAGXPipelineStateCacheManager;
+class FAGXQueryBufferPool;
+class FAGXRHIBuffer;
+
+class FAGXContext
+{
+	friend class FAGXCommandContextContainer;
+public:
+	FAGXContext(mtlpp::Device InDevice, FAGXCommandQueue& Queue, bool const bIsImmediate);
+	virtual ~FAGXContext();
+	
+	mtlpp::Device& GetDevice();
+	FAGXCommandQueue& GetCommandQueue();
+	FAGXCommandList& GetCommandList();
+	mtlpp::CommandBuffer const& GetCurrentCommandBuffer() const;
+	mtlpp::CommandBuffer& AGXRHI_API GetCurrentCommandBuffer();
+	FAGXStateCache& GetCurrentState() { return StateCache; }
+	FAGXRenderPass& GetCurrentRenderPass() { return RenderPass; }
+	
+	void InsertCommandBufferFence(FAGXCommandBufferFence& Fence, mtlpp::CommandBufferHandler Handler = nil);
+
+	/**
+	 * Do anything necessary to prepare for any kind of draw call 
+	 * @param PrimitiveType The UE4 primitive type for the draw call, needed to compile the correct render pipeline.
+	 * @returns True if the preparation completed and the draw call can be encoded, false to skip.
+	 */
+	bool PrepareToDraw(uint32 PrimitiveType);
+	
+	/**
+	 * Set the color, depth and stencil render targets, and then make the new command buffer/encoder
+	 */
+	void SetRenderPassInfo(const FRHIRenderPassInfo& RenderTargetsInfo, bool const bRestart = false);
+	
+	/**
+	 * Allocate from a dynamic ring buffer - by default align to the allowed alignment for offset field when setting buffers
+	 */
+	FAGXBuffer AllocateFromRingBuffer(uint32 Size, uint32 Alignment=0);
+
+	TSharedRef<FAGXQueryBufferPool, ESPMode::ThreadSafe> GetQueryBufferPool()
+	{
+		return QueryBuffer.ToSharedRef();
+	}
+
+    void SubmitCommandsHint(uint32 const bFlags = EAGXSubmitFlagsCreateCommandBuffer);
+	void SubmitCommandBufferAndWait();
+	void ResetRenderCommandEncoder();
+	
+	void DrawPrimitive(uint32 PrimitiveType, uint32 BaseVertexIndex, uint32 NumPrimitives, uint32 NumInstances);
+	
+	void DrawPrimitiveIndirect(uint32 PrimitiveType, FAGXVertexBuffer* VertexBuffer, uint32 ArgumentOffset);
+	
+	void DrawIndexedPrimitive(FAGXBuffer const& IndexBuffer, uint32 IndexStride, mtlpp::IndexType IndexType, uint32 PrimitiveType, int32 BaseVertexIndex, uint32 FirstInstance,
+							  uint32 NumVertices, uint32 StartIndex, uint32 NumPrimitives, uint32 NumInstances);
+	
+	void DrawIndexedIndirect(FAGXIndexBuffer* IndexBufferRHI, uint32 PrimitiveType, FAGXStructuredBuffer* VertexBufferRHI, int32 DrawArgumentsIndex, uint32 NumInstances);
+	
+	void DrawIndexedPrimitiveIndirect(uint32 PrimitiveType,FAGXIndexBuffer* IndexBufferRHI,FAGXVertexBuffer* VertexBufferRHI,uint32 ArgumentOffset);
+	
+	void DrawPatches(uint32 PrimitiveType, FAGXBuffer const& IndexBuffer, uint32 IndexBufferStride, int32 BaseVertexIndex, uint32 FirstInstance, uint32 StartIndex,
+					 uint32 NumPrimitives, uint32 NumInstances);
+	
+	void CopyFromTextureToBuffer(FAGXTexture const& Texture, uint32 sourceSlice, uint32 sourceLevel, mtlpp::Origin sourceOrigin, mtlpp::Size sourceSize, FAGXBuffer const& toBuffer, uint32 destinationOffset, uint32 destinationBytesPerRow, uint32 destinationBytesPerImage, mtlpp::BlitOption options);
+	
+	void CopyFromBufferToTexture(FAGXBuffer const& Buffer, uint32 sourceOffset, uint32 sourceBytesPerRow, uint32 sourceBytesPerImage, mtlpp::Size sourceSize, FAGXTexture const& toTexture, uint32 destinationSlice, uint32 destinationLevel, mtlpp::Origin destinationOrigin, mtlpp::BlitOption options);
+	
+	void CopyFromTextureToTexture(FAGXTexture const& Texture, uint32 sourceSlice, uint32 sourceLevel, mtlpp::Origin sourceOrigin, mtlpp::Size sourceSize, FAGXTexture const& toTexture, uint32 destinationSlice, uint32 destinationLevel, mtlpp::Origin destinationOrigin);
+	
+	void CopyFromBufferToBuffer(FAGXBuffer const& SourceBuffer, NSUInteger SourceOffset, FAGXBuffer const& DestinationBuffer, NSUInteger DestinationOffset, NSUInteger Size);
+	
+    bool AsyncCopyFromBufferToTexture(FAGXBuffer const& Buffer, uint32 sourceOffset, uint32 sourceBytesPerRow, uint32 sourceBytesPerImage, mtlpp::Size sourceSize, FAGXTexture const& toTexture, uint32 destinationSlice, uint32 destinationLevel, mtlpp::Origin destinationOrigin, mtlpp::BlitOption options);
+    
+    bool AsyncCopyFromTextureToTexture(FAGXTexture const& Texture, uint32 sourceSlice, uint32 sourceLevel, mtlpp::Origin sourceOrigin, mtlpp::Size sourceSize, FAGXTexture const& toTexture, uint32 destinationSlice, uint32 destinationLevel, mtlpp::Origin destinationOrigin);
+	
+	bool CanAsyncCopyToBuffer(FAGXBuffer const& DestinationBuffer);
+	
+    void AsyncCopyFromBufferToBuffer(FAGXBuffer const& SourceBuffer, NSUInteger SourceOffset, FAGXBuffer const& DestinationBuffer, NSUInteger DestinationOffset, NSUInteger Size);
+	
+    void AsyncGenerateMipmapsForTexture(FAGXTexture const& Texture);
+    
+	void SubmitAsyncCommands(mtlpp::CommandBufferHandler ScheduledHandler, mtlpp::CommandBufferHandler CompletionHandler, bool const bWait);
+	
+	void SynchronizeTexture(FAGXTexture const& Texture, uint32 Slice, uint32 Level);
+	
+	void SynchroniseResource(mtlpp::Resource const& Resource);
+	
+	void FillBuffer(FAGXBuffer const& Buffer, ns::Range Range, uint8 Value);
+
+	void Dispatch(uint32 ThreadGroupCountX, uint32 ThreadGroupCountY, uint32 ThreadGroupCountZ);
+	void DispatchIndirect(FAGXVertexBuffer* ArgumentBuffer, uint32 ArgumentOffset);
+
+	void StartTiming(class FAGXEventNode* EventNode);
+	void EndTiming(class FAGXEventNode* EventNode);
+
+#if ENABLE_METAL_GPUPROFILE
+	static void MakeCurrent(FAGXContext* Context);
+	static FAGXContext* GetCurrentContext();
+#endif
+	
+	void SetParallelPassFences(FAGXFence* Start, FAGXFence* End);
+	TRefCountPtr<FAGXFence> const& GetParallelPassStartFence(void) const;
+	TRefCountPtr<FAGXFence> const& GetParallelPassEndFence(void) const;
+	
+	void InitFrame(bool const bImmediateContext, uint32 Index, uint32 Num);
+	void FinishFrame(bool const bImmediateContext);
+
+	// Track Write->Read transitions for TBDR Fragment->Verex fencing
+	void TransitionResource(FRHIUnorderedAccessView* InResource);
+	void TransitionResource(FRHITexture* InResource);
+
+	template<typename T>
+	void TransitionRHIResource(T* InResource);
+
+protected:
+	/** The underlying Metal device */
+	mtlpp::Device Device;
+	
+	/** The wrapper around the device command-queue for creating & committing command buffers to */
+	FAGXCommandQueue& CommandQueue;
+	
+	/** The wrapper around commabd buffers for ensuring correct parallel execution order */
+	FAGXCommandList CommandList;
+	
+	/** The cache of all tracked & accessible state. */
+	FAGXStateCache StateCache;
+	
+	/** The render pass handler that actually encodes our commands. */
+	FAGXRenderPass RenderPass;
+	
+	/** A sempahore used to ensure that wait for previous frames to complete if more are in flight than we permit */
+	dispatch_semaphore_t CommandBufferSemaphore;
+	
+	/** A pool of buffers for writing visibility query results. */
+	TSharedPtr<FAGXQueryBufferPool, ESPMode::ThreadSafe> QueryBuffer;
+	
+	/** Initial fence to wait on for parallel contexts */
+	TRefCountPtr<FAGXFence> StartFence;
+	
+	/** Fence to update at the end for parallel contexts */
+	TRefCountPtr<FAGXFence> EndFence;
+	
+#if ENABLE_METAL_GPUPROFILE
+	/** the slot to store a per-thread context ref */
+	static uint32 CurrentContextTLSSlot;
+#endif
+	
+	/** Total number of parallel contexts that constitute the current pass. */
+	int32 NumParallelContextsInPass;
+	
+	/** Whether the validation layer is enabled */
+	bool bValidationEnabled;
+};
+
+template<typename T>
+void FAGXContext::TransitionRHIResource(T* InResource)
+{
+	auto Resource = ResourceCast(InResource);
+	if (Resource->GetCurrentBufferOrNil())
+	{
+		RenderPass.TransitionResources(Resource->GetCurrentBuffer());
+	}
+}
+
+class FAGXDeviceContext : public FAGXContext
+{
+public:
+	static FAGXDeviceContext* CreateDeviceContext();
+	virtual ~FAGXDeviceContext();
+	
+	void Init(void);
+	
+	inline bool SupportsFeature(EAGXFeatures InFeature) { return CommandQueue.SupportsFeature(InFeature); }
+	
+	inline FAGXResourceHeap& GetResourceHeap(void) { return Heap; }
+	
+	FAGXTexture CreateTexture(FAGXSurface* Surface, mtlpp::TextureDescriptor Descriptor);
+	FAGXBuffer CreatePooledBuffer(FAGXPooledBufferArgs const& Args);
+	void ReleaseBuffer(FAGXBuffer& Buf);
+	void ReleaseObject(id Object);
+	void ReleaseTexture(FAGXSurface* Surface, FAGXTexture& Texture);
+	void ReleaseTexture(FAGXTexture& Texture);
+	void ReleaseFence(FAGXFence* Fence);
+	
+	void BeginFrame();
+	void FlushFreeList(bool const bFlushFences = true);
+	void ClearFreeList();
+	void DrainHeap();
+	void EndFrame();
+	
+	/** RHIBeginScene helper */
+	void BeginScene();
+	/** RHIEndScene helper */
+	void EndScene();
+	
+	void BeginDrawingViewport(FAGXViewport* Viewport);
+	void EndDrawingViewport(FAGXViewport* Viewport, bool bPresent, bool bLockToVsync);
+	
+	/** Take a parallel FAGXContext from the free-list or allocate a new one if required */
+	FAGXRHICommandContext* AcquireContext(int32 NewIndex, int32 NewNum);
+	
+	/** Release a parallel FAGXContext back into the free-list */
+	void ReleaseContext(FAGXRHICommandContext* Context);
+	
+	/** Returns the number of concurrent contexts encoding commands, including the device context. */
+	uint32 GetNumActiveContexts(void) const;
+
+	void BeginParallelRenderCommandEncoding(uint32 Num);
+	void SetParallelRenderPassDescriptor(FRHIRenderPassInfo const& TargetInfo);
+	mtlpp::RenderCommandEncoder GetParallelRenderCommandEncoder(uint32 Index, mtlpp::ParallelRenderCommandEncoder& ParallelEncoder, mtlpp::CommandBuffer& CommandBuffer);
+	void EndParallelRenderCommandEncoding(void);
+	
+	/** Get the index of the bound Metal device in the global list of rendering devices. */
+	uint32 GetDeviceIndex(void) const;
+	
+	FAGXFrameAllocator* GetTransferAllocator()
+	{
+		return TransferBufferAllocator;
+	}
+    
+    FAGXFrameAllocator* GetUniformAllocator()
+    {
+        return UniformBufferAllocator;
+    }
+    
+    uint32 GetFrameNumberRHIThread()
+    {
+        return FrameNumberRHIThread;
+    }
+	
+	void NewLock(FAGXRHIBuffer* Buffer, FAGXFrameAllocator::AllocationEntry& Allocation);
+	FAGXFrameAllocator::AllocationEntry FetchAndRemoveLock(FAGXRHIBuffer* Buffer);
+	
+#if METAL_DEBUG_OPTIONS
+    void AddActiveBuffer(FAGXBuffer const& Buffer);
+    void RemoveActiveBuffer(FAGXBuffer const& Buffer);
+	bool ValidateIsInactiveBuffer(FAGXBuffer const& Buffer);
+	void ScribbleBuffer(FAGXBuffer& Buffer);
+#endif
+	
+private:
+	FAGXDeviceContext(mtlpp::Device MetalDevice, uint32 DeviceIndex, FAGXCommandQueue* Queue);
+	
+private:
+	/** The index into the GPU device list for the selected Metal device */
+	uint32 DeviceIndex;
+	
+	/** Dynamic memory heap */
+	FAGXResourceHeap Heap;
+	
+	/** GPU Frame Capture Manager */
+	FAGXCaptureManager CaptureManager;
+	
+	/** Free lists for releasing objects only once it is safe to do so */
+	TSet<FAGXBuffer> UsedBuffers;
+	TSet<FAGXTexture> UsedTextures;
+	TSet<FAGXFence*> UsedFences;
+	TLockFreePointerListLIFO<FAGXFence> FenceFreeList;
+	TSet<id> ObjectFreeList;
+	struct FAGXDelayedFreeList
+	{
+		bool IsComplete() const;
+		TArray<mtlpp::CommandBufferFence> Fences;
+		TSet<FAGXBuffer> UsedBuffers;
+		TSet<FAGXTexture> UsedTextures;
+		TSet<FAGXFence*> FenceFreeList;
+		TSet<id> ObjectFreeList;
+#if METAL_DEBUG_OPTIONS
+		int32 DeferCount;
+#endif
+	};
+	TArray<FAGXDelayedFreeList*> DelayedFreeLists;
+	
+//	TSet<FAGXUniformBuffer*> UniformBuffers;
+    FAGXFrameAllocator* UniformBufferAllocator;
+	FAGXFrameAllocator* TransferBufferAllocator;
+	
+	TMap<FAGXRHIBuffer*, FAGXFrameAllocator::AllocationEntry> OutstandingLocks;
+	
+#if METAL_DEBUG_OPTIONS
+	/** The list of fences for the current frame */
+	TArray<FAGXFence*> FrameFences;
+    
+    FCriticalSection ActiveBuffersMutex;
+    
+    /** These are the active buffers that cannot be CPU modified */
+    TMap<id<MTLBuffer>, TArray<NSRange>> ActiveBuffers;
+#endif
+	
+	/** Free-list of contexts for parallel encoding */
+	TLockFreePointerListLIFO<FAGXRHICommandContext> ParallelContexts;
+	
+	/** Fences for parallel execution */
+	TArray<TRefCountPtr<FAGXFence>> ParallelFences;
+	
+	/** Critical section for FreeList */
+	FCriticalSection FreeListMutex;
+	
+	/** Event for coordinating pausing of render thread to keep inline with the ios display link. */
+	FEvent* FrameReadyEvent;
+	
+	/** Internal frame counter, incremented on each call to RHIBeginScene. */
+	uint32 SceneFrameCounter;
+	
+	/** Internal frame counter, used to ensure that we only drain the buffer pool one after each frame within RHIEndFrame. */
+	uint32 FrameCounter;
+	
+	/** Bitfield of supported Metal features with varying availability depending on OS/device */
+	uint32 Features;
+	
+	/** Count of concurrent contexts encoding commands. */
+	int32 ActiveContexts;
+	
+	/** Count of concurrent parallel contexts encoding commands. */
+	int32 ActiveParallelContexts;
+	
+	/** Whether we presented this frame - only used to track when to introduce debug markers */
+	bool bPresented;
+	
+	/** PSO cache manager */
+	FAGXPipelineStateCacheManager* PSOManager;
+
+    /** Thread index owned by the RHI Thread. Monotonically increases every call to EndFrame() */
+    uint32 FrameNumberRHIThread;
+};
+
diff --git a/Engine/Source/Runtime/Apple/AGXRHI/Private/AGXDebugCommandEncoder.cpp b/Engine/Source/Runtime/Apple/AGXRHI/Private/AGXDebugCommandEncoder.cpp
new file mode 100644
--- /dev/null
+++ b/Engine/Source/Runtime/Apple/AGXRHI/Private/AGXDebugCommandEncoder.cpp
@@ -0,0 +1,64 @@
+// Copyright Epic Games, Inc. All Rights Reserved.
+
+#include "AGXRHIPrivate.h"
+
+#include "AGXDebugCommandEncoder.h"
+#include "AGXCommandBuffer.h"
+#include "AGXFence.h"
+
+#if MTLPP_CONFIG_VALIDATE && METAL_DEBUG_OPTIONS
+extern int32 GAGXRuntimeDebugLevel;
+
+@implementation FAGXDebugCommandEncoder
+
+APPLE_PLATFORM_OBJECT_ALLOC_OVERRIDES(FAGXDebugCommandEncoder)
+
+-(instancetype)init
+{
+	id Self = [super init];
+	if (Self)
+	{
+		UpdatedFences = [[NSHashTable weakObjectsHashTable] retain];
+		WaitingFences = [[NSHashTable weakObjectsHashTable] retain];
+	}
+	return Self;
+}
+
+-(void)dealloc
+{
+	[UpdatedFences release];
+	[WaitingFences release];
+	[super dealloc];
+}
+@end
+
+FAGXCommandEncoderDebugging::FAGXCommandEncoderDebugging()
+{
+	
+}
+
+FAGXCommandEncoderDebugging::FAGXCommandEncoderDebugging(FAGXDebugCommandEncoder* handle)
+: ns::Object<FAGXDebugCommandEncoder*>(handle)
+{
+	
+}
+
+void FAGXCommandEncoderDebugging::AddUpdateFence(id Fence)
+{
+	if ((EAGXDebugLevel)GAGXRuntimeDebugLevel >= EAGXDebugLevelValidation && Fence)
+	{
+		[m_ptr->UpdatedFences addObject:(FAGXDebugFence*)Fence];
+		[(FAGXDebugFence*)Fence updatingEncoder:m_ptr];
+	}
+}
+
+void FAGXCommandEncoderDebugging::AddWaitFence(id Fence)
+{
+	if ((EAGXDebugLevel)GAGXRuntimeDebugLevel >= EAGXDebugLevelValidation && Fence)
+	{
+		[m_ptr->WaitingFences addObject:(FAGXDebugFence*)Fence];
+		[(FAGXDebugFence*)Fence waitingEncoder:m_ptr];
+	}
+}
+#endif
+
diff --git a/Engine/Source/Runtime/Apple/AGXRHI/Private/AGXDebugCommandEncoder.h b/Engine/Source/Runtime/Apple/AGXRHI/Private/AGXDebugCommandEncoder.h
new file mode 100644
--- /dev/null
+++ b/Engine/Source/Runtime/Apple/AGXRHI/Private/AGXDebugCommandEncoder.h
@@ -0,0 +1,109 @@
+// Copyright Epic Games, Inc. All Rights Reserved.
+
+#pragma once
+
+#include "AGXCommandBuffer.h"
+
+// For some reason when including this file while building the editor Clang 9 ignores this pragma from MacPlatformCompilerPreSetup.h,
+// resulting in errors in FAGXDebugBufferBindings, FAGXDebugTextureBindings and FAGXDebugSamplerBindings. Readding it here works around this problem.
+#if (__clang_major__ >= 9)
+#pragma clang diagnostic ignored "-Wnullability-inferred-on-nested-type"
+#endif
+
+NS_ASSUME_NONNULL_BEGIN
+
+/**
+ * The sampler, buffer and texture resource limits as defined here:
+ * https://developer.apple.com/library/ios/documentation/Miscellaneous/Conceptual/MetalProgrammingGuide/Render-Ctx/Render-Ctx.html
+ */
+#if PLATFORM_IOS
+#define METAL_MAX_BUFFERS 31
+#define METAL_MAX_TEXTURES 31
+typedef uint32 FAGXTextureMask;
+#elif PLATFORM_MAC
+#define METAL_MAX_BUFFERS 31
+#define METAL_MAX_TEXTURES 128
+typedef __uint128_t FAGXTextureMask;
+#else
+#error "Unsupported Platform!"
+#endif
+typedef uint32 FAGXBufferMask;
+typedef uint16 FAGXSamplerMask;
+
+enum EAGXLimits
+{
+	ML_MaxSamplers = 16, /** Maximum number of samplers */
+	ML_MaxBuffers = METAL_MAX_BUFFERS, /** Maximum number of buffers */
+	ML_MaxTextures = METAL_MAX_TEXTURES, /** Maximum number of textures - there are more textures available on Mac than iOS */
+	ML_MaxViewports = 16 /** Technically this may be different at runtime, but this is the likely absolute upper-bound */
+};
+
+enum EAGXShaderFrequency
+{
+    EAGXShaderVertex = 0,
+    EAGXShaderFragment = 1,
+    EAGXShaderCompute = 2,
+	EAGXShaderStream = 3,
+    EAGXShaderRenderNum = 2,
+	EAGXShaderStagesNum = 4
+};
+
+/** A structure for quick mask-testing of shader-stage resource bindings */
+struct FAGXDebugShaderResourceMask
+{
+	FAGXTextureMask TextureMask;
+	FAGXBufferMask BufferMask;
+	FAGXSamplerMask SamplerMask;
+};
+
+/** A structure of arrays for the current buffer binding settings. */
+struct FAGXDebugBufferBindings
+{
+    /** The bound buffers or nil. */
+    id<MTLBuffer> _Nullable Buffers[ML_MaxBuffers];
+    /** Optional bytes buffer used instead of an id<MTLBuffer> */
+    void const* _Nullable Bytes[ML_MaxBuffers];
+    /** The bound buffer offsets or 0. */
+    NSUInteger Offsets[ML_MaxBuffers];
+};
+
+/** A structure of arrays for the current texture binding settings. */
+struct FAGXDebugTextureBindings
+{
+    /** The bound textures or nil. */
+    id<MTLTexture> _Nullable Textures[ML_MaxTextures];
+};
+
+/** A structure of arrays for the current sampler binding settings. */
+struct FAGXDebugSamplerBindings
+{
+    /** The bound sampler states or nil. */
+    id<MTLSamplerState> _Nullable Samplers[ML_MaxSamplers];
+};
+
+#if MTLPP_CONFIG_VALIDATE && METAL_DEBUG_OPTIONS
+class FAGXCommandBufferDebugging;
+@class FAGXDebugCommandBuffer;
+@class FAGXDebugFence;
+
+@interface FAGXDebugCommandEncoder : FApplePlatformObject
+{
+@public
+	NSHashTable<FAGXDebugFence*>* UpdatedFences;
+	NSHashTable<FAGXDebugFence*>* WaitingFences;
+}
+-(instancetype)init;
+@end
+
+class FAGXCommandEncoderDebugging : public ns::Object<FAGXDebugCommandEncoder*>
+{
+public:
+	FAGXCommandEncoderDebugging();
+	FAGXCommandEncoderDebugging(FAGXDebugCommandEncoder* handle);
+	
+	void AddUpdateFence(id Fence);
+	void AddWaitFence(id Fence);
+};
+
+#endif
+NS_ASSUME_NONNULL_END
diff --git a/Engine/Source/Runtime/Apple/AGXRHI/Private/AGXDynamicRHI.cpp b/Engine/Source/Runtime/Apple/AGXRHI/Private/AGXDynamicRHI.cpp
new file mode 100644
--- /dev/null
+++ b/Engine/Source/Runtime/Apple/AGXRHI/Private/AGXDynamicRHI.cpp
@@ -0,0 +1,151 @@
+// Copyright Epic Games, Inc. All Rights Reserved.
+
+/*=============================================================================
+	AGXDynamicRHI.cpp: AGX Dynamic RHI Class Implementation.
+=============================================================================*/
+
+
+#include "AGXRHIPrivate.h"
+#include "AGXRHIRenderQuery.h"
+#include "AGXRHIStagingBuffer.h"
+#include "AGXShaderTypes.h"
+#include "AGXVertexDeclaration.h"
+#include "AGXGraphicsPipelineState.h"
+#include "AGXComputePipelineState.h"
+#include "AGXTransitionData.h"
+
+
+//------------------------------------------------------------------------------
+
+#pragma mark - AGX Dynamic RHI Vertex Declaration Methods -
+
+
+FVertexDeclarationRHIRef FAGXDynamicRHI::RHICreateVertexDeclaration(const FVertexDeclarationElementList& Elements)
+{
+	@autoreleasepool {
+		uint32 Key = FCrc::MemCrc32(Elements.GetData(), Elements.Num() * sizeof(FVertexElement));
+
+		// look up an existing declaration
+		FVertexDeclarationRHIRef* VertexDeclarationRefPtr = VertexDeclarationCache.Find(Key);
+		if (VertexDeclarationRefPtr == NULL)
+		{
+			// create and add to the cache if it doesn't exist.
+			VertexDeclarationRefPtr = &VertexDeclarationCache.Add(Key, new FAGXVertexDeclaration(Elements));
+		}
+
+		return *VertexDeclarationRefPtr;
+	} // autoreleasepool
+}
+
+
+//------------------------------------------------------------------------------
+
+#pragma mark - AGX Dynamic RHI Pipeline State Methods -
+
+
+FGraphicsPipelineStateRHIRef FAGXDynamicRHI::RHICreateGraphicsPipelineState(const FGraphicsPipelineStateInitializer& Initializer)
+{
+	@autoreleasepool {
+		FAGXGraphicsPipelineState* State = new FAGXGraphicsPipelineState(Initializer);
+
+		if(!State->Compile())
+		{
+			// Compilation failures are propagated up to the caller.
+			State->DoNoDeferDelete();
+			delete State;
+			return nullptr;
+		}
+
+		State->VertexDeclaration = ResourceCast(Initializer.BoundShaderState.VertexDeclarationRHI);
+		State->VertexShader = ResourceCast(Initializer.BoundShaderState.VertexShaderRHI);
+		State->PixelShader = ResourceCast(Initializer.BoundShaderState.PixelShaderRHI);
+#if PLATFORM_SUPPORTS_GEOMETRY_SHADERS
+		State->GeometryShader = ResourceCast(Initializer.BoundShaderState.GeometryShaderRHI);
+#endif // PLATFORM_SUPPORTS_GEOMETRY_SHADERS
+
+		State->DepthStencilState = ResourceCast(Initializer.DepthStencilState);
+		State->RasterizerState = ResourceCast(Initializer.RasterizerState);
+
+		return State;
+	} // autoreleasepool
+}
+
+TRefCountPtr<FRHIComputePipelineState> FAGXDynamicRHI::RHICreateComputePipelineState(FRHIComputeShader* ComputeShader)
+{
+	@autoreleasepool {
+		return new FAGXComputePipelineState(ResourceCast(ComputeShader));
+	} // autoreleasepool
+}
+
+
+//------------------------------------------------------------------------------
+
+#pragma mark - AGX Dynamic RHI Staging Buffer Methods -
+
+
+FStagingBufferRHIRef FAGXDynamicRHI::RHICreateStagingBuffer()
+{
+	return new FAGXRHIStagingBuffer();
+}
+
+void* FAGXDynamicRHI::RHILockStagingBuffer(FRHIStagingBuffer* StagingBuffer, FRHIGPUFence* Fence, uint32 Offset, uint32 SizeRHI)
+{
+	FAGXRHIStagingBuffer* Buffer = ResourceCast(StagingBuffer);
+	return Buffer->Lock(Offset, SizeRHI);
+}
+
+void FAGXDynamicRHI::RHIUnlockStagingBuffer(FRHIStagingBuffer* StagingBuffer)
+{
+	FAGXRHIStagingBuffer* Buffer = ResourceCast(StagingBuffer);
+	Buffer->Unlock();
+}
+
+
+//------------------------------------------------------------------------------
+
+#pragma mark - AGX Dynamic RHI Resource Transition Methods -
+
+
+void FAGXDynamicRHI::RHICreateTransition(FRHITransition* Transition, const FRHITransitionCreateInfo& CreateInfo)
+{
+	checkf(FMath::IsPowerOfTwo(uint32(CreateInfo.SrcPipelines)) && FMath::IsPowerOfTwo(uint32(CreateInfo.DstPipelines)), TEXT("Support for multi-pipe resources is not yet implemented."));
+
+	// Construct the data in-place on the transition instance
+	new (Transition->GetPrivateData<FAGXTransitionData>()) FAGXTransitionData(CreateInfo.SrcPipelines, CreateInfo.DstPipelines, CreateInfo.Flags, CreateInfo.TransitionInfos);
+}
+
+void FAGXDynamicRHI::RHIReleaseTransition(FRHITransition* Transition)
+{
+	// Destruct the private data object of the transition instance.
+	Transition->GetPrivateData<FAGXTransitionData>()->~FAGXTransitionData();
+}
+
+
+//------------------------------------------------------------------------------
+
+#pragma mark - AGX Dynamic RHI Render Query Methods -
+
+
+FRenderQueryRHIRef FAGXDynamicRHI::RHICreateRenderQuery(ERenderQueryType QueryType)
+{
+	@autoreleasepool {
+		FRenderQueryRHIRef Query = new FAGXRHIRenderQuery(QueryType);
+		return Query;
+	}
+}
+
+FRenderQueryRHIRef FAGXDynamicRHI::RHICreateRenderQuery_RenderThread(class FRHICommandListImmediate& RHICmdList, ERenderQueryType QueryType)
+{
+	@autoreleasepool {
+		return GDynamicRHI->RHICreateRenderQuery(QueryType);
+	}
+}
+
+bool FAGXDynamicRHI::RHIGetRenderQueryResult(FRHIRenderQuery* QueryRHI, uint64& OutNumPixels, bool bWait, uint32 GPUIndex)
+{
+	@autoreleasepool {
+		check(IsInRenderingThread());
+		FAGXRHIRenderQuery* Query = ResourceCast(QueryRHI);
+		return Query->GetResult(OutNumPixels, bWait, GPUIndex);
+	}
+}
