Initial import of AGXRHI

diff --git a/Engine/Source/Runtime/Apple/AGXRHI/Private/AGXDynamicRHI.h b/Engine/Source/Runtime/Apple/AGXRHI/Private/AGXDynamicRHI.h
new file mode 100644
--- /dev/null
+++ b/Engine/Source/Runtime/Apple/AGXRHI/Private/AGXDynamicRHI.h
@@ -0,0 +1,197 @@
+// Copyright Epic Games, Inc. All Rights Reserved.
+
+/*=============================================================================
+	AGXRHI.h: Public AGX RHI definitions.
+=============================================================================*/
+
+#pragma once 
+
+#include "CoreMinimal.h"
+#include "RHI.h"
+
+DECLARE_LOG_CATEGORY_EXTERN(LogAGX, Display, All);
+
+#include <Metal/Metal.h>
+#include "AGXState.h"
+#include "AGXResources.h"
+#include "AGXRHIContext.h"
+#include "AGXViewport.h"
+
+class FAGXDeviceContext;
+
+/** The interface which is implemented by the dynamically bound RHI. */
+class FAGXDynamicRHI : public FDynamicRHI
+{
+public:
+
+	/** Initialization constructor. */
+	FAGXDynamicRHI(ERHIFeatureLevel::Type RequestedFeatureLevel);
+
+	/** Destructor */
+	~FAGXDynamicRHI();
+	
+	/** Setup for necessary internal resources */
+	void SetupRecursiveResources();
+
+	// FDynamicRHI interface.
+	virtual void Init();
+	virtual void Shutdown() {}
+	virtual const TCHAR* GetName() override { return TEXT("AGX"); }
+
+	virtual FRHIShaderLibraryRef RHICreateShaderLibrary(EShaderPlatform Platform, FString const& FilePath, FString const& Name) final override;
+	virtual FSamplerStateRHIRef RHICreateSamplerState(const FSamplerStateInitializerRHI& Initializer) final override;
+	virtual FRasterizerStateRHIRef RHICreateRasterizerState(const FRasterizerStateInitializerRHI& Initializer) final override;
+	virtual FDepthStencilStateRHIRef RHICreateDepthStencilState(const FDepthStencilStateInitializerRHI& Initializer) final override;
+	virtual FBlendStateRHIRef RHICreateBlendState(const FBlendStateInitializerRHI& Initializer) final override;
+	virtual FVertexDeclarationRHIRef RHICreateVertexDeclaration(const FVertexDeclarationElementList& Elements) final override;
+	virtual FPixelShaderRHIRef RHICreatePixelShader(TArrayView<const uint8> Code, const FSHAHash& Hash) final override;
+	virtual FVertexShaderRHIRef RHICreateVertexShader(TArrayView<const uint8> Code, const FSHAHash& Hash) final override;
+	virtual FGeometryShaderRHIRef RHICreateGeometryShader(TArrayView<const uint8> Code, const FSHAHash& Hash) final override;
+	virtual FComputeShaderRHIRef RHICreateComputeShader(TArrayView<const uint8> Code, const FSHAHash& Hash) final override;
+	virtual FBoundShaderStateRHIRef RHICreateBoundShaderState(FRHIVertexDeclaration* VertexDeclaration, FRHIVertexShader* VertexShader, FRHIPixelShader* PixelShader, FRHIGeometryShader* GeometryShader) final override;
+	virtual FGraphicsPipelineStateRHIRef RHICreateGraphicsPipelineState(const FGraphicsPipelineStateInitializer& Initializer) final override;
+	virtual TRefCountPtr<FRHIComputePipelineState> RHICreateComputePipelineState(FRHIComputeShader* ComputeShader) final override;
+	virtual FUniformBufferRHIRef RHICreateUniformBuffer(const void* Contents, const FRHIUniformBufferLayout& Layout, EUniformBufferUsage Usage, EUniformBufferValidation Validation) final override;
+	virtual FBufferRHIRef RHICreateBuffer(uint32 Size, EBufferUsageFlags Usage, uint32 Stride, ERHIAccess InResourceState, FRHIResourceCreateInfo& CreateInfo) final override;
+	virtual void RHITransferBufferUnderlyingResource(FRHIBuffer* DestBuffer, FRHIBuffer* SrcBuffer) final override;
+	virtual void RHICopyBuffer(FRHIBuffer* SourceBuffer, FRHIBuffer* DestBuffer) final override;
+	virtual void* LockBuffer_BottomOfPipe(FRHICommandListImmediate& RHICmdList, FRHIBuffer* Buffer, uint32 Offset, uint32 SizeRHI, EResourceLockMode LockMode) final override;
+	virtual void UnlockBuffer_BottomOfPipe(FRHICommandListImmediate& RHICmdList, FRHIBuffer* Buffer) final override;
+	virtual FUnorderedAccessViewRHIRef RHICreateUnorderedAccessView(FRHIBuffer* Buffer, bool bUseUAVCounter, bool bAppendBuffer) final override;
+	virtual FUnorderedAccessViewRHIRef RHICreateUnorderedAccessView(FRHITexture* Texture, uint32 MipLevel) final override;
+	virtual FUnorderedAccessViewRHIRef RHICreateUnorderedAccessView(FRHIBuffer* Buffer, uint8 Format) final override;
+	virtual FShaderResourceViewRHIRef RHICreateShaderResourceView(FRHIBuffer* Buffer) final override;
+	virtual FShaderResourceViewRHIRef RHICreateShaderResourceView(FRHITexture* Texture2DRHI, const FRHITextureSRVCreateInfo& CreateInfo) final override;
+	virtual FShaderResourceViewRHIRef RHICreateShaderResourceView(FRHIBuffer* Buffer, uint32 Stride, uint8 Format) final override;
+	virtual FShaderResourceViewRHIRef RHICreateShaderResourceView(const FShaderResourceViewInitializer& Initializer) final override;
+	virtual void RHIUpdateShaderResourceView(FRHIShaderResourceView* SRV, FRHIBuffer* Buffer, uint32 Stride, uint8 Format) final override;
+	virtual void RHIUpdateShaderResourceView(FRHIShaderResourceView* SRV, FRHIBuffer* Buffer) final override;
+	virtual uint64 RHICalcTexture2DPlatformSize(uint32 SizeX, uint32 SizeY, uint8 Format, uint32 NumMips, uint32 NumSamples, ETextureCreateFlags Flags, const FRHIResourceCreateInfo& CreateInfo, uint32& OutAlign) final override;
+	virtual uint64 RHICalcTexture3DPlatformSize(uint32 SizeX, uint32 SizeY, uint32 SizeZ, uint8 Format, uint32 NumMips, ETextureCreateFlags Flags, const FRHIResourceCreateInfo& CreateInfo, uint32& OutAlign) final override;
+	virtual uint64 RHICalcTextureCubePlatformSize(uint32 Size, uint8 Format, uint32 NumMips, ETextureCreateFlags Flags, const FRHIResourceCreateInfo& CreateInfo, uint32& OutAlign) final override;
+	virtual uint64 RHIGetMinimumAlignmentForBufferBackedSRV(EPixelFormat Format) final override;
+	virtual void RHIGetTextureMemoryStats(FTextureMemoryStats& OutStats) final override;
+	virtual bool RHIGetTextureMemoryVisualizeData(FColor* TextureData,int32 SizeX,int32 SizeY,int32 Pitch,int32 PixelSize) final override;
+	virtual FTexture2DRHIRef RHICreateTexture2D(uint32 SizeX, uint32 SizeY, uint8 Format, uint32 NumMips, uint32 NumSamples, ETextureCreateFlags Flags, ERHIAccess InResourceState, FRHIResourceCreateInfo& CreateInfo) final override;
+	virtual FTexture2DRHIRef RHICreateTextureExternal2D(uint32 SizeX, uint32 SizeY, uint8 Format, uint32 NumMips, uint32 NumSamples, ETextureCreateFlags Flags, ERHIAccess InResourceState, FRHIResourceCreateInfo& Info) final override;
+	virtual FTexture2DRHIRef RHIAsyncCreateTexture2D(uint32 SizeX, uint32 SizeY, uint8 Format, uint32 NumMips, ETextureCreateFlags Flags, ERHIAccess InResourceState, void** InitialMipData, uint32 NumInitialMips) final override;
+	virtual void RHICopySharedMips(FRHITexture2D* DestTexture2D, FRHITexture2D* SrcTexture2D) final override;
+	virtual FTexture2DArrayRHIRef RHICreateTexture2DArray(uint32 SizeX, uint32 SizeY, uint32 SizeZ, uint8 Format, uint32 NumMips, uint32 NumSamples, ETextureCreateFlags Flags, ERHIAccess InResourceState, FRHIResourceCreateInfo& CreateInfo) final override;
+	virtual FTexture3DRHIRef RHICreateTexture3D(uint32 SizeX, uint32 SizeY, uint32 SizeZ, uint8 Format, uint32 NumMips, ETextureCreateFlags Flags, ERHIAccess InResourceState, FRHIResourceCreateInfo& CreateInfo) final override;
+	virtual void RHIGetResourceInfo(FRHITexture* Ref, FRHIResourceInfo& OutInfo) final override;
+	virtual void RHIGenerateMips(FRHITexture* Texture) final override;
+	virtual uint32 RHIComputeMemorySize(FRHITexture* TextureRHI) final override;
+	virtual FTexture2DRHIRef RHIAsyncReallocateTexture2D(FRHITexture2D* Texture2D, int32 NewMipCount, int32 NewSizeX, int32 NewSizeY, FThreadSafeCounter* RequestStatus) final override;
+	virtual ETextureReallocationStatus RHIFinalizeAsyncReallocateTexture2D(FRHITexture2D* Texture2D, bool bBlockUntilCompleted) final override;
+	virtual ETextureReallocationStatus RHICancelAsyncReallocateTexture2D(FRHITexture2D* Texture2D, bool bBlockUntilCompleted) final override;
+	virtual void* RHILockTexture2D(FRHITexture2D* Texture, uint32 MipIndex, EResourceLockMode LockMode, uint32& DestStride, bool bLockWithinMiptail) final override;
+	virtual void RHIUnlockTexture2D(FRHITexture2D* Texture, uint32 MipIndex, bool bLockWithinMiptail) final override;
+	virtual void* RHILockTexture2DArray(FRHITexture2DArray* Texture, uint32 TextureIndex, uint32 MipIndex, EResourceLockMode LockMode, uint32& DestStride, bool bLockWithinMiptail) final override;
+	virtual void RHIUnlockTexture2DArray(FRHITexture2DArray* Texture, uint32 TextureIndex, uint32 MipIndex, bool bLockWithinMiptail) final override;
+	virtual void RHIUpdateTexture2D(FRHITexture2D* Texture, uint32 MipIndex, const struct FUpdateTextureRegion2D& UpdateRegion, uint32 SourcePitch, const uint8* SourceData) final override;
+	virtual void RHIUpdateTexture3D(FRHITexture3D* Texture, uint32 MipIndex, const struct FUpdateTextureRegion3D& UpdateRegion, uint32 SourceRowPitch, uint32 SourceDepthPitch, const uint8* SourceData) final override;
+	virtual FTextureCubeRHIRef RHICreateTextureCube(uint32 Size, uint8 Format, uint32 NumMips, ETextureCreateFlags Flags, ERHIAccess InResourceState, FRHIResourceCreateInfo& CreateInfo) final override;
+	virtual FTextureCubeRHIRef RHICreateTextureCubeArray(uint32 Size, uint32 ArraySize, uint8 Format, uint32 NumMips, ETextureCreateFlags Flags, ERHIAccess InResourceState, FRHIResourceCreateInfo& CreateInfo) final override;
+	virtual void* RHILockTextureCubeFace(FRHITextureCube* Texture, uint32 FaceIndex, uint32 ArrayIndex, uint32 MipIndex, EResourceLockMode LockMode, uint32& DestStride, bool bLockWithinMiptail) final override;
+	virtual void RHIUnlockTextureCubeFace(FRHITextureCube* Texture, uint32 FaceIndex, uint32 ArrayIndex, uint32 MipIndex, bool bLockWithinMiptail) final override;
+	virtual void RHIBindDebugLabelName(FRHITexture* Texture, const TCHAR* Name) final override;
+	virtual void RHIReadSurfaceData(FRHITexture* Texture,FIntRect Rect,TArray<FColor>& OutData,FReadSurfaceDataFlags InFlags) final override;
+	virtual void RHIReadSurfaceData(FRHITexture* TextureRHI, FIntRect InRect, TArray<FLinearColor>& OutData, FReadSurfaceDataFlags InFlags) final override;
+	virtual void RHIMapStagingSurface(FRHITexture* Texture, FRHIGPUFence* Fence, void*& OutData, int32& OutWidth, int32& OutHeight, uint32 GPUIndex = 0) final override;
+	virtual void RHIUnmapStagingSurface(FRHITexture* Texture, uint32 GPUIndex = 0) final override;
+	virtual void RHIReadSurfaceFloatData(FRHITexture* Texture,FIntRect Rect,TArray<FFloat16Color>& OutData,ECubeFace CubeFace,int32 ArrayIndex,int32 MipIndex) final override;
+	virtual void RHIRead3DSurfaceFloatData(FRHITexture* Texture,FIntRect Rect,FIntPoint ZMinMax,TArray<FFloat16Color>& OutData) final override;
+	virtual FRenderQueryRHIRef RHICreateRenderQuery(ERenderQueryType QueryType) final override;
+	virtual bool RHIGetRenderQueryResult(FRHIRenderQuery* RenderQuery, uint64& OutResult, bool bWait, uint32 GPUIndex = INDEX_NONE) final override;
+	virtual FTexture2DRHIRef RHIGetViewportBackBuffer(FRHIViewport* Viewport) final override;
+	virtual void RHIAdvanceFrameForGetViewportBackBuffer(FRHIViewport* Viewport) final override;
+	virtual void RHIAcquireThreadOwnership() final override;
+	virtual void RHIReleaseThreadOwnership() final override;
+	virtual void RHIFlushResources() final override;
+	virtual uint32 RHIGetGPUFrameCycles(uint32 GPUIndex = 0) final override;
+	virtual FViewportRHIRef RHICreateViewport(void* WindowHandle, uint32 SizeX, uint32 SizeY, bool bIsFullscreen, EPixelFormat PreferredPixelFormat) final override;
+	virtual void RHIResizeViewport(FRHIViewport* Viewport, uint32 SizeX, uint32 SizeY, bool bIsFullscreen) final override;
+	virtual void RHIResizeViewport(FRHIViewport* Viewport, uint32 SizeX, uint32 SizeY, bool bIsFullscreen, EPixelFormat PreferredPixelFormat) final override;
+	virtual void RHITick(float DeltaTime) final override;
+	virtual void RHIBlockUntilGPUIdle() final override;
+	virtual bool RHIGetAvailableResolutions(FScreenResolutionArray& Resolutions, bool bIgnoreRefreshRate) final override;
+	virtual void RHIGetSupportedResolution(uint32& Width, uint32& Height) final override;
+	virtual void RHIVirtualTextureSetFirstMipInMemory(FRHITexture2D* Texture, uint32 FirstMip) final override;
+	virtual void RHIVirtualTextureSetFirstMipVisible(FRHITexture2D* Texture, uint32 FirstMip) final override;
+	virtual void RHIExecuteCommandList(FRHICommandList* CmdList) final override;
+	virtual void* AGXRHI_API RHIGetNativeDevice() final override;
+	virtual void* AGXRHI_API RHIGetNativeInstance() final override;
+	virtual class IRHICommandContext* AGXRHI_API RHIGetDefaultContext() final override;
+	virtual IRHIComputeContext* RHIGetDefaultAsyncComputeContext() final override;
+	virtual class IRHICommandContextContainer* RHIGetCommandContextContainer(int32 Index, int32 Num) final override;
+	
+	virtual FBufferRHIRef CreateBuffer_RenderThread(class FRHICommandListImmediate& RHICmdList, uint32 Size, EBufferUsageFlags Usage, uint32 Stride, ERHIAccess ResourceState, FRHIResourceCreateInfo& CreateInfo) final override;
+	
+	virtual FTexture2DRHIRef AsyncReallocateTexture2D_RenderThread(class FRHICommandListImmediate& RHICmdList, FRHITexture2D* Texture2D, int32 NewMipCount, int32 NewSizeX, int32 NewSizeY, FThreadSafeCounter* RequestStatus) final override;
+	virtual ETextureReallocationStatus FinalizeAsyncReallocateTexture2D_RenderThread(class FRHICommandListImmediate& RHICmdList, FRHITexture2D* Texture2D, bool bBlockUntilCompleted) final override;
+	virtual ETextureReallocationStatus CancelAsyncReallocateTexture2D_RenderThread(class FRHICommandListImmediate& RHICmdList, FRHITexture2D* Texture2D, bool bBlockUntilCompleted) final override;
+
+	virtual void* LockTexture2D_RenderThread(class FRHICommandListImmediate& RHICmdList, FRHITexture2D* Texture, uint32 MipIndex, EResourceLockMode LockMode, uint32& DestStride, bool bLockWithinMiptail, bool bNeedsDefaultRHIFlush = true) final override;
+	virtual void UnlockTexture2D_RenderThread(class FRHICommandListImmediate& RHICmdList, FRHITexture2D* Texture, uint32 MipIndex, bool bLockWithinMiptail, bool bNeedsDefaultRHIFlush = true) final override;
+	virtual void UpdateTexture2D_RenderThread(class FRHICommandListImmediate& RHICmdList, FRHITexture2D* Texture, uint32 MipIndex, const struct FUpdateTextureRegion2D& UpdateRegion, uint32 SourcePitch, const uint8* SourceData) final override;
+	
+	virtual FUpdateTexture3DData BeginUpdateTexture3D_RenderThread(class FRHICommandListImmediate& RHICmdList, FRHITexture3D* Texture, uint32 MipIndex, const struct FUpdateTextureRegion3D& UpdateRegion) final override;
+	virtual void EndUpdateTexture3D_RenderThread(class FRHICommandListImmediate& RHICmdList, FUpdateTexture3DData& UpdateData) final override;
+	virtual void UpdateTexture3D_RenderThread(class FRHICommandListImmediate& RHICmdList, FRHITexture3D* Texture, uint32 MipIndex, const struct FUpdateTextureRegion3D& UpdateRegion, uint32 SourceRowPitch, uint32 SourceDepthPitch, const uint8* SourceData) final override;
+	
+	virtual FTexture2DRHIRef RHICreateTexture2D_RenderThread(class FRHICommandListImmediate& RHICmdList, uint32 SizeX, uint32 SizeY, uint8 Format, uint32 NumMips, uint32 NumSamples, ETextureCreateFlags Flags, ERHIAccess InResourceState, FRHIResourceCreateInfo& CreateInfo) final override;
+	virtual FTexture2DArrayRHIRef RHICreateTexture2DArray_RenderThread(class FRHICommandListImmediate& RHICmdList, uint32 SizeX, uint32 SizeY, uint32 SizeZ, uint8 Format, uint32 NumMips, uint32 NumSamples, ETextureCreateFlags Flags, ERHIAccess InResourceState, FRHIResourceCreateInfo& CreateInfo) final override;
+	virtual FTexture3DRHIRef RHICreateTexture3D_RenderThread(class FRHICommandListImmediate& RHICmdList, uint32 SizeX, uint32 SizeY, uint32 SizeZ, uint8 Format, uint32 NumMips, ETextureCreateFlags Flags, ERHIAccess InResourceState, FRHIResourceCreateInfo& CreateInfo) final override;
+	virtual FTextureCubeRHIRef RHICreateTextureCube_RenderThread(class FRHICommandListImmediate& RHICmdList, uint32 Size, uint8 Format, uint32 NumMips, ETextureCreateFlags Flags, ERHIAccess InResourceState, FRHIResourceCreateInfo& CreateInfo) final override;
+	virtual FTextureCubeRHIRef RHICreateTextureCubeArray_RenderThread(class FRHICommandListImmediate& RHICmdList, uint32 Size, uint32 ArraySize, uint8 Format, uint32 NumMips, ETextureCreateFlags Flags, ERHIAccess InResourceState, FRHIResourceCreateInfo& CreateInfo) final override;
+	
+	virtual FRenderQueryRHIRef RHICreateRenderQuery_RenderThread(class FRHICommandListImmediate& RHICmdList, ERenderQueryType QueryType) final override;
+	
+	virtual FUnorderedAccessViewRHIRef RHICreateUnorderedAccessView_RenderThread(class FRHICommandListImmediate& RHICmdList, FRHIBuffer* Buffer, bool bUseUAVCounter, bool bAppendBuffer) final override;
+	virtual FUnorderedAccessViewRHIRef RHICreateUnorderedAccessView_RenderThread(class FRHICommandListImmediate& RHICmdList, FRHITexture* Texture, uint32 MipLevel) final override;
+	virtual FUnorderedAccessViewRHIRef RHICreateUnorderedAccessView_RenderThread(class FRHICommandListImmediate& RHICmdList, FRHIBuffer* Buffer, uint8 Format) final override;
+
+	virtual FShaderResourceViewRHIRef RHICreateShaderResourceView_RenderThread(class FRHICommandListImmediate& RHICmdList, FRHIBuffer* Buffer, uint32 Stride, uint8 Format) final override;
+	virtual FShaderResourceViewRHIRef RHICreateShaderResourceView_RenderThread(class FRHICommandListImmediate& RHICmdList, const FShaderResourceViewInitializer& Initializer) final override;
+	virtual FShaderResourceViewRHIRef CreateShaderResourceView_RenderThread(class FRHICommandListImmediate& RHICmdList, const FShaderResourceViewInitializer& Initializer) final override;
+	virtual FShaderResourceViewRHIRef CreateShaderResourceView_RenderThread(class FRHICommandListImmediate& RHICmdList, FRHIBuffer* Buffer) final override;
+	virtual FShaderResourceViewRHIRef CreateShaderResourceView_RenderThread(class FRHICommandListImmediate& RHICmdList, FRHIBuffer* Buffer, uint32 Stride, uint8 Format) final override;
+	virtual FShaderResourceViewRHIRef RHICreateShaderResourceView_RenderThread(class FRHICommandListImmediate& RHICmdList, FRHIBuffer* Buffer) final override;
+	virtual FShaderResourceViewRHIRef RHICreateShaderResourceView_RenderThread(class FRHICommandListImmediate& RHICmdList, FRHITexture* Texture2DRHI, const FRHITextureSRVCreateInfo& CreateInfo) final override;
+
+	virtual void RHICreateTransition(FRHITransition* Transition, const FRHITransitionCreateInfo& CreateInfo) final override;
+	virtual void RHIReleaseTransition(FRHITransition* Transition) final override;
+
+	virtual FGPUFenceRHIRef RHICreateGPUFence(const FName &Name) final override;
+
+	virtual FStagingBufferRHIRef RHICreateStagingBuffer() final override;
+	virtual void* RHILockStagingBuffer(FRHIStagingBuffer* StagingBuffer, FRHIGPUFence* Fence, uint32 Offset, uint32 SizeRHI) final override;
+	virtual void RHIUnlockStagingBuffer(FRHIStagingBuffer* StagingBuffer) final override;
+	
+	virtual void RHIAcquireTransientResource_RenderThread(FRHITexture* Texture) final override;
+	virtual void RHIDiscardTransientResource_RenderThread(FRHITexture* Texture) final override;
+	virtual void RHIAcquireTransientResource_RenderThread(FRHIBuffer* Buffer) final override;
+	virtual void RHIDiscardTransientResource_RenderThread(FRHIBuffer* Buffer) final override;
+
+	/* Copy the source box pixels in the destination box texture, return true if implemented for the current platform*/
+	virtual void RHICopySubTextureRegion_RenderThread(class FRHICommandListImmediate& RHICmdList, FRHITexture2D* SourceTexture, FRHITexture2D* DestinationTexture, FBox2D SourceBox, FBox2D DestinationBox) final override;
+	virtual void RHICopySubTextureRegion(FRHITexture2D* SourceTexture, FRHITexture2D* DestinationTexture, FBox2D SourceBox, FBox2D DestinationBox) final override;
+
+	virtual FVertexShaderRHIRef CreateVertexShader_RenderThread(class FRHICommandListImmediate& RHICmdList, TArrayView<const uint8> Code, const FSHAHash& Hash) final override;
+	virtual FPixelShaderRHIRef CreatePixelShader_RenderThread(class FRHICommandListImmediate& RHICmdList, TArrayView<const uint8> Code, const FSHAHash& Hash) final override;
+	virtual FGeometryShaderRHIRef CreateGeometryShader_RenderThread(class FRHICommandListImmediate& RHICmdList, TArrayView<const uint8> Code, const FSHAHash& Hash) final override;
+	virtual FComputeShaderRHIRef CreateComputeShader_RenderThread(class FRHICommandListImmediate& RHICmdList, TArrayView<const uint8> Code, const FSHAHash& Hash) final override;
+	virtual FRHIShaderLibraryRef RHICreateShaderLibrary_RenderThread(class FRHICommandListImmediate& RHICmdList, EShaderPlatform Platform, FString FilePath, FString Name) final override;
+
+	virtual void RHIUpdateUniformBuffer(FRHIUniformBuffer* UniformBufferRHI, const void* Contents) final override;
+
+	virtual bool RHIIsTypedUAVLoadSupported(EPixelFormat PixelFormat) override;
+
+	virtual uint16 RHIGetPlatformTextureMaxSampleCount() override;
+
+private:
+	FTextureMemoryStats MemoryStats;
+	FAGXRHIImmediateCommandContext ImmediateContext;
+	FAGXRHICommandContext* AsyncComputeContext;
+	TMap<uint32, FVertexDeclarationRHIRef> VertexDeclarationCache;
+};
diff --git a/Engine/Source/Runtime/Apple/AGXRHI/Private/AGXDynamicRHIModule.cpp b/Engine/Source/Runtime/Apple/AGXRHI/Private/AGXDynamicRHIModule.cpp
new file mode 100644
--- /dev/null
+++ b/Engine/Source/Runtime/Apple/AGXRHI/Private/AGXDynamicRHIModule.cpp
@@ -0,0 +1,36 @@
+// Copyright Epic Games, Inc. All Rights Reserved.
+
+/*==============================================================================
+	AGXDynamicRHIModule.cpp: AGX Dynamic RHI Module Class Implementation.
+==============================================================================*/
+
+
+#include "AGXLLM.h"
+#include "DynamicRHI.h"
+#include "AGXDynamicRHIModule.h"
+#include "Modules/ModuleManager.h"
+
+
+//------------------------------------------------------------------------------
+
+#pragma mark - AGX Dynamic RHI Module Implementation
+
+
+IMPLEMENT_MODULE(FAGXDynamicRHIModule, AGXRHI);
+
+
+//------------------------------------------------------------------------------
+
+#pragma mark - AGX Dynamic RHI Module Class Methods
+
+
+bool FAGXDynamicRHIModule::IsSupported()
+{
+	return true;
+}
+
+FDynamicRHI* FAGXDynamicRHIModule::CreateRHI(ERHIFeatureLevel::Type RequestedFeatureLevel)
+{
+	LLM(AGXLLM::Initialise());
+	return new FAGXDynamicRHI(RequestedFeatureLevel);
+}
diff --git a/Engine/Source/Runtime/Apple/AGXRHI/Private/AGXDynamicRHIModule.h b/Engine/Source/Runtime/Apple/AGXRHI/Private/AGXDynamicRHIModule.h
new file mode 100644
--- /dev/null
+++ b/Engine/Source/Runtime/Apple/AGXRHI/Private/AGXDynamicRHIModule.h
@@ -0,0 +1,21 @@
+// Copyright Epic Games, Inc. All Rights Reserved.
+
+/*==============================================================================
+	AGXDynamicRHIModule.h: AGX Dynamic RHI Module Class.
+==============================================================================*/
+
+#pragma once
+
+
+//------------------------------------------------------------------------------
+
+#pragma mark - AGX Dynamic RHI Module Class
+
+
+class FAGXDynamicRHIModule : public IDynamicRHIModule
+{
+public:
+	virtual bool IsSupported() override final;
+
+	virtual FDynamicRHI* CreateRHI(ERHIFeatureLevel::Type RequestedFeatureLevel = ERHIFeatureLevel::Num) override final;
+};
diff --git a/Engine/Source/Runtime/Apple/AGXRHI/Private/AGXDynamicRHI_Shaders.cpp b/Engine/Source/Runtime/Apple/AGXRHI/Private/AGXDynamicRHI_Shaders.cpp
new file mode 100644
--- /dev/null
+++ b/Engine/Source/Runtime/Apple/AGXRHI/Private/AGXDynamicRHI_Shaders.cpp
@@ -0,0 +1,175 @@
+// Copyright Epic Games, Inc. All Rights Reserved.
+
+/*=============================================================================
+	AGXDynamicRHI_Shaders.cpp: AGX Dynamic RHI Class Shader Methods.
+=============================================================================*/
+
+
+#include "AGXRHIPrivate.h"
+#include "AGXShaderTypes.h"
+#include "Shaders/AGXShaderLibrary.h"
+
+
+//------------------------------------------------------------------------------
+
+#pragma mark - AGX Dynamic RHI Shader Methods
+
+
+FVertexShaderRHIRef FAGXDynamicRHI::RHICreateVertexShader(TArrayView<const uint8> Code, const FSHAHash& Hash)
+{
+	@autoreleasepool {
+		FAGXVertexShader* Shader = new FAGXVertexShader(Code);
+		return Shader;
+	}
+}
+
+FPixelShaderRHIRef FAGXDynamicRHI::RHICreatePixelShader(TArrayView<const uint8> Code, const FSHAHash& Hash)
+{
+	@autoreleasepool {
+		FAGXPixelShader* Shader = new FAGXPixelShader(Code);
+		return Shader;
+	}
+}
+
+FGeometryShaderRHIRef FAGXDynamicRHI::RHICreateGeometryShader(TArrayView<const uint8> Code, const FSHAHash& Hash)
+{
+	@autoreleasepool {
+		FAGXGeometryShader* Shader = new FAGXGeometryShader;
+		FMetalCodeHeader Header;
+		Shader->Init(Code, Header);
+		return Shader;
+	}
+}
+
+FComputeShaderRHIRef FAGXDynamicRHI::RHICreateComputeShader(TArrayView<const uint8> Code, const FSHAHash& Hash)
+{
+	@autoreleasepool {
+		return new FAGXComputeShader(Code);
+	}
+}
+
+FRHIShaderLibraryRef FAGXDynamicRHI::RHICreateShaderLibrary(EShaderPlatform Platform, FString const& FilePath, FString const& Name)
+{
+	FString METAL_MAP_EXTENSION(TEXT(".metalmap"));
+
+	@autoreleasepool {
+		FRHIShaderLibraryRef Result = nullptr;
+
+		FName PlatformName = LegacyShaderPlatformToShaderFormat(Platform);
+		FString LibName = FString::Printf(TEXT("%s_%s"), *Name, *PlatformName.GetPlainNameString());
+		LibName.ToLowerInline();
+
+		FString BinaryShaderFile = FilePath / LibName + METAL_MAP_EXTENSION;
+
+		if (IFileManager::Get().FileExists(*BinaryShaderFile) == false)
+		{
+			// the metal map files are stored in UFS file system
+			// for pak files this means they might be stored in a different location as the pak files will mount them to the project content directory
+			// the metal libraries are stores non UFS and could be anywhere on the file system.
+			// if we don't find the metalmap file straight away try the pak file path
+			BinaryShaderFile = FPaths::ProjectContentDir() / LibName + METAL_MAP_EXTENSION;
+		}
+
+		FScopeLock Lock(&FAGXShaderLibrary::LoadedShaderLibraryMutex);
+
+		FRHIShaderLibrary** FoundShaderLibrary = FAGXShaderLibrary::LoadedShaderLibraryMap.Find(BinaryShaderFile);
+		if (FoundShaderLibrary)
+		{
+			return *FoundShaderLibrary;
+		}
+
+		FArchive* BinaryShaderAr = IFileManager::Get().CreateFileReader(*BinaryShaderFile);
+
+		if( BinaryShaderAr != NULL )
+		{
+			FMetalShaderLibraryHeader Header;
+			FSerializedShaderArchive SerializedShaders;
+			TArray<uint8> ShaderCode;
+
+			*BinaryShaderAr << Header;
+			*BinaryShaderAr << SerializedShaders;
+			*BinaryShaderAr << ShaderCode;
+			BinaryShaderAr->Flush();
+			delete BinaryShaderAr;
+
+			// Would be good to check the language version of the library with the archive format here.
+			if (Header.Format == PlatformName.GetPlainNameString())
+			{
+				check(((SerializedShaders.GetNumShaders() + Header.NumShadersPerLibrary - 1) / Header.NumShadersPerLibrary) == Header.NumLibraries);
+
+				TArray<mtlpp::Library> Libraries;
+				Libraries.Empty(Header.NumLibraries);
+
+				for (uint32 i = 0; i < Header.NumLibraries; i++)
+				{
+					FString MetalLibraryFilePath = (FilePath / LibName) + FString::Printf(TEXT(".%d.metallib"), i);
+					MetalLibraryFilePath = FPaths::ConvertRelativePathToFull(MetalLibraryFilePath);
+#if !PLATFORM_MAC
+					MetalLibraryFilePath = IFileManager::Get().ConvertToAbsolutePathForExternalAppForRead(*MetalLibraryFilePath);
+#endif // !PLATFORM_MAC
+
+					METAL_GPUPROFILE(FAGXScopedCPUStats CPUStat(FString::Printf(TEXT("NewLibraryFile: %s"), *MetalLibraryFilePath)));
+					NSError* Error;
+					mtlpp::Library Library = [GetAGXDeviceContext().GetDevice() newLibraryWithFile:MetalLibraryFilePath.GetNSString() error:&Error];
+					if (Library != nil)
+					{
+						Libraries.Add(Library);
+					}
+					else
+					{
+						UE_LOG(LogAGX, Display, TEXT("Failed to create library: %s"), *FString(Error.description));
+						return nullptr;
+					}
+				}
+
+				Result = new FAGXShaderLibrary(Platform, Name, BinaryShaderFile, Header, SerializedShaders, ShaderCode, Libraries);
+				FAGXShaderLibrary::LoadedShaderLibraryMap.Add(BinaryShaderFile, Result.GetReference());
+			}
+			//else
+			//{
+			//	UE_LOG(LogAGX, Display, TEXT("Wrong shader platform wanted: %s, got: %s"), *LibName, *Map.Format);
+			//}
+		}
+		else
+		{
+			UE_LOG(LogAGX, Display, TEXT("No .metalmap file found for %s!"), *LibName);
+		}
+
+		return Result;
+	}
+}
+
+FBoundShaderStateRHIRef FAGXDynamicRHI::RHICreateBoundShaderState(
+	FRHIVertexDeclaration* VertexDeclarationRHI,
+	FRHIVertexShader* VertexShaderRHI,
+	FRHIPixelShader* PixelShaderRHI,
+	FRHIGeometryShader* GeometryShaderRHI)
+{
+	NOT_SUPPORTED("RHICreateBoundShaderState");
+	return nullptr;
+}
+
+FVertexShaderRHIRef FAGXDynamicRHI::CreateVertexShader_RenderThread(class FRHICommandListImmediate& RHICmdList, TArrayView<const uint8> Code, const FSHAHash& Hash)
+{
+	return RHICreateVertexShader(Code, Hash);
+}
+
+FGeometryShaderRHIRef FAGXDynamicRHI::CreateGeometryShader_RenderThread(class FRHICommandListImmediate& RHICmdList, TArrayView<const uint8> Code, const FSHAHash& Hash)
+{
+	return RHICreateGeometryShader(Code, Hash);
+}
+
+FPixelShaderRHIRef FAGXDynamicRHI::CreatePixelShader_RenderThread(class FRHICommandListImmediate& RHICmdList, TArrayView<const uint8> Code, const FSHAHash& Hash)
+{
+	return RHICreatePixelShader(Code, Hash);
+}
+
+FComputeShaderRHIRef FAGXDynamicRHI::CreateComputeShader_RenderThread(class FRHICommandListImmediate& RHICmdList, TArrayView<const uint8> Code, const FSHAHash& Hash)
+{
+	return RHICreateComputeShader(Code, Hash);
+}
+
+FRHIShaderLibraryRef FAGXDynamicRHI::RHICreateShaderLibrary_RenderThread(class FRHICommandListImmediate& RHICmdList, EShaderPlatform Platform, FString FilePath, FString Name)
+{
+	return RHICreateShaderLibrary(Platform, FilePath, Name);
+}
diff --git a/Engine/Source/Runtime/Apple/AGXRHI/Private/AGXFence.cpp b/Engine/Source/Runtime/Apple/AGXRHI/Private/AGXFence.cpp
new file mode 100644
--- /dev/null
+++ b/Engine/Source/Runtime/Apple/AGXRHI/Private/AGXFence.cpp
@@ -0,0 +1,229 @@
+// Copyright Epic Games, Inc. All Rights Reserved.
+
+#include "AGXRHIPrivate.h"
+
+#include "AGXFence.h"
+#include "AGXCommandBuffer.h"
+#include "AGXCommandQueue.h"
+#include "AGXContext.h"
+#include "AGXProfiler.h"
+
+@implementation FAGXDebugFence
+@synthesize Inner;
+
+APPLE_PLATFORM_OBJECT_ALLOC_OVERRIDES(FAGXDebugFence)
+
+-(instancetype)init
+{
+	id Self = [super init];
+	if (Self)
+	{
+		Label = nil;
+	}
+	return Self;
+}
+
+-(void)dealloc
+{
+	[self validate];
+	FAGXDebugCommandEncoder* Encoder = nil;
+	while ((Encoder = UpdatingEncoders.Pop()))
+	{
+		[Encoder release];
+	}
+	Encoder = nil;
+	while ((Encoder = WaitingEncoders.Pop()))
+	{
+		[Encoder release];
+	}
+	[Label release];
+	[super dealloc];
+}
+
+-(id <MTLDevice>) device
+{
+	if (Inner)
+	{
+		return Inner.device;
+	}
+	else
+	{
+		return nil;
+	}
+}
+
+-(NSString *_Nullable)label
+{
+	return Label;
+}
+
+-(void)setLabel:(NSString *_Nullable)Text
+{
+	[Text retain];
+	[Label release];
+	Label = Text;
+	if(Inner)
+	{
+		Inner.label = Text;
+	}
+}
+
+-(void)validate
+{
+	UE_CLOG(UpdatingEncoders.IsEmpty() != WaitingEncoders.IsEmpty(), LogAGX, Fatal, TEXT("Fence with unmatched updates/waits destructed - there's a gap in fence (%p) %s"), self, Label ? *FString(Label) : TEXT("Null"));
+}
+
+-(void)updatingEncoder:(FAGXDebugCommandEncoder*)Encoder
+{
+	check(Encoder);
+	UpdatingEncoders.Push([Encoder retain]);
+}
+
+-(void)waitingEncoder:(FAGXDebugCommandEncoder*)Encoder
+{
+	check(Encoder);
+	WaitingEncoders.Push([Encoder retain]);
+}
+
+-(TLockFreePointerListLIFO<FAGXDebugCommandEncoder>*)updatingEncoders
+{
+	return &UpdatingEncoders;
+}
+
+-(TLockFreePointerListLIFO<FAGXDebugCommandEncoder>*)waitingEncoders
+{
+	return &WaitingEncoders;
+}
+@end
+
+#if METAL_DEBUG_OPTIONS
+extern int32 GAGXRuntimeDebugLevel;
+#endif
+
+uint32 FAGXFence::Release() const
+{
+	uint32 Refs = uint32(FPlatformAtomics::InterlockedDecrement(&NumRefs));
+	if(Refs == 0)
+	{
+#if METAL_DEBUG_OPTIONS // When using validation we need to use fences only once per-frame in order to make it tractable
+		if (GAGXRuntimeDebugLevel >= EAGXDebugLevelValidation)
+		{
+			AGXSafeReleaseMetalFence(const_cast<FAGXFence*>(this));
+		}
+		else // However in a final game, we need to reuse fences aggressively so that we don't run out when loading into projects
+#endif
+		{
+			FAGXFencePool::Get().ReleaseFence(const_cast<FAGXFence*>(this));
+		}
+	}
+	return Refs;
+}
+
+#if METAL_DEBUG_OPTIONS
+void FAGXFence::Validate(void) const
+{
+	if (GetAGXDeviceContext().GetCommandQueue().GetRuntimeDebuggingLevel() >= EAGXDebugLevelValidation && Get(mtlpp::RenderStages::Vertex))
+	{
+		[(FAGXDebugFence*)Get(mtlpp::RenderStages::Vertex).GetPtr() validate];
+		if (Get(mtlpp::RenderStages::Fragment))
+		{
+			[(FAGXDebugFence*)Get(mtlpp::RenderStages::Fragment).GetPtr() validate];
+		}
+	}
+}
+#endif
+
+void FAGXFencePool::Initialise(mtlpp::Device const& InDevice)
+{
+	Device = InDevice;
+	for (int32 i = 0; i < FAGXFencePool::NumFences; i++)
+	{
+#if METAL_DEBUG_OPTIONS
+		if (GAGXRuntimeDebugLevel >= EAGXDebugLevelValidation)
+		{
+			FAGXDebugFence* VertexFence = [[FAGXDebugFence new] autorelease];
+			VertexFence.Inner = Device.NewFence();
+			FAGXDebugFence* FragmentFence = [[FAGXDebugFence new] autorelease];
+			FragmentFence.Inner = Device.NewFence();
+			FAGXFence* F = new FAGXFence;
+			F->Set(mtlpp::RenderStages::Vertex, VertexFence);
+			F->Set(mtlpp::RenderStages::Fragment, FragmentFence);
+			Fences.Add(F);
+			Lifo.Push(F);
+		}
+		else
+#endif
+		{
+			FAGXFence* F = new FAGXFence;
+			F->Set(mtlpp::RenderStages::Vertex, Device.NewFence());
+			F->Set(mtlpp::RenderStages::Fragment, Device.NewFence());
+#if METAL_DEBUG_OPTIONS
+			if (GAGXRuntimeDebugLevel >= EAGXDebugLevelValidation)
+			{
+				Fences.Add(F);
+			}
+#endif
+			Lifo.Push(F);
+		}
+	}
+	Count = FAGXFencePool::NumFences;
+    Allocated = 0;
+}
+
+FAGXFence* FAGXFencePool::AllocateFence()
+{
+	FAGXFence* Fence = Lifo.Pop();
+	if (Fence)
+	{
+        INC_DWORD_STAT(STAT_AGXFenceCount);
+        FPlatformAtomics::InterlockedDecrement(&Count);
+        FPlatformAtomics::InterlockedIncrement(&Allocated);
+#if METAL_DEBUG_OPTIONS
+		if (GAGXRuntimeDebugLevel >= EAGXDebugLevelValidation)
+		{
+			FScopeLock Lock(&Mutex);
+			check(Fences.Contains(Fence));
+			Fences.Remove(Fence);
+		}
+#endif
+	}
+	check(Fence);
+	Fence->Reset();
+	return Fence;
+}
+
+void FAGXFence::ValidateUsage(FAGXFence* InFence)
+{
+	if (InFence)
+	{
+		if (InFence->NumWrites(mtlpp::RenderStages::Vertex) != InFence->NumWaits(mtlpp::RenderStages::Vertex))
+		{
+			UE_LOG(LogAGX, Warning, TEXT("%p (%s) writes %d waits %d"), InFence, *FString(InFence->Get(mtlpp::RenderStages::Vertex).GetLabel()), (uint32)InFence->NumWrites(mtlpp::RenderStages::Vertex), (uint32)InFence->NumWaits(mtlpp::RenderStages::Vertex));
+		}
+		if (InFence->NumWrites(mtlpp::RenderStages::Fragment) != InFence->NumWaits(mtlpp::RenderStages::Fragment))
+		{
+			UE_LOG(LogAGX, Warning, TEXT("%p (%s) writes %d waits %d"), InFence, *FString(InFence->Get(mtlpp::RenderStages::Fragment).GetLabel()), (uint32)InFence->NumWrites(mtlpp::RenderStages::Fragment), (uint32)InFence->NumWaits(mtlpp::RenderStages::Fragment));
+		}
+	}
+}
+
+void FAGXFencePool::ReleaseFence(FAGXFence* const InFence)
+{
+	if (InFence)
+	{
+        DEC_DWORD_STAT(STAT_AGXFenceCount);
+        FPlatformAtomics::InterlockedDecrement(&Allocated);
+#if METAL_DEBUG_OPTIONS
+		if (GAGXRuntimeDebugLevel >= EAGXDebugLevelValidation)
+		{
+			FScopeLock Lock(&Mutex);
+			FAGXFence::ValidateUsage(InFence);
+			check(!Fences.Contains(InFence));
+			Fences.Add(InFence);
+		}
+#endif
+		FPlatformAtomics::InterlockedIncrement(&Count);
+		check(Count <= FAGXFencePool::NumFences);
+		Lifo.Push(InFence);
+	}
+}
diff --git a/Engine/Source/Runtime/Apple/AGXRHI/Private/AGXFence.h b/Engine/Source/Runtime/Apple/AGXRHI/Private/AGXFence.h
new file mode 100644
--- /dev/null
+++ b/Engine/Source/Runtime/Apple/AGXRHI/Private/AGXFence.h
@@ -0,0 +1,208 @@
+// Copyright Epic Games, Inc. All Rights Reserved.
+
+
+#pragma once
+
+#import <Metal/Metal.h>
+
+#pragma clang diagnostic ignored "-Wnullability-completeness"
+
+@class FAGXDebugCommandEncoder;
+
+@interface FAGXDebugFence : FApplePlatformObject<MTLFence>
+{
+	TLockFreePointerListLIFO<FAGXDebugCommandEncoder> UpdatingEncoders;
+	TLockFreePointerListLIFO<FAGXDebugCommandEncoder> WaitingEncoders;
+	NSString* Label;
+}
+@property (retain) id<MTLFence> Inner;
+-(void)updatingEncoder:(FAGXDebugCommandEncoder*)Encoder;
+-(void)waitingEncoder:(FAGXDebugCommandEncoder*)Encoder;
+-(TLockFreePointerListLIFO<FAGXDebugCommandEncoder>*)updatingEncoders;
+-(TLockFreePointerListLIFO<FAGXDebugCommandEncoder>*)waitingEncoders;
+-(void)validate;
+@end
+
+@protocol MTLDeviceExtensions <MTLDevice>
+/*!
+ @method newFence
+ @abstract Create a new MTLFence object
+ */
+- (id <MTLFence>)newFence;
+@end
+
+@protocol MTLBlitCommandEncoderExtensions <MTLBlitCommandEncoder>
+/*!
+ @abstract Update the event to capture all GPU work so far enqueued by this encoder. */
+-(void) updateFence:(id <MTLFence>)fence;
+/*!
+ @abstract Prevent further GPU work until the event is reached. */
+-(void) waitForFence:(id <MTLFence>)fence;
+@end
+@protocol MTLComputeCommandEncoderExtensions <MTLComputeCommandEncoder>
+/*!
+ @abstract Update the event to capture all GPU work so far enqueued by this encoder. */
+-(void) updateFence:(id <MTLFence>)fence;
+/*!
+ @abstract Prevent further GPU work until the event is reached. */
+-(void) waitForFence:(id <MTLFence>)fence;
+@end
+@protocol MTLRenderCommandEncoderExtensions <MTLRenderCommandEncoder>
+/*!
+ @abstract Update the event to capture all GPU work so far enqueued by this encoder for the given
+ stages.
+ @discussion Unlike <st>updateFence:</st>, this method will update the event when the given stage(s) complete, allowing for commands to overlap in execution.
+ */
+-(void) updateFence:(id <MTLFence>)fence afterStages:(MTLRenderStages)stages;
+/*!
+ @abstract Prevent further GPU work until the event is reached for the given stages.
+ @discussion Unlike <st>waitForFence:</st>, this method will only block commands assoicated with the given stage(s), allowing for commands to overlap in execution.
+ */
+-(void) waitForFence:(id <MTLFence>)fence beforeStages:(MTLRenderStages)stages;
+@end
+
+class FAGXFence
+{
+	enum
+	{
+		NumFenceStages = 2
+	};
+public:
+	FAGXFence()
+	: NumRefs(0)
+	{
+		for (uint32 i = 0; i < NumFenceStages; i++)
+		{
+			Writes[i] = 0;
+			Waits[i] = 0;
+		}
+	}
+	
+	explicit FAGXFence(FAGXFence const& Other)
+	: NumRefs(0)
+	{
+		operator=(Other);
+	}
+	
+	~FAGXFence()
+	{
+		check(!NumRefs);
+	}
+	
+	uint32 AddRef() const
+	{
+		return uint32(FPlatformAtomics::InterlockedIncrement(&NumRefs));
+	}
+	uint32 Release() const;
+	uint32 GetRefCount() const
+	{
+		return uint32(FPlatformAtomics::AtomicRead(&NumRefs));
+	}
+	
+	FAGXFence& operator=(FAGXFence const& Other)
+	{
+		if (&Other != this)
+		{
+			for (uint32 i = 0; i < NumFenceStages; i++)
+			{
+				Fences[i] = Other.Fences[i];
+			}
+		}
+		return *this;
+	}
+	
+#if METAL_DEBUG_OPTIONS
+	void Validate(void) const;
+#endif
+	
+	void Reset(void)
+	{
+		for (uint32 i = 0; i < NumFenceStages; i++)
+		{
+			Writes[i] = 0;
+			Waits[i] = 0;
+		}
+	}
+	
+	void Write(mtlpp::RenderStages Stage)
+	{
+		Writes[(uint32)Stage - 1u]++;
+	}
+	
+	void Wait(mtlpp::RenderStages Stage)
+	{
+		Waits[(uint32)Stage - 1u]++;
+	}
+	
+	int8 NumWrites(mtlpp::RenderStages Stage) const
+	{
+		return Writes[(uint32)Stage - 1u];
+	}
+	
+	int8 NumWaits(mtlpp::RenderStages Stage) const
+	{
+		return Waits[(uint32)Stage - 1u];
+	}
+	
+	bool NeedsWrite(mtlpp::RenderStages Stage) const
+	{
+		return Writes[(uint32)Stage - 1u] == 0 || (Waits[(uint32)Stage - 1u] > Writes[(uint32)Stage - 1u]);
+	}
+	
+	bool NeedsWait(mtlpp::RenderStages Stage) const
+	{
+		return Waits[(uint32)Stage - 1u] == 0 || (Writes[(uint32)Stage - 1u] > Waits[(uint32)Stage - 1u]);
+	}
+	
+	mtlpp::Fence Get(mtlpp::RenderStages Stage) const
+	{
+		return Fences[(uint32)Stage - 1u];
+	}
+	
+	void Set(mtlpp::RenderStages Stage, mtlpp::Fence Fence)
+	{
+		Fences[(uint32)Stage - 1u] = Fence;
+	}
+	
+	static void ValidateUsage(FAGXFence* Fence);
+	
+private:
+	mtlpp::Fence Fences[NumFenceStages];
+	int8 Writes[NumFenceStages];
+	int8 Waits[NumFenceStages];
+	mutable int32 NumRefs;	
+};
+
+class FAGXFencePool
+{
+	enum
+	{
+		NumFences = 2048
+	};
+public:
+	FAGXFencePool() {}
+	
+	static FAGXFencePool& Get()
+	{
+		static FAGXFencePool sSelf;
+		return sSelf;
+	}
+	
+	void Initialise(mtlpp::Device const& InDevice);
+	
+	FAGXFence* AllocateFence();
+	void ReleaseFence(FAGXFence* const InFence);
+	
+	int32 Max() const { return Count; }
+	int32 Num() const { return Allocated; }
+	
+private:
+	int32 Count;
+	int32 Allocated;
+	mtlpp::Device Device;
+#if METAL_DEBUG_OPTIONS
+	TSet<FAGXFence*> Fences;
+	FCriticalSection Mutex;
+#endif
+	TLockFreePointerListFIFO<FAGXFence, PLATFORM_CACHE_LINE_SIZE> Lifo;
+};
diff --git a/Engine/Source/Runtime/Apple/AGXRHI/Private/AGXFrameAllocator.cpp b/Engine/Source/Runtime/Apple/AGXRHI/Private/AGXFrameAllocator.cpp
new file mode 100644
--- /dev/null
+++ b/Engine/Source/Runtime/Apple/AGXRHI/Private/AGXFrameAllocator.cpp
@@ -0,0 +1,205 @@
+// Copyright Epic Games, Inc. All Rights Reserved.
+
+#include "AGXRHIPrivate.h"
+#include "AGXFrameAllocator.h"
+#include "AGXLLM.h"
+#include "AGXProfiler.h"
+
+#pragma mark Constructor/Destructor
+
+FAGXFrameAllocator::FAGXFrameAllocator(id <MTLDevice> TargetDevice)
+    : Device(TargetDevice)
+#if METAL_FRAME_ALLOCATOR_VALIDATION
+    , AllAllocatedBuffers(nil)
+#endif
+    , AvailableBuffers(nil)
+    , BuffersInFlight(nil)
+    , CurrentBuffer(nil)
+    , CurrentCursor(0)
+    , CurrentFrame(0)
+    , TargetAllocationLimit(0)
+    , DefaultAllocationSize(0)
+    , TotalBytesAllocated(0)
+    , BytesAcquiredInCurrentFrame(0)
+    , BytesInFlight(0)
+{
+    check(Device);
+#if METAL_FRAME_ALLOCATOR_VALIDATION
+    AllAllocatedBuffers = [NSMutableArray new];
+#endif
+    AvailableBuffers = [NSMutableArray new];
+    BuffersInFlight = [NSMutableArray new];
+    UnfairLock = OS_UNFAIR_LOCK_INIT;
+    
+    TotalAllocationStat = GET_STATID(STAT_AGXFrameAllocatorAllocatedMemory);
+    MemoryInFlightStat = GET_STATID(STAT_AGXFrameAllocatorMemoryInFlight);
+    BytesPerFrameStat = GET_STATID(STAT_AGXFrameAllocatorBytesPerFrame);
+    
+    [Device retain];
+}
+
+FAGXFrameAllocator::~FAGXFrameAllocator()
+{
+    // Owner of this object is responsible for finishing all rendering before destroying this object.
+    // This does not check and will crash.
+    os_unfair_lock_lock(&UnfairLock);
+    
+    [BuffersInFlight release];
+    [AvailableBuffers release];
+#if METAL_FRAME_ALLOCATOR_VALIDATION
+    [AllAllocatedBuffers release];
+#endif
+    
+    [Device release];
+    
+    os_unfair_lock_unlock(&UnfairLock);
+}
+
+#pragma mark -
+
+#pragma mark API Entrypoints
+
+void FAGXFrameAllocator::SetTargetAllocationLimitInBytes(uint32 LimitInBytes)
+{
+    TargetAllocationLimit = LimitInBytes;
+}
+
+void FAGXFrameAllocator::SetDefaultAllocationSizeInBytes(uint32 DefaultInBytes)
+{
+    DefaultAllocationSize = DefaultInBytes;
+}
+
+void FAGXFrameAllocator::SetStatIds(const TStatId& InTotalAllocationStat, const TStatId& InMemoryInFlightStat, const TStatId& InBytesPerFrameStat)
+{
+    TotalAllocationStat = InTotalAllocationStat;
+    MemoryInFlightStat = InMemoryInFlightStat;
+    BytesPerFrameStat = InBytesPerFrameStat;
+}
+
+FAGXFrameAllocator::AllocationEntry FAGXFrameAllocator::AcquireSpace(uint32 SizeInBytes)
+{
+    uint32 AlignedSize = Align(SizeInBytes, BufferOffsetAlignment);
+    
+    os_unfair_lock_lock(&UnfairLock);
+    
+    NSUInteger SpaceRequired = this->CurrentCursor + AlignedSize;
+    
+    if(this->CurrentBuffer == nil || (this->CurrentBuffer && SpaceRequired > [this->CurrentBuffer length]))
+    {
+        this->CurrentBuffer = this->FindOrAllocateBufferUnsafe(AlignedSize);
+        this->CurrentCursor = 0;
+        [this->BuffersInFlight addObject:this->CurrentBuffer];
+        this->BytesInFlight += [this->CurrentBuffer length];
+    }
+    
+    check(this->CurrentBuffer != nil);
+    
+    AllocationEntry Entry;
+    Entry.Backing = this->CurrentBuffer;
+    Entry.Offset = this->CurrentCursor;
+    this->CurrentCursor += AlignedSize;
+    
+    this->BytesAcquiredInCurrentFrame += AlignedSize;
+    
+    os_unfair_lock_unlock(&UnfairLock);
+    
+    return Entry;
+}
+
+void FAGXFrameAllocator::MarkEndOfFrame(uint32 FrameNumberThatEnded, id <MTLCommandBuffer> LastCommandBufferInFrame)
+{
+    __block FAGXFrameAllocator* Allocator = this;
+    __block NSMutableArray<id <MTLBuffer>>* BuffersToRecycle = this->BuffersInFlight;
+    
+    // Pass ownership of the BuffersInFlight list to the completion handler.
+    // It will be responsible for destroying the array.
+    [LastCommandBufferInFrame addCompletedHandler:
+    ^(id <MTLCommandBuffer> CompletedCommandBuffer)
+    {
+        Allocator->RecycleBuffers(BuffersToRecycle);
+        [BuffersToRecycle release];
+    }];
+    
+    this->BuffersInFlight = [NSMutableArray new];
+    this->CurrentBuffer = nil;
+    this->CurrentCursor = 0;
+    
+    SET_MEMORY_STAT_FName(BytesPerFrameStat.GetName(), this->BytesAcquiredInCurrentFrame);
+    SET_MEMORY_STAT_FName(MemoryInFlightStat.GetName(), this->BytesInFlight);
+    this->BytesAcquiredInCurrentFrame = 0;
+}
+
+#pragma mark Private
+
+id <MTLBuffer> FAGXFrameAllocator::FindOrAllocateBufferUnsafe(uint32 SizeInBytes)
+{
+    check(SizeInBytes % BufferOffsetAlignment == 0);
+    
+    id <MTLBuffer> Found = nil;
+    for(id <MTLBuffer> Buf : this->AvailableBuffers)
+    {
+        if(Buf.length > SizeInBytes)
+        {
+            Found = Buf;
+            break;
+        }
+    }
+    
+    if(Found != nil)
+    {
+        [this->AvailableBuffers removeObject:Found];
+        [Found setPurgeableState:MTLPurgeableStateNonVolatile];
+    }
+    else
+    {
+        uint32 AllocationSize = FMath::Max(DefaultAllocationSize, SizeInBytes);
+        Found = [Device newBufferWithLength:AllocationSize options:MTLResourceStorageModeShared | MTLResourceCPUCacheModeWriteCombined];
+        [Found setLabel: [NSString stringWithFormat:@"UniformBacking %u", AllocationSize]];
+        
+        check(Found);
+        
+        this->TotalBytesAllocated += AllocationSize;
+        INC_MEMORY_STAT_BY_FName(TotalAllocationStat.GetName(), [Found length]);
+#if STATS || ENABLE_LOW_LEVEL_MEM_TRACKER
+        AGXLLM::LogAllocBuffer(GetAGXDeviceContext().GetDevice(), Found);
+#endif
+
+#if METAL_FRAME_ALLOCATOR_VALIDATION
+        [this->AllAllocatedBuffers addObject:Found];
+#endif
+    }
+    
+    check(Found);
+    check([Found setPurgeableState:MTLPurgeableStateKeepCurrent] == MTLPurgeableStateNonVolatile);
+    
+    return Found;
+}
+
+void FAGXFrameAllocator::RecycleBuffers(NSArray <id <MTLBuffer>>* BuffersToRecycle)
+{
+    // This will run on the completion handler.
+    os_unfair_lock_lock(&UnfairLock);
+    
+    for(id <MTLBuffer> Buffer : BuffersToRecycle)
+    {
+        this->BytesInFlight -= [Buffer length];
+        
+        // Very simple logic: if we are over the target we will destroy the buffer instead of returning it to the pool.
+        if(this->TotalBytesAllocated < this->TargetAllocationLimit)
+        {
+            [Buffer setPurgeableState:MTLPurgeableStateVolatile];
+            [this->AvailableBuffers addObject:Buffer];
+        }
+        else
+        {
+            this->TotalBytesAllocated -= [Buffer length];
+            DEC_MEMORY_STAT_BY_FName(TotalAllocationStat.GetName(), [Buffer length]);
+#if METAL_FRAME_ALLOCATOR_VALIDATION
+            [this->AllAllocatedBuffers removeObject:Buffer];
+#endif
+            [Buffer release];
+        }
+    }
+
+    os_unfair_lock_unlock(&UnfairLock);
+}
diff --git a/Engine/Source/Runtime/Apple/AGXRHI/Private/AGXGraphicsPipelineState.cpp b/Engine/Source/Runtime/Apple/AGXRHI/Private/AGXGraphicsPipelineState.cpp
new file mode 100644
--- /dev/null
+++ b/Engine/Source/Runtime/Apple/AGXRHI/Private/AGXGraphicsPipelineState.cpp
@@ -0,0 +1,61 @@
+// Copyright Epic Games, Inc. All Rights Reserved.
+
+/*=============================================================================
+	AGXGraphicsPipelineState.cpp: AGX RHI graphics pipeline state class.
+=============================================================================*/
+
+
+#include "AGXRHIPrivate.h"
+#include "AGXVertexDeclaration.h"
+#include "AGXShaderTypes.h"
+#include "AGXGraphicsPipelineState.h"
+
+
+//------------------------------------------------------------------------------
+
+#pragma mark - AGX RHI Graphics Pipeline State Support Routines
+
+
+// From AGXPipeline.cpp:
+extern FAGXShaderPipeline* GetMTLRenderPipeline(bool const bSync, FAGXGraphicsPipelineState const* State, const FGraphicsPipelineStateInitializer& Init);
+extern void ReleaseMTLRenderPipeline(FAGXShaderPipeline* Pipeline);
+
+
+//------------------------------------------------------------------------------
+
+#pragma mark - AGX RHI Graphics Pipeline State Class Implementation
+
+
+FAGXGraphicsPipelineState::FAGXGraphicsPipelineState(const FGraphicsPipelineStateInitializer& Init)
+	: Initializer(Init)
+	, PipelineState(nil)
+{
+	// void
+}
+
+bool FAGXGraphicsPipelineState::Compile()
+{
+	check(PipelineState == nil);
+	PipelineState = [GetMTLRenderPipeline(true, this, Initializer) retain];
+	return (PipelineState != nil);
+}
+
+FAGXGraphicsPipelineState::~FAGXGraphicsPipelineState()
+{
+	if (PipelineState != nil)
+	{
+		ReleaseMTLRenderPipeline(PipelineState);
+		PipelineState = nil;
+	}
+}
+
+FAGXShaderPipeline* FAGXGraphicsPipelineState::GetPipeline()
+{
+	if (PipelineState == nil)
+	{
+		PipelineState = [GetMTLRenderPipeline(true, this, Initializer) retain];
+		check(PipelineState != nil);
+	}
+
+    return PipelineState;
+}
diff --git a/Engine/Source/Runtime/Apple/AGXRHI/Private/AGXGraphicsPipelineState.h b/Engine/Source/Runtime/Apple/AGXRHI/Private/AGXGraphicsPipelineState.h
new file mode 100644
--- /dev/null
+++ b/Engine/Source/Runtime/Apple/AGXRHI/Private/AGXGraphicsPipelineState.h
@@ -0,0 +1,53 @@
+// Copyright Epic Games, Inc. All Rights Reserved.
+
+/*=============================================================================
+	AGXGraphicsPipelineState.h: AGX RHI graphics pipeline state class.
+=============================================================================*/
+
+#pragma once
+
+class FAGXGraphicsPipelineState : public FRHIGraphicsPipelineState
+{
+	friend class FAGXDynamicRHI;
+
+public:
+	virtual ~FAGXGraphicsPipelineState();
+
+	FAGXShaderPipeline* GetPipeline();
+
+	/** Cached vertex structure */
+	TRefCountPtr<FAGXVertexDeclaration> VertexDeclaration;
+
+	/** Cached shaders */
+	TRefCountPtr<FAGXVertexShader> VertexShader;
+	TRefCountPtr<FAGXPixelShader> PixelShader;
+#if PLATFORM_SUPPORTS_GEOMETRY_SHADERS
+	TRefCountPtr<FAGXGeometryShader> GeometryShader;
+#endif
+
+	/** Cached state objects */
+	TRefCountPtr<FAGXDepthStencilState> DepthStencilState;
+	TRefCountPtr<FAGXRasterizerState> RasterizerState;
+
+	EPrimitiveType GetPrimitiveType()
+	{
+		return Initializer.PrimitiveType;
+	}
+
+	bool GetDepthBounds() const
+	{
+		return Initializer.bDepthBounds;
+	}
+
+private:
+	// This can only be created through the RHI to make sure Compile() is called.
+	FAGXGraphicsPipelineState(const FGraphicsPipelineStateInitializer& Init);
+
+	// Compiles the underlying gpu pipeline objects. This must be called before usage.
+	bool Compile();
+
+	// Needed to runtime refine shaders currently.
+	FGraphicsPipelineStateInitializer Initializer;
+
+	FAGXShaderPipeline* PipelineState;
+};
diff --git a/Engine/Source/Runtime/Apple/AGXRHI/Private/AGXHashedVertexDescriptor.cpp b/Engine/Source/Runtime/Apple/AGXRHI/Private/AGXHashedVertexDescriptor.cpp
new file mode 100644
--- /dev/null
+++ b/Engine/Source/Runtime/Apple/AGXRHI/Private/AGXHashedVertexDescriptor.cpp
@@ -0,0 +1,106 @@
+// Copyright Epic Games, Inc. All Rights Reserved.
+
+/*=============================================================================
+	AGXHashedVertexDescriptor.cpp: AGX RHI Hashed Vertex Descriptor.
+=============================================================================*/
+
+
+#include "AGXRHIPrivate.h"
+#include "AGXHashedVertexDescriptor.h"
+
+
+//------------------------------------------------------------------------------
+
+#pragma mark - AGX RHI Hashed Vertex Descriptor
+
+
+FAGXHashedVertexDescriptor::FAGXHashedVertexDescriptor()
+	: VertexDescHash(0)
+	, VertexDesc(nil)
+{
+	// void
+}
+
+FAGXHashedVertexDescriptor::FAGXHashedVertexDescriptor(mtlpp::VertexDescriptor Desc, uint32 Hash)
+	: VertexDescHash(Hash)
+	, VertexDesc(Desc)
+{
+	// void
+}
+
+FAGXHashedVertexDescriptor::FAGXHashedVertexDescriptor(FAGXHashedVertexDescriptor const& Other)
+	: VertexDescHash(0)
+	, VertexDesc(nil)
+{
+	operator=(Other);
+}
+
+FAGXHashedVertexDescriptor::~FAGXHashedVertexDescriptor()
+{
+	// void
+}
+
+FAGXHashedVertexDescriptor& FAGXHashedVertexDescriptor::operator=(FAGXHashedVertexDescriptor const& Other)
+{
+	if (this != &Other)
+	{
+		VertexDescHash = Other.VertexDescHash;
+		VertexDesc = Other.VertexDesc;
+	}
+	return *this;
+}
+
+bool FAGXHashedVertexDescriptor::operator==(FAGXHashedVertexDescriptor const& Other) const
+{
+	bool bEqual = false;
+
+	if (this != &Other)
+	{
+		if (VertexDescHash == Other.VertexDescHash)
+		{
+			bEqual = true;
+			if (VertexDesc.GetPtr() != Other.VertexDesc.GetPtr())
+			{
+				ns::Array<mtlpp::VertexBufferLayoutDescriptor> Layouts = VertexDesc.GetLayouts();
+				ns::Array<mtlpp::VertexAttributeDescriptor> Attributes = VertexDesc.GetAttributes();
+
+				ns::Array<mtlpp::VertexBufferLayoutDescriptor> OtherLayouts = Other.VertexDesc.GetLayouts();
+				ns::Array<mtlpp::VertexAttributeDescriptor> OtherAttributes = Other.VertexDesc.GetAttributes();
+				check(Layouts && Attributes && OtherLayouts && OtherAttributes);
+
+				for (uint32 i = 0; bEqual && i < MaxVertexElementCount; i++)
+				{
+					mtlpp::VertexBufferLayoutDescriptor LayoutDesc = Layouts[(NSUInteger)i];
+					mtlpp::VertexBufferLayoutDescriptor OtherLayoutDesc = OtherLayouts[(NSUInteger)i];
+
+					bEqual &= ((LayoutDesc != nil) == (OtherLayoutDesc != nil));
+
+					if (LayoutDesc && OtherLayoutDesc)
+					{
+						bEqual &= (LayoutDesc.GetStride() == OtherLayoutDesc.GetStride());
+						bEqual &= (LayoutDesc.GetStepFunction() == OtherLayoutDesc.GetStepFunction());
+						bEqual &= (LayoutDesc.GetStepRate() == OtherLayoutDesc.GetStepRate());
+					}
+
+					mtlpp::VertexAttributeDescriptor AttrDesc = Attributes[(NSUInteger)i];
+					mtlpp::VertexAttributeDescriptor OtherAttrDesc = OtherAttributes[(NSUInteger)i];
+
+					bEqual &= ((AttrDesc != nil) == (OtherAttrDesc != nil));
+
+					if (AttrDesc && OtherAttrDesc)
+					{
+						bEqual &= (AttrDesc.GetFormat() == OtherAttrDesc.GetFormat());
+						bEqual &= (AttrDesc.GetOffset() == OtherAttrDesc.GetOffset());
+						bEqual &= (AttrDesc.GetBufferIndex() == OtherAttrDesc.GetBufferIndex());
+					}
+				}
+			}
+		}
+	}
+	else
+	{
+		bEqual = true;
+	}
+
+	return bEqual;
+}
diff --git a/Engine/Source/Runtime/Apple/AGXRHI/Private/AGXHashedVertexDescriptor.h b/Engine/Source/Runtime/Apple/AGXRHI/Private/AGXHashedVertexDescriptor.h
new file mode 100644
--- /dev/null
+++ b/Engine/Source/Runtime/Apple/AGXRHI/Private/AGXHashedVertexDescriptor.h
@@ -0,0 +1,37 @@
+// Copyright Epic Games, Inc. All Rights Reserved.
+
+/*=============================================================================
+	AGXHashedVertexDescriptor.h: AGX RHI Hashed Vertex Descriptor.
+=============================================================================*/
+
+#pragma once
+
+
+//------------------------------------------------------------------------------
+
+#pragma mark - AGX RHI Hashed Vertex Descriptor
+
+
+/**
+ * The MTLVertexDescriptor and a pre-calculated hash value used to simplify
+ * comparisons (as vendor MTLVertexDescriptor implementations are not all
+ * comparable).
+ */
+struct FAGXHashedVertexDescriptor
+{
+	NSUInteger VertexDescHash;
+	mtlpp::VertexDescriptor VertexDesc;
+
+	FAGXHashedVertexDescriptor();
+	FAGXHashedVertexDescriptor(mtlpp::VertexDescriptor Desc, uint32 Hash);
+	FAGXHashedVertexDescriptor(FAGXHashedVertexDescriptor const& Other);
+	~FAGXHashedVertexDescriptor();
+
+	FAGXHashedVertexDescriptor& operator=(FAGXHashedVertexDescriptor const& Other);
+	bool operator==(FAGXHashedVertexDescriptor const& Other) const;
+
+	friend uint32 GetTypeHash(FAGXHashedVertexDescriptor const& Hash)
+	{
+		return Hash.VertexDescHash;
+	}
+};
diff --git a/Engine/Source/Runtime/Apple/AGXRHI/Private/AGXLLM.cpp b/Engine/Source/Runtime/Apple/AGXRHI/Private/AGXLLM.cpp
new file mode 100644
--- /dev/null
+++ b/Engine/Source/Runtime/Apple/AGXRHI/Private/AGXLLM.cpp
@@ -0,0 +1,322 @@
+// Copyright Epic Games, Inc. All Rights Reserved.
+
+#include "AGXLLM.h"
+#include "AGXProfiler.h"
+#include "RenderUtils.h"
+#include "HAL/LowLevelMemStats.h"
+
+#include <objc/runtime.h>
+
+#if ENABLE_LOW_LEVEL_MEM_TRACKER
+
+struct FLLMTagInfoAGX
+{
+	const TCHAR* Name;
+	FName StatName;				// shows in the LLMFULL stat group
+	FName SummaryStatName;		// shows in the LLM summary stat group
+};
+
+DECLARE_LLM_MEMORY_STAT(TEXT("Metal Buffers"), STAT_AGXBuffersLLM, STATGROUP_LLMPlatform);
+DECLARE_LLM_MEMORY_STAT(TEXT("Metal Textures"), STAT_AGXTexturesLLM, STATGROUP_LLMPlatform);
+DECLARE_LLM_MEMORY_STAT(TEXT("Metal Heaps"), STAT_AGXHeapsLLM, STATGROUP_LLMPlatform);
+DECLARE_LLM_MEMORY_STAT(TEXT("Metal RenderTargets"), STAT_AGXRenderTargetsLLM, STATGROUP_LLMPlatform);
+
+// *** order must match ELLMTagAGX enum ***
+const FLLMTagInfoAGX ELLMTagNamesAGX[] =
+{
+	// csv name									// stat name										// summary stat name						// enum value
+	{ TEXT("Metal Buffers"),		GET_STATFNAME(STAT_AGXBuffersLLM),		GET_STATFNAME(STAT_EngineSummaryLLM) },		// ELLMTagAGX::Buffers
+	{ TEXT("Metal Textures"),		GET_STATFNAME(STAT_AGXTexturesLLM),		GET_STATFNAME(STAT_EngineSummaryLLM) },		// ELLMTagAGX::Textures
+	{ TEXT("Metal Heaps"),			GET_STATFNAME(STAT_AGXHeapsLLM),			GET_STATFNAME(STAT_EngineSummaryLLM) },		// ELLMTagAGX::Heaps
+	{ TEXT("Metal Render Targets"),	GET_STATFNAME(STAT_AGXRenderTargetsLLM),	GET_STATFNAME(STAT_EngineSummaryLLM) },		// ELLMTagAGX::RenderTargets
+};
+
+/*
+ * Register Metal tags with LLM
+ */
+void AGXLLM::Initialise()
+{
+	int32 TagCount = sizeof(ELLMTagNamesAGX) / sizeof(FLLMTagInfoAGX);
+
+	for (int32 Index = 0; Index < TagCount; ++Index)
+	{
+		int32 Tag = (int32)ELLMTagApple::AppleMetalTagsStart + Index;
+		const FLLMTagInfoAGX& TagInfo = ELLMTagNamesAGX[Index];
+
+		FLowLevelMemTracker::Get().RegisterPlatformTag(Tag, TagInfo.Name, TagInfo.StatName, TagInfo.SummaryStatName);
+	}
+}
+
+#endif		// #if ENABLE_LOW_LEVEL_MEM_TRACKER
+
+@implementation FAGXDeallocHandler
+
+APPLE_PLATFORM_OBJECT_ALLOC_OVERRIDES(FAGXDeallocHandler)
+
+-(instancetype)initWithBlock:(dispatch_block_t)InBlock
+{
+	id Self = [super init];
+	if (Self)
+	{
+		self->Block = Block_copy(InBlock);
+	}
+	return Self;
+}
+-(void)dealloc
+{
+	self->Block();
+	Block_release(self->Block);
+	[super dealloc];
+}
+@end
+
+static mtlpp::PixelFormat FromSRGBFormat(mtlpp::PixelFormat Format)
+{
+	mtlpp::PixelFormat MTLFormat = Format;
+	
+	switch (Format)
+	{
+		case mtlpp::PixelFormat::RGBA8Unorm_sRGB:
+			MTLFormat = mtlpp::PixelFormat::RGBA8Unorm;
+			break;
+		case mtlpp::PixelFormat::BGRA8Unorm_sRGB:
+			MTLFormat = mtlpp::PixelFormat::BGRA8Unorm;
+			break;
+#if PLATFORM_MAC
+		case mtlpp::PixelFormat::BC1_RGBA_sRGB:
+			MTLFormat = mtlpp::PixelFormat::BC1_RGBA;
+			break;
+		case mtlpp::PixelFormat::BC2_RGBA_sRGB:
+			MTLFormat = mtlpp::PixelFormat::BC2_RGBA;
+			break;
+		case mtlpp::PixelFormat::BC3_RGBA_sRGB:
+			MTLFormat = mtlpp::PixelFormat::BC3_RGBA;
+			break;
+		case mtlpp::PixelFormat::BC7_RGBAUnorm_sRGB:
+			MTLFormat = mtlpp::PixelFormat::BC7_RGBAUnorm;
+			break;
+#endif //PLATFORM_MAC
+#if PLATFORM_IOS
+		case mtlpp::PixelFormat::R8Unorm_sRGB:
+			MTLFormat = mtlpp::PixelFormat::R8Unorm;
+			break;
+		case mtlpp::PixelFormat::PVRTC_RGBA_2BPP_sRGB:
+			MTLFormat = mtlpp::PixelFormat::PVRTC_RGBA_2BPP;
+			break;
+		case mtlpp::PixelFormat::PVRTC_RGBA_4BPP_sRGB:
+			MTLFormat = mtlpp::PixelFormat::PVRTC_RGBA_4BPP;
+			break;
+		case mtlpp::PixelFormat::ASTC_4x4_sRGB:
+			MTLFormat = mtlpp::PixelFormat::ASTC_4x4_LDR;
+			break;
+		case mtlpp::PixelFormat::ASTC_6x6_sRGB:
+			MTLFormat = mtlpp::PixelFormat::ASTC_6x6_LDR;
+			break;
+		case mtlpp::PixelFormat::ASTC_8x8_sRGB:
+			MTLFormat = mtlpp::PixelFormat::ASTC_8x8_LDR;
+			break;
+		case mtlpp::PixelFormat::ASTC_10x10_sRGB:
+			MTLFormat = mtlpp::PixelFormat::ASTC_10x10_LDR;
+			break;
+		case mtlpp::PixelFormat::ASTC_12x12_sRGB:
+			MTLFormat = mtlpp::PixelFormat::ASTC_12x12_LDR;
+			break;
+#endif //PLATFORM_IOS
+		default:
+			break;
+	}
+	
+	return MTLFormat;
+}
+
+static EPixelFormat MetalToRHIPixelFormat(mtlpp::PixelFormat Format)
+{
+	Format = FromSRGBFormat(Format);
+	for (uint32 i = 0; i < PF_MAX; i++)
+	{
+		if((mtlpp::PixelFormat)GPixelFormats[i].PlatformFormat == Format)
+		{
+			return (EPixelFormat)i;
+		}
+	}
+	check(false);
+	return PF_MAX;
+}
+
+static mtlpp::SizeAndAlign TextureSizeAndAlign(mtlpp::TextureType TextureType, uint32 Width, uint32 Height, uint32 Depth, mtlpp::PixelFormat Format, uint32 MipCount, uint32 SampleCount, uint32 ArrayCount)
+{
+	mtlpp::SizeAndAlign SizeAlign;
+	SizeAlign.Size = 0;
+	SizeAlign.Align = 0;
+	
+	uint32 Align = 0;
+	FRHIResourceCreateInfo CreateInfo(TEXT(""));
+	switch (TextureType)
+	{
+		case mtlpp::TextureType::Texture2D:
+		case mtlpp::TextureType::Texture2DMultisample:
+			SizeAlign.Size = RHICalcTexture2DPlatformSize(Width, Height, MetalToRHIPixelFormat(Format), MipCount, SampleCount, TexCreate_None, CreateInfo, Align);
+			SizeAlign.Align = Align;
+			break;
+		case mtlpp::TextureType::Texture2DArray:
+			SizeAlign.Size = RHICalcTexture2DPlatformSize(Width, Height, MetalToRHIPixelFormat(Format), MipCount, SampleCount, TexCreate_None, CreateInfo, Align) * ArrayCount;
+			SizeAlign.Align = Align;
+			break;
+		case mtlpp::TextureType::TextureCube:
+			SizeAlign.Size = RHICalcTextureCubePlatformSize(Width, MetalToRHIPixelFormat(Format), MipCount, TexCreate_None, CreateInfo, Align);
+			SizeAlign.Align = Align;
+			break;
+#if PLATFORM_MAC || __IPHONE_OS_VERSION_MAX_ALLOWED >= 110000
+		case mtlpp::TextureType::TextureCubeArray:
+			SizeAlign.Size = RHICalcTextureCubePlatformSize(Width, MetalToRHIPixelFormat(Format), MipCount, TexCreate_None, CreateInfo, Align) * ArrayCount;
+			SizeAlign.Align = Align;
+			break;
+#endif
+		case mtlpp::TextureType::Texture3D:
+			SizeAlign.Size = RHICalcTexture3DPlatformSize(Width, Height, Depth, MetalToRHIPixelFormat(Format), MipCount, TexCreate_None, CreateInfo, Align);
+			SizeAlign.Align = Align;
+			break;
+		case mtlpp::TextureType::Texture1D:
+		case mtlpp::TextureType::Texture1DArray:
+		default:
+			check(false);
+			break;
+	}
+	
+	return SizeAlign;
+}
+
+void AGXLLM::LogAllocTexture(mtlpp::Device& Device, mtlpp::TextureDescriptor const& Desc, mtlpp::Texture const& Texture)
+{
+	mtlpp::SizeAndAlign SizeAlign;
+#if PLATFORM_MAC || __IPHONE_OS_VERSION_MAX_ALLOWED >= 110000
+	if (FAGXCommandQueue::SupportsFeature(EAGXFeaturesGPUCaptureManager))
+	{
+		SizeAlign = Device.HeapTextureSizeAndAlign(Desc);
+	}
+	else
+#endif
+	{
+		SizeAlign = TextureSizeAndAlign(Desc.GetTextureType(), Desc.GetWidth(), Desc.GetHeight(), Desc.GetDepth(), Desc.GetPixelFormat(), Desc.GetMipmapLevelCount(), Desc.GetSampleCount(), Desc.GetArrayLength());
+	}
+	
+	void* Ptr = (void*)Texture.GetPtr();
+	uint64 Size = SizeAlign.Size;
+	
+#if PLATFORM_IOS
+	bool bMemoryless = (Texture.GetStorageMode() == mtlpp::StorageMode::Memoryless);
+	if (!bMemoryless)
+#endif
+	{
+		INC_MEMORY_STAT_BY(STAT_AGXTextureMemory, Size);
+	}
+	INC_DWORD_STAT(STAT_AGXTextureCount);
+	
+	LLM(FLowLevelMemTracker::Get().OnLowLevelAlloc(ELLMTracker::Platform, Ptr, Size, ELLMTag::Untagged, ELLMAllocType::System));
+	// Assign a dealloc handler to untrack the memory - but don't track the dispatch block!
+	{
+		LLM_SCOPED_PAUSE_TRACKING(ELLMAllocType::System);
+		
+		if (Desc.GetUsage() & mtlpp::TextureUsage::RenderTarget)
+		{
+			objc_setAssociatedObject(Texture.GetPtr(), (void*)&AGXLLM::LogAllocTexture,
+			[[[FAGXDeallocHandler alloc] initWithBlock:^{
+				LLM_PLATFORM_SCOPE_METAL(ELLMTagAGX::RenderTargets);
+				
+				LLM(FLowLevelMemTracker::Get().OnLowLevelFree(ELLMTracker::Platform, Ptr, ELLMAllocType::System));
+				
+#if PLATFORM_IOS
+				if (!bMemoryless)
+#endif
+				{
+					DEC_MEMORY_STAT_BY(STAT_AGXTextureMemory, Size);
+				}
+				DEC_DWORD_STAT(STAT_AGXTextureCount);
+			}] autorelease],
+			OBJC_ASSOCIATION_RETAIN);
+		}
+		else
+		{
+			objc_setAssociatedObject(Texture.GetPtr(), (void*)&AGXLLM::LogAllocTexture,
+			[[[FAGXDeallocHandler alloc] initWithBlock:^{
+				LLM_PLATFORM_SCOPE_METAL(ELLMTagAGX::Textures);
+			
+				LLM(FLowLevelMemTracker::Get().OnLowLevelFree(ELLMTracker::Platform, Ptr, ELLMAllocType::System));
+			
+#if PLATFORM_IOS
+				if (!bMemoryless)
+#endif
+				{
+					DEC_MEMORY_STAT_BY(STAT_AGXTextureMemory, Size);
+				}
+				DEC_DWORD_STAT(STAT_AGXTextureCount);
+			}] autorelease],
+			OBJC_ASSOCIATION_RETAIN);
+		}
+	}
+}
+
+void AGXLLM::LogAllocBuffer(mtlpp::Device& Device, mtlpp::Buffer const& Buffer)
+{
+	void* Ptr = (void*)Buffer.GetPtr();
+	uint64 Size = Buffer.GetLength();
+	
+	INC_MEMORY_STAT_BY(STAT_AGXBufferMemory, Size);
+	INC_DWORD_STAT(STAT_AGXBufferCount);
+	
+	LLM(FLowLevelMemTracker::Get().OnLowLevelAlloc(ELLMTracker::Platform, Ptr, Size, ELLMTag::Untagged, ELLMAllocType::System));
+	// Assign a dealloc handler to untrack the memory - but don't track the dispatch block!
+	{
+		LLM_SCOPED_PAUSE_TRACKING(ELLMAllocType::System);
+		
+		objc_setAssociatedObject(Buffer.GetPtr(), (void*)&AGXLLM::LogAllocBuffer,
+		[[[FAGXDeallocHandler alloc] initWithBlock:^{
+			LLM_PLATFORM_SCOPE_METAL(ELLMTagAGX::Buffers);
+			
+			LLM(FLowLevelMemTracker::Get().OnLowLevelFree(ELLMTracker::Platform, Ptr, ELLMAllocType::System));
+			
+			DEC_MEMORY_STAT_BY(STAT_AGXBufferMemory, Size);
+			DEC_DWORD_STAT(STAT_AGXBufferCount);
+		}] autorelease],
+		OBJC_ASSOCIATION_RETAIN);
+	}
+}
+
+void AGXLLM::LogAllocHeap(mtlpp::Device& Device, mtlpp::Heap const& Heap)
+{
+	void* Ptr = (void*)Heap.GetPtr();
+	uint64 Size = Heap.GetSize();
+	
+	INC_MEMORY_STAT_BY(STAT_AGXHeapMemory, Size);
+	INC_DWORD_STAT(STAT_AGXHeapCount);
+	
+	LLM(FLowLevelMemTracker::Get().OnLowLevelAlloc(ELLMTracker::Platform, Ptr, Size, ELLMTag::Untagged, ELLMAllocType::System));
+	// Assign a dealloc handler to untrack the memory - but don't track the dispatch block!
+	{
+		LLM_SCOPED_PAUSE_TRACKING(ELLMAllocType::System);
+		
+		objc_setAssociatedObject(Heap.GetPtr(), (void*)&AGXLLM::LogAllocHeap,
+								 [[[FAGXDeallocHandler alloc] initWithBlock:^{
+			LLM_SCOPE_METAL(ELLMTagAGX::Heaps);
+			LLM_PLATFORM_SCOPE_METAL(ELLMTagAGX::Heaps);
+			
+			LLM(FLowLevelMemTracker::Get().OnLowLevelFree(ELLMTracker::Platform, Ptr, ELLMAllocType::System));
+			
+			DEC_MEMORY_STAT_BY(STAT_AGXHeapMemory, Size);
+			DEC_DWORD_STAT(STAT_AGXHeapCount);
+		}] autorelease],
+		OBJC_ASSOCIATION_RETAIN);
+	}
+}
+
+void AGXLLM::LogAliasTexture(mtlpp::Texture const& Texture)
+{
+	objc_setAssociatedObject(Texture.GetPtr(), (void*)&AGXLLM::LogAllocTexture, nil, OBJC_ASSOCIATION_RETAIN);
+}
+
+void AGXLLM::LogAliasBuffer(mtlpp::Buffer const& Buffer)
+{
+	objc_setAssociatedObject(Buffer.GetPtr(), (void*)&AGXLLM::LogAllocBuffer, nil, OBJC_ASSOCIATION_RETAIN);
+}
+
+
diff --git a/Engine/Source/Runtime/Apple/AGXRHI/Private/AGXLLM.h b/Engine/Source/Runtime/Apple/AGXRHI/Private/AGXLLM.h
new file mode 100644
--- /dev/null
+++ b/Engine/Source/Runtime/Apple/AGXRHI/Private/AGXLLM.h
@@ -0,0 +1,57 @@
+// Copyright Epic Games, Inc. All Rights Reserved.
+
+#pragma once
+
+#include "AGXRHIPrivate.h"
+#include "HAL/LowLevelMemTracker.h"
+#include "Apple/AppleLLM.h"
+
+@interface FAGXDeallocHandler : FApplePlatformObject<NSObject>
+{
+	dispatch_block_t Block;
+}
+-(instancetype)initWithBlock:(dispatch_block_t)InBlock;
+-(void)dealloc;
+@end
+
+#if ENABLE_LOW_LEVEL_MEM_TRACKER
+
+#define LLM_SCOPE_METAL(Tag) LLM_SCOPE((ELLMTag)Tag)
+#define LLM_PLATFORM_SCOPE_METAL(Tag) LLM_PLATFORM_SCOPE((ELLMTag)Tag)
+
+enum class ELLMTagAGX : LLM_TAG_TYPE
+{
+	Buffers = (LLM_TAG_TYPE)ELLMTagApple::AppleMetalTagsStart,
+	Textures,
+	Heaps,
+	RenderTargets,
+	
+	Count
+};
+
+static_assert((int32)ELLMTagAGX::Count <= (int32)ELLMTagApple::AppleMetalTagsEnd, "too many ELLMTagAGX tags. Need to increase LLM_TAG_APPLE_NUM_METAL_TAGS_RESERVED");
+
+namespace AGXLLM
+{
+	void Initialise();
+}
+
+#else
+
+#define LLM_SCOPE_METAL(...)
+#define LLM_PLATFORM_SCOPE_METAL(...)
+
+#endif		// #if ENABLE_LOW_LEVEL_MEM_TRACKER
+
+// These work without the LLM module
+namespace AGXLLM
+{
+	void LogAllocTexture(mtlpp::Device& Device, mtlpp::TextureDescriptor const& Desc, mtlpp::Texture const& Texture);
+	void LogAllocBuffer(mtlpp::Device& Device, mtlpp::Buffer const& Buffer);
+	void LogAllocHeap(mtlpp::Device& Device, mtlpp::Heap const& Heap);
+	
+	void LogAliasTexture(mtlpp::Texture const& Texture);
+	void LogAliasBuffer(mtlpp::Buffer const& Buffer);
+}
+
+
