Initial import of AGXRHI

diff --git a/Engine/Source/Runtime/Apple/AGXRHI/Private/AGXPipeline.cpp b/Engine/Source/Runtime/Apple/AGXRHI/Private/AGXPipeline.cpp
new file mode 100644
--- /dev/null
+++ b/Engine/Source/Runtime/Apple/AGXRHI/Private/AGXPipeline.cpp
@@ -0,0 +1,1034 @@
+// Copyright Epic Games, Inc. All Rights Reserved.
+
+/*=============================================================================
+	AGXPipeline.cpp: AGX RHI shader pipeline implementation.
+=============================================================================*/
+
+#include "AGXRHIPrivate.h"
+#include "AGXVertexDeclaration.h"
+#include "AGXShaderTypes.h"
+#include "AGXGraphicsPipelineState.h"
+#include "AGXComputePipelineState.h"
+#include "AGXPipeline.h"
+#include "MetalShaderResources.h"
+#include "AGXProfiler.h"
+#include "AGXCommandQueue.h"
+#include "AGXCommandBuffer.h"
+#include "RenderUtils.h"
+#include "Misc/ScopeRWLock.h"
+#include "HAL/PThreadEvent.h"
+#include <objc/runtime.h>
+
+static int32 GAGXCacheShaderPipelines = 1;
+static FAutoConsoleVariableRef CVarAGXCacheShaderPipelines(
+	TEXT("rhi.AGX.CacheShaderPipelines"),
+	GAGXCacheShaderPipelines,
+	TEXT("When enabled (1, default) cache all graphics pipeline state objects created in AGXRHI for the life of the program, this trades memory for performance as creating PSOs is expensive in Metal.\n")
+	TEXT("Disable in the project configuration to allow PSOs to be released to save memory at the expense of reduced performance and increased hitching in-game\n. (On by default (1))"), ECVF_ReadOnly);
+
+static int32 GAGXCacheMinSize = 32;
+static FAutoConsoleVariableRef CVarAGXCacheMinSize(
+	TEXT("r.ShaderPipelineCache.AGXCacheMinSizeInMB"),
+	GAGXCacheMinSize,
+	TEXT("Sets the minimum size that we expect the metal OS cache to be (in MB). This is used to determine if we need to cache PSOs again (Default: 32).\n"), ECVF_ReadOnly);
+
+static uint32 BlendBitOffsets[] = { Offset_BlendState0, Offset_BlendState1, Offset_BlendState2, Offset_BlendState3, Offset_BlendState4, Offset_BlendState5, Offset_BlendState6, Offset_BlendState7 };
+static uint32 RTBitOffsets[] = { Offset_RenderTargetFormat0, Offset_RenderTargetFormat1, Offset_RenderTargetFormat2, Offset_RenderTargetFormat3, Offset_RenderTargetFormat4, Offset_RenderTargetFormat5, Offset_RenderTargetFormat6, Offset_RenderTargetFormat7 };
+static_assert(Offset_RasterEnd < 64 && Offset_End < 128, "Offset_RasterEnd must be < 64 && Offset_End < 128");
+
+static float RoundUpNearestEven(const float f)
+{
+	const float ret = FMath::CeilToFloat(f);
+	const float isOdd = (float)(((int)ret) & 1);
+	return ret + isOdd;
+}
+
+// A tile-based or vertex-based debug shader for trying to emulate Aftermath style failure reporting
+static NSString* GAGXDebugShader = @"#include <metal_stdlib>\n"
+"#include <metal_compute>\n"
+"\n"
+"using namespace metal;\n"
+"\n"
+"struct DebugInfo\n"
+"{\n"
+"   uint CmdBuffIndex;\n"
+"	uint EncoderIndex;\n"
+"   uint ContextIndex;\n"
+"   uint CommandIndex;\n"
+"   uint CommandBuffer[2];\n"
+"	uint PSOSignature[4];\n"
+"};\n"
+"\n"
+#if !PLATFORM_MAC
+"// Executes once per-tile\n"
+"kernel void Main_Debug(constant DebugInfo *debugTable [[ buffer(0) ]], device DebugInfo* debugBuffer [[ buffer(1) ]], uint2 threadgroup_position_in_grid [[ threadgroup_position_in_grid ]], uint2 threadgroups_per_grid [[ threadgroups_per_grid ]])\n"
+"{\n"
+"	// Write Pass, Draw indices\n"
+"	// Write Vertex+Fragment PSO sig (in form VertexLen, VertexCRC, FragLen, FragCRC)\n"
+"   uint tile_index = threadgroup_position_in_grid.x + (threadgroup_position_in_grid.y * threadgroups_per_grid.x);"
+"	debugBuffer[tile_index] = debugTable[0];\n"
+"}";
+#else
+"// Executes once as a point draw call\n"
+"vertex void Main_Debug(constant DebugInfo *debugTable [[ buffer(0) ]], device DebugInfo* debugBuffer [[ buffer(1) ]])\n"
+"{\n"
+"	// Write Pass, Draw indices\n"
+"	// Write Vertex+Fragment PSO sig (in form VertexLen, VertexCRC, FragLen, FragCRC)\n"
+"	debugBuffer[0] = debugTable[0];\n"
+"}";
+#endif
+
+// A compute debug shader for trying to emulate Aftermath style failure reporting
+static NSString* GAGXDebugMarkerComputeShader = @"#include <metal_stdlib>\n"
+"#include <metal_compute>\n"
+"\n"
+"using namespace metal;\n"
+"\n"
+"struct DebugInfo\n"
+"{\n"
+"   uint CmdBuffIndex;\n"
+"	uint EncoderIndex;\n"
+"   uint ContextIndex;\n"
+"   uint CommandIndex;\n"
+"   uint CommandBuffer[2];\n"
+"	uint PSOSignature[4];\n"
+"};\n"
+"\n"
+"// Executes once\n"
+"kernel void Main_Debug(constant DebugInfo *debugTable [[ buffer(0) ]], device DebugInfo* debugBuffer [[ buffer(1) ]])\n"
+"{\n"
+"	// Write Pass, Draw indices\n"
+"	// Write Vertex+Fragment PSO sig (in form VertexLen, VertexCRC, FragLen, FragCRC)\n"
+"	debugBuffer[0] = debugTable[0];\n"
+"}";
+
+struct FAGXHelperFunctions
+{
+    mtlpp::Library DebugShadersLib;
+    mtlpp::Function DebugFunc;
+	
+	mtlpp::Library DebugComputeShadersLib;
+	mtlpp::Function DebugComputeFunc;
+	mtlpp::ComputePipelineState DebugComputeState;
+	
+    FAGXHelperFunctions()
+    {
+#if !PLATFORM_TVOS
+        if (GAGXCommandBufferDebuggingEnabled)
+        {
+            mtlpp::CompileOptions CompileOptions;
+            ns::AutoReleasedError Error;
+
+			DebugShadersLib = GetAGXDeviceContext().GetDevice().NewLibrary(GAGXDebugShader, CompileOptions, &Error);
+            DebugFunc = DebugShadersLib.NewFunction(@"Main_Debug");
+			
+			DebugComputeShadersLib = GetAGXDeviceContext().GetDevice().NewLibrary(GAGXDebugMarkerComputeShader, CompileOptions, &Error);
+			DebugComputeFunc = DebugComputeShadersLib.NewFunction(@"Main_Debug");
+			
+			DebugComputeState = GetAGXDeviceContext().GetDevice().NewComputePipelineState(DebugComputeFunc, &Error);
+        }
+#endif
+    }
+    
+    static FAGXHelperFunctions& Get()
+    {
+        static FAGXHelperFunctions sSelf;
+        return sSelf;
+    }
+    
+    mtlpp::Function GetDebugFunction()
+    {
+        return DebugFunc;
+    }
+	
+	mtlpp::ComputePipelineState GetDebugComputeState()
+	{
+		return DebugComputeState;
+	}
+};
+
+mtlpp::ComputePipelineState AGXGetMetalDebugComputeState()
+{
+	return FAGXHelperFunctions::Get().GetDebugComputeState();
+}
+
+struct FAGXGraphicsPipelineKey
+{
+	FAGXRenderPipelineHash RenderPipelineHash;
+	FAGXHashedVertexDescriptor VertexDescriptorHash;
+	FSHAHash VertexFunction;
+	FSHAHash PixelFunction;
+
+	template<typename Type>
+	inline void SetHashValue(uint32 Offset, uint32 NumBits, Type Value)
+	{
+		if (Offset < Offset_RasterEnd)
+		{
+			uint64 BitMask = ((((uint64)1ULL) << NumBits) - 1) << Offset;
+			RenderPipelineHash.RasterBits = (RenderPipelineHash.RasterBits & ~BitMask) | (((uint64)Value << Offset) & BitMask);
+		}
+		else
+		{
+			Offset -= Offset_RenderTargetFormat0;
+			uint64 BitMask = ((((uint64)1ULL) << NumBits) - 1) << Offset;
+			RenderPipelineHash.TargetBits = (RenderPipelineHash.TargetBits & ~BitMask) | (((uint64)Value << Offset) & BitMask);
+		}
+	}
+
+	bool operator==(FAGXGraphicsPipelineKey const& Other) const
+	{
+		return (RenderPipelineHash == Other.RenderPipelineHash
+		&& VertexDescriptorHash == Other.VertexDescriptorHash
+		&& VertexFunction == Other.VertexFunction
+		&& PixelFunction == Other.PixelFunction);
+	}
+	
+	friend uint32 GetTypeHash(FAGXGraphicsPipelineKey const& Key)
+	{
+		uint32 H = FCrc::MemCrc32(&Key.RenderPipelineHash, sizeof(Key.RenderPipelineHash), GetTypeHash(Key.VertexDescriptorHash));
+		H = FCrc::MemCrc32(Key.VertexFunction.Hash, sizeof(Key.VertexFunction.Hash), H);
+		H = FCrc::MemCrc32(Key.PixelFunction.Hash, sizeof(Key.PixelFunction.Hash), H);
+		return H;
+	}
+	
+	friend void InitMetalGraphicsPipelineKey(FAGXGraphicsPipelineKey& Key, const FGraphicsPipelineStateInitializer& Init)
+	{
+		uint32 const NumActiveTargets = Init.ComputeNumValidRenderTargets();
+		check(NumActiveTargets <= MaxSimultaneousRenderTargets);
+	
+		FAGXBlendState* BlendState = (FAGXBlendState*)Init.BlendState;
+		
+		FMemory::Memzero(Key.RenderPipelineHash);
+		
+		bool bHasActiveTargets = false;
+		for (uint32 i = 0; i < NumActiveTargets; i++)
+		{
+			EPixelFormat TargetFormat = (EPixelFormat)Init.RenderTargetFormats[i];
+			if (TargetFormat == PF_Unknown) { continue; }
+
+			mtlpp::PixelFormat MetalFormat = (mtlpp::PixelFormat)GPixelFormats[TargetFormat].PlatformFormat;
+			ETextureCreateFlags Flags = (ETextureCreateFlags)Init.RenderTargetFlags[i];
+			if (Flags & TexCreate_SRGB)
+			{
+#if PLATFORM_MAC // Expand as R8_sRGB is iOS only.
+				if (MetalFormat == mtlpp::PixelFormat::R8Unorm)
+				{
+					MetalFormat = mtlpp::PixelFormat::RGBA8Unorm;
+				}
+#endif
+				MetalFormat = AGXToSRGBFormat(MetalFormat);
+			}
+			
+			uint8 FormatKey = AGXGetMetalPixelFormatKey(MetalFormat);;
+			Key.SetHashValue(RTBitOffsets[i], NumBits_RenderTargetFormat, FormatKey);
+			Key.SetHashValue(BlendBitOffsets[i], NumBits_BlendState, BlendState->RenderTargetStates[i].BlendStateKey);
+			
+			bHasActiveTargets |= true;
+		}
+		
+		uint8 DepthFormatKey = 0;
+		uint8 StencilFormatKey = 0;
+		switch(Init.DepthStencilTargetFormat)
+		{
+			case PF_DepthStencil:
+			{
+				mtlpp::PixelFormat MetalFormat = (mtlpp::PixelFormat)GPixelFormats[PF_DepthStencil].PlatformFormat;
+				if (Init.DepthTargetLoadAction != ERenderTargetLoadAction::ENoAction || Init.DepthTargetStoreAction != ERenderTargetStoreAction::ENoAction)
+				{
+					DepthFormatKey = AGXGetMetalPixelFormatKey(MetalFormat);
+				}
+				if (Init.StencilTargetLoadAction != ERenderTargetLoadAction::ENoAction || Init.StencilTargetStoreAction != ERenderTargetStoreAction::ENoAction)
+				{
+					StencilFormatKey = AGXGetMetalPixelFormatKey(mtlpp::PixelFormat::Stencil8);
+				}
+				bHasActiveTargets |= true;
+				break;
+			}
+			case PF_ShadowDepth:
+			{
+				DepthFormatKey = AGXGetMetalPixelFormatKey((mtlpp::PixelFormat)GPixelFormats[PF_ShadowDepth].PlatformFormat);
+				bHasActiveTargets |= true;
+				break;
+			}
+			default:
+			{
+				break;
+			}
+		}
+		
+		// If the pixel shader writes depth then we must compile with depth access, so we may bind the dummy depth.
+		// If the pixel shader writes to UAVs but not target is bound we must also bind the dummy depth.
+		FAGXPixelShader* PixelShader = (FAGXPixelShader*)Init.BoundShaderState.PixelShaderRHI;
+		if ( PixelShader && ( ( PixelShader->Bindings.InOutMask.IsFieldEnabled(CrossCompiler::FShaderBindingInOutMask::DepthStencilMaskIndex) && DepthFormatKey == 0 ) || (bHasActiveTargets == false && PixelShader->Bindings.NumUAVs > 0) ) )
+		{
+			mtlpp::PixelFormat MetalFormat = (mtlpp::PixelFormat)GPixelFormats[PF_DepthStencil].PlatformFormat;
+			DepthFormatKey = AGXGetMetalPixelFormatKey(MetalFormat);
+		}
+		
+		Key.SetHashValue(Offset_DepthFormat, NumBits_DepthFormat, DepthFormatKey);
+		Key.SetHashValue(Offset_StencilFormat, NumBits_StencilFormat, StencilFormatKey);
+
+		Key.SetHashValue(Offset_SampleCount, NumBits_SampleCount, Init.NumSamples);
+
+		Key.SetHashValue(Offset_AlphaToCoverage, NumBits_AlphaToCoverage, Init.NumSamples > 1 && BlendState->bUseAlphaToCoverage ? 1 : 0);
+		
+#if PLATFORM_MAC
+		Key.SetHashValue(Offset_PrimitiveTopology, NumBits_PrimitiveTopology, TranslatePrimitiveTopology(Init.PrimitiveType));
+#endif
+
+		FAGXVertexDeclaration* VertexDecl = (FAGXVertexDeclaration*)Init.BoundShaderState.VertexDeclarationRHI;
+		Key.VertexDescriptorHash = VertexDecl->Layout;
+		
+		FAGXVertexShader* VertexShader = (FAGXVertexShader*)Init.BoundShaderState.VertexShaderRHI;
+		Key.VertexFunction = VertexShader->GetHash();
+
+		if (PixelShader)
+		{
+			Key.PixelFunction = PixelShader->GetHash();
+		}
+	}
+};
+
+static FAGXShaderPipeline* CreateMTLRenderPipeline(bool const bSync, FAGXGraphicsPipelineKey const& Key, const FGraphicsPipelineStateInitializer& Init);
+
+class FAGXShaderPipelineCache
+{
+public:
+	static FAGXShaderPipelineCache& Get()
+	{
+		static FAGXShaderPipelineCache sSelf;
+		return sSelf;
+	}
+	
+	FAGXShaderPipeline* GetRenderPipeline(bool const bSync, FAGXGraphicsPipelineState const* State, const FGraphicsPipelineStateInitializer& Init)
+	{
+		SCOPE_CYCLE_COUNTER(STAT_AGXPipelineStateTime);
+		
+		FAGXGraphicsPipelineKey Key;
+		InitMetalGraphicsPipelineKey(Key, Init);
+		
+		// By default there'll be more threads trying to read this than to write it.
+		PipelineMutex.ReadLock();
+
+		// Try to find the entry in the cache.
+		FAGXShaderPipeline* Desc = Pipelines.FindRef(Key);
+
+		PipelineMutex.ReadUnlock();
+
+		if (Desc == nil)
+		{
+
+			// By default there'll be more threads trying to read this than to write it.
+			EventsMutex.ReadLock();
+
+			// Try to find a pipeline creation event for this key. If it's found, we already have a thread creating this pipeline and we just have to wait.
+			TSharedPtr<FPThreadEvent, ESPMode::ThreadSafe> Event = PipelineEvents.FindRef(Key);
+
+			EventsMutex.ReadUnlock();
+
+			bool bCompile = false;
+			if (!Event.IsValid())
+			{
+				// Create an event other threads can use to wait if they request the same pipeline this thread is creating
+				EventsMutex.WriteLock();
+
+				Event = PipelineEvents.FindRef(Key);
+				if (!Event.IsValid())
+				{
+					Event = PipelineEvents.Add(Key, MakeShareable(new FPThreadEvent()));
+					Event->Create(true);
+					bCompile = true;
+				}
+				check(Event.IsValid());
+
+				EventsMutex.WriteUnlock();
+			}
+
+			if (bCompile)
+			{
+				Desc = CreateMTLRenderPipeline(bSync, Key, Init);
+
+				if (Desc != nil)
+				{
+					PipelineMutex.WriteLock();
+
+					Pipelines.Add(Key, Desc);
+					ReverseLookup.Add(Desc, Key);
+
+					PipelineMutex.WriteUnlock();
+
+					if (GAGXCacheShaderPipelines == 0)
+					{
+						// When we aren't caching for program lifetime we autorelease so that the PSO is released to the OS once all RHI references are released.
+						[Desc autorelease];
+					}
+				}
+
+				EventsMutex.WriteLock();
+
+				Event->Trigger();
+				PipelineEvents.Remove(Key);
+
+				EventsMutex.WriteUnlock();
+			}
+			else
+			{
+				check(Event.IsValid());
+				Event->Wait();
+
+				PipelineMutex.ReadLock();
+				Desc = Pipelines.FindRef(Key);
+				PipelineMutex.ReadUnlock();
+				check(Desc);
+			}
+		}
+		
+		return Desc;
+	}
+	
+	void ReleaseRenderPipeline(FAGXShaderPipeline* Pipeline)
+	{
+		if (GAGXCacheShaderPipelines)
+		{
+			[Pipeline release];
+		}
+		else
+		{
+			// We take a mutex here to prevent anyone from acquiring a reference to the state which might just be about to return memory to the OS.
+			FRWScopeLock Lock(PipelineMutex, SLT_Write);
+			[Pipeline release];
+		}
+	}
+	
+	void RemoveRenderPipeline(FAGXShaderPipeline* Pipeline)
+	{
+		check (GAGXCacheShaderPipelines == 0);
+		{
+			FAGXGraphicsPipelineKey* Desc = ReverseLookup.Find(Pipeline);
+			if (Desc)
+			{
+				Pipelines.Remove(*Desc);
+				ReverseLookup.Remove(Pipeline);
+			}
+		}
+	}
+	
+	
+private:
+	FRWLock PipelineMutex;
+	FRWLock EventsMutex;
+	TMap<FAGXGraphicsPipelineKey, FAGXShaderPipeline*> Pipelines;
+	TMap<FAGXShaderPipeline*, FAGXGraphicsPipelineKey> ReverseLookup;
+	TMap<FAGXGraphicsPipelineKey, TSharedPtr<FPThreadEvent, ESPMode::ThreadSafe>> PipelineEvents;
+};
+
+@implementation FAGXShaderPipeline
+
+APPLE_PLATFORM_OBJECT_ALLOC_OVERRIDES(FAGXShaderPipeline)
+
+- (void)dealloc
+{
+	// For render pipeline states we might need to remove the PSO from the cache when we aren't caching them for program lifetime
+	if (GAGXCacheShaderPipelines == 0 && RenderPipelineState)
+	{
+		FAGXShaderPipelineCache::Get().RemoveRenderPipeline(self);
+	}
+	[super dealloc];
+}
+
+- (instancetype)init
+{
+	id Self = [super init];
+	if (Self)
+	{
+		RenderPipelineReflection = mtlpp::RenderPipelineReflection(nil);
+		ComputePipelineReflection = mtlpp::ComputePipelineReflection(nil);
+		StreamPipelineReflection = mtlpp::RenderPipelineReflection(nil);
+#if METAL_DEBUG_OPTIONS
+		RenderDesc = mtlpp::RenderPipelineDescriptor(nil);
+		StreamDesc = mtlpp::RenderPipelineDescriptor(nil);
+		ComputeDesc = mtlpp::ComputePipelineDescriptor(nil);
+#endif
+	}
+	return Self;
+}
+
+- (void)initResourceMask
+{
+	if (RenderPipelineReflection)
+	{
+		[self initResourceMask:EAGXShaderVertex];
+		[self initResourceMask:EAGXShaderFragment];
+		
+		if (AGXSafeGetRuntimeDebuggingLevel() < EAGXDebugLevelValidation)
+		{
+			RenderPipelineReflection = mtlpp::RenderPipelineReflection(nil);
+		}
+	}
+	if (ComputePipelineReflection)
+	{
+		[self initResourceMask:EAGXShaderCompute];
+		
+		if (AGXSafeGetRuntimeDebuggingLevel() < EAGXDebugLevelValidation)
+		{
+			ComputePipelineReflection = mtlpp::ComputePipelineReflection(nil);
+		}
+	}
+	if (StreamPipelineReflection)
+	{
+		[self initResourceMask:EAGXShaderStream];
+		
+		if (AGXSafeGetRuntimeDebuggingLevel() < EAGXDebugLevelValidation)
+		{
+			StreamPipelineReflection = mtlpp::RenderPipelineReflection(nil);
+		}
+	}
+}
+- (void)initResourceMask:(EAGXShaderFrequency)Frequency
+{
+	NSArray<MTLArgument*>* Arguments = nil;
+	switch(Frequency)
+	{
+		case EAGXShaderVertex:
+		{
+			MTLRenderPipelineReflection* Reflection = RenderPipelineReflection;
+			check(Reflection);
+			
+			Arguments = Reflection.vertexArguments;
+			break;
+		}
+		case EAGXShaderFragment:
+		{
+			MTLRenderPipelineReflection* Reflection = RenderPipelineReflection;
+			check(Reflection);
+			
+			Arguments = Reflection.fragmentArguments;
+			break;
+		}
+		case EAGXShaderCompute:
+		{
+			MTLComputePipelineReflection* Reflection = ComputePipelineReflection;
+			check(Reflection);
+			
+			Arguments = Reflection.arguments;
+			break;
+		}
+		case EAGXShaderStream:
+		{
+			MTLRenderPipelineReflection* Reflection = StreamPipelineReflection;
+			check(Reflection);
+			
+			Arguments = Reflection.vertexArguments;
+			break;
+		}
+		default:
+			check(false);
+			break;
+	}
+	
+	for (uint32 i = 0; i < Arguments.count; i++)
+	{
+		MTLArgument* Arg = [Arguments objectAtIndex:i];
+		check(Arg);
+
+		if (!Arg.active)
+		{
+			continue;
+		}
+
+		switch(Arg.type)
+		{
+			case MTLArgumentTypeBuffer:
+			{
+				checkf(Arg.index < ML_MaxBuffers, TEXT("Metal buffer index exceeded!"));
+				if (FString(Arg.name) != TEXT("BufferSizes") && FString(Arg.name) != TEXT("spvBufferSizeConstants"))
+				{
+					ResourceMask[Frequency].BufferMask |= (1 << Arg.index);
+				
+					if(BufferDataSizes[Frequency].Num() < 31)
+						BufferDataSizes[Frequency].SetNumZeroed(31);
+				
+					BufferDataSizes[Frequency][Arg.index] = Arg.bufferDataSize;
+				}
+				break;
+			}
+			case MTLArgumentTypeThreadgroupMemory:
+			{
+				break;
+			}
+			case MTLArgumentTypeTexture:
+			{
+				checkf(Arg.index < ML_MaxTextures, TEXT("Metal texture index exceeded!"));
+				ResourceMask[Frequency].TextureMask |= (FAGXTextureMask(1) << Arg.index);
+				TextureTypes[Frequency].Add(Arg.index, (uint8)Arg.textureType);
+				break;
+			}
+			case MTLArgumentTypeSampler:
+			{
+				checkf(Arg.index < ML_MaxSamplers, TEXT("Metal sampler index exceeded!"));
+				ResourceMask[Frequency].SamplerMask |= (1 << Arg.index);
+				break;
+			}
+			default:
+				check(false);
+				break;
+		}
+	}
+}
+@end
+
+static MTLVertexDescriptor* GetMaskedVertexDescriptor(MTLVertexDescriptor* InputDesc, const CrossCompiler::FShaderBindingInOutMask& InOutMask)
+{
+	for (uint32 Attr = 0; Attr < MaxMetalStreams; Attr++)
+	{
+		if (!InOutMask.IsFieldEnabled((int32)Attr) && [InputDesc.attributes objectAtIndexedSubscript:Attr] != nil)
+		{
+			MTLVertexDescriptor* Desc = [[InputDesc copy] autorelease];
+			CrossCompiler::FShaderBindingInOutMask BuffersUsed;
+			for (int32 MetalStreamIndex = 0; MetalStreamIndex < MaxMetalStreams; ++MetalStreamIndex)
+			{
+				if (!InOutMask.IsFieldEnabled(MetalStreamIndex))
+				{
+					[Desc.attributes setObject:nil atIndexedSubscript:MetalStreamIndex];
+				}
+				else
+				{
+					BuffersUsed.EnableField([Desc.attributes objectAtIndexedSubscript : MetalStreamIndex].bufferIndex);
+				}
+			}
+			for (int32 BufferIndex = 0; BufferIndex < ML_MaxBuffers; ++BufferIndex)
+			{
+				if (!BuffersUsed.IsFieldEnabled(BufferIndex))
+				{
+					[Desc.layouts setObject:nil atIndexedSubscript:BufferIndex];
+				}
+			}
+			return Desc;
+		}
+	}
+	
+	return InputDesc;
+}
+
+static bool ConfigureRenderPipelineDescriptor(mtlpp::RenderPipelineDescriptor& RenderPipelineDesc,
+#if PLATFORM_MAC
+											  mtlpp::RenderPipelineDescriptor& DebugPipelineDesc,
+#elif !PLATFORM_TVOS
+											  mtlpp::TileRenderPipelineDescriptor& DebugPipelineDesc,
+#endif
+											  FAGXGraphicsPipelineKey const& Key,
+											  const FGraphicsPipelineStateInitializer& Init)
+{
+	FAGXPixelShader* PixelShader = (FAGXPixelShader*)Init.BoundShaderState.PixelShaderRHI;
+	uint32 const NumActiveTargets = Init.ComputeNumValidRenderTargets();
+	check(NumActiveTargets <= MaxSimultaneousRenderTargets);
+	if (PixelShader)
+	{
+		if (PixelShader->Bindings.InOutMask.Bitmask == 0 && PixelShader->Bindings.NumUAVs == 0 && PixelShader->Bindings.bDiscards == false)
+		{
+			UE_LOG(LogAGX, Error, TEXT("Pixel shader has no outputs which is not permitted. No Discards, In-Out Mask: %x\nNumber UAVs: %d\nSource Code:\n%s"), PixelShader->Bindings.InOutMask.Bitmask, PixelShader->Bindings.NumUAVs, *FString(PixelShader->GetSourceCode()));
+			return false;
+		}
+		
+		const uint32 MaxNumActiveTargets = __builtin_popcount(PixelShader->Bindings.InOutMask.Bitmask & ((1u << CrossCompiler::FShaderBindingInOutMask::MaxIndex) - 1));
+		UE_CLOG((NumActiveTargets < MaxNumActiveTargets), LogAGX, Verbose, TEXT("NumActiveTargets doesn't match pipeline's pixel shader output mask: %u, %hx"), NumActiveTargets, PixelShader->Bindings.InOutMask.Bitmask);
+	}
+	
+	FAGXBlendState* BlendState = (FAGXBlendState*)Init.BlendState;
+	
+	ns::Array<mtlpp::RenderPipelineColorAttachmentDescriptor> ColorAttachments = RenderPipelineDesc.GetColorAttachments();
+#if !PLATFORM_TVOS
+	auto DebugColorAttachements = DebugPipelineDesc.GetColorAttachments();
+#endif
+	
+	uint32 TargetWidth = 0;
+	for (uint32 ActiveTargetIndex = 0; ActiveTargetIndex < NumActiveTargets; ActiveTargetIndex++)
+	{
+		EPixelFormat TargetFormat = (EPixelFormat)Init.RenderTargetFormats[ActiveTargetIndex];
+		
+		const bool bIsActiveTargetBound = (PixelShader && PixelShader->Bindings.InOutMask.IsFieldEnabled(ActiveTargetIndex));
+		METAL_FATAL_ASSERT(!(TargetFormat == PF_Unknown && bIsActiveTargetBound), TEXT("Pipeline pixel shader expects target %u to be bound but it isn't: %s."), ActiveTargetIndex, *FString(PixelShader->GetSourceCode()));
+		
+		TargetWidth += GPixelFormats[TargetFormat].BlockBytes;
+		
+		mtlpp::PixelFormat MetalFormat = (mtlpp::PixelFormat)GPixelFormats[TargetFormat].PlatformFormat;
+		ETextureCreateFlags Flags = (ETextureCreateFlags)Init.RenderTargetFlags[ActiveTargetIndex];
+		if (Flags & TexCreate_SRGB)
+		{
+#if PLATFORM_MAC // Expand as R8_sRGB is iOS only.
+			if (MetalFormat == mtlpp::PixelFormat::R8Unorm)
+			{
+				MetalFormat = mtlpp::PixelFormat::RGBA8Unorm;
+			}
+#endif
+			MetalFormat = AGXToSRGBFormat(MetalFormat);
+		}
+		
+		mtlpp::RenderPipelineColorAttachmentDescriptor Attachment = ColorAttachments[ActiveTargetIndex];
+		Attachment.SetPixelFormat(MetalFormat);
+		
+#if !PLATFORM_TVOS
+		auto DebugAttachment = DebugColorAttachements[ActiveTargetIndex];
+		DebugAttachment.SetPixelFormat(MetalFormat);
+#endif
+		
+		mtlpp::RenderPipelineColorAttachmentDescriptor Blend = BlendState->RenderTargetStates[ActiveTargetIndex].BlendState;
+		if(TargetFormat != PF_Unknown)
+		{
+			// assign each property manually, would be nice if this was faster
+			Attachment.SetBlendingEnabled(Blend.IsBlendingEnabled());
+			Attachment.SetSourceRgbBlendFactor(Blend.GetSourceRgbBlendFactor());
+			Attachment.SetDestinationRgbBlendFactor(Blend.GetDestinationRgbBlendFactor());
+			Attachment.SetRgbBlendOperation(Blend.GetRgbBlendOperation());
+			Attachment.SetSourceAlphaBlendFactor(Blend.GetSourceAlphaBlendFactor());
+			Attachment.SetDestinationAlphaBlendFactor(Blend.GetDestinationAlphaBlendFactor());
+			Attachment.SetAlphaBlendOperation(Blend.GetAlphaBlendOperation());
+			Attachment.SetWriteMask(Blend.GetWriteMask());
+			
+#if PLATFORM_MAC
+			DebugAttachment.SetBlendingEnabled(Blend.IsBlendingEnabled());
+			DebugAttachment.SetSourceRgbBlendFactor(Blend.GetSourceRgbBlendFactor());
+			DebugAttachment.SetDestinationRgbBlendFactor(Blend.GetDestinationRgbBlendFactor());
+			DebugAttachment.SetRgbBlendOperation(Blend.GetRgbBlendOperation());
+			DebugAttachment.SetSourceAlphaBlendFactor(Blend.GetSourceAlphaBlendFactor());
+			DebugAttachment.SetDestinationAlphaBlendFactor(Blend.GetDestinationAlphaBlendFactor());
+			DebugAttachment.SetAlphaBlendOperation(Blend.GetAlphaBlendOperation());
+			DebugAttachment.SetWriteMask(Blend.GetWriteMask());
+#endif
+		}
+		else
+		{
+			Attachment.SetBlendingEnabled(NO);
+			Attachment.SetWriteMask(mtlpp::ColorWriteMask::None);
+#if PLATFORM_MAC
+			DebugAttachment.SetBlendingEnabled(NO);
+			DebugAttachment.SetWriteMask(mtlpp::ColorWriteMask::None);
+#endif
+		}
+	}
+	
+	// don't allow a PSO that is too wide
+	if (!GSupportsWideMRT && TargetWidth > 16)
+	{
+		return false;
+	}
+	
+	switch(Init.DepthStencilTargetFormat)
+	{
+		case PF_DepthStencil:
+		{
+			mtlpp::PixelFormat MetalFormat = (mtlpp::PixelFormat)GPixelFormats[PF_DepthStencil].PlatformFormat;
+			if(MetalFormat == mtlpp::PixelFormat::Depth32Float)
+			{
+				if (Init.DepthTargetLoadAction != ERenderTargetLoadAction::ENoAction || Init.DepthTargetStoreAction != ERenderTargetStoreAction::ENoAction)
+				{
+					RenderPipelineDesc.SetDepthAttachmentPixelFormat(MetalFormat);
+#if PLATFORM_MAC
+					DebugPipelineDesc.SetDepthAttachmentPixelFormat(MetalFormat);
+#endif
+				}
+				if (Init.StencilTargetLoadAction != ERenderTargetLoadAction::ENoAction || Init.StencilTargetStoreAction != ERenderTargetStoreAction::ENoAction)
+				{
+					RenderPipelineDesc.SetStencilAttachmentPixelFormat(mtlpp::PixelFormat::Stencil8);
+#if PLATFORM_MAC
+					DebugPipelineDesc.SetStencilAttachmentPixelFormat(mtlpp::PixelFormat::Stencil8);
+#endif
+				}
+			}
+			else
+			{
+				RenderPipelineDesc.SetDepthAttachmentPixelFormat(MetalFormat);
+				RenderPipelineDesc.SetStencilAttachmentPixelFormat(MetalFormat);
+#if PLATFORM_MAC
+				DebugPipelineDesc.SetDepthAttachmentPixelFormat(MetalFormat);
+				DebugPipelineDesc.SetStencilAttachmentPixelFormat(MetalFormat);
+#endif
+			}
+			break;
+		}
+		case PF_ShadowDepth:
+		{
+			RenderPipelineDesc.SetDepthAttachmentPixelFormat((mtlpp::PixelFormat)GPixelFormats[PF_ShadowDepth].PlatformFormat);
+#if PLATFORM_MAC
+			DebugPipelineDesc.SetDepthAttachmentPixelFormat((mtlpp::PixelFormat)GPixelFormats[PF_ShadowDepth].PlatformFormat);
+#endif
+			break;
+		}
+		default:
+		{
+			break;
+		}
+	}
+	
+	check(Init.BoundShaderState.VertexShaderRHI != nullptr);
+#if PLATFORM_SUPPORTS_GEOMETRY_SHADERS
+	check(Init.BoundShaderState.GeometryShaderRHI == nullptr);
+#endif
+	
+	if( RenderPipelineDesc.GetDepthAttachmentPixelFormat() == mtlpp::PixelFormat::Invalid &&
+		PixelShader && ( PixelShader->Bindings.InOutMask.IsFieldEnabled(CrossCompiler::FShaderBindingInOutMask::DepthStencilMaskIndex) || ( NumActiveTargets == 0 && PixelShader->Bindings.NumUAVs > 0) ) )
+	{
+		RenderPipelineDesc.SetDepthAttachmentPixelFormat((mtlpp::PixelFormat)GPixelFormats[PF_DepthStencil].PlatformFormat);
+		RenderPipelineDesc.SetStencilAttachmentPixelFormat((mtlpp::PixelFormat)GPixelFormats[PF_DepthStencil].PlatformFormat);
+		
+#if PLATFORM_MAC
+		DebugPipelineDesc.SetDepthAttachmentPixelFormat((mtlpp::PixelFormat)GPixelFormats[PF_DepthStencil].PlatformFormat);
+		DebugPipelineDesc.SetStencilAttachmentPixelFormat((mtlpp::PixelFormat)GPixelFormats[PF_DepthStencil].PlatformFormat);
+#endif
+	}
+	
+	static bool bNoMSAA = FParse::Param(FCommandLine::Get(), TEXT("nomsaa"));
+	uint16 NumSamples = !bNoMSAA ? FMath::Max(Init.NumSamples, (uint16)1u) : (uint16)1u;
+	RenderPipelineDesc.SetSampleCount(NumSamples);
+	RenderPipelineDesc.SetAlphaToCoverageEnabled(NumSamples > 1 && BlendState->bUseAlphaToCoverage);
+#if PLATFORM_MAC
+	RenderPipelineDesc.SetInputPrimitiveTopology(TranslatePrimitiveTopology(Init.PrimitiveType));
+	DebugPipelineDesc.SetSampleCount(!bNoMSAA ? FMath::Max(Init.NumSamples, (uint16)1u) : (uint16)1u);
+	DebugPipelineDesc.SetInputPrimitiveTopology(mtlpp::PrimitiveTopologyClass::Point);
+#endif
+	
+	if (FAGXCommandQueue::SupportsFeature(EAGXFeaturesPipelineBufferMutability))
+	{
+		FAGXVertexShader* VertexShader = (FAGXVertexShader*)Init.BoundShaderState.VertexShaderRHI;
+		
+		ns::AutoReleased<ns::Array<mtlpp::PipelineBufferDescriptor>> VertexPipelineBuffers = RenderPipelineDesc.GetVertexBuffers();
+		FMetalShaderBindings& VertexBindings = VertexShader->Bindings;
+		int8 VertexSideTable = VertexShader->SideTableBinding;
+		{
+			uint32 ImmutableBuffers = VertexBindings.ConstantBuffers | VertexBindings.ArgumentBuffers;
+			while(ImmutableBuffers)
+			{
+				uint32 Index = __builtin_ctz(ImmutableBuffers);
+				ImmutableBuffers &= ~(1 << Index);
+				
+				if (Index < ML_MaxBuffers)
+				{
+					ns::AutoReleased<mtlpp::PipelineBufferDescriptor> PipelineBuffer = VertexPipelineBuffers[Index];
+					PipelineBuffer.SetMutability(mtlpp::Mutability::Immutable);
+				}
+			}
+			if (VertexSideTable > 0)
+			{
+				ns::AutoReleased<mtlpp::PipelineBufferDescriptor> PipelineBuffer = VertexPipelineBuffers[VertexSideTable];
+				PipelineBuffer.SetMutability(mtlpp::Mutability::Immutable);
+			}
+		}
+		
+		if (PixelShader)
+		{
+			ns::AutoReleased<ns::Array<mtlpp::PipelineBufferDescriptor>> FragmentPipelineBuffers = RenderPipelineDesc.GetFragmentBuffers();
+			uint32 ImmutableBuffers = PixelShader->Bindings.ConstantBuffers | PixelShader->Bindings.ArgumentBuffers;
+			while(ImmutableBuffers)
+			{
+				uint32 Index = __builtin_ctz(ImmutableBuffers);
+				ImmutableBuffers &= ~(1 << Index);
+				
+				if (Index < ML_MaxBuffers)
+				{
+					ns::AutoReleased<mtlpp::PipelineBufferDescriptor> PipelineBuffer = FragmentPipelineBuffers[Index];
+					PipelineBuffer.SetMutability(mtlpp::Mutability::Immutable);
+				}
+			}
+			if (PixelShader->SideTableBinding > 0)
+			{
+				ns::AutoReleased<mtlpp::PipelineBufferDescriptor> PipelineBuffer = FragmentPipelineBuffers[PixelShader->SideTableBinding];
+				PipelineBuffer.SetMutability(mtlpp::Mutability::Immutable);
+			}
+		}
+	}
+	
+	return true;
+}
+
+static FAGXShaderPipeline* CreateMTLRenderPipeline(bool const bSync, FAGXGraphicsPipelineKey const& Key, const FGraphicsPipelineStateInitializer& Init)
+{
+    FAGXVertexShader* VertexShader = (FAGXVertexShader*)Init.BoundShaderState.VertexShaderRHI;
+    FAGXPixelShader* PixelShader = (FAGXPixelShader*)Init.BoundShaderState.PixelShaderRHI;
+    
+    mtlpp::Function vertexFunction = VertexShader->GetFunction();
+    mtlpp::Function fragmentFunction = PixelShader ? PixelShader->GetFunction() : nil;
+
+    FAGXShaderPipeline* Pipeline = nil;
+    if (vertexFunction && ((PixelShader != nullptr) == (fragmentFunction != nil)))
+    {
+		ns::Error Error;
+		mtlpp::Device Device = GetAGXDeviceContext().GetDevice();
+
+		uint32 const NumActiveTargets = Init.ComputeNumValidRenderTargets();
+        check(NumActiveTargets <= MaxSimultaneousRenderTargets);
+		
+		Pipeline = [FAGXShaderPipeline new];
+		METAL_DEBUG_OPTION(FMemory::Memzero(Pipeline->ResourceMask, sizeof(Pipeline->ResourceMask)));
+
+		mtlpp::RenderPipelineDescriptor RenderPipelineDesc;
+        mtlpp::ComputePipelineDescriptor ComputePipelineDesc(nil);
+#if PLATFORM_MAC
+        mtlpp::RenderPipelineDescriptor DebugPipelineDesc;
+#elif !PLATFORM_TVOS
+        mtlpp::TileRenderPipelineDescriptor DebugPipelineDesc;
+#endif
+		
+#if PLATFORM_TVOS
+		if (!ConfigureRenderPipelineDescriptor(RenderPipelineDesc, Key, Init))
+#else
+		if (!ConfigureRenderPipelineDescriptor(RenderPipelineDesc, DebugPipelineDesc, Key, Init))
+#endif
+		{
+			return nil;
+		}
+        
+        FAGXVertexDeclaration* VertexDecl = (FAGXVertexDeclaration*)Init.BoundShaderState.VertexDeclarationRHI;
+		
+		RenderPipelineDesc.SetVertexDescriptor(GetMaskedVertexDescriptor(VertexDecl->Layout.VertexDesc, VertexShader->Bindings.InOutMask));
+		RenderPipelineDesc.SetVertexFunction(vertexFunction);
+		RenderPipelineDesc.SetFragmentFunction(fragmentFunction);
+#if ENABLE_METAL_GPUPROFILE
+		ns::String VertexName = vertexFunction.GetName();
+		ns::String FragmentName = fragmentFunction ? fragmentFunction.GetName() : @"";
+		RenderPipelineDesc.SetLabel([NSString stringWithFormat:@"%@+%@", VertexName.GetPtr(), FragmentName.GetPtr()]);
+#endif
+
+		NSUInteger RenderOption = mtlpp::PipelineOption::NoPipelineOption;
+		mtlpp::AutoReleasedRenderPipelineReflection* Reflection = nullptr;
+		mtlpp::AutoReleasedRenderPipelineReflection OutReflection;
+		Reflection = &OutReflection;
+		if (GetAGXDeviceContext().GetCommandQueue().GetRuntimeDebuggingLevel() >= EAGXDebugLevelFastValidation)
+		{
+			RenderOption = mtlpp::PipelineOption::ArgumentInfo | mtlpp::PipelineOption::BufferTypeInfo;
+		}
+
+		{
+			ns::AutoReleasedError RenderError;
+			METAL_GPUPROFILE(FAGXScopedCPUStats CPUStat(FString::Printf(TEXT("NewRenderPipeline: %s"), TEXT("")/**FString([RenderPipelineDesc.GetPtr() description])*/)));
+			Pipeline->RenderPipelineState = Device.NewRenderPipelineState(RenderPipelineDesc, (mtlpp::PipelineOption)RenderOption, Reflection, &RenderError);
+			if (Reflection)
+			{
+				Pipeline->RenderPipelineReflection = *Reflection;
+#if METAL_DEBUG_OPTIONS
+				Pipeline->RenderDesc = RenderPipelineDesc;
+#endif
+			}
+			Error = RenderError;
+		}
+		
+		UE_CLOG((Pipeline->RenderPipelineState == nil), LogAGX, Error, TEXT("Failed to generate a pipeline state object: %s"), *FString(Error.GetPtr().description));
+		UE_CLOG((Pipeline->RenderPipelineState == nil), LogAGX, Error, TEXT("Vertex shader: %s"), *FString(VertexShader->GetSourceCode()));
+		UE_CLOG((Pipeline->RenderPipelineState == nil), LogAGX, Error, TEXT("Pixel shader: %s"), PixelShader ? *FString(PixelShader->GetSourceCode()) : TEXT("NULL"));
+		UE_CLOG((Pipeline->RenderPipelineState == nil), LogAGX, Error, TEXT("Descriptor: %s"), *FString(RenderPipelineDesc.GetPtr().description));
+		UE_CLOG((Pipeline->RenderPipelineState == nil), LogAGX, Error, TEXT("Failed to generate a render pipeline state object:\n\n %s\n\n"), *FString(Error.GetLocalizedDescription()));
+		
+		// We need to pass a failure up the chain, so we'll clean up here.
+		if(Pipeline->RenderPipelineState == nil)
+		{
+			[Pipeline release];
+			return nil;
+		}
+		
+#if METAL_DEBUG_OPTIONS
+		Pipeline->VertexSource = VertexShader->GetSourceCode();
+        Pipeline->FragmentSource = PixelShader ? PixelShader->GetSourceCode() : nil;
+#endif
+		
+#if !PLATFORM_TVOS
+		if (GAGXCommandBufferDebuggingEnabled)
+		{
+#if PLATFORM_MAC
+			DebugPipelineDesc.SetVertexFunction(FAGXHelperFunctions::Get().GetDebugFunction());
+			DebugPipelineDesc.SetRasterizationEnabled(false);
+#else
+			DebugPipelineDesc.SetTileFunction(FAGXHelperFunctions::Get().GetDebugFunction());
+			DebugPipelineDesc.SetRasterSampleCount(RenderPipelineDesc.GetSampleCount());
+			DebugPipelineDesc.SetThreadgroupSizeMatchesTileSize(false);
+#endif
+#if ENABLE_METAL_GPUPROFILE
+			DebugPipelineDesc.SetLabel(@"Main_Debug");
+#endif
+
+			ns::AutoReleasedError RenderError;
+			METAL_GPUPROFILE(FAGXScopedCPUStats CPUStat(FString::Printf(TEXT("NewDebugPipeline: %s"), TEXT("")/**FString([RenderPipelineDesc.GetPtr() description])*/)));
+			Pipeline->DebugPipelineState = Device.NewRenderPipelineState(DebugPipelineDesc, mtlpp::PipelineOption::NoPipelineOption, Reflection, nullptr);
+		}
+#endif
+        
+#if METAL_DEBUG_OPTIONS
+        if (GFrameCounter > 3)
+        {
+            UE_LOG(LogAGX, Verbose, TEXT("Created a hitchy pipeline state for hash %llx %llx %llx"), (uint64)Key.RenderPipelineHash.RasterBits, (uint64)(Key.RenderPipelineHash.TargetBits), (uint64)Key.VertexDescriptorHash.VertexDescHash);
+        }
+#endif
+    }
+	
+	if (Pipeline && AGXSafeGetRuntimeDebuggingLevel() >= EAGXDebugLevelFastValidation)
+	{
+		[Pipeline initResourceMask];
+	}
+	
+    return !bSync ? nil : Pipeline;
+}
+
+FAGXShaderPipeline* GetMTLRenderPipeline(bool const bSync, FAGXGraphicsPipelineState const* State, const FGraphicsPipelineStateInitializer& Init)
+{
+	return FAGXShaderPipelineCache::Get().GetRenderPipeline(bSync, State, Init);
+}
+
+void ReleaseMTLRenderPipeline(FAGXShaderPipeline* Pipeline)
+{
+	FAGXShaderPipelineCache::Get().ReleaseRenderPipeline(Pipeline);
+}
+
+FAGXPipelineStateCacheManager::FAGXPipelineStateCacheManager()
+{
+#if PLATFORM_IOS
+	OnShaderPipelineCachePreOpenDelegate = FShaderPipelineCache::GetCachePreOpenDelegate().AddRaw(this, &FAGXPipelineStateCacheManager::OnShaderPipelineCachePreOpen);
+	OnShaderPipelineCacheOpenedDelegate = FShaderPipelineCache::GetCacheOpenedDelegate().AddRaw(this, &FAGXPipelineStateCacheManager::OnShaderPipelineCacheOpened);
+	OnShaderPipelineCachePrecompilationCompleteDelegate = FShaderPipelineCache::GetPrecompilationCompleteDelegate().AddRaw(this, &FAGXPipelineStateCacheManager::OnShaderPipelineCachePrecompilationComplete);
+#endif
+}
+
+FAGXPipelineStateCacheManager::~FAGXPipelineStateCacheManager()
+{
+	if (OnShaderPipelineCacheOpenedDelegate.IsValid())
+	{
+		FShaderPipelineCache::GetCacheOpenedDelegate().Remove(OnShaderPipelineCacheOpenedDelegate);
+	}
+	
+	if (OnShaderPipelineCachePrecompilationCompleteDelegate.IsValid())
+	{
+		FShaderPipelineCache::GetPrecompilationCompleteDelegate().Remove(OnShaderPipelineCachePrecompilationCompleteDelegate);
+	}
+}
+
+void FAGXPipelineStateCacheManager::OnShaderPipelineCachePreOpen(FString const& Name, EShaderPlatform Platform, bool& bReady)
+{
+	// only do this when haven't gotten a full pso cache already
+	struct stat FileInfo;
+	static FString PrivateWritePathBase = FString([NSSearchPathForDirectoriesInDomains(NSLibraryDirectory, NSUserDomainMask, YES) objectAtIndex:0]) + TEXT("/");
+	FString Result = PrivateWritePathBase + FString([NSString stringWithFormat:@"/Caches/%@/com.apple.metal/functions.data", [NSBundle mainBundle].bundleIdentifier]);
+	FString Result2 = PrivateWritePathBase + FString([NSString stringWithFormat:@"/Caches/%@/com.apple.metal/usecache.txt", [NSBundle mainBundle].bundleIdentifier]);
+	if (stat(TCHAR_TO_UTF8(*Result), &FileInfo) != -1 && ((FileInfo.st_size / 1024 / 1024) > GAGXCacheMinSize) && stat(TCHAR_TO_UTF8(*Result2), &FileInfo) != -1)
+	{
+		bReady = false;
+		FShaderPipelineCache::SetBatchMode(FShaderPipelineCache::BatchMode::Background);
+	}
+	else
+	{
+		bReady = true;
+		FShaderPipelineCache::SetBatchMode(FShaderPipelineCache::BatchMode::Precompile);
+	}
+}
+
+void FAGXPipelineStateCacheManager::OnShaderPipelineCacheOpened(FString const& Name, EShaderPlatform Platform, uint32 Count, const FGuid& VersionGuid, FShaderPipelineCache::FShaderCachePrecompileContext& ShaderCachePrecompileContext)
+{
+	ShaderCachePrecompileContext.SetPrecompilationIsSlowTask();
+}
+
+void FAGXPipelineStateCacheManager::OnShaderPipelineCachePrecompilationComplete(uint32 Count, double Seconds, const FShaderPipelineCache::FShaderCachePrecompileContext& ShaderCachePrecompileContext)
+{
+	// Want to ignore any subsequent Shader Pipeline Cache opening/closing, eg when loading modules
+	FShaderPipelineCache::GetCachePreOpenDelegate().Remove(OnShaderPipelineCachePreOpenDelegate);
+	FShaderPipelineCache::GetCacheOpenedDelegate().Remove(OnShaderPipelineCacheOpenedDelegate);
+	FShaderPipelineCache::GetPrecompilationCompleteDelegate().Remove(OnShaderPipelineCachePrecompilationCompleteDelegate);
+	OnShaderPipelineCachePreOpenDelegate.Reset();
+	OnShaderPipelineCacheOpenedDelegate.Reset();
+	OnShaderPipelineCachePrecompilationCompleteDelegate.Reset();
+}
diff --git a/Engine/Source/Runtime/Apple/AGXRHI/Private/AGXPipeline.h b/Engine/Source/Runtime/Apple/AGXRHI/Private/AGXPipeline.h
new file mode 100755
--- /dev/null
+++ b/Engine/Source/Runtime/Apple/AGXRHI/Private/AGXPipeline.h
@@ -0,0 +1,94 @@
+// Copyright Epic Games, Inc. All Rights Reserved.
+
+#pragma once
+
+#include "AGXRHIPrivate.h"
+#if METAL_DEBUG_OPTIONS
+#include "AGXDebugCommandEncoder.h"
+#endif
+
+#include "ShaderPipelineCache.h"
+
+enum EAGXPipelineHashBits
+{
+	NumBits_RenderTargetFormat = 6, //(x8=48),
+	NumBits_DepthFormat = 3, //(x1=3),
+	NumBits_StencilFormat = 3, //(x1=3),
+	NumBits_SampleCount = 3, //(x1=3),
+
+	NumBits_BlendState = 7, //(x8=56),
+	NumBits_PrimitiveTopology = 2, //(x1=2)
+	NumBits_AlphaToCoverage = 1, //(x1=1)
+};
+
+enum EAGXPipelineHashOffsets
+{
+	Offset_BlendState0 = 0,
+	Offset_BlendState1 = Offset_BlendState0 + NumBits_BlendState,
+	Offset_BlendState2 = Offset_BlendState1 + NumBits_BlendState,
+	Offset_BlendState3 = Offset_BlendState2 + NumBits_BlendState,
+	Offset_BlendState4 = Offset_BlendState3 + NumBits_BlendState,
+	Offset_BlendState5 = Offset_BlendState4 + NumBits_BlendState,
+	Offset_BlendState6 = Offset_BlendState5 + NumBits_BlendState,
+	Offset_BlendState7 = Offset_BlendState6 + NumBits_BlendState,
+	Offset_PrimitiveTopology = Offset_BlendState7 + NumBits_BlendState,
+	Offset_RasterEnd = Offset_PrimitiveTopology + NumBits_PrimitiveTopology,
+
+	Offset_RenderTargetFormat0 = 64,
+	Offset_RenderTargetFormat1 = Offset_RenderTargetFormat0 + NumBits_RenderTargetFormat,
+	Offset_RenderTargetFormat2 = Offset_RenderTargetFormat1 + NumBits_RenderTargetFormat,
+	Offset_RenderTargetFormat3 = Offset_RenderTargetFormat2 + NumBits_RenderTargetFormat,
+	Offset_RenderTargetFormat4 = Offset_RenderTargetFormat3 + NumBits_RenderTargetFormat,
+	Offset_RenderTargetFormat5 = Offset_RenderTargetFormat4 + NumBits_RenderTargetFormat,
+	Offset_RenderTargetFormat6 = Offset_RenderTargetFormat5 + NumBits_RenderTargetFormat,
+	Offset_RenderTargetFormat7 = Offset_RenderTargetFormat6 + NumBits_RenderTargetFormat,
+	Offset_DepthFormat = Offset_RenderTargetFormat7 + NumBits_RenderTargetFormat,
+	Offset_StencilFormat = Offset_DepthFormat + NumBits_DepthFormat,
+	Offset_SampleCount = Offset_StencilFormat + NumBits_StencilFormat,
+	Offset_AlphaToCoverage = Offset_SampleCount + NumBits_SampleCount,
+	Offset_End = Offset_AlphaToCoverage + NumBits_AlphaToCoverage
+};
+
+class FAGXPipelineStateCacheManager
+{
+public:
+	FAGXPipelineStateCacheManager();
+	~FAGXPipelineStateCacheManager();
+	
+private:
+	FDelegateHandle OnShaderPipelineCachePreOpenDelegate;
+	FDelegateHandle OnShaderPipelineCacheOpenedDelegate;
+	FDelegateHandle OnShaderPipelineCachePrecompilationCompleteDelegate;
+	
+	/** Delegate handlers to track the ShaderPipelineCache precompile. */
+	void OnShaderPipelineCachePreOpen(FString const& Name, EShaderPlatform Platform, bool& bReady);
+	void OnShaderPipelineCacheOpened(FString const& Name, EShaderPlatform Platform, uint32 Count, const FGuid& VersionGuid, FShaderPipelineCache::FShaderCachePrecompileContext& ShaderCachePrecompileContext);
+	void OnShaderPipelineCachePrecompilationComplete(uint32 Count, double Seconds, const FShaderPipelineCache::FShaderCachePrecompileContext& ShaderCachePrecompileContext);
+};
+
+@interface FAGXShaderPipeline : FApplePlatformObject
+{
+@public
+	mtlpp::RenderPipelineState RenderPipelineState;
+	mtlpp::ComputePipelineState ComputePipelineState;
+	mtlpp::RenderPipelineState StreamPipelineState;
+	mtlpp::RenderPipelineState DebugPipelineState;
+	TArray<uint32> BufferDataSizes[EAGXShaderStagesNum];
+	TMap<uint8, uint8> TextureTypes[EAGXShaderStagesNum];
+	FAGXDebugShaderResourceMask ResourceMask[EAGXShaderStagesNum];
+	mtlpp::RenderPipelineReflection RenderPipelineReflection;
+	mtlpp::RenderPipelineReflection StreamPipelineReflection;
+	mtlpp::ComputePipelineReflection ComputePipelineReflection;
+#if METAL_DEBUG_OPTIONS
+	ns::String VertexSource;
+	ns::String FragmentSource;
+	ns::String ComputeSource;
+	mtlpp::RenderPipelineDescriptor RenderDesc;
+	mtlpp::RenderPipelineDescriptor StreamDesc;
+	mtlpp::ComputePipelineDescriptor ComputeDesc;
+#endif
+}
+- (instancetype)init;
+- (void)initResourceMask;
+- (void)initResourceMask:(EAGXShaderFrequency)Frequency;
+@end
diff --git a/Engine/Source/Runtime/Apple/AGXRHI/Private/AGXProfiler.cpp b/Engine/Source/Runtime/Apple/AGXRHI/Private/AGXProfiler.cpp
new file mode 100644
--- /dev/null
+++ b/Engine/Source/Runtime/Apple/AGXRHI/Private/AGXProfiler.cpp
@@ -0,0 +1,873 @@
+// Copyright Epic Games, Inc. All Rights Reserved.
+
+#include "AGXRHIPrivate.h"
+#include "AGXProfiler.h"
+#include "EngineGlobals.h"
+#include "StaticBoundShaderState.h"
+#include "AGXCommandBuffer.h"
+#include "HAL/FileManager.h"
+
+DEFINE_STAT(STAT_AGXUniformMemAlloc);
+DEFINE_STAT(STAT_AGXUniformMemFreed);
+DEFINE_STAT(STAT_AGXVertexMemAlloc);
+DEFINE_STAT(STAT_AGXVertexMemFreed);
+DEFINE_STAT(STAT_AGXIndexMemAlloc);
+DEFINE_STAT(STAT_AGXIndexMemFreed);
+DEFINE_STAT(STAT_AGXTextureMemUpdate);
+
+DEFINE_STAT(STAT_AGXDrawCallTime);
+DEFINE_STAT(STAT_AGXPipelineStateTime);
+DEFINE_STAT(STAT_AGXPrepareDrawTime);
+
+DEFINE_STAT(STAT_AGXSwitchToRenderTime);
+DEFINE_STAT(STAT_AGXSwitchToComputeTime);
+DEFINE_STAT(STAT_AGXSwitchToBlitTime);
+DEFINE_STAT(STAT_AGXSwitchToAsyncBlitTime);
+DEFINE_STAT(STAT_AGXPrepareToRenderTime);
+DEFINE_STAT(STAT_AGXPrepareToDispatchTime);
+DEFINE_STAT(STAT_AGXCommitRenderResourceTablesTime);
+DEFINE_STAT(STAT_AGXSetRenderStateTime);
+DEFINE_STAT(STAT_AGXSetRenderPipelineStateTime);
+
+DEFINE_STAT(STAT_AGXMakeDrawableTime);
+DEFINE_STAT(STAT_AGXBufferPageOffTime);
+DEFINE_STAT(STAT_AGXTexturePageOnTime);
+DEFINE_STAT(STAT_AGXTexturePageOffTime);
+DEFINE_STAT(STAT_AGXGPUWorkTime);
+DEFINE_STAT(STAT_AGXGPUIdleTime);
+DEFINE_STAT(STAT_AGXPresentTime);
+DEFINE_STAT(STAT_AGXCustomPresentTime);
+DEFINE_STAT(STAT_AGXCommandBufferCreatedPerFrame);
+DEFINE_STAT(STAT_AGXCommandBufferCommittedPerFrame);
+DEFINE_STAT(STAT_AGXBufferMemory);
+DEFINE_STAT(STAT_AGXTextureMemory);
+DEFINE_STAT(STAT_AGXHeapMemory);
+DEFINE_STAT(STAT_AGXBufferUnusedMemory);
+DEFINE_STAT(STAT_AGXTextureUnusedMemory);
+DEFINE_STAT(STAT_AGXBufferCount);
+DEFINE_STAT(STAT_AGXTextureCount);
+DEFINE_STAT(STAT_AGXHeapCount);
+DEFINE_STAT(STAT_AGXFenceCount);
+
+DEFINE_STAT(STAT_AGXUniformMemoryInFlight);
+DEFINE_STAT(STAT_AGXUniformAllocatedMemory);
+DEFINE_STAT(STAT_AGXUniformBytesPerFrame);
+
+DEFINE_STAT(STAT_AGXFrameAllocatorMemoryInFlight);
+DEFINE_STAT(STAT_AGXFrameAllocatorAllocatedMemory);
+DEFINE_STAT(STAT_AGXFrameAllocatorBytesPerFrame);
+
+int64 volatile GAGXTexturePageOnTime = 0;
+int64 volatile GAGXGPUWorkTime = 0;
+int64 volatile GAGXGPUIdleTime = 0;
+int64 volatile GAGXPresentTime = 0;
+
+void WriteString(FArchive* OutputFile, const char* String)
+{
+	OutputFile->Serialize((void*)String, sizeof(ANSICHAR)*FCStringAnsi::Strlen(String));
+}
+
+FAGXEventNode::~FAGXEventNode()
+{
+}
+
+float FAGXEventNode::GetTiming()
+{
+	return FPlatformTime::ToSeconds(EndTime - StartTime);
+}
+
+void FAGXEventNode::StartTiming()
+{
+	StartTime = 0;
+	EndTime = 0;
+
+	Context->StartTiming(this);
+}
+
+mtlpp::CommandBufferHandler FAGXEventNode::Start(void)
+{
+	mtlpp::CommandBufferHandler Block = [this](mtlpp::CommandBuffer const& CompletedBuffer)
+	{
+		const CFTimeInterval GpuTimeSeconds = CompletedBuffer.GetGpuStartTime();
+		const double CyclesPerSecond = 1.0 / FPlatformTime::GetSecondsPerCycle();
+		StartTime = GpuTimeSeconds * CyclesPerSecond;
+	};
+    return Block_copy(Block);
+}
+
+void FAGXEventNode::StopTiming()
+{
+	Context->EndTiming(this);
+}
+
+mtlpp::CommandBufferHandler FAGXEventNode::Stop(void)
+{
+	mtlpp::CommandBufferHandler Block = [this](mtlpp::CommandBuffer const& CompletedBuffer)
+	{
+		// This is still used by ProfileGPU
+		const CFTimeInterval GpuTimeSeconds = CompletedBuffer.GetGpuEndTime();
+		const double CyclesPerSecond = 1.0 / FPlatformTime::GetSecondsPerCycle();
+		EndTime = GpuTimeSeconds * CyclesPerSecond;
+		
+		if(bRoot)
+		{
+			if(!bFullProfiling)
+			{
+				delete this;
+			}
+		}
+	};
+	return Block_copy(Block);
+}
+
+bool AGXGPUProfilerIsInSafeThread()
+{
+	return (GIsAGXInitialized && !GIsRHIInitialized) || (IsInRHIThread() || IsInActualRenderingThread());
+}
+	
+/** Start this frame of per tracking */
+void FAGXEventNodeFrame::StartFrame()
+{
+	RootNode->StartTiming();
+}
+
+/** End this frame of per tracking, but do not block yet */
+void FAGXEventNodeFrame::EndFrame()
+{
+	RootNode->StopTiming();
+}
+
+/** Calculates root timing base frequency (if needed by this RHI) */
+float FAGXEventNodeFrame::GetRootTimingResults()
+{
+	return RootNode->GetTiming();
+}
+
+void FAGXEventNodeFrame::LogDisjointQuery()
+{
+	
+}
+
+FGPUProfilerEventNode* FAGXGPUProfiler::CreateEventNode(const TCHAR* InName, FGPUProfilerEventNode* InParent)
+{
+#if ENABLE_METAL_GPUPROFILE
+	FAGXEventNode* EventNode = new FAGXEventNode(FAGXContext::GetCurrentContext(), InName, InParent, false, false);
+	return EventNode;
+#else
+	return nullptr;
+#endif
+}
+
+void FAGXGPUProfiler::Cleanup()
+{
+	
+}
+
+void FAGXGPUProfiler::PushEvent(const TCHAR* Name, FColor Color)
+{
+	if(AGXGPUProfilerIsInSafeThread())
+	{
+		FGPUProfiler::PushEvent(Name, Color);
+	}
+}
+
+void FAGXGPUProfiler::PopEvent()
+{
+	if(AGXGPUProfilerIsInSafeThread())
+	{
+		FGPUProfiler::PopEvent();
+	}
+}
+
+//TGlobalResource<FVector4VertexDeclaration> GAGXVector4VertexDeclaration;
+TGlobalResource<FTexture> GAGXLongTaskRT;
+
+void FAGXGPUProfiler::BeginFrame()
+{
+	if(!CurrentEventNodeFrame)
+	{
+		// Start tracking the frame
+		CurrentEventNodeFrame = new FAGXEventNodeFrame(Context, GTriggerGPUProfile);
+		CurrentEventNodeFrame->StartFrame();
+		
+		if(GNumAlternateFrameRenderingGroups > 1)
+		{
+			GTriggerGPUProfile = false;
+		}
+
+		if(GTriggerGPUProfile)
+		{
+			bTrackingEvents = true;
+			bLatchedGProfilingGPU = true;
+			GTriggerGPUProfile = false;
+		}
+	}
+	NumNestedFrames++;
+}
+
+void FAGXGPUProfiler::EndFrame()
+{
+	if(--NumNestedFrames == 0)
+	{
+		dispatch_async(dispatch_get_global_queue(QOS_CLASS_UTILITY, 0), ^{
+#if PLATFORM_MAC
+			FPlatformMisc::UpdateDriverMonitorStatistics(GetAGXDeviceContext().GetDeviceIndex());
+#endif
+		});
+#if STATS
+		SET_CYCLE_COUNTER(STAT_AGXTexturePageOnTime, GAGXTexturePageOnTime);
+		GAGXTexturePageOnTime = 0;
+		
+		SET_CYCLE_COUNTER(STAT_AGXGPUIdleTime, GAGXGPUIdleTime);
+		SET_CYCLE_COUNTER(STAT_AGXGPUWorkTime, GAGXGPUWorkTime);
+		SET_CYCLE_COUNTER(STAT_AGXPresentTime, GAGXPresentTime);
+#endif
+		
+		if(CurrentEventNodeFrame)
+		{
+			CurrentEventNodeFrame->EndFrame();
+			
+			if(bLatchedGProfilingGPU)
+			{
+				bTrackingEvents = false;
+				bLatchedGProfilingGPU = false;
+			
+				UE_LOG(LogRHI, Warning, TEXT(""));
+				UE_LOG(LogRHI, Warning, TEXT(""));
+				CurrentEventNodeFrame->DumpEventTree();
+			}
+			
+			delete CurrentEventNodeFrame;
+			CurrentEventNodeFrame = NULL;
+		}
+	}
+}
+
+// WARNING:
+// All these recording functions MUST be called from within scheduled/completion handlers.
+// Ordering is enforced by libdispatch so calling these outside of that context WILL result in
+// incorrect values.
+void FAGXGPUProfiler::RecordFrame(TArray<FAGXCommandBufferTiming>& CommandBufferTimings, FAGXCommandBufferTiming& LastBufferTiming)
+{
+	const double CyclesPerSecond = 1.0 / FPlatformTime::GetSecondsPerCycle();
+
+	double RunningFrameTimeSeconds = 0.0;
+	uint64 FrameStartGPUCycles = 0;
+	uint64 FrameEndGPUCycles = 0;
+
+	// Sort the timings
+	CommandBufferTimings.Sort();
+
+	CFTimeInterval FirstStartTime = 0.0;
+
+	// Add the timings excluding any overlapping time
+	for (const FAGXCommandBufferTiming& Timing : CommandBufferTimings)
+	{
+		if (FirstStartTime == 0.0)
+		{
+			FirstStartTime = Timing.StartTime;
+		}
+		
+		// Only process if the previous buffer finished before the end of this one
+		if (LastBufferTiming.EndTime < Timing.EndTime)
+		{
+			// Check if the end of the previous buffer finished before the start of this one
+			if (LastBufferTiming.EndTime > Timing.StartTime)
+			{
+				// Segment from end of last buffer to end of current
+				RunningFrameTimeSeconds += Timing.EndTime - LastBufferTiming.EndTime;
+			}
+			else
+			{
+				// Full timing of this buffer
+				RunningFrameTimeSeconds += Timing.EndTime - Timing.StartTime;
+			}
+
+			LastBufferTiming = Timing;
+		}
+	}
+    
+	FrameStartGPUCycles = FirstStartTime * CyclesPerSecond;
+	FrameEndGPUCycles = LastBufferTiming.EndTime * CyclesPerSecond;
+    
+	uint64 FrameGPUTimeCycles = uint64(CyclesPerSecond * RunningFrameTimeSeconds);
+	FPlatformAtomics::AtomicStore_Relaxed((int32*)&GGPUFrameTime, int32(FrameGPUTimeCycles));
+	
+#if STATS
+	FPlatformAtomics::AtomicStore_Relaxed(&GAGXGPUWorkTime, FrameGPUTimeCycles);
+	int64 FrameIdleTimeCycles = int64(FrameEndGPUCycles - FrameStartGPUCycles - FrameGPUTimeCycles);
+	FPlatformAtomics::AtomicStore_Relaxed(&GAGXGPUIdleTime, FrameIdleTimeCycles);
+#endif //STATS
+}
+
+void FAGXGPUProfiler::RecordPresent(const mtlpp::CommandBuffer& Buffer)
+{
+	const CFTimeInterval GpuStartTimeSeconds = Buffer.GetGpuStartTime();
+	const CFTimeInterval GpuEndTimeSeconds = Buffer.GetGpuEndTime();
+	const double CyclesPerSecond = 1.0 / FPlatformTime::GetSecondsPerCycle();
+	uint64 StartTimeCycles = uint64(GpuStartTimeSeconds * CyclesPerSecond);
+	uint64 EndTimeCycles = uint64(GpuEndTimeSeconds * CyclesPerSecond);
+	int64 Time = int64(EndTimeCycles - StartTimeCycles);
+	FPlatformAtomics::AtomicStore_Relaxed(&GAGXPresentTime, Time);
+}
+// END WARNING
+
+IAGXStatsScope::~IAGXStatsScope()
+{
+	for (IAGXStatsScope* Stat : Children)
+	{
+		delete Stat;
+	}
+}
+
+FString IAGXStatsScope::GetJSONRepresentation(uint32 Pid)
+{
+	FString JSONOutput;
+	
+	{
+		if (GPUStartTime && GPUEndTime)
+		{
+			uint64 ChildStartCallTime = GPUStartTime;
+			uint64 ChildDrawCallTime = GPUEndTime - GPUStartTime;
+			
+			JSONOutput += FString::Printf(TEXT("{\"pid\":%d, \"tid\":%d, \"ph\": \"X\", \"name\": \"%s\", \"ts\": %llu, \"dur\": %llu, \"args\":{\"num_child\":%u}},\n"),
+				  Pid,
+				  GPUThreadIndex,
+				  *Name,
+				  ChildStartCallTime,
+				  ChildDrawCallTime,
+				  Children.Num()
+			  );
+		}
+	}
+	
+	if (CPUStartTime && CPUEndTime)
+	{
+		uint64 ChildStartCallTime = CPUStartTime;
+		uint64 ChildDrawCallTime = FMath::Max(CPUEndTime - CPUStartTime, 1llu);
+		
+		JSONOutput += FString::Printf(TEXT("{\"pid\":%d, \"tid\":%d, \"ph\": \"X\", \"name\": \"%s\", \"ts\": %llu, \"dur\": %llu, \"args\":{\"num_child\":%u}},\n"),
+			 Pid,
+			 CPUThreadIndex,
+			 *Name,
+			 ChildStartCallTime,
+			 ChildDrawCallTime,
+			 Children.Num()
+		);
+	}
+	
+	return JSONOutput;
+}
+
+FAGXCommandBufferStats::FAGXCommandBufferStats(mtlpp::CommandBuffer const& Buffer, uint64 InGPUThreadIndex)
+{
+	CmdBuffer = Buffer;
+	
+	Name = FString::Printf(TEXT("CommandBuffer: %p %s"), CmdBuffer.GetPtr(), *FString(CmdBuffer.GetLabel().GetPtr()));
+	
+	CPUThreadIndex = FPlatformTLS::GetCurrentThreadId();
+	GPUThreadIndex = InGPUThreadIndex;
+	
+	Start(Buffer);
+}
+
+FAGXCommandBufferStats::~FAGXCommandBufferStats()
+{
+}
+
+void FAGXCommandBufferStats::Start(mtlpp::CommandBuffer const& Buffer)
+{
+	CPUStartTime = FPlatformTime::ToMilliseconds64(mach_absolute_time()) * 1000.0;
+	CPUEndTime = 0;
+	
+	GPUStartTime = 0;
+	GPUEndTime = 0;
+}
+
+void FAGXCommandBufferStats::End(mtlpp::CommandBuffer const& Buffer)
+{
+	check(Buffer.GetPtr() == CmdBuffer.GetPtr());
+	
+	bool const bTracing = FAGXProfiler::GetProfiler() && FAGXProfiler::GetProfiler()->TracingEnabled();
+	CmdBuffer.AddCompletedHandler(^(const mtlpp::CommandBuffer & InnerBuffer) {
+		const CFTimeInterval GpuTimeSeconds = InnerBuffer.GetGpuStartTime();
+		GPUStartTime = GpuTimeSeconds * 1000000.0;
+		
+		const CFTimeInterval GpuEndTimeSeconds = InnerBuffer.GetGpuEndTime();
+		GPUEndTime = GpuEndTimeSeconds * 1000000.0;
+	
+		if (bTracing)
+		{
+			FAGXProfiler::GetProfiler()->AddCommandBuffer(this);
+		}
+		else
+		{
+			delete this;
+		}
+	});
+	
+	CPUEndTime = FPlatformTime::ToMilliseconds64(mach_absolute_time()) * 1000.0;
+}
+
+
+#pragma mark -- FAGXProfiler
+// ----------------------------------------------------------------
+
+
+FAGXProfiler* FAGXProfiler::Self = nullptr;
+static FAGXViewportPresentHandler PresentHandler = ^(uint32 DisplayID, double OutputSeconds, double OutputDuration){
+	FAGXProfiler* Profiler = FAGXProfiler::GetProfiler();
+	Profiler->AddDisplayVBlank(DisplayID, OutputSeconds, OutputDuration);
+};
+
+FAGXDisplayStats::FAGXDisplayStats(uint32 DisplayID, double OutputSeconds, double Duration)
+{
+	Name = TEXT("V-Blank");
+	
+	CPUThreadIndex = FPlatformTLS::GetCurrentThreadId();
+	GPUThreadIndex = DisplayID;
+	
+	CPUStartTime = FPlatformTime::ToMilliseconds64(mach_absolute_time()) * 1000.0;
+	CPUEndTime = CPUStartTime+1;
+	
+	GPUStartTime = OutputSeconds * 1000000.0;
+	GPUEndTime = GPUStartTime + (Duration * 1000000.0);
+}
+FAGXDisplayStats::~FAGXDisplayStats()
+{
+}
+
+void FAGXDisplayStats::Start(mtlpp::CommandBuffer const& Buffer)
+{
+}
+void FAGXDisplayStats::End(mtlpp::CommandBuffer const& Buffer)
+{
+}
+
+FAGXCPUStats::FAGXCPUStats(FString const& InName)
+{
+	Name = InName;
+	
+	CPUThreadIndex = 0;
+	GPUThreadIndex = 0;
+	
+	CPUStartTime = 0;
+	CPUEndTime = 0;
+	
+	GPUStartTime = 0;
+	GPUEndTime = 0;
+}
+FAGXCPUStats::~FAGXCPUStats()
+{
+	
+}
+
+void FAGXCPUStats::Start(void)
+{
+	CPUThreadIndex = FPlatformTLS::GetCurrentThreadId();
+	CPUStartTime = FPlatformTime::ToMilliseconds64(mach_absolute_time()) * 1000.0;
+	
+}
+void FAGXCPUStats::End(void)
+{
+	CPUEndTime = FPlatformTime::ToMilliseconds64(mach_absolute_time()) * 1000.0;
+}
+
+void FAGXCPUStats::Start(mtlpp::CommandBuffer const& Buffer)
+{
+	
+}
+void FAGXCPUStats::End(mtlpp::CommandBuffer const& Buffer)
+{
+	
+}
+
+void FAGXProfiler::AddDisplayVBlank(uint32 DisplayID, double OutputSeconds, double OutputDuration)
+{
+	if (GIsRHIInitialized && bEnabled)
+	{
+		FScopeLock Lock(&Mutex);
+		DisplayStats.Add(new FAGXDisplayStats(DisplayID, OutputSeconds, OutputDuration));
+	}
+}
+
+FAGXProfiler::FAGXProfiler(FAGXContext* Context)
+: FAGXGPUProfiler(Context)
+, bEnabled(false)
+{
+	NumFramesToCapture = -1;
+	CaptureFrameNumber = 0;
+	
+	bRequestStartCapture = false;
+	bRequestStopCapture = false;
+	
+	if (FPlatformRHIFramePacer::IsEnabled())
+	{
+		FPlatformRHIFramePacer::AddHandler(PresentHandler);
+	}
+}
+
+FAGXProfiler::~FAGXProfiler()
+{
+	check(bEnabled == false);
+	if (FPlatformRHIFramePacer::IsEnabled())
+	{
+		FPlatformRHIFramePacer::RemoveHandler(PresentHandler);
+	}
+}
+
+FAGXProfiler* FAGXProfiler::CreateProfiler(FAGXContext *InContext)
+{
+	if (!Self)
+	{
+		Self = new FAGXProfiler(InContext);
+		
+		int32 CaptureFrames = 0;
+		if (FParse::Value(FCommandLine::Get(), TEXT("AGXProfileFrames="), CaptureFrames))
+		{
+			Self->BeginCapture(CaptureFrames);
+		}
+	}
+	return Self;
+}
+
+FAGXProfiler* FAGXProfiler::GetProfiler()
+{
+	return Self;
+}
+
+void FAGXProfiler::DestroyProfiler()
+{
+	delete Self;
+	Self = nullptr;
+}
+
+void FAGXProfiler::BeginCapture(int InNumFramesToCapture)
+{
+	check(IsInGameThread());
+	
+	NumFramesToCapture = InNumFramesToCapture;
+	CaptureFrameNumber = 0;
+	
+	bRequestStartCapture = true;
+}
+
+void FAGXProfiler::EndCapture()
+{
+	bRequestStopCapture = true;
+}
+
+bool FAGXProfiler::TracingEnabled() const
+{
+	return bEnabled;
+}
+
+void FAGXProfiler::BeginFrame()
+{
+	if (AGXGPUProfilerIsInSafeThread())
+	{
+		if (bRequestStartCapture && !bEnabled)
+		{
+			bEnabled = true;
+			bRequestStartCapture = false;
+		}
+	}
+	
+	FAGXGPUProfiler::BeginFrame();
+	
+	if (AGXGPUProfilerIsInSafeThread() && GetEmitDrawEvents())
+	{
+		PushEvent(TEXT("FRAME"), FColor(0, 255, 0, 255));
+	}
+}
+
+void FAGXProfiler::EndFrame()
+{
+	if (AGXGPUProfilerIsInSafeThread() && GetEmitDrawEvents())
+	{
+		PopEvent();
+	}
+	
+	FAGXGPUProfiler::EndFrame();
+	
+	if (AGXGPUProfilerIsInSafeThread() && bEnabled)
+	{
+		CaptureFrameNumber++;
+		if (bRequestStopCapture || (NumFramesToCapture > 0 && CaptureFrameNumber >= NumFramesToCapture))
+		{
+			bRequestStopCapture = false;
+			NumFramesToCapture = -1;
+			bEnabled = false;
+			SaveTrace();
+		}
+	}
+}
+
+void FAGXProfiler::EncodeDraw(FAGXCommandBufferStats* CmdBufStats, char const* DrawCall, uint32 RHIPrimitives, uint32 RHIVertices, uint32 RHIInstances)
+{
+	if (AGXGPUProfilerIsInSafeThread())
+	{
+		FAGXGPUProfiler::RegisterGPUWork(RHIPrimitives, RHIVertices);
+	}
+}
+
+void FAGXProfiler::EncodeBlit(FAGXCommandBufferStats* CmdBufStats, char const* DrawCall)
+{
+	if (AGXGPUProfilerIsInSafeThread())
+	{
+		FAGXGPUProfiler::RegisterGPUWork(1, 1);
+	}
+}
+
+void FAGXProfiler::EncodeBlit(FAGXCommandBufferStats* CmdBufStats, FString DrawCall)
+{
+	if (AGXGPUProfilerIsInSafeThread())
+	{
+		FAGXGPUProfiler::RegisterGPUWork(1, 1);
+	}
+}
+
+void FAGXProfiler::EncodeDispatch(FAGXCommandBufferStats* CmdBufStats, char const* DrawCall)
+{
+	if (AGXGPUProfilerIsInSafeThread())
+	{
+		FAGXGPUProfiler::RegisterGPUWork(1, 1);
+	}
+}
+
+FAGXCPUStats* FAGXProfiler::AddCPUStat(FString const& Name)
+{
+	if (GIsRHIInitialized && bEnabled)
+	{
+		FScopeLock Lock(&Mutex);
+		FAGXCPUStats* Stat = new FAGXCPUStats(Name);
+		CPUStats.Add(Stat);
+		return Stat;
+	}
+	else
+	{
+		return nullptr;
+	}
+}
+
+FAGXCommandBufferStats* FAGXProfiler::AllocateCommandBuffer(const mtlpp::CommandBuffer &Buffer, uint64 GPUThreadIndex)
+{
+	return new FAGXCommandBufferStats(Buffer, GPUThreadIndex);
+}
+
+void FAGXProfiler::AddCommandBuffer(FAGXCommandBufferStats *CommandBuffer)
+{
+	if (GIsRHIInitialized)
+	{
+		FScopeLock Lock(&Mutex);
+		TracedBuffers.Add(CommandBuffer);
+	}
+	else
+	{
+		delete CommandBuffer;
+	}
+}
+
+void FAGXProfiler::PushEvent(const TCHAR *Name, FColor Color)
+{
+	FAGXGPUProfiler::PushEvent(Name, Color);
+}
+
+void FAGXProfiler::PopEvent()
+{
+	FAGXGPUProfiler::PopEvent();
+}
+
+void FAGXProfiler::SaveTrace()
+{
+	Context->SubmitCommandBufferAndWait();
+	{
+		FScopeLock Lock(&Mutex);
+		
+		TSet<uint32> ThreadIDs;
+		
+		for (FAGXCommandBufferStats* CmdBufStats : TracedBuffers)
+		{
+			ThreadIDs.Add(CmdBufStats->CPUThreadIndex);
+			
+			for (IAGXStatsScope* ES : CmdBufStats->Children)
+			{
+				ThreadIDs.Add(ES->CPUThreadIndex);
+				
+				for (IAGXStatsScope* DS : ES->Children)
+				{
+					ThreadIDs.Add(DS->CPUThreadIndex);
+				}
+			}
+		}
+		
+		TSet<uint32> Displays;
+		for (FAGXDisplayStats* DisplayStat : DisplayStats)
+		{
+			ThreadIDs.Add(DisplayStat->CPUThreadIndex);
+			Displays.Add(DisplayStat->GPUThreadIndex);
+		}
+		
+		for (FAGXCPUStats* CPUStat : CPUStats)
+		{
+			ThreadIDs.Add(CPUStat->CPUThreadIndex);
+		}
+		
+		FString Filename = FString::Printf(TEXT("Profile(%s)"), *FDateTime::Now().ToString(TEXT("%Y%m%d_%H%M%S")));
+		FString TracingRootPath = FPaths::ProfilingDir() + TEXT("Traces/");
+		FString OutputFilename = TracingRootPath + Filename + TEXT(".json");
+		
+		FArchive* OutputFile = IFileManager::Get().CreateFileWriter(*OutputFilename);
+		
+		WriteString(OutputFile, R"({"traceEvents":[)" "\n");
+		
+		int32 SortIndex = 0; // Lower numbers result in higher position in the visualizer.
+		const uint32 Pid = FPlatformProcess::GetCurrentProcessId();
+		
+		for (int32 GPUIndex = 0; GPUIndex <= 0/*MaxGPUIndex*/; ++GPUIndex)
+		{
+			FString Output = FString::Printf(TEXT("{\"pid\":%d, \"tid\":%d, \"ph\": \"M\", \"name\": \"thread_name\", \"args\":{\"name\":\"GPU %d Command Buffers\"}},{\"pid\":%d, \"tid\":%d, \"ph\": \"M\", \"name\": \"thread_sort_index\", \"args\":{\"sort_index\": %d}},\n"),
+											 Pid, GPUIndex, GPUIndex, Pid, GPUIndex, SortIndex
+											 );
+			
+			WriteString(OutputFile, TCHAR_TO_UTF8(*Output));
+			SortIndex++;
+			
+			Output = FString::Printf(TEXT("{\"pid\":%d, \"tid\":%d, \"ph\": \"M\", \"name\": \"thread_name\", \"args\":{\"name\":\"GPU %d Operations\"}},{\"pid\":%d, \"tid\":%d, \"ph\": \"M\", \"name\": \"thread_sort_index\", \"args\":{\"sort_index\": %d}},\n"),
+											 Pid, GPUIndex+SortIndex, GPUIndex, Pid, GPUIndex+SortIndex, SortIndex
+											 );
+			
+			WriteString(OutputFile, TCHAR_TO_UTF8(*Output));
+			SortIndex++;
+			
+			Output = FString::Printf(TEXT("{\"pid\":%d, \"tid\":%d, \"ph\": \"M\", \"name\": \"thread_name\", \"args\":{\"name\":\"Render Events %d\"}},{\"pid\":%d, \"tid\":%d, \"ph\": \"M\", \"name\": \"thread_sort_index\", \"args\":{\"sort_index\": %d}},\n"),
+									 Pid, GPUIndex+SortIndex, GPUIndex, Pid, GPUIndex+SortIndex, SortIndex
+									 );
+			
+			WriteString(OutputFile, TCHAR_TO_UTF8(*Output));
+			SortIndex++;
+			
+			Output = FString::Printf(TEXT("{\"pid\":%d, \"tid\":%d, \"ph\": \"M\", \"name\": \"thread_name\", \"args\":{\"name\":\"Driver Stats %d\"}},{\"pid\":%d, \"tid\":%d, \"ph\": \"M\", \"name\": \"thread_sort_index\", \"args\":{\"sort_index\": %d}},\n"),
+									 Pid, GPUIndex+SortIndex, GPUIndex, Pid, GPUIndex+SortIndex, SortIndex
+									 );
+			
+			WriteString(OutputFile, TCHAR_TO_UTF8(*Output));
+			SortIndex++;
+			
+			for (uint32 Display : Displays)
+			{
+				Output = FString::Printf(TEXT("{\"pid\":%d, \"tid\":%d, \"ph\": \"M\", \"name\": \"thread_name\", \"args\":{\"name\":\"Display %d\"}},{\"pid\":%d, \"tid\":%d, \"ph\": \"M\", \"name\": \"thread_sort_index\", \"args\":{\"sort_index\": %d}},\n"),
+										 Pid, Display + SortIndex, SortIndex - 3, Pid, Display + SortIndex, SortIndex
+										 );
+				
+				WriteString(OutputFile, TCHAR_TO_UTF8(*Output));
+				SortIndex++;
+			}
+		}
+		
+		static const uint32 BufferSize = 128;
+		char Buffer[BufferSize];
+		for (uint32 CPUIndex : ThreadIDs)
+		{
+			bool bThreadName = false;
+			pthread_t PThread = pthread_from_mach_thread_np((mach_port_t)CPUIndex);
+			if (PThread)
+			{
+				if (!pthread_getname_np(PThread,Buffer,BufferSize))
+				{
+					bThreadName = true;
+				}
+			}
+			if (!bThreadName)
+			{
+				sprintf(Buffer, "Thread %d", CPUIndex);
+			}
+			
+			FString Output = FString::Printf(TEXT("{\"pid\":%d, \"tid\":%d, \"ph\": \"M\", \"name\": \"thread_name\", \"args\":{\"name\":\"%s\"}},{\"pid\":%d, \"tid\":%d, \"ph\": \"M\", \"name\": \"thread_sort_index\", \"args\":{\"sort_index\": %d}},\n"),
+											 Pid, CPUIndex, UTF8_TO_TCHAR(Buffer), Pid, CPUIndex, SortIndex
+											 );
+			
+			
+			WriteString(OutputFile, TCHAR_TO_UTF8(*Output));
+			SortIndex++;
+		}
+		
+		for (FAGXCommandBufferStats* CmdBufStats : TracedBuffers)
+		{
+			WriteString(OutputFile, TCHAR_TO_UTF8(*CmdBufStats->GetJSONRepresentation(Pid)));
+			
+			for (IAGXStatsScope* ES : CmdBufStats->Children)
+			{
+				WriteString(OutputFile, TCHAR_TO_UTF8(*ES->GetJSONRepresentation(Pid)));
+				
+				uint64 PrevTime = ES->GPUStartTime;
+				for (IAGXStatsScope* DS : ES->Children)
+				{
+					WriteString(OutputFile, TCHAR_TO_UTF8(*DS->GetJSONRepresentation(Pid)));
+					if (!DS->GPUStartTime)
+					{
+						DS->GPUStartTime = FMath::Max(PrevTime, DS->GPUStartTime);
+						DS->GPUEndTime = DS->GPUStartTime + 1llu;
+						WriteString(OutputFile, TCHAR_TO_UTF8(*DS->GetJSONRepresentation(Pid)));
+					}
+					PrevTime = DS->GPUEndTime;
+				}
+			}
+			
+			delete CmdBufStats;
+		}
+		TracedBuffers.Empty();
+		
+		for (FAGXDisplayStats* DisplayStat : DisplayStats)
+		{
+			DisplayStat->GPUThreadIndex += 3;
+			WriteString(OutputFile, TCHAR_TO_UTF8(*DisplayStat->GetJSONRepresentation(Pid)));
+			delete DisplayStat;
+		}
+		DisplayStats.Empty();
+		
+		for (FAGXCPUStats* CPUStat : CPUStats)
+		{
+			WriteString(OutputFile, TCHAR_TO_UTF8(*CPUStat->GetJSONRepresentation(Pid)));
+			delete CPUStat;
+		}
+		CPUStats.Empty();
+		
+		// All done
+		
+		WriteString(OutputFile, "{}]}");
+		
+		OutputFile->Close();
+	}
+}
+
+static void HandleAGXProfileCommand(const TArray<FString>& Args, UWorld*, FOutputDevice& Ar)
+{
+	if (Args.Num() < 1)
+	{
+		return;
+	}
+	FString Param = Args[0];
+	if (Param == TEXT("START"))
+	{
+		FAGXProfiler::GetProfiler()->BeginCapture();
+	}
+	else if (Param == TEXT("STOP"))
+	{
+		FAGXProfiler::GetProfiler()->EndCapture();
+	}
+	else
+	{
+		int32 CaptureFrames = 0;
+		if (FParse::Value(*Param, TEXT("FRAMES="), CaptureFrames))
+		{
+			FAGXProfiler::GetProfiler()->BeginCapture(CaptureFrames);
+		}
+	}
+}
+
+static FAutoConsoleCommand HandleAGXProfilerCmd(
+	TEXT("AGXProfiler"),
+	TEXT("Starts or stops AGX profiler"),
+	FConsoleCommandWithWorldArgsAndOutputDeviceDelegate::CreateStatic(&HandleAGXProfileCommand)
+);
diff --git a/Engine/Source/Runtime/Apple/AGXRHI/Private/AGXProfiler.h b/Engine/Source/Runtime/Apple/AGXRHI/Private/AGXProfiler.h
new file mode 100644
--- /dev/null
+++ b/Engine/Source/Runtime/Apple/AGXRHI/Private/AGXProfiler.h
@@ -0,0 +1,374 @@
+// Copyright Epic Games, Inc. All Rights Reserved.
+
+#pragma once
+
+#include "AGXRHIPrivate.h"
+#include "AGXCommandQueue.h"
+#include "GPUProfiler.h"
+
+// Stats
+DECLARE_CYCLE_STAT_EXTERN(TEXT("MakeDrawable time"),STAT_AGXMakeDrawableTime,STATGROUP_AGXRHI, );
+DECLARE_CYCLE_STAT_EXTERN(TEXT("Draw call time"),STAT_AGXDrawCallTime,STATGROUP_AGXRHI, );
+DECLARE_CYCLE_STAT_EXTERN(TEXT("PrepareDraw time"),STAT_AGXPrepareDrawTime,STATGROUP_AGXRHI, );
+DECLARE_CYCLE_STAT_EXTERN(TEXT("SwitchToRender time"),STAT_AGXSwitchToRenderTime,STATGROUP_AGXRHI, );
+DECLARE_CYCLE_STAT_EXTERN(TEXT("SwitchToCompute time"),STAT_AGXSwitchToComputeTime,STATGROUP_AGXRHI, );
+DECLARE_CYCLE_STAT_EXTERN(TEXT("SwitchToBlit time"),STAT_AGXSwitchToBlitTime,STATGROUP_AGXRHI, );
+DECLARE_CYCLE_STAT_EXTERN(TEXT("SwitchToAsyncBlit time"),STAT_AGXSwitchToAsyncBlitTime,STATGROUP_AGXRHI, );
+DECLARE_CYCLE_STAT_EXTERN(TEXT("PrepareToRender time"),STAT_AGXPrepareToRenderTime,STATGROUP_AGXRHI, );
+DECLARE_CYCLE_STAT_EXTERN(TEXT("PrepareToDispatch time"),STAT_AGXPrepareToDispatchTime,STATGROUP_AGXRHI, );
+DECLARE_CYCLE_STAT_EXTERN(TEXT("CommitRenderResourceTables time"),STAT_AGXCommitRenderResourceTablesTime,STATGROUP_AGXRHI, );
+DECLARE_CYCLE_STAT_EXTERN(TEXT("SetRenderState time"),STAT_AGXSetRenderStateTime,STATGROUP_AGXRHI, );
+DECLARE_CYCLE_STAT_EXTERN(TEXT("SetRenderPipelineState time"),STAT_AGXSetRenderPipelineStateTime,STATGROUP_AGXRHI, );
+DECLARE_CYCLE_STAT_EXTERN(TEXT("PipelineState time"),STAT_AGXPipelineStateTime,STATGROUP_AGXRHI, );
+DECLARE_CYCLE_STAT_EXTERN(TEXT("Buffer Page-Off time"), STAT_AGXBufferPageOffTime, STATGROUP_AGXRHI, );
+DECLARE_CYCLE_STAT_EXTERN(TEXT("Texture Page-Off time"), STAT_AGXTexturePageOffTime, STATGROUP_AGXRHI, );
+
+DECLARE_DWORD_COUNTER_STAT_EXTERN(TEXT("Uniform Memory Allocated Per-Frame"), STAT_AGXUniformMemAlloc, STATGROUP_AGXRHI, );
+DECLARE_DWORD_COUNTER_STAT_EXTERN(TEXT("Uniform Memory Freed Per-Frame"), STAT_AGXUniformMemFreed, STATGROUP_AGXRHI, );
+DECLARE_DWORD_COUNTER_STAT_EXTERN(TEXT("Vertex Memory Allocated Per-Frame"), STAT_AGXVertexMemAlloc, STATGROUP_AGXRHI, );
+DECLARE_DWORD_COUNTER_STAT_EXTERN(TEXT("Vertex Memory Freed Per-Frame"), STAT_AGXVertexMemFreed, STATGROUP_AGXRHI, );
+DECLARE_DWORD_COUNTER_STAT_EXTERN(TEXT("Index Memory Allocated Per-Frame"), STAT_AGXIndexMemAlloc, STATGROUP_AGXRHI, );
+DECLARE_DWORD_COUNTER_STAT_EXTERN(TEXT("Index Memory Freed Per-Frame"), STAT_AGXIndexMemFreed, STATGROUP_AGXRHI, );
+DECLARE_DWORD_COUNTER_STAT_EXTERN(TEXT("Texture Memory Updated Per-Frame"), STAT_AGXTextureMemUpdate, STATGROUP_AGXRHI, );
+DECLARE_MEMORY_STAT_EXTERN(TEXT("Buffer Memory"), STAT_AGXBufferMemory, STATGROUP_AGXRHI, );
+DECLARE_MEMORY_STAT_EXTERN(TEXT("Texture Memory"), STAT_AGXTextureMemory, STATGROUP_AGXRHI, );
+DECLARE_MEMORY_STAT_EXTERN(TEXT("Heap Memory"), STAT_AGXHeapMemory, STATGROUP_AGXRHI, );
+DECLARE_MEMORY_STAT_EXTERN(TEXT("Unused Buffer Memory"), STAT_AGXBufferUnusedMemory, STATGROUP_AGXRHI, );
+DECLARE_MEMORY_STAT_EXTERN(TEXT("Unused Texture Memory"), STAT_AGXTextureUnusedMemory, STATGROUP_AGXRHI, );
+
+DECLARE_MEMORY_STAT_EXTERN(TEXT("Uniform Memory In Flight"), STAT_AGXUniformMemoryInFlight, STATGROUP_AGXRHI, );
+DECLARE_MEMORY_STAT_EXTERN(TEXT("Allocated Uniform Pool Memory"), STAT_AGXUniformAllocatedMemory, STATGROUP_AGXRHI, );
+DECLARE_MEMORY_STAT_EXTERN(TEXT("Uniform Memory Per Frame"), STAT_AGXUniformBytesPerFrame, STATGROUP_AGXRHI, );
+
+DECLARE_MEMORY_STAT_EXTERN(TEXT("General Frame Allocator Memory In Flight"), STAT_AGXFrameAllocatorMemoryInFlight, STATGROUP_AGXRHI, );
+DECLARE_MEMORY_STAT_EXTERN(TEXT("Allocated Frame Allocator Memory"), STAT_AGXFrameAllocatorAllocatedMemory, STATGROUP_AGXRHI, );
+DECLARE_MEMORY_STAT_EXTERN(TEXT("Frame Allocator Memory Per Frame"), STAT_AGXFrameAllocatorBytesPerFrame, STATGROUP_AGXRHI, );
+
+DECLARE_DWORD_ACCUMULATOR_STAT_EXTERN(TEXT("Buffer Count"), STAT_AGXBufferCount, STATGROUP_AGXRHI, );
+DECLARE_DWORD_ACCUMULATOR_STAT_EXTERN(TEXT("Texture Count"), STAT_AGXTextureCount, STATGROUP_AGXRHI, );
+DECLARE_DWORD_ACCUMULATOR_STAT_EXTERN(TEXT("Heap Count"), STAT_AGXHeapCount, STATGROUP_AGXRHI, );
+DECLARE_DWORD_ACCUMULATOR_STAT_EXTERN(TEXT("Fence Count"), STAT_AGXFenceCount, STATGROUP_AGXRHI, );
+
+DECLARE_CYCLE_STAT_EXTERN(TEXT("Texture Page-On time"), STAT_AGXTexturePageOnTime, STATGROUP_AGXRHI, );
+DECLARE_CYCLE_STAT_EXTERN(TEXT("GPU Work time"), STAT_AGXGPUWorkTime, STATGROUP_AGXRHI, );
+DECLARE_CYCLE_STAT_EXTERN(TEXT("GPU Idle time"), STAT_AGXGPUIdleTime, STATGROUP_AGXRHI, );
+DECLARE_CYCLE_STAT_EXTERN(TEXT("Present time"), STAT_AGXPresentTime, STATGROUP_AGXRHI, );
+DECLARE_CYCLE_STAT_EXTERN(TEXT("CustomPresent time"), STAT_AGXCustomPresentTime, STATGROUP_AGXRHI, );
+#if STATS
+extern int64 volatile GAGXTexturePageOnTime;
+extern int64 volatile GAGXGPUWorkTime;
+extern int64 volatile GAGXGPUIdleTime;
+extern int64 volatile GAGXPresentTime;
+#endif
+
+DECLARE_DWORD_COUNTER_STAT_EXTERN(TEXT("Number Command Buffers Created Per-Frame"), STAT_AGXCommandBufferCreatedPerFrame, STATGROUP_AGXRHI, );
+DECLARE_DWORD_COUNTER_STAT_EXTERN(TEXT("Number Command Buffers Committed Per-Frame"), STAT_AGXCommandBufferCommittedPerFrame, STATGROUP_AGXRHI, );
+
+/** A single perf event node, which tracks information about a appBeginDrawEvent/appEndDrawEvent range. */
+class FAGXEventNode : public FGPUProfilerEventNode
+{
+public:
+	
+	FAGXEventNode(FAGXContext* InContext, const TCHAR* InName, FGPUProfilerEventNode* InParent, bool bIsRoot, bool bInFullProfiling)
+	: FGPUProfilerEventNode(InName, InParent)
+	, StartTime(0)
+	, EndTime(0)
+	, Context(InContext)
+	, bRoot(bIsRoot)
+    , bFullProfiling(bInFullProfiling)
+	{
+	}
+	
+	virtual ~FAGXEventNode();
+	
+	/**
+	 * Returns the time in ms that the GPU spent in this draw event.
+	 * This blocks the CPU if necessary, so can cause hitching.
+	 */
+	virtual float GetTiming() override;
+	
+	virtual void StartTiming() override;
+	
+	virtual void StopTiming() override;
+	
+	mtlpp::CommandBufferHandler Start(void);
+	mtlpp::CommandBufferHandler Stop(void);
+
+	bool Wait() const { return bRoot && bFullProfiling; }
+	bool IsRoot() const { return bRoot; }
+	
+	uint64 GetCycles() { return EndTime - StartTime; }
+	
+	uint64 StartTime;
+	uint64 EndTime;
+private:
+	FAGXContext* Context;
+	bool bRoot;
+    bool bFullProfiling;
+};
+
+/** An entire frame of perf event nodes, including ancillary timers. */
+class FAGXEventNodeFrame : public FGPUProfilerEventNodeFrame
+{
+public:
+	FAGXEventNodeFrame(FAGXContext* InContext, bool bInFullProfiling)
+	: RootNode(new FAGXEventNode(InContext, TEXT("Frame"), nullptr, true, bInFullProfiling))
+    , bFullProfiling(bInFullProfiling)
+	{
+	}
+	
+	virtual ~FAGXEventNodeFrame()
+	{
+        if(bFullProfiling)
+        {
+            delete RootNode;
+        }
+	}
+	
+	/** Start this frame of per tracking */
+	virtual void StartFrame() override;
+	
+	/** End this frame of per tracking, but do not block yet */
+	virtual void EndFrame() override;
+	
+	/** Calculates root timing base frequency (if needed by this RHI) */
+	virtual float GetRootTimingResults() override;
+	
+	virtual void LogDisjointQuery() override;
+	
+	FAGXEventNode* RootNode;
+    bool bFullProfiling;
+};
+
+// This class has multiple inheritance but really FGPUTiming is a static class
+class FAGXGPUTiming : public FGPUTiming
+{
+public:
+	
+	/**
+	 * Constructor.
+	 */
+	FAGXGPUTiming()
+	{
+		StaticInitialize(nullptr, PlatformStaticInitialize);
+	}
+	
+	void SetCalibrationTimestamp(uint64 GPU, uint64 CPU)
+	{
+		FGPUTiming::SetCalibrationTimestamp({ GPU, CPU });
+	}
+	
+private:
+	
+	/**
+	 * Initializes the static variables, if necessary.
+	 */
+	static void PlatformStaticInitialize(void* UserData)
+	{
+		// Are the static variables initialized?
+		if ( !GAreGlobalsInitialized )
+		{
+			GIsSupported = true;
+			SetTimingFrequency(1000 * 1000 * 1000);
+			GAreGlobalsInitialized = true;
+		}
+	}
+};
+
+struct IAGXStatsScope
+{
+	FString Name;
+	FString Parent;
+	TArray<IAGXStatsScope*> Children;
+	
+	uint64 CPUStartTime;
+	uint64 CPUEndTime;
+	
+	uint64 GPUStartTime;
+	uint64 GPUEndTime;
+	
+	uint64 CPUThreadIndex;
+	uint64 GPUThreadIndex;
+	
+	virtual ~IAGXStatsScope();
+	
+	virtual void Start(mtlpp::CommandBuffer const& Buffer) = 0;
+	virtual void End(mtlpp::CommandBuffer const& Buffer) = 0;
+	
+	FString GetJSONRepresentation(uint32 PID);
+};
+
+struct FAGXCPUStats : public IAGXStatsScope
+{
+	FAGXCPUStats(FString const& Name);
+	virtual ~FAGXCPUStats();
+	
+	void Start(void);
+	void End(void);
+	
+	virtual void Start(mtlpp::CommandBuffer const& Buffer) final override;
+	virtual void End(mtlpp::CommandBuffer const& Buffer) final override;
+};
+
+struct FAGXDisplayStats : public IAGXStatsScope
+{
+	FAGXDisplayStats(uint32 DisplayID, double OutputSeconds, double Duration);
+	virtual ~FAGXDisplayStats();
+	
+	virtual void Start(mtlpp::CommandBuffer const& Buffer) final override;
+	virtual void End(mtlpp::CommandBuffer const& Buffer) final override;
+};
+
+enum EMTLFenceType
+{
+	EMTLFenceTypeWait,
+	EMTLFenceTypeUpdate,
+};
+
+struct FAGXCommandBufferStats : public IAGXStatsScope
+{
+	FAGXCommandBufferStats(mtlpp::CommandBuffer const& Buffer, uint64 GPUThreadIndex);
+	virtual ~FAGXCommandBufferStats();
+	
+	virtual void Start(mtlpp::CommandBuffer const& Buffer) final override;
+	virtual void End(mtlpp::CommandBuffer const& Buffer) final override;
+
+	ns::AutoReleased<mtlpp::CommandBuffer> CmdBuffer;
+};
+
+/**
+ * Simple struct to hold sortable command buffer start and end timestamps.
+ */
+struct FAGXCommandBufferTiming
+{
+	CFTimeInterval StartTime;
+	CFTimeInterval EndTime;
+
+	bool operator<(const FAGXCommandBufferTiming& RHS) const
+	{
+		// Sort by start time and then by length if the commandbuffer started at the same time
+		if (this->StartTime < RHS.StartTime)
+		{
+			return true;
+		}
+		else if ((this->StartTime == RHS.StartTime) && (this->EndTime > RHS.EndTime))
+		{
+			return true;
+		}
+		return false;
+	}
+};
+
+/**
+ * Encapsulates GPU profiling logic and data.
+ * There's only one global instance of this struct so it should only contain global data, nothing specific to a frame.
+ */
+struct FAGXGPUProfiler : public FGPUProfiler
+{
+	/** GPU hitch profile histories */
+	TIndirectArray<FAGXEventNodeFrame> GPUHitchEventNodeFrames;
+	
+	FAGXGPUProfiler(FAGXContext* InContext)
+	:	FGPUProfiler()
+	,	Context(InContext)
+	,   NumNestedFrames(0)
+	{}
+	
+	virtual ~FAGXGPUProfiler() {}
+	
+	virtual FGPUProfilerEventNode* CreateEventNode(const TCHAR* InName, FGPUProfilerEventNode* InParent) override;
+	
+	void Cleanup();
+	
+	virtual void PushEvent(const TCHAR* Name, FColor Color) override;
+	virtual void PopEvent() override;
+	
+	void BeginFrame();
+	void EndFrame();
+	
+	// WARNING:
+	// These functions MUST be called from within Metal scheduled/completion handlers
+	// since they depend on libdispatch to enforce ordering.
+	static void RecordFrame(TArray<FAGXCommandBufferTiming>& CommandBufferTimings, FAGXCommandBufferTiming& LastPresentBufferTiming);
+	static void RecordPresent(const mtlpp::CommandBuffer& Buffer);
+	// END WARNING
+	
+	FAGXGPUTiming TimingSupport;
+	FAGXContext* Context;
+	int32 NumNestedFrames;
+};
+
+class FAGXProfiler : public FAGXGPUProfiler
+{
+	static FAGXProfiler* Self;
+public:
+	FAGXProfiler(FAGXContext* InContext);
+	~FAGXProfiler();
+	
+	static FAGXProfiler* CreateProfiler(FAGXContext* InContext);
+	static FAGXProfiler* GetProfiler();
+	static void DestroyProfiler();
+	
+	void BeginCapture(int InNumFramesToCapture = -1);
+	void EndCapture();
+	bool TracingEnabled() const;
+	
+	void BeginFrame();
+	void EndFrame();
+	
+	void AddDisplayVBlank(uint32 DisplayID, double OutputSeconds, double OutputDuration);
+	
+	void EncodeDraw(FAGXCommandBufferStats* CmdBufStats, char const* DrawCall, uint32 RHIPrimitives, uint32 RHIVertices, uint32 RHIInstances);
+	void EncodeBlit(FAGXCommandBufferStats* CmdBufStats, char const* DrawCall);
+	void EncodeBlit(FAGXCommandBufferStats* CmdBufStats, FString DrawCall);
+	void EncodeDispatch(FAGXCommandBufferStats* CmdBufStats, char const* DrawCall);
+	
+	FAGXCPUStats* AddCPUStat(FString const& Name);
+	FAGXCommandBufferStats* AllocateCommandBuffer(mtlpp::CommandBuffer const& Buffer, uint64 GPUThreadIndex);
+	void AddCommandBuffer(FAGXCommandBufferStats* CommandBuffer);
+	virtual void PushEvent(const TCHAR* Name, FColor Color) final override;
+	virtual void PopEvent() final override;
+	
+	void SaveTrace();
+	
+private:
+	FCriticalSection Mutex;
+	
+	TArray<FAGXCommandBufferStats*> TracedBuffers;
+	TArray<FAGXDisplayStats*> DisplayStats;
+	TArray<FAGXCPUStats*> CPUStats;
+	
+	int32 NumFramesToCapture;
+	int32 CaptureFrameNumber;
+	
+	bool bRequestStartCapture;
+	bool bRequestStopCapture;
+	bool bEnabled;
+};
+
+struct FAGXScopedCPUStats
+{
+	FAGXScopedCPUStats(FString const& Name)
+	: Stats(nullptr)
+	{
+		FAGXProfiler* Profiler = FAGXProfiler::GetProfiler();
+		if (Profiler)
+		{
+			Stats = Profiler->AddCPUStat(Name);
+			if (Stats)
+			{
+				Stats->Start();
+			}
+		}
+	}
+	
+	~FAGXScopedCPUStats()
+	{
+		if (Stats)
+		{
+			Stats->End();
+		}
+	}
+	
+	FAGXCPUStats* Stats;
+};
