Initial import of AGXRHI

diff --git a/Engine/Source/Runtime/Apple/AGXRHI/Private/AGXRenderCommandEncoder.cpp b/Engine/Source/Runtime/Apple/AGXRHI/Private/AGXRenderCommandEncoder.cpp
new file mode 100644
--- /dev/null
+++ b/Engine/Source/Runtime/Apple/AGXRHI/Private/AGXRenderCommandEncoder.cpp
@@ -0,0 +1,780 @@
+// Copyright Epic Games, Inc. All Rights Reserved.
+
+/*=============================================================================
+	AGXRenderCommandEncoder.cpp: AGX RHI command encoder wrapper.
+=============================================================================*/
+
+#include "AGXRHIPrivate.h"
+
+#include "AGXRenderCommandEncoder.h"
+#include "AGXCommandBuffer.h"
+#include "AGXCommandQueue.h"
+#include "AGXFence.h"
+#include "AGXPipeline.h"
+
+#if MTLPP_CONFIG_VALIDATE && METAL_DEBUG_OPTIONS
+NS_ASSUME_NONNULL_BEGIN
+
+@interface FAGXDebugRenderCommandEncoder : FAGXDebugCommandEncoder
+{
+@public
+#pragma mark - Private Member Variables -
+#if METAL_DEBUG_OPTIONS
+	FAGXDebugShaderResourceMask ResourceMask[EAGXShaderRenderNum];
+	FAGXDebugBufferBindings ShaderBuffers[EAGXShaderRenderNum];
+	FAGXDebugTextureBindings ShaderTextures[EAGXShaderRenderNum];
+	FAGXDebugSamplerBindings ShaderSamplers[EAGXShaderRenderNum];
+	id<MTLRenderPipelineState> DebugState;
+#endif
+	MTLRenderPassDescriptor* RenderPassDesc;
+	id<MTLRenderCommandEncoder> Inner;
+	FAGXCommandBufferDebugging Buffer;
+	FAGXShaderPipeline* Pipeline;
+}
+
+/** Initialise the wrapper with the provided command-buffer. */
+-(id)initWithEncoder:(id<MTLRenderCommandEncoder>)Encoder fromDescriptor:(MTLRenderPassDescriptor*)RenderPassDesc andCommandBuffer:(FAGXCommandBufferDebugging const&)Buffer;
+
+@end
+
+@interface FAGXDebugParallelRenderCommandEncoder : FApplePlatformObject
+{
+@public
+	TArray<FAGXRenderCommandEncoderDebugging> RenderEncoders;
+	MTLRenderPassDescriptor* RenderPassDesc;
+	id<MTLParallelRenderCommandEncoder> Inner;
+	FAGXCommandBufferDebugging Buffer;
+}
+-(id)initWithEncoder:(id<MTLParallelRenderCommandEncoder>)Encoder fromDescriptor:(MTLRenderPassDescriptor*)RenderPassDesc andCommandBuffer:(FAGXCommandBufferDebugging const&)Buffer;
+@end
+
+
+#if METAL_DEBUG_OPTIONS
+static NSString* GAGXDebugVertexShader = @"#include <metal_stdlib>\n"
+@"using namespace metal;\n"
+@"struct VertexInput\n"
+@"{\n"
+@"};\n"
+@"vertex void WriteCommandIndexVS(VertexInput StageIn [[stage_in]], constant uint* Input [[ buffer(0) ]], device uint* Output  [[ buffer(1) ]])\n"
+@"{\n"
+@"	Output[0] = Input[0];\n"
+@"}\n";
+
+static id <MTLRenderPipelineState> GetDebugVertexShaderState(id<MTLDevice> Device, MTLRenderPassDescriptor* PassDesc)
+{
+	static id<MTLFunction> Func = nil;
+	static FCriticalSection Mutex;
+	static NSMutableDictionary* Dict = [NSMutableDictionary new];
+	if(!Func)
+	{
+		id<MTLLibrary> Lib = [Device newLibraryWithSource:GAGXDebugVertexShader options:nil error:nullptr];
+		check(Lib);
+		Func = [Lib newFunctionWithName:@"WriteCommandIndexVS"];
+		check(Func);
+		[Lib release];
+	}
+	
+	FScopeLock Lock(&Mutex);
+	id<MTLRenderPipelineState> State = [Dict objectForKey:PassDesc];
+	if (!State)
+	{
+		MTLRenderPipelineDescriptor* Desc = [MTLRenderPipelineDescriptor new];
+		
+		Desc.vertexFunction = Func;
+		
+		if (PassDesc.depthAttachment)
+		{
+			Desc.depthAttachmentPixelFormat = PassDesc.depthAttachment.texture.pixelFormat;
+		}
+		if (PassDesc.stencilAttachment)
+		{
+			Desc.stencilAttachmentPixelFormat = PassDesc.stencilAttachment.texture.pixelFormat;
+		}
+		if (PassDesc.colorAttachments)
+		{
+			for (NSUInteger i = 0; i < 8; i++)
+			{
+				MTLRenderPassColorAttachmentDescriptor* CD = [PassDesc.colorAttachments objectAtIndexedSubscript:i];
+				if (CD.texture.pixelFormat != MTLPixelFormatInvalid)
+				{
+					MTLRenderPipelineColorAttachmentDescriptor* CD0 = [[MTLRenderPipelineColorAttachmentDescriptor new] autorelease];
+					CD0.pixelFormat = CD.texture.pixelFormat;
+					[Desc.colorAttachments setObject:CD0 atIndexedSubscript:i];
+				}
+			}
+		}
+		Desc.rasterizationEnabled = false;
+		
+		State = [[Device newRenderPipelineStateWithDescriptor:Desc error:nil] autorelease];
+		check(State);
+		
+		[Dict setObject:State forKey:PassDesc];
+		
+		[Desc release];
+	}
+	check(State);
+	return State;
+}
+#endif
+
+@implementation FAGXDebugParallelRenderCommandEncoder
+
+APPLE_PLATFORM_OBJECT_ALLOC_OVERRIDES(FAGXDebugParallelRenderCommandEncoder)
+
+-(id)initWithEncoder:(id<MTLParallelRenderCommandEncoder>)Encoder fromDescriptor:(MTLRenderPassDescriptor*)Desc andCommandBuffer:(FAGXCommandBufferDebugging const&)SourceBuffer
+{
+	id Self = [super init];
+	if (Self)
+	{
+        Inner = Encoder;
+		Buffer = SourceBuffer;
+		RenderPassDesc = [Desc retain];
+	}
+	return Self;
+}
+
+-(void)dealloc
+{
+	[RenderPassDesc release];
+	[super dealloc];
+}
+
+@end
+
+@implementation FAGXDebugRenderCommandEncoder
+
+APPLE_PLATFORM_OBJECT_ALLOC_OVERRIDES(FAGXDebugRenderCommandEncoder)
+
+-(id)initWithEncoder:(id<MTLRenderCommandEncoder>)Encoder fromDescriptor:(MTLRenderPassDescriptor*)Desc andCommandBuffer:(FAGXCommandBufferDebugging const&)SourceBuffer
+{
+	id Self = [super init];
+	if (Self)
+	{
+        Inner = Encoder;
+		Buffer = SourceBuffer;
+		RenderPassDesc = [Desc retain];
+#if METAL_DEBUG_OPTIONS
+		DebugState = Buffer.GetPtr()->DebugLevel >= EAGXDebugLevelValidation ? [GetDebugVertexShaderState(Buffer.GetPtr()->InnerBuffer.device, Desc) retain] : nil;
+#endif
+        Pipeline = nil;
+	}
+	return Self;
+}
+
+-(void)dealloc
+{
+	[RenderPassDesc release];
+#if METAL_DEBUG_OPTIONS
+	[DebugState release];
+#endif
+	[Pipeline release];
+	[super dealloc];
+}
+
+@end
+
+NS_ASSUME_NONNULL_END
+
+void FAGXRenderCommandEncoderDebugging::InsertDebugDraw()
+{
+#if METAL_DEBUG_OPTIONS
+//	switch (((FAGXDebugRenderCommandEncoder*)m_ptr)->Buffer.GetPtr()->DebugLevel)
+//	{
+//		case EAGXDebugLevelConditionalSubmit:
+//		case EAGXDebugLevelWaitForComplete:
+//		case EAGXDebugLevelLogOperations:
+//		case EAGXDebugLevelValidation:
+//		{
+//			uint32 const Index = Buffer->DebugCommands.Num();
+//#if PLATFORM_MAC
+//			[Inner textureBarrier];
+//#endif
+//			[Inner setVertexBytes:&Index length:sizeof(Index) atIndex:0];
+//			[Inner setVertexBuffer:Buffer->DebugInfoBuffer offset:0 atIndex:1];
+//			[Inner setRenderPipelineState:DebugState];
+//			[Inner drawPrimitives:mtlpp::PrimitiveTypePoint vertexStart:0 vertexCount:1];
+//#if PLATFORM_MAC
+//			[Inner textureBarrier];
+//#endif
+//			if (Pipeline && Pipeline->RenderPipelineState)
+//			{
+//				[Inner setRenderPipelineState:Pipeline->RenderPipelineState];
+//			}
+//			
+//			if (ShaderBuffers[EAGXShaderVertex].Buffers[0])
+//			{
+//				[Inner setVertexBuffer:ShaderBuffers[EAGXShaderVertex].Buffers[0] offset:ShaderBuffers[EAGXShaderVertex].Offsets[0] atIndex:0];
+//			}
+//			else if (ShaderBuffers[EAGXShaderVertex].Bytes[0])
+//			{
+//				[Inner setVertexBytes:ShaderBuffers[EAGXShaderVertex].Bytes[0] length:ShaderBuffers[EAGXShaderVertex].Offsets[0] atIndex:0];
+//			}
+//			
+//			if (ShaderBuffers[EAGXShaderVertex].Buffers[1])
+//			{
+//				[Inner setVertexBuffer:ShaderBuffers[EAGXShaderVertex].Buffers[1] offset:ShaderBuffers[EAGXShaderVertex].Offsets[1] atIndex:1];
+//			}
+//			else if (ShaderBuffers[EAGXShaderVertex].Bytes[1])
+//			{
+//				[Inner setVertexBytes:ShaderBuffers[EAGXShaderVertex].Bytes[1] length:ShaderBuffers[EAGXShaderVertex].Offsets[1] atIndex:1];
+//			}
+//		}
+//		default:
+//		{
+//			break;
+//		}
+//	}
+#endif
+}
+
+FAGXRenderCommandEncoderDebugging::FAGXRenderCommandEncoderDebugging()
+{
+	
+}
+FAGXRenderCommandEncoderDebugging::FAGXRenderCommandEncoderDebugging(mtlpp::RenderCommandEncoder& Encoder, mtlpp::RenderPassDescriptor const& Desc, FAGXCommandBufferDebugging& Buffer)
+: FAGXCommandEncoderDebugging((FAGXDebugCommandEncoder*)[[[FAGXDebugRenderCommandEncoder alloc] initWithEncoder:Encoder.GetPtr() fromDescriptor:Desc.GetPtr() andCommandBuffer:Buffer] autorelease])
+{
+	Buffer.BeginRenderCommandEncoder([NSString stringWithFormat:@"Render: %@", Encoder.GetLabel().GetPtr()], Desc);
+	Encoder.SetAssociatedObject((void const*)&FAGXRenderCommandEncoderDebugging::Get, (FAGXCommandEncoderDebugging const&)*this);
+}
+FAGXRenderCommandEncoderDebugging::FAGXRenderCommandEncoderDebugging(FAGXDebugCommandEncoder* handle)
+: FAGXCommandEncoderDebugging(handle)
+{
+	
+}
+
+FAGXRenderCommandEncoderDebugging FAGXRenderCommandEncoderDebugging::Get(mtlpp::RenderCommandEncoder& Encoder)
+{
+	return Encoder.GetAssociatedObject<FAGXRenderCommandEncoderDebugging>((void const*)&FAGXRenderCommandEncoderDebugging::Get);
+}
+
+void FAGXRenderCommandEncoderDebugging::SetPipeline(FAGXShaderPipeline* Pipeline)
+{
+#if METAL_DEBUG_OPTIONS
+	((FAGXDebugRenderCommandEncoder*)m_ptr)->Pipeline = [Pipeline retain];
+	switch (((FAGXDebugRenderCommandEncoder*)m_ptr)->Buffer.GetPtr()->DebugLevel)
+	{
+		case EAGXDebugLevelConditionalSubmit:
+		case EAGXDebugLevelWaitForComplete:
+		case EAGXDebugLevelLogOperations:
+		{
+			((FAGXDebugRenderCommandEncoder*)m_ptr)->Buffer.SetPipeline(Pipeline->RenderPipelineState.GetLabel());
+		}
+		default:
+		{
+			break;
+		}
+	}
+#endif
+}
+
+void FAGXRenderCommandEncoderDebugging::SetBytes(EAGXShaderFrequency Freq, const void * bytes, NSUInteger length, NSUInteger index)
+{
+#if METAL_DEBUG_OPTIONS
+	switch (((FAGXDebugRenderCommandEncoder*)m_ptr)->Buffer.GetPtr()->DebugLevel)
+	{
+		case EAGXDebugLevelConditionalSubmit:
+		case EAGXDebugLevelWaitForComplete:
+		case EAGXDebugLevelLogOperations:
+		case EAGXDebugLevelValidation:
+		{
+			((FAGXDebugRenderCommandEncoder*)m_ptr)->ShaderBuffers[Freq].Buffers[index] = nil;
+			((FAGXDebugRenderCommandEncoder*)m_ptr)->ShaderBuffers[Freq].Bytes[index] = bytes;
+			((FAGXDebugRenderCommandEncoder*)m_ptr)->ShaderBuffers[Freq].Offsets[index] = length;
+			((FAGXDebugRenderCommandEncoder*)m_ptr)->ResourceMask[Freq].BufferMask = bytes ? (((FAGXDebugRenderCommandEncoder*)m_ptr)->ResourceMask[Freq].BufferMask | (1 << (index))) : (((FAGXDebugRenderCommandEncoder*)m_ptr)->ResourceMask[Freq].BufferMask & ~(1 << (index)));
+		}
+		default:
+		{
+			break;
+		}
+	}
+#endif
+}
+void FAGXRenderCommandEncoderDebugging::SetBuffer(EAGXShaderFrequency Freq,  FAGXBuffer const& buffer, NSUInteger offset, NSUInteger index)
+{
+#if METAL_DEBUG_OPTIONS
+	switch (((FAGXDebugRenderCommandEncoder*)m_ptr)->Buffer.GetPtr()->DebugLevel)
+	{
+		case EAGXDebugLevelConditionalSubmit:
+		case EAGXDebugLevelWaitForComplete:
+		case EAGXDebugLevelLogOperations:
+		case EAGXDebugLevelValidation:
+		{
+			((FAGXDebugRenderCommandEncoder*)m_ptr)->ShaderBuffers[Freq].Buffers[index] = buffer;
+			((FAGXDebugRenderCommandEncoder*)m_ptr)->ShaderBuffers[Freq].Bytes[index] = nil;
+			((FAGXDebugRenderCommandEncoder*)m_ptr)->ShaderBuffers[Freq].Offsets[index] = offset;
+			((FAGXDebugRenderCommandEncoder*)m_ptr)->ResourceMask[Freq].BufferMask = buffer ? (((FAGXDebugRenderCommandEncoder*)m_ptr)->ResourceMask[Freq].BufferMask | (1 << (index))) : (((FAGXDebugRenderCommandEncoder*)m_ptr)->ResourceMask[Freq].BufferMask & ~(1 << (index)));
+		}
+		default:
+		{
+			break;
+		}
+	}
+#endif
+}
+void FAGXRenderCommandEncoderDebugging::SetBufferOffset(EAGXShaderFrequency Freq, NSUInteger offset, NSUInteger index)
+{
+#if METAL_DEBUG_OPTIONS
+	switch (((FAGXDebugRenderCommandEncoder*)m_ptr)->Buffer.GetPtr()->DebugLevel)
+	{
+		case EAGXDebugLevelConditionalSubmit:
+		case EAGXDebugLevelWaitForComplete:
+		case EAGXDebugLevelLogOperations:
+		case EAGXDebugLevelValidation:
+		{
+			((FAGXDebugRenderCommandEncoder*)m_ptr)->ShaderBuffers[Freq].Offsets[index] = offset;
+			check(((FAGXDebugRenderCommandEncoder*)m_ptr)->ResourceMask[Freq].BufferMask & (1 << (index)));
+		}
+		default:
+		{
+			break;
+		}
+	}
+#endif
+}
+
+void FAGXRenderCommandEncoderDebugging::SetTexture(EAGXShaderFrequency Freq,  FAGXTexture const& texture, NSUInteger index)
+{
+#if METAL_DEBUG_OPTIONS
+	switch (((FAGXDebugRenderCommandEncoder*)m_ptr)->Buffer.GetPtr()->DebugLevel)
+	{
+		case EAGXDebugLevelConditionalSubmit:
+		case EAGXDebugLevelWaitForComplete:
+		case EAGXDebugLevelLogOperations:
+		case EAGXDebugLevelValidation:
+		{
+			((FAGXDebugRenderCommandEncoder*)m_ptr)->ShaderTextures[Freq].Textures[index] = texture;
+			((FAGXDebugRenderCommandEncoder*)m_ptr)->ResourceMask[Freq].TextureMask = texture ? (((FAGXDebugRenderCommandEncoder*)m_ptr)->ResourceMask[Freq].TextureMask | (FAGXTextureMask(1) << index)) : (((FAGXDebugRenderCommandEncoder*)m_ptr)->ResourceMask[Freq].TextureMask & ~(FAGXTextureMask(1) << index));
+		}
+		default:
+		{
+			break;
+		}
+	}
+#endif
+}
+
+void FAGXRenderCommandEncoderDebugging::SetSamplerState(EAGXShaderFrequency Freq,  mtlpp::SamplerState const& sampler, NSUInteger index)
+{
+#if METAL_DEBUG_OPTIONS
+	switch (((FAGXDebugRenderCommandEncoder*)m_ptr)->Buffer.GetPtr()->DebugLevel)
+	{
+		case EAGXDebugLevelConditionalSubmit:
+		case EAGXDebugLevelWaitForComplete:
+		case EAGXDebugLevelLogOperations:
+		case EAGXDebugLevelValidation:
+		{
+			((FAGXDebugRenderCommandEncoder*)m_ptr)->ShaderSamplers[Freq].Samplers[index] = sampler;
+			((FAGXDebugRenderCommandEncoder*)m_ptr)->ResourceMask[Freq].SamplerMask = sampler ? (((FAGXDebugRenderCommandEncoder*)m_ptr)->ResourceMask[Freq].SamplerMask | (1 << (index))) : (((FAGXDebugRenderCommandEncoder*)m_ptr)->ResourceMask[Freq].SamplerMask & ~(1 << (index)));
+		}
+		default:
+		{
+			break;
+		}
+	}
+#endif
+}
+
+void FAGXRenderCommandEncoderDebugging::SetSamplerState(EAGXShaderFrequency Freq,  mtlpp::SamplerState const& sampler, float lodMinClamp, float lodMaxClamp, NSUInteger index)
+{
+#if METAL_DEBUG_OPTIONS
+	switch (((FAGXDebugRenderCommandEncoder*)m_ptr)->Buffer.GetPtr()->DebugLevel)
+	{
+		case EAGXDebugLevelConditionalSubmit:
+		case EAGXDebugLevelWaitForComplete:
+		case EAGXDebugLevelLogOperations:
+		case EAGXDebugLevelValidation:
+		{
+			((FAGXDebugRenderCommandEncoder*)m_ptr)->ShaderSamplers[Freq].Samplers[index] = sampler;
+			((FAGXDebugRenderCommandEncoder*)m_ptr)->ResourceMask[Freq].SamplerMask = sampler ? (((FAGXDebugRenderCommandEncoder*)m_ptr)->ResourceMask[Freq].SamplerMask | (1 << (index))) : (((FAGXDebugRenderCommandEncoder*)m_ptr)->ResourceMask[Freq].SamplerMask & ~(1 << (index)));
+		}
+		default:
+		{
+			break;
+		}
+	}
+#endif
+}
+
+void FAGXRenderCommandEncoderDebugging::SetDepthStencilState( mtlpp::DepthStencilState const& depthStencilState)
+{
+}
+
+void FAGXRenderCommandEncoderDebugging::Draw(mtlpp::PrimitiveType primitiveType, NSUInteger vertexStart, NSUInteger vertexCount, NSUInteger instanceCount)
+{
+#if METAL_DEBUG_OPTIONS
+	switch(((FAGXDebugRenderCommandEncoder*)m_ptr)->Buffer.GetPtr()->DebugLevel)
+	{
+		case EAGXDebugLevelConditionalSubmit:
+		case EAGXDebugLevelWaitForComplete:
+		case EAGXDebugLevelLogOperations:
+		{
+			((FAGXDebugRenderCommandEncoder*)m_ptr)->Buffer.Draw([NSString stringWithFormat:@"%s", __PRETTY_FUNCTION__]);
+		}
+		case EAGXDebugLevelValidation:
+		{
+			Validate();
+		}
+		default:
+		{
+			break;
+		}
+	}
+#endif
+}
+
+void FAGXRenderCommandEncoderDebugging::Draw(mtlpp::PrimitiveType primitiveType, NSUInteger vertexStart, NSUInteger vertexCount)
+{
+#if METAL_DEBUG_OPTIONS
+	switch(((FAGXDebugRenderCommandEncoder*)m_ptr)->Buffer.GetPtr()->DebugLevel)
+	{
+		case EAGXDebugLevelConditionalSubmit:
+		case EAGXDebugLevelWaitForComplete:
+		case EAGXDebugLevelLogOperations:
+		{
+			((FAGXDebugRenderCommandEncoder*)m_ptr)->Buffer.Draw([NSString stringWithFormat:@"%s", __PRETTY_FUNCTION__]);
+		}
+		case EAGXDebugLevelValidation:
+		{
+			Validate();
+		}
+		default:
+		{
+			break;
+		}
+	}
+#endif
+}
+
+void FAGXRenderCommandEncoderDebugging::DrawIndexed(mtlpp::PrimitiveType primitiveType, NSUInteger indexCount, mtlpp::IndexType indexType, FAGXBuffer const& indexBuffer, NSUInteger indexBufferOffset, NSUInteger instanceCount)
+{
+#if METAL_DEBUG_OPTIONS
+	switch(((FAGXDebugRenderCommandEncoder*)m_ptr)->Buffer.GetPtr()->DebugLevel)
+	{
+		case EAGXDebugLevelConditionalSubmit:
+		case EAGXDebugLevelWaitForComplete:
+		case EAGXDebugLevelLogOperations:
+		{
+			((FAGXDebugRenderCommandEncoder*)m_ptr)->Buffer.Draw([NSString stringWithFormat:@"%s", __PRETTY_FUNCTION__]);
+		}
+		case EAGXDebugLevelValidation:
+		{
+			Validate();
+		}
+		default:
+		{
+			break;
+		}
+	}
+#endif
+}
+
+void FAGXRenderCommandEncoderDebugging::DrawIndexed(mtlpp::PrimitiveType primitiveType, NSUInteger indexCount, mtlpp::IndexType indexType, FAGXBuffer const& indexBuffer, NSUInteger indexBufferOffset)
+{
+#if METAL_DEBUG_OPTIONS
+	switch(((FAGXDebugRenderCommandEncoder*)m_ptr)->Buffer.GetPtr()->DebugLevel)
+	{
+		case EAGXDebugLevelConditionalSubmit:
+		case EAGXDebugLevelWaitForComplete:
+		case EAGXDebugLevelLogOperations:
+		{
+			((FAGXDebugRenderCommandEncoder*)m_ptr)->Buffer.Draw([NSString stringWithFormat:@"%s", __PRETTY_FUNCTION__]);
+		}
+		case EAGXDebugLevelValidation:
+		{
+			Validate();
+		}
+		default:
+		{
+			break;
+		}
+	}
+#endif
+}
+
+void FAGXRenderCommandEncoderDebugging::Draw(mtlpp::PrimitiveType primitiveType, NSUInteger vertexStart, NSUInteger vertexCount, NSUInteger instanceCount, NSUInteger baseInstance)
+{
+#if METAL_DEBUG_OPTIONS
+	switch(((FAGXDebugRenderCommandEncoder*)m_ptr)->Buffer.GetPtr()->DebugLevel)
+	{
+		case EAGXDebugLevelConditionalSubmit:
+		case EAGXDebugLevelWaitForComplete:
+		case EAGXDebugLevelLogOperations:
+		{
+			((FAGXDebugRenderCommandEncoder*)m_ptr)->Buffer.Draw([NSString stringWithFormat:@"%s:%u,%u,%u,%u,%u", __PRETTY_FUNCTION__, (uint32)primitiveType, (uint32)vertexStart, (uint32)vertexCount, (uint32)instanceCount, (uint32)baseInstance]);
+		}
+		case EAGXDebugLevelValidation:
+		{
+			Validate();
+		}
+		default:
+		{
+			break;
+		}
+	}
+#endif
+}
+
+void FAGXRenderCommandEncoderDebugging::DrawIndexed(mtlpp::PrimitiveType primitiveType, NSUInteger indexCount, mtlpp::IndexType indexType, FAGXBuffer const& indexBuffer, NSUInteger indexBufferOffset, NSUInteger instanceCount, NSInteger baseVertex, NSUInteger baseInstance)
+{
+#if METAL_DEBUG_OPTIONS
+	switch(((FAGXDebugRenderCommandEncoder*)m_ptr)->Buffer.GetPtr()->DebugLevel)
+	{
+		case EAGXDebugLevelConditionalSubmit:
+		case EAGXDebugLevelWaitForComplete:
+		case EAGXDebugLevelLogOperations:
+		{
+			((FAGXDebugRenderCommandEncoder*)m_ptr)->Buffer.Draw([NSString stringWithFormat:@"%s:%u,%u,%u,%u,%u,%u,%u", __PRETTY_FUNCTION__, (uint32)primitiveType, (uint32)indexCount, (uint32)indexType, (uint32)indexBufferOffset, (uint32)instanceCount, (uint32)baseVertex, (uint32)baseInstance]);
+		}
+		case EAGXDebugLevelValidation:
+		{
+			Validate();
+		}
+		default:
+		{
+			break;
+		}
+	}
+#endif
+}
+
+void FAGXRenderCommandEncoderDebugging::Draw(mtlpp::PrimitiveType primitiveType, FAGXBuffer const& indirectBuffer, NSUInteger indirectBufferOffset)
+{
+#if METAL_DEBUG_OPTIONS
+	switch(((FAGXDebugRenderCommandEncoder*)m_ptr)->Buffer.GetPtr()->DebugLevel)
+	{
+		case EAGXDebugLevelConditionalSubmit:
+		case EAGXDebugLevelWaitForComplete:
+		case EAGXDebugLevelLogOperations:
+		{
+			((FAGXDebugRenderCommandEncoder*)m_ptr)->Buffer.Draw([NSString stringWithFormat:@"%s", __PRETTY_FUNCTION__]);
+		}
+		case EAGXDebugLevelValidation:
+		{
+			Validate();
+		}
+		default:
+		{
+			break;
+		}
+	}
+#endif
+}
+
+void FAGXRenderCommandEncoderDebugging::DrawIndexed(mtlpp::PrimitiveType primitiveType, mtlpp::IndexType indexType, FAGXBuffer const& indexBuffer, NSUInteger indexBufferOffset, FAGXBuffer const& indirectBuffer, NSUInteger indirectBufferOffset)
+{
+#if METAL_DEBUG_OPTIONS
+	switch(((FAGXDebugRenderCommandEncoder*)m_ptr)->Buffer.GetPtr()->DebugLevel)
+	{
+		case EAGXDebugLevelConditionalSubmit:
+		case EAGXDebugLevelWaitForComplete:
+		case EAGXDebugLevelLogOperations:
+		{
+			((FAGXDebugRenderCommandEncoder*)m_ptr)->Buffer.Draw([NSString stringWithFormat:@"%s", __PRETTY_FUNCTION__]);
+		}
+		case EAGXDebugLevelValidation:
+		{
+			Validate();
+		}
+		default:
+		{
+			break;
+		}
+	}
+#endif
+}
+
+/** Validates the pipeline/binding state */
+bool FAGXRenderCommandEncoderDebugging::ValidateFunctionBindings(EAGXShaderFrequency Frequency)
+{
+	bool bOK = true;
+#if METAL_DEBUG_OPTIONS
+	switch (((FAGXDebugRenderCommandEncoder*)m_ptr)->Buffer.GetPtr()->DebugLevel)
+	{
+		case EAGXDebugLevelConditionalSubmit:
+		case EAGXDebugLevelWaitForComplete:
+		case EAGXDebugLevelLogOperations:
+		case EAGXDebugLevelValidation:
+		{
+			check(((FAGXDebugRenderCommandEncoder*)m_ptr)->Pipeline);
+			
+			MTLRenderPipelineReflection* Reflection = ((FAGXDebugRenderCommandEncoder*)m_ptr)->Pipeline->RenderPipelineReflection;
+			check(Reflection);
+			
+			NSArray<MTLArgument*>* Arguments = nil;
+			switch(Frequency)
+			{
+				case EAGXShaderVertex:
+				{
+					Arguments = Reflection.vertexArguments;
+					break;
+				}
+				case EAGXShaderFragment:
+				{
+					Arguments = Reflection.fragmentArguments;
+					break;
+				}
+				default:
+					check(false);
+					break;
+			}
+			
+			for (uint32 i = 0; i < Arguments.count; i++)
+			{
+				MTLArgument* Arg = [Arguments objectAtIndex:i];
+				check(Arg);
+				switch(Arg.type)
+				{
+					case MTLArgumentTypeBuffer:
+					{
+						checkf(Arg.index < ML_MaxBuffers, TEXT("Metal buffer index exceeded!"));
+						if ((((FAGXDebugRenderCommandEncoder*)m_ptr)->ShaderBuffers[Frequency].Buffers[Arg.index] == nil && ((FAGXDebugRenderCommandEncoder*)m_ptr)->ShaderBuffers[Frequency].Bytes[Arg.index] == nil))
+						{
+							bOK = false;
+							UE_LOG(LogAGX, Warning, TEXT("Unbound buffer at Metal index %u which will crash the driver: %s"), (uint32)Arg.index, *FString([Arg description]));
+						}
+						break;
+					}
+					case MTLArgumentTypeThreadgroupMemory:
+					{
+						break;
+					}
+					case MTLArgumentTypeTexture:
+					{
+						checkf(Arg.index < ML_MaxTextures, TEXT("Metal texture index exceeded!"));
+						if (((FAGXDebugRenderCommandEncoder*)m_ptr)->ShaderTextures[Frequency].Textures[Arg.index] == nil)
+						{
+							bOK = false;
+							UE_LOG(LogAGX, Warning, TEXT("Unbound texture at Metal index %u which will crash the driver: %s"), (uint32)Arg.index, *FString([Arg description]));
+						}
+						else if (((FAGXDebugRenderCommandEncoder*)m_ptr)->ShaderTextures[Frequency].Textures[Arg.index].textureType != Arg.textureType)
+						{
+							bOK = false;
+							UE_LOG(LogAGX, Warning, TEXT("Incorrect texture type bound at Metal index %u which will crash the driver: %s\n%s"), (uint32)Arg.index, *FString([Arg description]), *FString([((FAGXDebugRenderCommandEncoder*)m_ptr)->ShaderTextures[Frequency].Textures[Arg.index] description]));
+						}
+						break;
+					}
+					case MTLArgumentTypeSampler:
+					{
+						checkf(Arg.index < ML_MaxSamplers, TEXT("Metal sampler index exceeded!"));
+						if (((FAGXDebugRenderCommandEncoder*)m_ptr)->ShaderSamplers[Frequency].Samplers[Arg.index] == nil)
+						{
+							bOK = false;
+							UE_LOG(LogAGX, Warning, TEXT("Unbound sampler at Metal index %u which will crash the driver: %s"), (uint32)Arg.index, *FString([Arg description]));
+						}
+						break;
+					}
+					default:
+						check(false);
+						break;
+				}
+			}
+			break;
+		}
+		default:
+		{
+			break;
+		}
+	}
+#endif
+	return bOK;
+}
+void FAGXRenderCommandEncoderDebugging::Validate()
+{
+#if METAL_DEBUG_OPTIONS
+	bool bOK = ValidateFunctionBindings(EAGXShaderVertex);
+	if (!bOK)
+	{
+		UE_LOG(LogAGX, Error, TEXT("Metal Validation failures for vertex shader:\n%s"), ((FAGXDebugRenderCommandEncoder*)m_ptr)->Pipeline && ((FAGXDebugRenderCommandEncoder*)m_ptr)->Pipeline->VertexSource ? *FString(((FAGXDebugRenderCommandEncoder*)m_ptr)->Pipeline->VertexSource) : TEXT("nil"));
+	}
+	
+	bOK = ValidateFunctionBindings(EAGXShaderFragment);
+	if (!bOK)
+	{
+		UE_LOG(LogAGX, Error, TEXT("Metal Validation failures for fragment shader:\n%s"), ((FAGXDebugRenderCommandEncoder*)m_ptr)->Pipeline && ((FAGXDebugRenderCommandEncoder*)m_ptr)->Pipeline->FragmentSource ? *FString(((FAGXDebugRenderCommandEncoder*)m_ptr)->Pipeline->FragmentSource) : TEXT("nil"));
+	}
+#endif
+}
+
+void FAGXRenderCommandEncoderDebugging::InsertDebugSignpost(ns::String const& Label)
+{
+	((FAGXDebugRenderCommandEncoder*)m_ptr)->Buffer.InsertDebugSignpost(Label);
+}
+void FAGXRenderCommandEncoderDebugging::PushDebugGroup(ns::String const& Group)
+{
+    ((FAGXDebugRenderCommandEncoder*)m_ptr)->Buffer.PushDebugGroup(Group);
+}
+void FAGXRenderCommandEncoderDebugging::PopDebugGroup()
+{
+    ((FAGXDebugRenderCommandEncoder*)m_ptr)->Buffer.PopDebugGroup();
+#if METAL_DEBUG_OPTIONS
+	InsertDebugDraw();
+#endif
+}
+
+void FAGXRenderCommandEncoderDebugging::EndEncoder()
+{
+	((FAGXDebugRenderCommandEncoder*)m_ptr)->Buffer.EndCommandEncoder();
+}
+
+FAGXParallelRenderCommandEncoderDebugging::FAGXParallelRenderCommandEncoderDebugging()
+{
+
+}
+
+FAGXParallelRenderCommandEncoderDebugging::FAGXParallelRenderCommandEncoderDebugging(mtlpp::ParallelRenderCommandEncoder& Encoder, mtlpp::RenderPassDescriptor const& Desc, FAGXCommandBufferDebugging& Buffer)
+: ns::Object<FAGXDebugParallelRenderCommandEncoder*>([[FAGXDebugParallelRenderCommandEncoder alloc] initWithEncoder:Encoder.GetPtr() fromDescriptor:Desc.GetPtr() andCommandBuffer:Buffer], ns::Ownership::Assign)
+{
+	Buffer.BeginRenderCommandEncoder([NSString stringWithFormat:@"ParallelRender: %@", Encoder.GetLabel().GetPtr()], Desc);
+	Encoder.SetAssociatedObject((void const*)&FAGXParallelRenderCommandEncoderDebugging::Get, *this);
+}
+
+FAGXParallelRenderCommandEncoderDebugging::FAGXParallelRenderCommandEncoderDebugging(FAGXDebugParallelRenderCommandEncoder* handle)
+: ns::Object<FAGXDebugParallelRenderCommandEncoder*>(handle)
+{
+
+}
+
+FAGXDebugParallelRenderCommandEncoder* FAGXParallelRenderCommandEncoderDebugging::Get(mtlpp::ParallelRenderCommandEncoder& Buffer)
+{
+	return Buffer.GetAssociatedObject<FAGXParallelRenderCommandEncoderDebugging>((void const*)&FAGXParallelRenderCommandEncoderDebugging::Get);
+}
+
+FAGXRenderCommandEncoderDebugging FAGXParallelRenderCommandEncoderDebugging::GetRenderCommandEncoderDebugger(mtlpp::RenderCommandEncoder& Encoder)
+{
+	mtlpp::RenderPassDescriptor Desc(m_ptr->RenderPassDesc);
+	FAGXCommandBufferDebugging IndirectBuffer([[FAGXDebugCommandBuffer alloc] initWithCommandBuffer:m_ptr->Buffer.GetPtr()->InnerBuffer]);
+	FAGXRenderCommandEncoderDebugging EncoderDebugging(Encoder, Desc, IndirectBuffer);
+	m_ptr->RenderEncoders.Add(EncoderDebugging);
+	return EncoderDebugging;
+}
+
+void FAGXParallelRenderCommandEncoderDebugging::InsertDebugSignpost(ns::String const& Label)
+{
+	m_ptr->Buffer.InsertDebugSignpost(Label);
+}
+
+void FAGXParallelRenderCommandEncoderDebugging::PushDebugGroup(ns::String const& Group)
+{
+    m_ptr->Buffer.PushDebugGroup(Group);
+}
+
+void FAGXParallelRenderCommandEncoderDebugging::PopDebugGroup()
+{
+    m_ptr->Buffer.PopDebugGroup();
+}
+
+void FAGXParallelRenderCommandEncoderDebugging::EndEncoder()
+{
+	FAGXDebugCommandBuffer* CommandBuffer = m_ptr->Buffer.GetPtr();
+	for (FAGXRenderCommandEncoderDebugging& EncoderDebug : m_ptr->RenderEncoders)
+	{
+		FAGXDebugCommandBuffer* CmdBuffer = ((FAGXDebugRenderCommandEncoder*)EncoderDebug.GetPtr())->Buffer.GetPtr();
+		[CommandBuffer->DebugGroup addObjectsFromArray:CmdBuffer->DebugGroup];
+		CommandBuffer->DebugCommands.Append(CmdBuffer->DebugCommands);
+	}
+	m_ptr->Buffer.EndCommandEncoder();
+}
+
+#endif
+
diff --git a/Engine/Source/Runtime/Apple/AGXRHI/Private/AGXRenderCommandEncoder.h b/Engine/Source/Runtime/Apple/AGXRHI/Private/AGXRenderCommandEncoder.h
new file mode 100644
--- /dev/null
+++ b/Engine/Source/Runtime/Apple/AGXRHI/Private/AGXRenderCommandEncoder.h
@@ -0,0 +1,80 @@
+// Copyright Epic Games, Inc. All Rights Reserved.
+
+#pragma once
+
+#include "AGXDebugCommandEncoder.h"
+
+#if MTLPP_CONFIG_VALIDATE && METAL_DEBUG_OPTIONS
+NS_ASSUME_NONNULL_BEGIN
+
+@class FAGXShaderPipeline;
+@class FAGXDebugRenderCommandEncoder;
+@class FAGXDebugParallelRenderCommandEncoder;
+
+class FAGXRenderCommandEncoderDebugging : public FAGXCommandEncoderDebugging
+{
+	void InsertDebugDraw();
+public:
+	FAGXRenderCommandEncoderDebugging();
+	FAGXRenderCommandEncoderDebugging(mtlpp::RenderCommandEncoder& Encoder, mtlpp::RenderPassDescriptor const& Desc, FAGXCommandBufferDebugging& Buffer);
+	FAGXRenderCommandEncoderDebugging(FAGXDebugCommandEncoder* handle);
+	
+	static FAGXRenderCommandEncoderDebugging Get(mtlpp::RenderCommandEncoder& Buffer);
+	
+	void SetPipeline(FAGXShaderPipeline* Pipeline);
+	
+	void SetBytes(EAGXShaderFrequency Freq, const void * bytes, NSUInteger length, NSUInteger index);
+	void SetBuffer(EAGXShaderFrequency Freq,  FAGXBuffer const& buffer, NSUInteger offset, NSUInteger index);
+	void SetBufferOffset(EAGXShaderFrequency Freq, NSUInteger offset, NSUInteger index);
+	void SetTexture(EAGXShaderFrequency Freq,  FAGXTexture const& texture, NSUInteger index);
+	void SetSamplerState(EAGXShaderFrequency Freq,  mtlpp::SamplerState const& sampler, NSUInteger index);
+	void SetSamplerState(EAGXShaderFrequency Freq,  mtlpp::SamplerState const& sampler, float lodMinClamp, float lodMaxClamp, NSUInteger index);
+	void SetDepthStencilState( mtlpp::DepthStencilState const& depthStencilState);
+	
+	void Draw(mtlpp::PrimitiveType primitiveType, NSUInteger vertexStart, NSUInteger vertexCount, NSUInteger instanceCount);
+	
+	void Draw(mtlpp::PrimitiveType primitiveType, NSUInteger vertexStart, NSUInteger vertexCount);
+	
+	void DrawIndexed(mtlpp::PrimitiveType primitiveType, NSUInteger indexCount, mtlpp::IndexType indexType, FAGXBuffer const& indexBuffer, NSUInteger indexBufferOffset, NSUInteger instanceCount);
+	
+	void DrawIndexed(mtlpp::PrimitiveType primitiveType, NSUInteger indexCount, mtlpp::IndexType indexType, FAGXBuffer const& indexBuffer, NSUInteger indexBufferOffset);
+	
+	void Draw(mtlpp::PrimitiveType primitiveType, NSUInteger vertexStart, NSUInteger vertexCount, NSUInteger instanceCount, NSUInteger baseInstance);
+	
+	void DrawIndexed(mtlpp::PrimitiveType primitiveType, NSUInteger indexCount, mtlpp::IndexType indexType, FAGXBuffer const& indexBuffer, NSUInteger indexBufferOffset, NSUInteger instanceCount, NSInteger baseVertex, NSUInteger baseInstance);
+	
+	void Draw(mtlpp::PrimitiveType primitiveType, FAGXBuffer const& indirectBuffer, NSUInteger indirectBufferOffset);
+	
+	void DrawIndexed(mtlpp::PrimitiveType primitiveType, mtlpp::IndexType indexType, FAGXBuffer const& indexBuffer, NSUInteger indexBufferOffset, FAGXBuffer const& indirectBuffer, NSUInteger indirectBufferOffset);
+	
+	/** Validates the pipeline/binding state */
+	bool ValidateFunctionBindings(EAGXShaderFrequency Frequency);
+	void Validate();
+	
+	void InsertDebugSignpost(ns::String const& Label);
+	void PushDebugGroup(ns::String const& Group);
+	void PopDebugGroup();
+	
+	void EndEncoder();
+};
+
+class FAGXParallelRenderCommandEncoderDebugging : public ns::Object<FAGXDebugParallelRenderCommandEncoder*>
+{
+public:
+	FAGXParallelRenderCommandEncoderDebugging();
+	FAGXParallelRenderCommandEncoderDebugging(mtlpp::ParallelRenderCommandEncoder& Encoder, mtlpp::RenderPassDescriptor const& Desc, FAGXCommandBufferDebugging& Buffer);
+	FAGXParallelRenderCommandEncoderDebugging(FAGXDebugParallelRenderCommandEncoder* handle);
+
+	static FAGXDebugParallelRenderCommandEncoder* Get(mtlpp::ParallelRenderCommandEncoder& Buffer);
+
+	FAGXRenderCommandEncoderDebugging GetRenderCommandEncoderDebugger(mtlpp::RenderCommandEncoder& Encoder);
+
+	void InsertDebugSignpost(ns::String const& Label);
+	void PushDebugGroup(ns::String const& Group);
+	void PopDebugGroup();
+	
+	void EndEncoder();
+};
+
+NS_ASSUME_NONNULL_END
+#endif
