Initial import of AGXRHI

diff --git a/Engine/Source/Runtime/Apple/AGXRHI/Private/AGXRenderPass.cpp b/Engine/Source/Runtime/Apple/AGXRHI/Private/AGXRenderPass.cpp
new file mode 100644
--- /dev/null
+++ b/Engine/Source/Runtime/Apple/AGXRHI/Private/AGXRenderPass.cpp
@@ -0,0 +1,1767 @@
+// Copyright Epic Games, Inc. All Rights Reserved.
+
+/*=============================================================================
+	AGXRenderPass.cpp: AGX RHI command pass wrapper.
+=============================================================================*/
+
+
+#include "AGXRHIPrivate.h"
+#include "AGXShaderTypes.h"
+#include "AGXGraphicsPipelineState.h"
+#include "AGXVertexDeclaration.h"
+#include "AGXRenderPass.h"
+#include "AGXCommandBuffer.h"
+#include "AGXProfiler.h"
+#include "AGXFrameAllocator.h"
+
+#pragma mark - Private Console Variables -
+
+static int32 GAGXCommandBufferCommitThreshold = 0;
+static FAutoConsoleVariableRef CVarAGXCommandBufferCommitThreshold(
+	TEXT("rhi.AGX.CommandBufferCommitThreshold"),
+	GAGXCommandBufferCommitThreshold,
+	TEXT("When enabled (> 0) if the command buffer has more than this number of draw/dispatch command encoded then it will be committed at the next encoder boundary to keep the GPU busy. (Default: 0, set to <= 0 to disable)"));
+
+static int32 GAGXDeferRenderPasses = 1;
+static FAutoConsoleVariableRef CVarAGXDeferRenderPasses(
+	TEXT("rhi.AGX.DeferRenderPasses"),
+	GAGXDeferRenderPasses,
+	TEXT("Whether to defer creating render command encoders. (Default: 1)"));
+
+// Deliberately not static!
+int32 GAGXDebugOpsCount = PLATFORM_MAC ? 1 : 10;
+static FAutoConsoleVariableRef CVarAGXDebugOpsCount(
+	TEXT("rhi.AGX.DebugOpsCount"),
+	GAGXDebugOpsCount,
+	TEXT("The number of operations to allow between GPU debug markers for the r.GPUCrashDebugging reports. (Default: Mac = 1 : iOS/tvOS = 10)"));
+
+#pragma mark - Public C++ Boilerplate -
+
+FAGXRenderPass::FAGXRenderPass(FAGXCommandList& InCmdList, FAGXStateCache& Cache)
+: CmdList(InCmdList)
+, State(Cache)
+, CurrentEncoder(InCmdList, EAGXCommandEncoderCurrent)
+, PrologueEncoder(InCmdList, EAGXCommandEncoderPrologue)
+, RenderPassDesc(nil)
+, ComputeDispatchType(mtlpp::DispatchType::Serial)
+, NumOutstandingOps(0)
+, bWithinRenderPass(false)
+{
+}
+
+FAGXRenderPass::~FAGXRenderPass(void)
+{
+	check(!CurrentEncoder.GetCommandBuffer());
+	check(!PrologueEncoder.GetCommandBuffer());
+	check(!PassStartFence);
+}
+
+void FAGXRenderPass::SetDispatchType(mtlpp::DispatchType Type)
+{
+	ComputeDispatchType = Type;
+}
+
+void FAGXRenderPass::Begin(FAGXFence* Fence, bool const bParallelBegin)
+{
+	if (!bParallelBegin || !FAGXCommandQueue::SupportsFeature(EAGXFeaturesParallelRenderEncoders))
+	{
+		check(!PassStartFence || Fence == nullptr);
+		if (Fence)
+		{
+			PassStartFence = Fence;
+			PrologueStartEncoderFence = Fence;
+		}
+	}
+	else
+	{
+		check(!ParallelPassEndFence || Fence == nullptr);
+		if (Fence)
+		{
+			ParallelPassEndFence = Fence;
+			PrologueStartEncoderFence = Fence;
+		}
+	}
+	
+	if (!CmdList.IsParallel() && !CurrentEncoder.GetCommandBuffer())
+	{
+		CurrentEncoder.StartCommandBuffer();
+		check(CurrentEncoder.GetCommandBuffer());
+	}
+}
+
+void FAGXRenderPass::Wait(FAGXFence* Fence)
+{
+	if (Fence)
+	{
+		if (PrologueEncoder.IsBlitCommandEncoderActive() || PrologueEncoder.IsComputeCommandEncoderActive())
+		{
+			PrologueEncoder.WaitForFence(Fence);
+			METAL_DEBUG_LAYER(EAGXDebugLevelValidation, FAGXFence::ValidateUsage(Fence));
+		}
+		else if (CurrentEncoder.IsRenderCommandEncoderActive() || CurrentEncoder.IsBlitCommandEncoderActive() || CurrentEncoder.IsComputeCommandEncoderActive())
+		{
+			CurrentEncoder.WaitForFence(Fence);
+			METAL_DEBUG_LAYER(EAGXDebugLevelValidation, FAGXFence::ValidateUsage(Fence));
+		}
+        else
+        {
+            PassStartFence = Fence;
+			PrologueStartEncoderFence = Fence;
+        }
+	}
+		
+}
+
+void FAGXRenderPass::Update(FAGXFence* Fence)
+{
+	if (Fence)
+	{
+		// Force an encoder - possibly consuming the start fence so that we get the proper order
+		// the higher-level can generate empty contexts but we have no sane way to deal with that.
+		if (!CurrentEncoder.IsRenderCommandEncoderActive() && !CurrentEncoder.IsBlitCommandEncoderActive() && !CurrentEncoder.IsComputeCommandEncoderActive())
+		{
+			ConditionalSwitchToBlit();
+		}
+		CurrentEncoder.UpdateFence(Fence);
+		State.FlushVisibilityResults(CurrentEncoder);
+		TRefCountPtr<FAGXFence> NewFence = CurrentEncoder.EndEncoding();
+		check(!CurrentEncoderFence || !NewFence);
+		if (NewFence)
+		{
+			CurrentEncoderFence = NewFence;
+		}
+	}
+}
+
+TRefCountPtr<FAGXFence> const& FAGXRenderPass::Submit(EAGXSubmitFlags Flags)
+{
+	if (CurrentEncoder.GetCommandBuffer() || (Flags & EAGXSubmitFlagsAsyncCommandBuffer))
+	{
+		if (PrologueEncoder.IsBlitCommandEncoderActive() || PrologueEncoder.IsComputeCommandEncoderActive())
+		{
+			check(PrologueEncoder.GetCommandBuffer());
+			PrologueEncoderFence = PrologueEncoder.EndEncoding();
+		}
+		if (PrologueEncoder.GetCommandBuffer())
+		{
+			PrologueEncoder.CommitCommandBuffer((Flags & EAGXSubmitFlagsAsyncCommandBuffer) ? Flags : EAGXSubmitFlagsNone);
+        }
+    }
+
+    // Must be on the render thread if there's no RHI thread
+    // Must be on the RHI thread otherwise
+    CheckMetalThread();
+    
+    if (Flags & EAGXSubmitFlagsLastCommandBuffer)
+    {
+        check(CurrentEncoder.GetCommandBuffer());
+        id <MTLCommandBuffer> CommandBuffer = CurrentEncoder.GetCommandBuffer().GetPtr();
+        
+        FAGXDeviceContext& DeviceContext = (FAGXDeviceContext&)GetAGXDeviceContext();
+        FAGXFrameAllocator* UniformAllocator = DeviceContext.GetUniformAllocator();
+        
+        UniformAllocator->MarkEndOfFrame(DeviceContext.GetFrameNumberRHIThread(), CommandBuffer);
+		
+		FAGXFrameAllocator* TransferAllocator = DeviceContext.GetTransferAllocator();
+		TransferAllocator->MarkEndOfFrame(DeviceContext.GetFrameNumberRHIThread(), CommandBuffer);
+    }
+    
+    if (CurrentEncoder.GetCommandBuffer() && !(Flags & EAGXSubmitFlagsAsyncCommandBuffer))
+    {
+        if (CurrentEncoder.IsRenderCommandEncoderActive() || CurrentEncoder.IsBlitCommandEncoderActive() || CurrentEncoder.IsComputeCommandEncoderActive())
+        {
+            if (CurrentEncoder.IsRenderCommandEncoderActive())
+            {
+                State.SetRenderStoreActions(CurrentEncoder, (Flags & EAGXSubmitFlagsBreakCommandBuffer));
+				State.FlushVisibilityResults(CurrentEncoder);
+            }
+            CurrentEncoderFence = CurrentEncoder.EndEncoding();
+        }
+		
+        CurrentEncoder.CommitCommandBuffer(Flags);
+    }
+	
+	OutstandingBufferUploads.Empty();
+	if (Flags & EAGXSubmitFlagsResetState)
+	{
+		PrologueEncoder.Reset();
+		CurrentEncoder.Reset();
+	}
+	
+	return CurrentEncoderFence;
+}
+
+void FAGXRenderPass::BeginParallelRenderPass(mtlpp::RenderPassDescriptor RenderPass, uint32 NumParallelContextsInPass)
+{
+	check(!bWithinRenderPass);
+	check(!RenderPassDesc);
+	check(RenderPass);
+	check(CurrentEncoder.GetCommandBuffer());
+
+	if (!CurrentEncoder.GetParallelRenderCommandEncoder())
+	{
+		if (PrologueEncoder.IsBlitCommandEncoderActive() || PrologueEncoder.IsComputeCommandEncoderActive())
+		{
+			PrologueEncoderFence = PrologueEncoder.EndEncoding();
+		}
+		if (CurrentEncoder.IsRenderCommandEncoderActive() || CurrentEncoder.IsBlitCommandEncoderActive() || CurrentEncoder.IsComputeCommandEncoderActive())
+		{
+			State.FlushVisibilityResults(CurrentEncoder);
+			CurrentEncoderFence = CurrentEncoder.EndEncoding();
+		}
+
+		CurrentEncoder.SetRenderPassDescriptor(RenderPass);
+
+		CurrentEncoder.BeginParallelRenderCommandEncoding(NumParallelContextsInPass);
+
+		RenderPassDesc = RenderPass;
+
+		bWithinRenderPass = true;
+	}
+}
+
+void FAGXRenderPass::BeginRenderPass(mtlpp::RenderPassDescriptor RenderPass)
+{
+	check(!bWithinRenderPass);
+	check(!RenderPassDesc);
+	check(RenderPass);
+	check(!CurrentEncoder.IsRenderCommandEncoderActive());
+	if (!CmdList.IsParallel() && !CmdList.IsImmediate() && !CurrentEncoder.GetCommandBuffer())
+	{
+		CurrentEncoder.StartCommandBuffer();
+	}
+	check(CmdList.IsParallel() || CurrentEncoder.GetCommandBuffer());
+	
+	// EndEncoding should provide the encoder fence...
+	if (PrologueEncoder.IsBlitCommandEncoderActive() || PrologueEncoder.IsComputeCommandEncoderActive())
+	{
+		PrologueEncoderFence = PrologueEncoder.EndEncoding();
+	}
+	if (CurrentEncoder.IsRenderCommandEncoderActive() || CurrentEncoder.IsBlitCommandEncoderActive() || CurrentEncoder.IsComputeCommandEncoderActive())
+	{
+		State.FlushVisibilityResults(CurrentEncoder);
+		CurrentEncoderFence = CurrentEncoder.EndEncoding();
+	}
+	State.SetStateDirty();
+	State.SetRenderTargetsActive(true);
+	
+	RenderPassDesc = RenderPass;
+	
+	CurrentEncoder.SetRenderPassDescriptor(RenderPassDesc);
+
+	if (!GAGXDeferRenderPasses || !State.CanRestartRenderPass() || CmdList.IsParallel())
+	{
+		CurrentEncoder.BeginRenderCommandEncoding();
+		if (PassStartFence)
+		{
+			CurrentEncoder.WaitForFence(PassStartFence);
+			PassStartFence = nullptr;
+		}
+		if (ParallelPassEndFence)
+		{
+			CurrentEncoder.WaitForFence(ParallelPassEndFence);
+			ParallelPassEndFence = nullptr;
+		}
+		if (CurrentEncoderFence)
+		{
+			CurrentEncoder.WaitForFence(CurrentEncoderFence);
+			CurrentEncoderFence = nullptr;
+		}
+		if (PrologueEncoderFence)
+		{
+			// Consume on the current encoder but do not invalidate
+			CurrentEncoder.WaitForFence(PrologueEncoderFence);
+		}
+		if (PrologueEncoder.IsBlitCommandEncoderActive() || PrologueEncoder.IsComputeCommandEncoderActive())
+		{
+			CurrentEncoder.WaitForFence(PrologueEncoder.GetEncoderFence());
+		}
+		State.SetRenderStoreActions(CurrentEncoder, false);
+		check(CurrentEncoder.IsRenderCommandEncoderActive());
+	}
+	
+	bWithinRenderPass = true;
+	
+	check(!PrologueEncoder.IsBlitCommandEncoderActive() && !PrologueEncoder.IsComputeCommandEncoderActive());
+}
+
+void FAGXRenderPass::RestartRenderPass(mtlpp::RenderPassDescriptor RenderPass)
+{
+	check(bWithinRenderPass);
+	check(RenderPassDesc);
+	check(CmdList.IsParallel() || CurrentEncoder.GetCommandBuffer());
+	
+	mtlpp::RenderPassDescriptor StartDesc;
+	if (RenderPass != nil)
+	{
+		// Just restart with the render pass we were given - the caller should have ensured that this is restartable
+		check(State.CanRestartRenderPass());
+		StartDesc = RenderPass;
+	}
+	else if (State.PrepareToRestart(CurrentEncoder.IsRenderPassDescriptorValid() && (State.GetRenderPassDescriptor().GetPtr() == CurrentEncoder.GetRenderPassDescriptor().GetPtr())))
+	{
+		// Restart with the render pass we have in the state cache - the state cache says its safe
+		StartDesc = State.GetRenderPassDescriptor();
+	}
+	else
+	{
+		
+		METAL_FATAL_ERROR(TEXT("Failed to restart render pass with descriptor: %s"), *FString([RenderPassDesc.GetPtr() description]));
+	}
+	check(StartDesc);
+	
+	RenderPassDesc = StartDesc;
+	
+#if METAL_DEBUG_OPTIONS
+	if ((GetAGXDeviceContext().GetCommandQueue().GetRuntimeDebuggingLevel() >= EAGXDebugLevelValidation))
+	{
+		bool bAllLoadActionsOK = true;
+		ns::Array<mtlpp::RenderPassColorAttachmentDescriptor> Attachments = RenderPassDesc.GetColorAttachments();
+		for(uint i = 0; i < 8; i++)
+		{
+			mtlpp::RenderPassColorAttachmentDescriptor Desc = Attachments[i];
+			if(Desc && Desc.GetTexture())
+			{
+				bAllLoadActionsOK &= (Desc.GetLoadAction() != mtlpp::LoadAction::Clear);
+			}
+		}
+		if(RenderPassDesc.GetDepthAttachment() && RenderPassDesc.GetDepthAttachment().GetTexture())
+		{
+			bAllLoadActionsOK &= (RenderPassDesc.GetDepthAttachment().GetLoadAction() != mtlpp::LoadAction::Clear);
+		}
+		if(RenderPassDesc.GetStencilAttachment() && RenderPassDesc.GetStencilAttachment().GetTexture())
+		{
+			bAllLoadActionsOK &= (RenderPassDesc.GetStencilAttachment().GetLoadAction() != mtlpp::LoadAction::Clear);
+		}
+		
+		if (!bAllLoadActionsOK)
+		{
+			UE_LOG(LogAGX, Warning, TEXT("Tried to restart render encoding with a clear operation - this would erroneously re-clear any existing draw calls: %s"), *FString([RenderPassDesc.GetPtr() description]));
+			
+			for(uint i = 0; i< 8; i++)
+			{
+				mtlpp::RenderPassColorAttachmentDescriptor Desc = Attachments[i];
+				if(Desc && Desc.GetTexture())
+				{
+					Desc.SetLoadAction(mtlpp::LoadAction::Load);
+				}
+			}
+			if(RenderPassDesc.GetDepthAttachment() && RenderPassDesc.GetDepthAttachment().GetTexture())
+			{
+				RenderPassDesc.GetDepthAttachment().SetLoadAction(mtlpp::LoadAction::Load);
+			}
+			if(RenderPassDesc.GetStencilAttachment() && RenderPassDesc.GetStencilAttachment().GetTexture())
+			{
+				RenderPassDesc.GetStencilAttachment().SetLoadAction(mtlpp::LoadAction::Load);
+			}
+		}
+	}
+#endif
+	
+	// EndEncoding should provide the encoder fence...
+	if (CurrentEncoder.IsBlitCommandEncoderActive() || CurrentEncoder.IsComputeCommandEncoderActive() || CurrentEncoder.IsRenderCommandEncoderActive())
+	{
+		if (CurrentEncoder.IsRenderCommandEncoderActive())
+		{
+			State.SetRenderStoreActions(CurrentEncoder, true);
+			State.FlushVisibilityResults(CurrentEncoder);
+		}
+		CurrentEncoderFence = CurrentEncoder.EndEncoding();
+	}
+	State.SetStateDirty();
+	State.SetRenderTargetsActive(true);
+	
+	CurrentEncoder.SetRenderPassDescriptor(RenderPassDesc);
+	CurrentEncoder.BeginRenderCommandEncoding();
+	if (PassStartFence)
+	{
+		CurrentEncoder.WaitForFence(PassStartFence);
+		PassStartFence = nullptr;
+	}
+	if (ParallelPassEndFence)
+	{
+		CurrentEncoder.WaitForFence(ParallelPassEndFence);
+		ParallelPassEndFence = nullptr;
+	}
+	if (CurrentEncoderFence)
+	{
+		CurrentEncoder.WaitForFence(CurrentEncoderFence);
+		CurrentEncoderFence = nullptr;
+	}
+	if (PrologueEncoderFence)
+	{
+		// Consume on the current encoder but do not invalidate
+		CurrentEncoder.WaitForFence(PrologueEncoderFence);
+	}
+	if (PrologueEncoder.IsBlitCommandEncoderActive() || PrologueEncoder.IsComputeCommandEncoderActive())
+	{
+		CurrentEncoder.WaitForFence(PrologueEncoder.GetEncoderFence());
+	}
+	State.SetRenderStoreActions(CurrentEncoder, false);
+	
+	check(CurrentEncoder.IsRenderCommandEncoderActive());
+}
+
+void FAGXRenderPass::DrawPrimitive(uint32 PrimitiveType, uint32 BaseVertexIndex, uint32 NumPrimitives, uint32 NumInstances)
+{
+	NumInstances = FMath::Max(NumInstances,1u);
+	
+	ConditionalSwitchToRender();
+	check(CurrentEncoder.GetCommandBuffer());
+	check(CurrentEncoder.IsRenderCommandEncoderActive());
+	
+	PrepareToRender(PrimitiveType);
+
+	// draw!
+	// how many verts to render
+	uint32 NumVertices = GetVertexCountForPrimitiveCount(NumPrimitives, PrimitiveType);
+	
+	METAL_GPUPROFILE(FAGXProfiler::GetProfiler()->EncodeDraw(CurrentEncoder.GetCommandBufferStats(), __FUNCTION__, NumPrimitives, NumVertices, NumInstances));
+	CurrentEncoder.GetRenderCommandEncoder().Draw(AGXTranslatePrimitiveType(PrimitiveType), BaseVertexIndex, NumVertices, NumInstances);
+	METAL_DEBUG_LAYER(EAGXDebugLevelFastValidation, CurrentEncoder.GetRenderCommandEncoderDebugging().Draw(AGXTranslatePrimitiveType(PrimitiveType), BaseVertexIndex, NumVertices, NumInstances));
+
+	if (GAGXCommandBufferDebuggingEnabled)
+	{
+		FAGXCommandData Data;
+		Data.CommandType = FAGXCommandData::Type::DrawPrimitive;
+		Data.Draw.BaseInstance = 0;
+		Data.Draw.InstanceCount = NumInstances;
+		Data.Draw.VertexCount = NumVertices;
+		Data.Draw.VertexStart = BaseVertexIndex;
+		
+		InsertDebugDraw(Data);
+	}
+	
+	ConditionalSubmit();	
+}
+
+void FAGXRenderPass::DrawPrimitiveIndirect(uint32 PrimitiveType, FAGXVertexBuffer* VertexBuffer, uint32 ArgumentOffset)
+{
+	if (GetAGXDeviceContext().SupportsFeature(EAGXFeaturesIndirectBuffer))
+	{
+		ConditionalSwitchToRender();
+		check(CurrentEncoder.GetCommandBuffer());
+		check(CurrentEncoder.IsRenderCommandEncoderActive());
+		
+		FAGXBuffer TheBackingBuffer = VertexBuffer->GetCurrentBuffer();
+		check(TheBackingBuffer);
+		
+		PrepareToRender(PrimitiveType);
+		
+		METAL_GPUPROFILE(FAGXProfiler::GetProfiler()->EncodeDraw(CurrentEncoder.GetCommandBufferStats(), __FUNCTION__, 1, 1, 1));
+		CurrentEncoder.GetRenderCommandEncoder().Draw(AGXTranslatePrimitiveType(PrimitiveType), TheBackingBuffer, ArgumentOffset);
+		METAL_DEBUG_LAYER(EAGXDebugLevelFastValidation, CurrentEncoder.GetRenderCommandEncoderDebugging().Draw(AGXTranslatePrimitiveType(PrimitiveType), TheBackingBuffer, ArgumentOffset));
+
+		if (GAGXCommandBufferDebuggingEnabled)
+		{
+			FAGXCommandData Data;
+			Data.CommandType = FAGXCommandData::Type::DrawPrimitiveIndirect;
+
+			InsertDebugDraw(Data);
+		}
+
+		ConditionalSubmit();
+	}
+	else
+	{
+		NOT_SUPPORTED("RHIDrawPrimitiveIndirect");
+	}
+}
+
+void FAGXRenderPass::DrawIndexedPrimitive(FAGXBuffer const& IndexBuffer, uint32 IndexStride, uint32 PrimitiveType, int32 BaseVertexIndex, uint32 FirstInstance,
+											 uint32 NumVertices, uint32 StartIndex, uint32 NumPrimitives, uint32 NumInstances)
+{
+	// We need at least one to cover all use cases
+	NumInstances = FMath::Max(NumInstances,1u);
+	
+#if UE_BUILD_DEBUG || UE_BUILD_DEVELOPMENT
+	{
+		FAGXGraphicsPipelineState* PipelineState = State.GetGraphicsPSO();
+		check(PipelineState != nullptr);
+		FAGXVertexDeclaration* VertexDecl = PipelineState->VertexDeclaration;
+		check(VertexDecl != nullptr);
+		
+		// Set our local copy and try to disprove the passed in value
+		uint32 ClampedNumInstances = NumInstances;
+		const CrossCompiler::FShaderBindingInOutMask& InOutMask = PipelineState->VertexShader->Bindings.InOutMask;
+
+		// I think it is valid to have no elements in this list
+		for(int VertexElemIdx = 0;VertexElemIdx < VertexDecl->Elements.Num();++VertexElemIdx)
+		{
+			FVertexElement const & VertexElem = VertexDecl->Elements[VertexElemIdx];
+			if(VertexElem.Stride > 0 && VertexElem.bUseInstanceIndex && InOutMask.IsFieldEnabled(VertexElem.AttributeIndex))
+			{
+				uint32 AvailElementCount = 0;
+				
+				uint32 BufferSize = State.GetVertexBufferSize(VertexElem.StreamIndex);
+				uint32 ElementCount = (BufferSize / VertexElem.Stride);
+				
+				if(ElementCount > FirstInstance)
+				{
+					AvailElementCount = ElementCount - FirstInstance;
+				}
+				
+				ClampedNumInstances = FMath::Clamp<uint32>(ClampedNumInstances, 0, AvailElementCount);
+				
+				if(ClampedNumInstances < NumInstances)
+				{
+					FString ShaderName = TEXT("Unknown");
+#if !(UE_BUILD_SHIPPING || UE_BUILD_TEST)
+					ShaderName = PipelineState->PixelShader->ShaderName;
+#endif
+					// Setting NumInstances to ClampedNumInstances would fix any visual rendering bugs resulting from this bad call but these draw calls are wrong - don't hide the issue
+					UE_LOG(LogAGX, Error, TEXT("Metal DrawIndexedPrimitive requested to draw %d Instances but vertex stream only has %d instance data available. ShaderName: %s, Deficient Attribute Index: %u"), NumInstances, ClampedNumInstances,
+						   *ShaderName, VertexElem.AttributeIndex);
+				}
+			}
+		}
+	}
+#endif
+	
+	ConditionalSwitchToRender();
+	check(CurrentEncoder.GetCommandBuffer());
+	check(CurrentEncoder.IsRenderCommandEncoderActive());
+	
+	PrepareToRender(PrimitiveType);
+	
+	uint32 NumIndices = GetVertexCountForPrimitiveCount(NumPrimitives, PrimitiveType);
+	
+	METAL_GPUPROFILE(FAGXProfiler::GetProfiler()->EncodeDraw(CurrentEncoder.GetCommandBufferStats(), __FUNCTION__, NumPrimitives, NumVertices, NumInstances));
+	if (GRHISupportsBaseVertexIndex && GRHISupportsFirstInstance)
+	{
+		CurrentEncoder.GetRenderCommandEncoder().DrawIndexed(AGXTranslatePrimitiveType(PrimitiveType), NumIndices, ((IndexStride == 2) ? mtlpp::IndexType::UInt16 : mtlpp::IndexType::UInt32), IndexBuffer, StartIndex * IndexStride, NumInstances, BaseVertexIndex, FirstInstance);
+		METAL_DEBUG_LAYER(EAGXDebugLevelFastValidation, CurrentEncoder.GetRenderCommandEncoderDebugging().DrawIndexed(AGXTranslatePrimitiveType(PrimitiveType), NumIndices, ((IndexStride == 2) ? mtlpp::IndexType::UInt16 : mtlpp::IndexType::UInt32), IndexBuffer, StartIndex * IndexStride, NumInstances, BaseVertexIndex, FirstInstance));
+	}
+	else
+	{
+		CurrentEncoder.GetRenderCommandEncoder().DrawIndexed(AGXTranslatePrimitiveType(PrimitiveType), NumIndices, ((IndexStride == 2) ? mtlpp::IndexType::UInt16 : mtlpp::IndexType::UInt32), IndexBuffer, StartIndex * IndexStride, NumInstances);
+		METAL_DEBUG_LAYER(EAGXDebugLevelFastValidation, CurrentEncoder.GetRenderCommandEncoderDebugging().DrawIndexed(AGXTranslatePrimitiveType(PrimitiveType), NumIndices, ((IndexStride == 2) ? mtlpp::IndexType::UInt16 : mtlpp::IndexType::UInt32), IndexBuffer, StartIndex * IndexStride, NumInstances));
+	}
+	
+	if (GAGXCommandBufferDebuggingEnabled)
+	{
+		FAGXCommandData Data;
+		Data.CommandType = FAGXCommandData::Type::DrawPrimitiveIndexed;
+		Data.DrawIndexed.BaseInstance = FirstInstance;
+		Data.DrawIndexed.BaseVertex = BaseVertexIndex;
+		Data.DrawIndexed.IndexCount = NumIndices;
+		Data.DrawIndexed.IndexStart = StartIndex;
+		Data.DrawIndexed.InstanceCount = NumInstances;
+		
+		InsertDebugDraw(Data);
+	}
+	
+	ConditionalSubmit();
+}
+
+void FAGXRenderPass::DrawIndexedIndirect(FAGXIndexBuffer* IndexBuffer, uint32 PrimitiveType, FAGXStructuredBuffer* VertexBuffer, int32 DrawArgumentsIndex, uint32 NumInstances)
+{
+	if (GetAGXDeviceContext().SupportsFeature(EAGXFeaturesIndirectBuffer))
+	{
+		check(NumInstances > 1);
+		
+		ConditionalSwitchToRender();
+		check(CurrentEncoder.GetCommandBuffer());
+		check(CurrentEncoder.IsRenderCommandEncoderActive());
+		
+		FAGXBuffer TheBackingIndexBuffer = IndexBuffer->GetCurrentBuffer();
+		FAGXBuffer TheBackingBuffer = VertexBuffer->GetCurrentBuffer();
+		
+		check(TheBackingIndexBuffer);
+		check(TheBackingBuffer);
+		
+		// finalize any pending state
+		PrepareToRender(PrimitiveType);
+		
+		METAL_GPUPROFILE(FAGXProfiler::GetProfiler()->EncodeDraw(CurrentEncoder.GetCommandBufferStats(), __FUNCTION__, 1, 1, 1));
+		CurrentEncoder.GetRenderCommandEncoder().DrawIndexed(AGXTranslatePrimitiveType(PrimitiveType), (mtlpp::IndexType)IndexBuffer->IndexType, TheBackingIndexBuffer, 0, TheBackingBuffer, (DrawArgumentsIndex * 5 * sizeof(uint32)));
+		METAL_DEBUG_LAYER(EAGXDebugLevelFastValidation, CurrentEncoder.GetRenderCommandEncoderDebugging().DrawIndexed(AGXTranslatePrimitiveType(PrimitiveType), (mtlpp::IndexType)IndexBuffer->IndexType, TheBackingIndexBuffer, 0, TheBackingBuffer, (DrawArgumentsIndex * 5 * sizeof(uint32))));
+
+		if (GAGXCommandBufferDebuggingEnabled)
+		{
+			FAGXCommandData Data;
+			Data.CommandType = FAGXCommandData::Type::DrawPrimitiveIndexedIndirect;
+			
+			InsertDebugDraw(Data);
+		}
+		ConditionalSubmit();
+	}
+	else
+	{
+		NOT_SUPPORTED("RHIDrawIndexedIndirect");
+	}
+}
+
+void FAGXRenderPass::DrawIndexedPrimitiveIndirect(uint32 PrimitiveType,FAGXIndexBuffer* IndexBuffer,FAGXVertexBuffer* VertexBuffer,uint32 ArgumentOffset)
+{
+	if (GetAGXDeviceContext().SupportsFeature(EAGXFeaturesIndirectBuffer))
+	{		 
+		ConditionalSwitchToRender();
+		check(CurrentEncoder.GetCommandBuffer());
+		check(CurrentEncoder.IsRenderCommandEncoderActive());
+		
+		FAGXBuffer TheBackingIndexBuffer = IndexBuffer->GetCurrentBuffer();
+		FAGXBuffer TheBackingBuffer = VertexBuffer->GetCurrentBuffer();
+		
+		check(TheBackingIndexBuffer);
+		check(TheBackingBuffer);
+		
+		PrepareToRender(PrimitiveType);
+		
+		METAL_GPUPROFILE(FAGXProfiler::GetProfiler()->EncodeDraw(CurrentEncoder.GetCommandBufferStats(), __FUNCTION__, 1, 1, 1));
+		CurrentEncoder.GetRenderCommandEncoder().DrawIndexed(AGXTranslatePrimitiveType(PrimitiveType), (mtlpp::IndexType)IndexBuffer->IndexType, TheBackingIndexBuffer, 0, TheBackingBuffer, ArgumentOffset);
+		METAL_DEBUG_LAYER(EAGXDebugLevelFastValidation, CurrentEncoder.GetRenderCommandEncoderDebugging().DrawIndexed(AGXTranslatePrimitiveType(PrimitiveType), (mtlpp::IndexType)IndexBuffer->IndexType, TheBackingIndexBuffer, 0, TheBackingBuffer, ArgumentOffset));
+
+		if (GAGXCommandBufferDebuggingEnabled)
+		{
+			FAGXCommandData Data;
+			Data.CommandType = FAGXCommandData::Type::DrawPrimitiveIndirect;
+			
+			InsertDebugDraw(Data);
+		}
+		
+		ConditionalSubmit();
+	}
+	else
+	{
+		NOT_SUPPORTED("RHIDrawIndexedPrimitiveIndirect");
+	}
+}
+
+void FAGXRenderPass::Dispatch(uint32 ThreadGroupCountX, uint32 ThreadGroupCountY, uint32 ThreadGroupCountZ)
+{
+    if (CurrentEncoder.IsParallel() || CurrentEncoder.NumEncodedPasses() == 0)
+	{
+		ConditionalSwitchToAsyncCompute();
+		check(PrologueEncoder.GetCommandBuffer());
+		check(PrologueEncoder.IsComputeCommandEncoderActive());
+		
+		PrepareToAsyncDispatch();
+		
+		TRefCountPtr<FAGXComputeShader> ComputeShader = State.GetComputeShader();
+		check(ComputeShader);
+		
+		METAL_GPUPROFILE(FAGXProfiler::GetProfiler()->EncodeDispatch(PrologueEncoder.GetCommandBufferStats(), __FUNCTION__));
+		
+		mtlpp::Size ThreadgroupCounts = mtlpp::Size(ComputeShader->NumThreadsX, ComputeShader->NumThreadsY, ComputeShader->NumThreadsZ);
+		check(ComputeShader->NumThreadsX > 0 && ComputeShader->NumThreadsY > 0 && ComputeShader->NumThreadsZ > 0);
+		mtlpp::Size Threadgroups = mtlpp::Size(ThreadGroupCountX, ThreadGroupCountY, ThreadGroupCountZ);
+		PrologueEncoder.GetComputeCommandEncoder().DispatchThreadgroups(Threadgroups, ThreadgroupCounts);
+		METAL_DEBUG_LAYER(EAGXDebugLevelFastValidation, PrologueEncoder.GetComputeCommandEncoderDebugging().DispatchThreadgroups(Threadgroups, ThreadgroupCounts));
+		
+		ConditionalSubmit();
+	}
+	else
+	{
+		ConditionalSwitchToCompute();
+		check(CurrentEncoder.GetCommandBuffer());
+		check(CurrentEncoder.IsComputeCommandEncoderActive());
+
+		PrepareToDispatch();
+		
+		TRefCountPtr<FAGXComputeShader> ComputeShader = State.GetComputeShader();
+		check(ComputeShader);
+		
+		METAL_GPUPROFILE(FAGXProfiler::GetProfiler()->EncodeDispatch(CurrentEncoder.GetCommandBufferStats(), __FUNCTION__));
+		
+		mtlpp::Size ThreadgroupCounts = mtlpp::Size(ComputeShader->NumThreadsX, ComputeShader->NumThreadsY, ComputeShader->NumThreadsZ);
+		check(ComputeShader->NumThreadsX > 0 && ComputeShader->NumThreadsY > 0 && ComputeShader->NumThreadsZ > 0);
+		mtlpp::Size Threadgroups = mtlpp::Size(ThreadGroupCountX, ThreadGroupCountY, ThreadGroupCountZ);
+		CurrentEncoder.GetComputeCommandEncoder().DispatchThreadgroups(Threadgroups, ThreadgroupCounts);
+		METAL_DEBUG_LAYER(EAGXDebugLevelFastValidation, CurrentEncoder.GetComputeCommandEncoderDebugging().DispatchThreadgroups(Threadgroups, ThreadgroupCounts));
+		
+		if (GAGXCommandBufferDebuggingEnabled)
+		{
+			FAGXCommandData Data;
+			Data.CommandType = FAGXCommandData::Type::Dispatch;
+			Data.Dispatch.threadgroupsPerGrid[0] = ThreadGroupCountX;
+			Data.Dispatch.threadgroupsPerGrid[1] = ThreadGroupCountY;
+			Data.Dispatch.threadgroupsPerGrid[2] = ThreadGroupCountZ;
+			
+			InsertDebugDispatch(Data);
+		}
+		
+		ConditionalSubmit();
+	}
+}
+
+void FAGXRenderPass::DispatchIndirect(FAGXVertexBuffer* ArgumentBuffer, uint32 ArgumentOffset)
+{
+	check(ArgumentBuffer);
+	
+    if (CurrentEncoder.IsParallel() || CurrentEncoder.NumEncodedPasses() == 0)
+	{
+		ConditionalSwitchToAsyncCompute();
+		check(PrologueEncoder.GetCommandBuffer());
+		check(PrologueEncoder.IsComputeCommandEncoderActive());
+		check(ArgumentBuffer->GetCurrentBuffer());
+		
+		PrepareToAsyncDispatch();
+		
+		TRefCountPtr<FAGXComputeShader> ComputeShader = State.GetComputeShader();
+		check(ComputeShader);
+		
+		METAL_GPUPROFILE(FAGXProfiler::GetProfiler()->EncodeDispatch(PrologueEncoder.GetCommandBufferStats(), __FUNCTION__));
+		mtlpp::Size ThreadgroupCounts = mtlpp::Size(ComputeShader->NumThreadsX, ComputeShader->NumThreadsY, ComputeShader->NumThreadsZ);
+		check(ComputeShader->NumThreadsX > 0 && ComputeShader->NumThreadsY > 0 && ComputeShader->NumThreadsZ > 0);
+		
+		PrologueEncoder.GetComputeCommandEncoder().DispatchThreadgroupsWithIndirectBuffer(ArgumentBuffer->GetCurrentBuffer(), ArgumentOffset, ThreadgroupCounts);
+		METAL_DEBUG_LAYER(EAGXDebugLevelFastValidation, PrologueEncoder.GetComputeCommandEncoderDebugging().DispatchThreadgroupsWithIndirectBuffer(ArgumentBuffer->GetCurrentBuffer(), ArgumentOffset, ThreadgroupCounts));
+		
+		ConditionalSubmit();
+	}
+	else
+	{
+		ConditionalSwitchToCompute();
+		check(CurrentEncoder.GetCommandBuffer());
+		check(CurrentEncoder.IsComputeCommandEncoderActive());
+		
+		PrepareToDispatch();
+		
+		TRefCountPtr<FAGXComputeShader> ComputeShader = State.GetComputeShader();
+		check(ComputeShader);
+		
+		METAL_GPUPROFILE(FAGXProfiler::GetProfiler()->EncodeDispatch(CurrentEncoder.GetCommandBufferStats(), __FUNCTION__));
+		mtlpp::Size ThreadgroupCounts = mtlpp::Size(ComputeShader->NumThreadsX, ComputeShader->NumThreadsY, ComputeShader->NumThreadsZ);
+		check(ComputeShader->NumThreadsX > 0 && ComputeShader->NumThreadsY > 0 && ComputeShader->NumThreadsZ > 0);
+
+		CurrentEncoder.GetComputeCommandEncoder().DispatchThreadgroupsWithIndirectBuffer(ArgumentBuffer->GetCurrentBuffer(), ArgumentOffset, ThreadgroupCounts);
+		METAL_DEBUG_LAYER(EAGXDebugLevelFastValidation, CurrentEncoder.GetComputeCommandEncoderDebugging().DispatchThreadgroupsWithIndirectBuffer(ArgumentBuffer->GetCurrentBuffer(), ArgumentOffset, ThreadgroupCounts));
+
+		if (GAGXCommandBufferDebuggingEnabled)
+		{
+			FAGXCommandData Data;
+			Data.CommandType = FAGXCommandData::Type::DispatchIndirect;
+			Data.DispatchIndirect.ArgumentBuffer = ArgumentBuffer->GetCurrentBuffer();
+			Data.DispatchIndirect.ArgumentOffset = ArgumentOffset;
+			
+			InsertDebugDispatch(Data);
+		}
+	
+		ConditionalSubmit();
+	}
+}
+
+TRefCountPtr<FAGXFence> const& FAGXRenderPass::EndRenderPass(void)
+{
+	if (bWithinRenderPass)
+	{
+		check(RenderPassDesc);
+		check(CurrentEncoder.GetCommandBuffer());
+		
+		// This just calls End - it exists only to enforce assumptions
+		End();
+	}
+	return CurrentEncoderFence;
+}
+
+void FAGXRenderPass::InsertTextureBarrier()
+{
+#if PLATFORM_MAC
+	check(CurrentEncoder.IsRenderCommandEncoderActive());
+	
+	id <MTLRenderCommandEncoder> RenderEncoder = CurrentEncoder.GetRenderCommandEncoder().GetPtr();
+	check(RenderEncoder);
+	[RenderEncoder memoryBarrierWithScope:MTLBarrierScopeRenderTargets afterStages:MTLRenderStageFragment beforeStages:MTLRenderStageVertex];
+#endif
+}
+
+void FAGXRenderPass::CopyFromTextureToBuffer(FAGXTexture const& Texture, uint32 sourceSlice, uint32 sourceLevel, mtlpp::Origin sourceOrigin, mtlpp::Size sourceSize, FAGXBuffer const& toBuffer, uint32 destinationOffset, uint32 destinationBytesPerRow, uint32 destinationBytesPerImage, mtlpp::BlitOption options)
+{
+	ConditionalSwitchToBlit();
+	mtlpp::BlitCommandEncoder& Encoder = CurrentEncoder.GetBlitCommandEncoder();
+	check(Encoder.GetPtr());
+	
+	METAL_GPUPROFILE(FAGXProfiler::GetProfiler()->EncodeBlit(CurrentEncoder.GetCommandBufferStats(), __FUNCTION__));
+	{
+		MTLPP_VALIDATE(mtlpp::BlitCommandEncoder, Encoder, AGXSafeGetRuntimeDebuggingLevel() >= EAGXDebugLevelValidation, Copy(Texture, sourceSlice, sourceLevel, sourceOrigin, sourceSize, toBuffer, destinationOffset, destinationBytesPerRow, destinationBytesPerImage, options));
+		METAL_DEBUG_LAYER(EAGXDebugLevelFastValidation, CurrentEncoder.GetBlitCommandEncoderDebugging().Copy(Texture, sourceSlice, sourceLevel, sourceOrigin, sourceSize, toBuffer, destinationOffset, destinationBytesPerRow, destinationBytesPerImage, options));
+	}
+	ConditionalSubmit();
+}
+
+void FAGXRenderPass::CopyFromBufferToTexture(FAGXBuffer const& Buffer, uint32 sourceOffset, uint32 sourceBytesPerRow, uint32 sourceBytesPerImage, mtlpp::Size sourceSize, FAGXTexture const& toTexture, uint32 destinationSlice, uint32 destinationLevel, mtlpp::Origin destinationOrigin, mtlpp::BlitOption options)
+{
+	ConditionalSwitchToBlit();
+	mtlpp::BlitCommandEncoder& Encoder = CurrentEncoder.GetBlitCommandEncoder();
+	check(Encoder.GetPtr());
+	
+	
+	METAL_GPUPROFILE(FAGXProfiler::GetProfiler()->EncodeBlit(CurrentEncoder.GetCommandBufferStats(), __FUNCTION__));
+	if (options == mtlpp::BlitOption::None)
+	{
+		MTLPP_VALIDATE(mtlpp::BlitCommandEncoder, Encoder, AGXSafeGetRuntimeDebuggingLevel() >= EAGXDebugLevelValidation, Copy(Buffer, sourceOffset, sourceBytesPerRow, sourceBytesPerImage, sourceSize, toTexture, destinationSlice, destinationLevel, destinationOrigin));
+		METAL_DEBUG_LAYER(EAGXDebugLevelFastValidation, CurrentEncoder.GetBlitCommandEncoderDebugging().Copy(Buffer, sourceOffset, sourceBytesPerRow, sourceBytesPerImage, sourceSize, toTexture, destinationSlice, destinationLevel, destinationOrigin));
+	}
+	else
+	{
+		MTLPP_VALIDATE(mtlpp::BlitCommandEncoder, Encoder, AGXSafeGetRuntimeDebuggingLevel() >= EAGXDebugLevelValidation, Copy(Buffer, sourceOffset, sourceBytesPerRow, sourceBytesPerImage, sourceSize, toTexture, destinationSlice, destinationLevel, destinationOrigin, options));
+		METAL_DEBUG_LAYER(EAGXDebugLevelFastValidation, CurrentEncoder.GetBlitCommandEncoderDebugging().Copy(Buffer, sourceOffset, sourceBytesPerRow, sourceBytesPerImage, sourceSize, toTexture, destinationSlice, destinationLevel, destinationOrigin, options));
+	}
+	ConditionalSubmit();
+}
+
+void FAGXRenderPass::CopyFromTextureToTexture(FAGXTexture const& Texture, uint32 sourceSlice, uint32 sourceLevel, mtlpp::Origin sourceOrigin, mtlpp::Size sourceSize, FAGXTexture const& toTexture, uint32 destinationSlice, uint32 destinationLevel, mtlpp::Origin destinationOrigin)
+{
+	ConditionalSwitchToBlit();
+	mtlpp::BlitCommandEncoder& Encoder = CurrentEncoder.GetBlitCommandEncoder();
+	check(Encoder.GetPtr());
+	
+	METAL_GPUPROFILE(FAGXProfiler::GetProfiler()->EncodeBlit(CurrentEncoder.GetCommandBufferStats(), __FUNCTION__));
+	MTLPP_VALIDATE(mtlpp::BlitCommandEncoder, Encoder, AGXSafeGetRuntimeDebuggingLevel() >= EAGXDebugLevelValidation, Copy(Texture, sourceSlice, sourceLevel, sourceOrigin, sourceSize, toTexture, destinationSlice, destinationLevel, destinationOrigin));
+	METAL_DEBUG_LAYER(EAGXDebugLevelFastValidation, CurrentEncoder.GetBlitCommandEncoderDebugging().Copy(Texture, sourceSlice, sourceLevel, sourceOrigin, sourceSize, toTexture, destinationSlice, destinationLevel, destinationOrigin));
+	ConditionalSubmit();
+}
+
+void FAGXRenderPass::CopyFromBufferToBuffer(FAGXBuffer const& SourceBuffer, NSUInteger SourceOffset, FAGXBuffer const& DestinationBuffer, NSUInteger DestinationOffset, NSUInteger Size)
+{
+	ConditionalSwitchToBlit();
+	mtlpp::BlitCommandEncoder& Encoder = CurrentEncoder.GetBlitCommandEncoder();
+	check(Encoder.GetPtr());
+	
+	METAL_GPUPROFILE(FAGXProfiler::GetProfiler()->EncodeBlit(CurrentEncoder.GetCommandBufferStats(), __FUNCTION__));
+	MTLPP_VALIDATE(mtlpp::BlitCommandEncoder, Encoder, AGXSafeGetRuntimeDebuggingLevel() >= EAGXDebugLevelValidation, Copy(SourceBuffer, SourceOffset, DestinationBuffer, DestinationOffset, Size));
+	METAL_DEBUG_LAYER(EAGXDebugLevelFastValidation, CurrentEncoder.GetBlitCommandEncoderDebugging().Copy(SourceBuffer, SourceOffset, DestinationBuffer, DestinationOffset, Size));
+	ConditionalSubmit();
+}
+
+void FAGXRenderPass::PresentTexture(FAGXTexture const& Texture, uint32 sourceSlice, uint32 sourceLevel, mtlpp::Origin sourceOrigin, mtlpp::Size sourceSize, FAGXTexture const& toTexture, uint32 destinationSlice, uint32 destinationLevel, mtlpp::Origin destinationOrigin)
+{
+	ConditionalSwitchToBlit();
+	mtlpp::BlitCommandEncoder& Encoder = CurrentEncoder.GetBlitCommandEncoder();
+	check(Encoder.GetPtr());
+	
+	METAL_GPUPROFILE(FAGXProfiler::GetProfiler()->EncodeBlit(CurrentEncoder.GetCommandBufferStats(), __FUNCTION__));
+	MTLPP_VALIDATE(mtlpp::BlitCommandEncoder, Encoder, AGXSafeGetRuntimeDebuggingLevel() >= EAGXDebugLevelValidation, Copy(Texture, sourceSlice, sourceLevel, sourceOrigin, sourceSize, toTexture, destinationSlice, destinationLevel, destinationOrigin));
+	METAL_DEBUG_LAYER(EAGXDebugLevelFastValidation, CurrentEncoder.GetBlitCommandEncoderDebugging().Copy(Texture, sourceSlice, sourceLevel, sourceOrigin, sourceSize, toTexture, destinationSlice, destinationLevel, destinationOrigin));
+}
+
+void FAGXRenderPass::SynchronizeTexture(FAGXTexture const& Texture, uint32 Slice, uint32 Level)
+{
+	check(Texture);
+#if PLATFORM_MAC
+	ConditionalSwitchToBlit();
+	mtlpp::BlitCommandEncoder& Encoder = CurrentEncoder.GetBlitCommandEncoder();
+	check(Encoder.GetPtr());
+	
+	// METAL_GPUPROFILE(FAGXProfiler::GetProfiler()->EncodeBlit(CurrentEncoder.GetCommandBufferStats(), __FUNCTION__));
+	MTLPP_VALIDATE(mtlpp::BlitCommandEncoder, Encoder, AGXSafeGetRuntimeDebuggingLevel() >= EAGXDebugLevelValidation, Synchronize(Texture, Slice, Level));
+	METAL_DEBUG_LAYER(EAGXDebugLevelFastValidation, CurrentEncoder.GetBlitCommandEncoderDebugging().Synchronize(Texture, Slice, Level));
+	ConditionalSubmit();
+#endif
+}
+
+void FAGXRenderPass::SynchroniseResource(mtlpp::Resource const& Resource)
+{
+	check(Resource);
+#if PLATFORM_MAC
+	ConditionalSwitchToBlit();
+	mtlpp::BlitCommandEncoder& Encoder = CurrentEncoder.GetBlitCommandEncoder();
+	check(Encoder.GetPtr());
+	
+	// METAL_GPUPROFILE(FAGXProfiler::GetProfiler()->EncodeBlit(CurrentEncoder.GetCommandBufferStats(), __FUNCTION__));
+	MTLPP_VALIDATE(mtlpp::BlitCommandEncoder, Encoder, AGXSafeGetRuntimeDebuggingLevel() >= EAGXDebugLevelValidation, Synchronize(Resource));
+	METAL_DEBUG_LAYER(EAGXDebugLevelFastValidation, CurrentEncoder.GetBlitCommandEncoderDebugging().Synchronize(Resource));
+	ConditionalSubmit();
+#endif
+}
+
+void FAGXRenderPass::FillBuffer(FAGXBuffer const& Buffer, ns::Range Range, uint8 Value)
+{
+	check(Buffer);
+	
+	mtlpp::BlitCommandEncoder TargetEncoder;
+	METAL_DEBUG_ONLY(FAGXBlitCommandEncoderDebugging Debugging);
+	bool bAsync = !CurrentEncoder.HasBufferBindingHistory(Buffer);
+	if(bAsync)
+	{
+		ConditionalSwitchToAsyncBlit();
+		TargetEncoder = PrologueEncoder.GetBlitCommandEncoder();
+		METAL_GPUPROFILE(FAGXProfiler::GetProfiler()->EncodeBlit(PrologueEncoder.GetCommandBufferStats(), FString::Printf(TEXT("FillBuffer: %p %llu %llu"), Buffer.GetPtr(), Buffer.GetOffset() + Range.Location, Range.Length)));
+		METAL_DEBUG_LAYER(EAGXDebugLevelFastValidation, Debugging = PrologueEncoder.GetBlitCommandEncoderDebugging());
+	}
+	else
+	{
+		ConditionalSwitchToBlit();
+		TargetEncoder = CurrentEncoder.GetBlitCommandEncoder();
+		METAL_GPUPROFILE(FAGXProfiler::GetProfiler()->EncodeBlit(CurrentEncoder.GetCommandBufferStats(), FString::Printf(TEXT("FillBuffer: %p %llu %llu"), Buffer.GetPtr(), Buffer.GetOffset() + Range.Location, Range.Length)));
+		METAL_DEBUG_LAYER(EAGXDebugLevelFastValidation, Debugging = CurrentEncoder.GetBlitCommandEncoderDebugging());
+	}
+	
+	check(TargetEncoder.GetPtr());
+	
+	MTLPP_VALIDATE(mtlpp::BlitCommandEncoder, TargetEncoder, AGXSafeGetRuntimeDebuggingLevel() >= EAGXDebugLevelValidation, Fill(Buffer, Range, Value));
+	METAL_DEBUG_LAYER(EAGXDebugLevelFastValidation, (bAsync ? PrologueEncoder.GetBlitCommandEncoderDebugging() : CurrentEncoder.GetBlitCommandEncoderDebugging()).Fill(Buffer, Range, Value));
+	
+	if (!bAsync)
+	{
+		ConditionalSubmit();
+	}
+}
+
+bool FAGXRenderPass::AsyncCopyFromBufferToTexture(FAGXBuffer const& Buffer, uint32 sourceOffset, uint32 sourceBytesPerRow, uint32 sourceBytesPerImage, mtlpp::Size sourceSize, FAGXTexture const& toTexture, uint32 destinationSlice, uint32 destinationLevel, mtlpp::Origin destinationOrigin, mtlpp::BlitOption options)
+{
+	mtlpp::BlitCommandEncoder TargetEncoder;
+	METAL_DEBUG_ONLY(FAGXBlitCommandEncoderDebugging Debugging);
+	bool bAsync = !CurrentEncoder.HasTextureBindingHistory(toTexture);
+	if(bAsync)
+	{
+		ConditionalSwitchToAsyncBlit();
+		TargetEncoder = PrologueEncoder.GetBlitCommandEncoder();
+		METAL_GPUPROFILE(FAGXProfiler::GetProfiler()->EncodeBlit(PrologueEncoder.GetCommandBufferStats(), __FUNCTION__));
+		METAL_DEBUG_LAYER(EAGXDebugLevelFastValidation, Debugging = PrologueEncoder.GetBlitCommandEncoderDebugging());
+	}
+	else
+	{
+		ConditionalSwitchToBlit();
+		TargetEncoder = CurrentEncoder.GetBlitCommandEncoder();
+		METAL_GPUPROFILE(FAGXProfiler::GetProfiler()->EncodeBlit(CurrentEncoder.GetCommandBufferStats(), __FUNCTION__));
+		METAL_DEBUG_LAYER(EAGXDebugLevelFastValidation, Debugging = CurrentEncoder.GetBlitCommandEncoderDebugging());
+	}
+	
+	check(TargetEncoder.GetPtr());
+	
+	if (options == mtlpp::BlitOption::None)
+	{
+		MTLPP_VALIDATE(mtlpp::BlitCommandEncoder, TargetEncoder, AGXSafeGetRuntimeDebuggingLevel() >= EAGXDebugLevelValidation, Copy(Buffer, sourceOffset, sourceBytesPerRow, sourceBytesPerImage, sourceSize, toTexture, destinationSlice, destinationLevel, destinationOrigin));
+		METAL_DEBUG_LAYER(EAGXDebugLevelFastValidation, Debugging.Copy(Buffer, sourceOffset, sourceBytesPerRow, sourceBytesPerImage, sourceSize, toTexture, destinationSlice, destinationLevel, destinationOrigin));
+	}
+	else
+	{
+		MTLPP_VALIDATE(mtlpp::BlitCommandEncoder, TargetEncoder, AGXSafeGetRuntimeDebuggingLevel() >= EAGXDebugLevelValidation, Copy(Buffer, sourceOffset, sourceBytesPerRow, sourceBytesPerImage, sourceSize, toTexture, destinationSlice, destinationLevel, destinationOrigin, options));
+		METAL_DEBUG_LAYER(EAGXDebugLevelFastValidation, Debugging.Copy(Buffer, sourceOffset, sourceBytesPerRow, sourceBytesPerImage, sourceSize, toTexture, destinationSlice, destinationLevel, destinationOrigin, options));
+	}
+	
+	return bAsync;
+}
+
+bool FAGXRenderPass::AsyncCopyFromTextureToTexture(FAGXTexture const& Texture, uint32 sourceSlice, uint32 sourceLevel, mtlpp::Origin sourceOrigin, mtlpp::Size sourceSize, FAGXTexture const& toTexture, uint32 destinationSlice, uint32 destinationLevel, mtlpp::Origin destinationOrigin)
+{
+	mtlpp::BlitCommandEncoder TargetEncoder;
+	METAL_DEBUG_ONLY(FAGXBlitCommandEncoderDebugging Debugging);
+	bool bAsync = !CurrentEncoder.HasTextureBindingHistory(toTexture);
+	if(bAsync)
+	{
+		ConditionalSwitchToAsyncBlit();
+		TargetEncoder = PrologueEncoder.GetBlitCommandEncoder();
+		METAL_GPUPROFILE(FAGXProfiler::GetProfiler()->EncodeBlit(PrologueEncoder.GetCommandBufferStats(), __FUNCTION__));
+		METAL_DEBUG_LAYER(EAGXDebugLevelFastValidation, Debugging = PrologueEncoder.GetBlitCommandEncoderDebugging());
+	}
+	else
+	{
+		ConditionalSwitchToBlit();
+		TargetEncoder = CurrentEncoder.GetBlitCommandEncoder();
+		METAL_GPUPROFILE(FAGXProfiler::GetProfiler()->EncodeBlit(CurrentEncoder.GetCommandBufferStats(), __FUNCTION__));
+		METAL_DEBUG_LAYER(EAGXDebugLevelFastValidation, Debugging = CurrentEncoder.GetBlitCommandEncoderDebugging());
+	}
+	
+	check(TargetEncoder.GetPtr());
+	
+	MTLPP_VALIDATE(mtlpp::BlitCommandEncoder, TargetEncoder, AGXSafeGetRuntimeDebuggingLevel() >= EAGXDebugLevelValidation, Copy(Texture, sourceSlice, sourceLevel, sourceOrigin, sourceSize, toTexture, destinationSlice, destinationLevel, destinationOrigin));
+	METAL_DEBUG_LAYER(EAGXDebugLevelFastValidation, Debugging.Copy(Texture, sourceSlice, sourceLevel, sourceOrigin, sourceSize, toTexture, destinationSlice, destinationLevel, destinationOrigin));
+	
+	return bAsync;
+}
+
+bool FAGXRenderPass::CanAsyncCopyToBuffer(FAGXBuffer const& DestinationBuffer)
+{
+	return !CurrentEncoder.HasBufferBindingHistory(DestinationBuffer);
+}
+
+void FAGXRenderPass::AsyncCopyFromBufferToBuffer(FAGXBuffer const& SourceBuffer, NSUInteger SourceOffset, FAGXBuffer const& DestinationBuffer, NSUInteger DestinationOffset, NSUInteger Size)
+{
+	mtlpp::BlitCommandEncoder TargetEncoder;
+	METAL_DEBUG_ONLY(FAGXBlitCommandEncoderDebugging Debugging);
+	bool bAsync = !CurrentEncoder.HasBufferBindingHistory(DestinationBuffer);
+	if(bAsync)
+	{
+		ConditionalSwitchToAsyncBlit();
+		TargetEncoder = PrologueEncoder.GetBlitCommandEncoder();
+		METAL_GPUPROFILE(FAGXProfiler::GetProfiler()->EncodeBlit(PrologueEncoder.GetCommandBufferStats(), FString::Printf(TEXT("AsyncCopyFromBufferToBuffer: %p %llu %llu"), DestinationBuffer.GetPtr(), DestinationBuffer.GetOffset() + DestinationOffset, Size)));
+		METAL_DEBUG_LAYER(EAGXDebugLevelFastValidation, Debugging = PrologueEncoder.GetBlitCommandEncoderDebugging());
+	}
+	else
+	{
+		ConditionalSwitchToBlit();
+		TargetEncoder = CurrentEncoder.GetBlitCommandEncoder();
+		METAL_GPUPROFILE(FAGXProfiler::GetProfiler()->EncodeBlit(CurrentEncoder.GetCommandBufferStats(), FString::Printf(TEXT("AsyncCopyFromBufferToBuffer: %p %llu %llu"), DestinationBuffer.GetPtr(), DestinationBuffer.GetOffset() + DestinationOffset, Size)));
+		METAL_DEBUG_LAYER(EAGXDebugLevelFastValidation, Debugging = CurrentEncoder.GetBlitCommandEncoderDebugging());
+	}
+	
+	check(TargetEncoder.GetPtr());
+	
+    MTLPP_VALIDATE(mtlpp::BlitCommandEncoder, TargetEncoder, AGXSafeGetRuntimeDebuggingLevel() >= EAGXDebugLevelValidation, Copy(SourceBuffer, SourceOffset, DestinationBuffer, DestinationOffset, Size));
+	METAL_DEBUG_LAYER(EAGXDebugLevelFastValidation, Debugging.Copy(SourceBuffer, SourceOffset, DestinationBuffer, DestinationOffset, Size));
+}
+
+FAGXBuffer FAGXRenderPass::AllocateTemporyBufferForCopy(FAGXBuffer const& DestinationBuffer, NSUInteger Size, NSUInteger Align)
+{
+	FAGXBuffer Buffer;
+	bool bAsync = !CurrentEncoder.HasBufferBindingHistory(DestinationBuffer);
+	if(bAsync)
+	{
+		Buffer = PrologueEncoder.GetRingBuffer().NewBuffer(Size, Align);
+	}
+	else
+	{
+		Buffer = CurrentEncoder.GetRingBuffer().NewBuffer(Size, Align);
+	}
+	return Buffer;
+}
+
+void FAGXRenderPass::AsyncGenerateMipmapsForTexture(FAGXTexture const& Texture)
+{
+	// This must be a plain old error
+	check(!CurrentEncoder.HasTextureBindingHistory(Texture));
+	ConditionalSwitchToAsyncBlit();
+	mtlpp::BlitCommandEncoder Encoder = PrologueEncoder.GetBlitCommandEncoder();
+	check(Encoder.GetPtr());
+	
+	METAL_GPUPROFILE(FAGXProfiler::GetProfiler()->EncodeBlit(CurrentEncoder.GetCommandBufferStats(), __FUNCTION__));
+	MTLPP_VALIDATE(mtlpp::BlitCommandEncoder, Encoder, AGXSafeGetRuntimeDebuggingLevel() >= EAGXDebugLevelValidation, GenerateMipmaps(Texture));
+	METAL_DEBUG_LAYER(EAGXDebugLevelFastValidation, PrologueEncoder.GetBlitCommandEncoderDebugging().GenerateMipmaps(Texture));
+}
+
+TRefCountPtr<FAGXFence> const& FAGXRenderPass::End(void)
+{
+	// EndEncoding should provide the encoder fence...
+	if (PrologueEncoder.IsBlitCommandEncoderActive() || PrologueEncoder.IsComputeCommandEncoderActive())
+	{
+		PrologueEncoderFence = PrologueEncoder.EndEncoding();
+	}
+	
+	if (CmdList.IsImmediate() && IsWithinParallelPass() && CurrentEncoder.IsParallelRenderCommandEncoderActive())
+	{
+		State.SetRenderStoreActions(CurrentEncoder, false);
+		CurrentEncoder.EndEncoding();
+		
+		ConditionalSwitchToBlit();
+		CurrentEncoderFence = CurrentEncoder.EndEncoding();
+		ParallelPassEndFence = nullptr;
+		PassStartFence = nullptr;
+	}
+	else if (CurrentEncoder.IsRenderCommandEncoderActive() || CurrentEncoder.IsBlitCommandEncoderActive() || CurrentEncoder.IsComputeCommandEncoderActive())
+	{
+		State.FlushVisibilityResults(CurrentEncoder);
+		check(!CurrentEncoderFence);
+		check(!PassStartFence);
+		check(!ParallelPassEndFence);
+		CurrentEncoderFence = CurrentEncoder.EndEncoding();
+	}
+	else if (PassStartFence || ParallelPassEndFence)
+	{
+		ConditionalSwitchToBlit();
+		CurrentEncoderFence = CurrentEncoder.EndEncoding();
+		ParallelPassEndFence = nullptr;
+		PassStartFence = nullptr;
+	}
+	
+	check(!PassStartFence);
+	check(!ParallelPassEndFence);
+	
+	State.SetRenderTargetsActive(false);
+	
+	RenderPassDesc = nil;
+	bWithinRenderPass = false;
+	
+	return CurrentEncoderFence;
+}
+
+void FAGXRenderPass::InsertCommandBufferFence(FAGXCommandBufferFence& Fence, mtlpp::CommandBufferHandler Handler)
+{
+	CurrentEncoder.InsertCommandBufferFence(Fence, Handler);
+}
+
+void FAGXRenderPass::AddCompletionHandler(mtlpp::CommandBufferHandler Handler)
+{
+	CurrentEncoder.AddCompletionHandler(Handler);
+}
+
+void FAGXRenderPass::AddAsyncCommandBufferHandlers(mtlpp::CommandBufferHandler Scheduled, mtlpp::CommandBufferHandler Completion)
+{
+	check(PrologueEncoder.GetCommandBuffer() && PrologueEncoder.IsBlitCommandEncoderActive());
+	if (Scheduled)
+	{
+		PrologueEncoder.GetCommandBuffer().AddScheduledHandler(Scheduled);
+	}
+	if (Completion)
+	{
+		PrologueEncoder.AddCompletionHandler(Completion);
+	}
+}
+
+void FAGXRenderPass::TransitionResources(mtlpp::Resource const& Resource)
+{
+	PrologueEncoder.TransitionResources(Resource);
+	CurrentEncoder.TransitionResources(Resource);
+}
+
+#pragma mark - Public Debug Support -
+
+void FAGXRenderPass::InsertDebugEncoder()
+{
+	FAGXBuffer NewBuf = CurrentEncoder.GetRingBuffer().NewBuffer(BufferOffsetAlignment, BufferOffsetAlignment);
+	
+	check(NewBuf);
+	
+	mtlpp::BlitCommandEncoder TargetEncoder;
+	METAL_DEBUG_ONLY(FAGXBlitCommandEncoderDebugging Debugging);
+	ConditionalSwitchToBlit();
+	TargetEncoder = CurrentEncoder.GetBlitCommandEncoder();
+	METAL_GPUPROFILE(FAGXProfiler::GetProfiler()->EncodeBlit(CurrentEncoder.GetCommandBufferStats(), __FUNCTION__));
+	METAL_DEBUG_LAYER(EAGXDebugLevelFastValidation, Debugging = CurrentEncoder.GetBlitCommandEncoderDebugging());
+	
+	check(TargetEncoder.GetPtr());
+	
+	MTLPP_VALIDATE(mtlpp::BlitCommandEncoder, TargetEncoder, AGXSafeGetRuntimeDebuggingLevel() >= EAGXDebugLevelValidation, Fill(NewBuf, ns::Range(0, BufferOffsetAlignment), 0xff));
+	METAL_DEBUG_LAYER(EAGXDebugLevelFastValidation, CurrentEncoder.GetBlitCommandEncoderDebugging().Fill(NewBuf, ns::Range(0, BufferOffsetAlignment), 0xff));
+	
+	ConditionalSubmit();
+}
+
+void FAGXRenderPass::InsertDebugSignpost(ns::String const& String)
+{
+	CurrentEncoder.InsertDebugSignpost(String);
+	PrologueEncoder.InsertDebugSignpost(FString::Printf(TEXT("Prologue %s"), *FString(String.GetPtr())).GetNSString());
+}
+
+void FAGXRenderPass::PushDebugGroup(ns::String const& String)
+{
+	CurrentEncoder.PushDebugGroup(String);
+	PrologueEncoder.PushDebugGroup(FString::Printf(TEXT("Prologue %s"), *FString(String.GetPtr())).GetNSString());
+}
+
+void FAGXRenderPass::PopDebugGroup(void)
+{
+	CurrentEncoder.PopDebugGroup();
+	PrologueEncoder.PopDebugGroup();
+}
+
+#pragma mark - Public Accessors -
+	
+mtlpp::CommandBuffer const& FAGXRenderPass::GetCurrentCommandBuffer(void) const
+{
+	return CurrentEncoder.GetCommandBuffer();
+}
+
+mtlpp::CommandBuffer& FAGXRenderPass::GetCurrentCommandBuffer(void)
+{
+	return CurrentEncoder.GetCommandBuffer();
+}
+	
+FAGXSubBufferRing& FAGXRenderPass::GetRingBuffer(void)
+{
+	return CurrentEncoder.GetRingBuffer();
+}
+
+void FAGXRenderPass::ShrinkRingBuffers(void)
+{
+	PrologueEncoder.GetRingBuffer().Shrink();
+	CurrentEncoder.GetRingBuffer().Shrink();
+}
+
+bool FAGXRenderPass::IsWithinParallelPass(void)
+{
+	return bWithinRenderPass && CurrentEncoder.IsParallelRenderCommandEncoderActive();
+}
+
+mtlpp::RenderCommandEncoder FAGXRenderPass::GetParallelRenderCommandEncoder(uint32 Index, mtlpp::ParallelRenderCommandEncoder& ParallelEncoder)
+{
+	check(IsWithinParallelPass());
+	ParallelEncoder = CurrentEncoder.GetParallelRenderCommandEncoder();
+	return CurrentEncoder.GetChildRenderCommandEncoder(Index);
+}
+
+void FAGXRenderPass::ConditionalSwitchToRender(void)
+{
+	SCOPE_CYCLE_COUNTER(STAT_AGXSwitchToRenderTime);
+	
+	check(bWithinRenderPass);
+	check(RenderPassDesc);
+	check(CmdList.IsParallel() || CurrentEncoder.GetCommandBuffer());
+	
+	if (CurrentEncoder.IsComputeCommandEncoderActive() || CurrentEncoder.IsBlitCommandEncoderActive())
+	{
+		CurrentEncoderFence = CurrentEncoder.EndEncoding();
+	}
+	
+	if (!CurrentEncoder.IsRenderCommandEncoderActive())
+	{
+		RestartRenderPass(nil);
+	}
+	
+	check(CurrentEncoder.IsRenderCommandEncoderActive());
+}
+
+void FAGXRenderPass::ConditionalSwitchToCompute(void)
+{
+	SCOPE_CYCLE_COUNTER(STAT_AGXSwitchToComputeTime);
+	
+	check(CurrentEncoder.GetCommandBuffer());
+	check(!CurrentEncoder.IsParallel());
+	
+	if (CurrentEncoder.IsRenderCommandEncoderActive() || CurrentEncoder.IsBlitCommandEncoderActive())
+	{
+		if (CurrentEncoder.IsRenderCommandEncoderActive())
+		{
+			State.SetRenderStoreActions(CurrentEncoder, true);
+			State.FlushVisibilityResults(CurrentEncoder);
+		}
+		CurrentEncoderFence = CurrentEncoder.EndEncoding();
+		State.SetRenderTargetsActive(false);
+	}
+	
+	if (!CurrentEncoder.IsComputeCommandEncoderActive())
+	{
+		State.SetStateDirty();
+		CurrentEncoder.BeginComputeCommandEncoding(ComputeDispatchType);
+		if (PassStartFence)
+		{
+			CurrentEncoder.WaitForFence(PassStartFence);
+			PassStartFence = nullptr;
+		}
+		if (ParallelPassEndFence)
+		{
+			CurrentEncoder.WaitForFence(ParallelPassEndFence);
+			ParallelPassEndFence = nullptr;
+		}
+		if (CurrentEncoderFence)
+		{
+			CurrentEncoder.WaitForFence(CurrentEncoderFence);
+			CurrentEncoderFence = nullptr;
+		}
+		if (PrologueEncoderFence)
+		{
+			CurrentEncoder.WaitForFence(PrologueEncoderFence);
+		}
+	}
+	
+	check(CurrentEncoder.IsComputeCommandEncoderActive());
+	
+	if (PrologueEncoder.IsBlitCommandEncoderActive() || PrologueEncoder.IsComputeCommandEncoderActive())
+	{
+		CurrentEncoder.WaitForFence(PrologueEncoder.GetEncoderFence());
+	}
+}
+
+void FAGXRenderPass::ConditionalSwitchToBlit(void)
+{
+	SCOPE_CYCLE_COUNTER(STAT_AGXSwitchToBlitTime);
+	
+	check(CurrentEncoder.GetCommandBuffer());
+	check(!CurrentEncoder.IsParallel());
+	
+	if (CurrentEncoder.IsRenderCommandEncoderActive() || CurrentEncoder.IsComputeCommandEncoderActive())
+	{
+		if (CurrentEncoder.IsRenderCommandEncoderActive())
+		{
+			State.SetRenderStoreActions(CurrentEncoder, true);
+			State.FlushVisibilityResults(CurrentEncoder);
+		}
+		CurrentEncoderFence = CurrentEncoder.EndEncoding();
+		State.SetRenderTargetsActive(false);
+	}
+	
+	if (!CurrentEncoder.IsBlitCommandEncoderActive())
+	{
+		CurrentEncoder.BeginBlitCommandEncoding();
+		if (PassStartFence)
+		{
+			CurrentEncoder.WaitForFence(PassStartFence);
+			PassStartFence = nullptr;
+		}
+		if (ParallelPassEndFence)
+		{
+			CurrentEncoder.WaitForFence(ParallelPassEndFence);
+			ParallelPassEndFence = nullptr;
+		}
+		if (CurrentEncoderFence)
+		{
+			CurrentEncoder.WaitForFence(CurrentEncoderFence);
+			CurrentEncoderFence = nullptr;
+		}
+		if (PrologueEncoderFence)
+		{
+			CurrentEncoder.WaitForFence(PrologueEncoderFence);
+		}
+	}
+	
+	check(CurrentEncoder.IsBlitCommandEncoderActive());
+	
+	if (PrologueEncoder.IsBlitCommandEncoderActive() || PrologueEncoder.IsComputeCommandEncoderActive())
+	{
+		CurrentEncoder.WaitForFence(PrologueEncoder.GetEncoderFence());
+	}
+}
+
+void FAGXRenderPass::ConditionalSwitchToAsyncBlit(void)
+{
+	SCOPE_CYCLE_COUNTER(STAT_AGXSwitchToAsyncBlitTime);
+	
+	if (PrologueEncoder.IsComputeCommandEncoderActive() || PrologueEncoder.IsRenderCommandEncoderActive())
+	{
+		PrologueEncoderFence = PrologueEncoder.EndEncoding();
+	}
+	
+	if (!PrologueEncoder.IsBlitCommandEncoderActive())
+	{
+		if (!PrologueEncoder.GetCommandBuffer())
+		{
+			PrologueEncoder.StartCommandBuffer();
+		}
+		PrologueEncoder.BeginBlitCommandEncoding();
+		if (PrologueStartEncoderFence)
+		{
+			if (PrologueStartEncoderFence->NeedsWait(mtlpp::RenderStages::Vertex))
+			{
+				PrologueEncoder.WaitForFence(PrologueStartEncoderFence);
+			}
+			else
+			{
+				PrologueEncoder.WaitAndUpdateFence(PrologueStartEncoderFence);
+			}
+			PrologueStartEncoderFence = nullptr;
+		}
+		if (PrologueEncoderFence)
+		{
+			if (PrologueEncoderFence->NeedsWait(mtlpp::RenderStages::Vertex))
+			{
+				PrologueEncoder.WaitForFence(PrologueEncoderFence);
+			}
+			else
+			{
+				PrologueEncoder.WaitAndUpdateFence(PrologueEncoderFence);
+			}
+			PrologueEncoderFence = nullptr;
+		}
+#if METAL_DEBUG_OPTIONS
+//		if (GetEmitDrawEvents() && PrologueEncoder.GetEncoderFence())
+//		{
+//			for (uint32 i = mtlpp::RenderStages::Vertex; i <= mtlpp::RenderStages::Fragment && PrologueEncoder.GetEncoderFence()->Get((mtlpp::RenderStages)i).GetPtr(); i++)
+//			{
+//				if (CmdList.GetCommandQueue().GetRuntimeDebuggingLevel() >= EAGXDebugLevelValidation)
+//				{
+//					PrologueEncoder.GetEncoderFence()->Get((mtlpp::RenderStages)i).GetPtr().label = [NSString stringWithFormat:@"Prologue %@", PrologueEncoderFence->Get((mtlpp::RenderStages)i).GetLabel().GetPtr()];
+//				}
+//				else
+//				{
+//					PrologueEncoder.GetEncoderFence()->Get((mtlpp::RenderStages)i).SetLabel([NSString stringWithFormat:@"Prologue %@", PrologueEncoderFence->Get((mtlpp::RenderStages)i).GetLabel().GetPtr()]);
+//				}
+//			}
+//		}
+#endif
+
+		if (CurrentEncoder.IsRenderCommandEncoderActive() || CurrentEncoder.IsComputeCommandEncoderActive() || CurrentEncoder.IsBlitCommandEncoderActive())
+		{
+			CurrentEncoder.WaitForFence(PrologueEncoder.GetEncoderFence());
+		}
+	}
+	
+	check(PrologueEncoder.IsBlitCommandEncoderActive());
+}
+
+void FAGXRenderPass::ConditionalSwitchToAsyncCompute(void)
+{
+	SCOPE_CYCLE_COUNTER(STAT_AGXSwitchToComputeTime);
+	
+	if (PrologueEncoder.IsBlitCommandEncoderActive() || PrologueEncoder.IsRenderCommandEncoderActive())
+	{
+		PrologueEncoderFence = PrologueEncoder.EndEncoding();
+	}
+	
+	if (!PrologueEncoder.IsComputeCommandEncoderActive())
+	{
+		if (!PrologueEncoder.GetCommandBuffer())
+		{
+			PrologueEncoder.StartCommandBuffer();
+		}
+		State.SetStateDirty();
+		PrologueEncoder.BeginComputeCommandEncoding(ComputeDispatchType);
+		
+		if (PrologueStartEncoderFence)
+		{
+			if (PrologueStartEncoderFence->NeedsWait(mtlpp::RenderStages::Vertex))
+			{
+				PrologueEncoder.WaitForFence(PrologueStartEncoderFence);
+			}
+			else
+			{
+				PrologueEncoder.WaitAndUpdateFence(PrologueStartEncoderFence);
+			}
+			PrologueStartEncoderFence = nullptr;
+		}
+		if (PrologueEncoderFence)
+		{
+			if (PrologueEncoderFence->NeedsWait(mtlpp::RenderStages::Vertex))
+			{
+				PrologueEncoder.WaitForFence(PrologueEncoderFence);
+			}
+			else
+			{
+				PrologueEncoder.WaitAndUpdateFence(PrologueEncoderFence);
+			}
+			PrologueEncoderFence = nullptr;
+		}
+		if (PassStartFence)
+		{
+			if(PassStartFence->NeedsWait(mtlpp::RenderStages::Vertex))
+			{
+				PrologueEncoder.WaitForFence(PassStartFence);
+			}
+			else
+			{
+				PrologueEncoder.WaitAndUpdateFence(PassStartFence);
+			}
+			PassStartFence = nullptr;
+		}
+#if METAL_DEBUG_OPTIONS
+//		if (GetEmitDrawEvents() && PrologueEncoder.GetEncoderFence())
+//		{
+//			for (uint32 i = mtlpp::RenderStages::Vertex; i <= mtlpp::RenderStages::Fragment && PrologueEncoder.GetEncoderFence()->Get((mtlpp::RenderStages)i).GetPtr(); i++)
+//			{
+//				if (CmdList.GetCommandQueue().GetRuntimeDebuggingLevel() >= EAGXDebugLevelValidation)
+//				{
+//					PrologueEncoder.GetEncoderFence()->Get((mtlpp::RenderStages)i).GetPtr().label = [NSString stringWithFormat:@"Prologue %@", PrologueEncoderFence->Get((mtlpp::RenderStages)i).GetLabel().GetPtr()];
+//				}
+//				else
+//				{
+//					PrologueEncoder.GetEncoderFence()->Get((mtlpp::RenderStages)i).SetLabel([NSString stringWithFormat:@"Prologue %@", PrologueEncoderFence->Get((mtlpp::RenderStages)i).GetLabel().GetPtr()]);
+//				}
+//			}
+//		}
+#endif
+		
+		if (CurrentEncoder.IsRenderCommandEncoderActive() || CurrentEncoder.IsComputeCommandEncoderActive() || CurrentEncoder.IsBlitCommandEncoderActive())
+		{
+			CurrentEncoder.WaitForFence(PrologueEncoder.GetEncoderFence());
+		}
+	}
+	
+	check(PrologueEncoder.IsComputeCommandEncoderActive());
+}
+
+void FAGXRenderPass::CommitRenderResourceTables(void)
+{
+	SCOPE_CYCLE_COUNTER(STAT_AGXCommitRenderResourceTablesTime);
+	
+	State.CommitRenderResources(&CurrentEncoder);
+	
+	State.CommitResourceTable(EAGXShaderStages::Vertex, mtlpp::FunctionType::Vertex, CurrentEncoder);
+	
+	FAGXGraphicsPipelineState const* BoundShaderState = State.GetGraphicsPSO();
+	
+	if (BoundShaderState->VertexShader->SideTableBinding >= 0)
+	{
+		CurrentEncoder.SetShaderSideTable(mtlpp::FunctionType::Vertex, BoundShaderState->VertexShader->SideTableBinding);
+		State.SetShaderBuffer(EAGXShaderStages::Vertex, nil, nil, 0, 0, BoundShaderState->VertexShader->SideTableBinding, mtlpp::ResourceUsage(0));
+	}
+	
+	if (IsValidRef(BoundShaderState->PixelShader))
+	{
+		State.CommitResourceTable(EAGXShaderStages::Pixel, mtlpp::FunctionType::Fragment, CurrentEncoder);
+		if (BoundShaderState->PixelShader->SideTableBinding >= 0)
+		{
+			CurrentEncoder.SetShaderSideTable(mtlpp::FunctionType::Fragment, BoundShaderState->PixelShader->SideTableBinding);
+			State.SetShaderBuffer(EAGXShaderStages::Pixel, nil, nil, 0, 0, BoundShaderState->PixelShader->SideTableBinding, mtlpp::ResourceUsage(0));
+		}
+	}
+}
+
+void FAGXRenderPass::CommitDispatchResourceTables(void)
+{
+	State.CommitComputeResources(&CurrentEncoder);
+	
+	State.CommitResourceTable(EAGXShaderStages::Compute, mtlpp::FunctionType::Kernel, CurrentEncoder);
+	
+	FAGXComputeShader const* ComputeShader = State.GetComputeShader();
+	if (ComputeShader->SideTableBinding >= 0)
+	{
+		CurrentEncoder.SetShaderSideTable(mtlpp::FunctionType::Kernel, ComputeShader->SideTableBinding);
+		State.SetShaderBuffer(EAGXShaderStages::Compute, nil, nil, 0, 0, ComputeShader->SideTableBinding, mtlpp::ResourceUsage(0));
+	}
+}
+
+void FAGXRenderPass::CommitAsyncDispatchResourceTables(void)
+{
+	State.CommitComputeResources(&PrologueEncoder);
+	
+	State.CommitResourceTable(EAGXShaderStages::Compute, mtlpp::FunctionType::Kernel, PrologueEncoder);
+	
+	FAGXComputeShader const* ComputeShader = State.GetComputeShader();
+	if (ComputeShader->SideTableBinding >= 0)
+	{
+		PrologueEncoder.SetShaderSideTable(mtlpp::FunctionType::Kernel, ComputeShader->SideTableBinding);
+		State.SetShaderBuffer(EAGXShaderStages::Compute, nil, nil, 0, 0, ComputeShader->SideTableBinding, mtlpp::ResourceUsage(0));
+	}
+}
+
+void FAGXRenderPass::PrepareToRender(uint32 PrimitiveType)
+{
+	SCOPE_CYCLE_COUNTER(STAT_AGXPrepareToRenderTime);
+	
+	check(CurrentEncoder.GetCommandBuffer());
+	check(CurrentEncoder.IsRenderCommandEncoderActive());
+	
+	// Set raster state
+	State.SetRenderState(CurrentEncoder, nullptr);
+	
+	// Bind shader resources
+	CommitRenderResourceTables();
+    
+    State.SetRenderPipelineState(CurrentEncoder, nullptr);
+}
+
+void FAGXRenderPass::PrepareToDispatch(void)
+{
+	SCOPE_CYCLE_COUNTER(STAT_AGXPrepareToDispatchTime);
+	
+	check(CurrentEncoder.GetCommandBuffer());
+	check(CurrentEncoder.IsComputeCommandEncoderActive());
+	
+	// Bind shader resources
+	CommitDispatchResourceTables();
+    
+    State.SetComputePipelineState(CurrentEncoder);
+}
+
+void FAGXRenderPass::PrepareToAsyncDispatch(void)
+{
+	SCOPE_CYCLE_COUNTER(STAT_AGXPrepareToDispatchTime);
+	
+	check(PrologueEncoder.GetCommandBuffer());
+	check(PrologueEncoder.IsComputeCommandEncoderActive());
+	
+	// Bind shader resources
+	CommitAsyncDispatchResourceTables();
+	
+	State.SetComputePipelineState(PrologueEncoder);
+}
+
+void FAGXRenderPass::ConditionalSubmit()
+{
+	NumOutstandingOps++;
+	
+	bool bCanForceSubmit = State.CanRestartRenderPass();
+
+	FRHIRenderPassInfo CurrentRenderTargets = State.GetRenderPassInfo();
+	
+	// Force a command-encoder when GAGXRuntimeDebugLevel is enabled to help track down intermittent command-buffer failures.
+	if (GAGXCommandBufferCommitThreshold > 0 && NumOutstandingOps >= GAGXCommandBufferCommitThreshold && CmdList.GetCommandQueue().GetRuntimeDebuggingLevel() >= EAGXDebugLevelConditionalSubmit)
+	{
+		bool bCanChangeRT = true;
+		
+		if (bWithinRenderPass)
+		{
+			const bool bIsMSAAActive = State.GetHasValidRenderTarget() && State.GetSampleCount() != 1;
+			bCanChangeRT = !bIsMSAAActive;
+			
+			for (int32 RenderTargetIndex = 0; bCanChangeRT && RenderTargetIndex < CurrentRenderTargets.GetNumColorRenderTargets(); RenderTargetIndex++)
+			{
+				FRHIRenderPassInfo::FColorEntry& RenderTargetView = CurrentRenderTargets.ColorRenderTargets[RenderTargetIndex];
+				
+				if (GetStoreAction(RenderTargetView.Action) != ERenderTargetStoreAction::EMultisampleResolve)
+				{
+					RenderTargetView.Action = MakeRenderTargetActions(ERenderTargetLoadAction::ELoad, ERenderTargetStoreAction::EStore);
+				}
+				else
+				{
+					bCanChangeRT = false;
+				}
+			}
+			
+			if (bCanChangeRT && CurrentRenderTargets.DepthStencilRenderTarget.DepthStencilTarget)
+			{
+				if (GetStoreAction(GetDepthActions(CurrentRenderTargets.DepthStencilRenderTarget.Action)) != ERenderTargetStoreAction::EMultisampleResolve && GetStoreAction(GetStencilActions(CurrentRenderTargets.DepthStencilRenderTarget.Action)) != ERenderTargetStoreAction::EMultisampleResolve)
+				{
+					ERenderTargetActions Actions = MakeRenderTargetActions(ERenderTargetLoadAction::ELoad, ERenderTargetStoreAction::EStore);
+					CurrentRenderTargets.DepthStencilRenderTarget.Action = MakeDepthStencilTargetActions(Actions, Actions);
+				}
+				else
+				{
+					bCanChangeRT = false;
+				}
+			}
+		}
+		
+		bCanForceSubmit = bCanChangeRT;
+	}
+	
+	if (GAGXCommandBufferCommitThreshold > 0 && NumOutstandingOps > 0 && NumOutstandingOps >= GAGXCommandBufferCommitThreshold && bCanForceSubmit && !CurrentEncoder.IsParallel())
+	{
+		if (CurrentEncoder.GetCommandBuffer())
+		{
+			Submit(EAGXSubmitFlagsCreateCommandBuffer);
+			NumOutstandingOps = 0;
+		}
+		
+		// Force a command-encoder when GAGXRuntimeDebugLevel is enabled to help track down intermittent command-buffer failures.
+		if (bWithinRenderPass && CmdList.GetCommandQueue().GetRuntimeDebuggingLevel() >= EAGXDebugLevelConditionalSubmit && State.GetHasValidRenderTarget())
+		{
+			bool bSet = false;
+			State.InvalidateRenderTargets();
+			if (IsFeatureLevelSupported( GMaxRHIShaderPlatform, ERHIFeatureLevel::SM5 ))
+			{
+				bSet = State.SetRenderPassInfo(CurrentRenderTargets, State.GetVisibilityResultsBuffer(), false);
+			}
+			else
+			{
+				bSet = State.SetRenderPassInfo(CurrentRenderTargets, NULL, false);
+			}
+			
+			if (bSet)
+			{
+				RestartRenderPass(State.GetRenderPassDescriptor());
+			}
+		}
+	}
+}
+
+uint32 FAGXRenderPass::GetEncoderIndex(void) const
+{
+	if (!CmdList.IsParallel())
+	{
+		return PrologueEncoder.NumEncodedPasses() + CurrentEncoder.NumEncodedPasses();
+	}
+	else
+	{
+		return GetAGXDeviceContext().GetCurrentRenderPass().GetEncoderIndex();
+	}
+}
+
+uint32 FAGXRenderPass::GetCommandBufferIndex(void) const
+{
+	if (!CmdList.IsParallel())
+	{
+		return CurrentEncoder.GetCommandBufferIndex();
+	}
+	else
+	{
+		return GetAGXDeviceContext().GetCurrentRenderPass().GetCommandBufferIndex();
+	}
+}
+
+#if PLATFORM_MAC
+@protocol IMTLRenderCommandEncoder
+- (void)memoryBarrierWithResources:(const id<MTLResource>[])resources count:(NSUInteger)count afterStages:(MTLRenderStages)after beforeStages:(MTLRenderStages)before;
+@end
+#endif
+@protocol IMTLComputeCommandEncoder
+- (void)memoryBarrierWithResources:(const id<MTLResource>[])resources count:(NSUInteger)count;
+@end
+
+void FAGXRenderPass::InsertDebugDraw(FAGXCommandData& Data)
+{
+#if !PLATFORM_TVOS
+	if (GAGXCommandBufferDebuggingEnabled && (!FAGXCommandQueue::SupportsFeature(EAGXFeaturesValidation) || State.GetVisibilityResultMode() == mtlpp::VisibilityResultMode::Disabled))
+	{
+		FAGXGraphicsPipelineState* BoundShaderState = State.GetGraphicsPSO();
+		
+		uint32 NumCommands = CurrentEncoder.GetMarkers().AddCommand(GetCommandBufferIndex(), GetEncoderIndex(), CmdList.GetParallelIndex(), State.GetDebugBuffer(), BoundShaderState, nullptr, Data);
+		
+		if ((NumCommands % GAGXDebugOpsCount) == 0)
+		{
+			FAGXDebugInfo DebugInfo;
+			DebugInfo.EncoderIndex = GetEncoderIndex();
+			DebugInfo.ContextIndex = CmdList.GetParallelIndex();
+			DebugInfo.CommandIndex = NumCommands;
+			DebugInfo.CmdBuffIndex = GetCommandBufferIndex();
+			DebugInfo.CommandBuffer = reinterpret_cast<uintptr_t>(CurrentEncoder.GetCommandBuffer().GetPtr());
+			DebugInfo.PSOSignature[0] = BoundShaderState->VertexShader->SourceLen;
+			DebugInfo.PSOSignature[1] = BoundShaderState->VertexShader->SourceCRC;
+			if (IsValidRef(BoundShaderState->PixelShader))
+			{
+				DebugInfo.PSOSignature[2] = BoundShaderState->PixelShader->SourceLen;
+				DebugInfo.PSOSignature[3] = BoundShaderState->PixelShader->SourceCRC;
+			}
+			else
+			{
+				DebugInfo.PSOSignature[2] = 0;
+				DebugInfo.PSOSignature[3] = 0;
+			}
+			
+			FAGXShaderPipeline* PSO = State.GetPipelineState();
+			
+			CurrentEncoder.GetRenderCommandEncoder().SetRenderPipelineState(PSO->DebugPipelineState);
+
+			mtlpp::VisibilityResultMode VisMode = State.GetVisibilityResultMode();
+			uint32 VisibilityOffset = State.GetVisibilityResultOffset();
+			if (VisMode != mtlpp::VisibilityResultMode::Disabled)
+			{
+				CurrentEncoder.GetRenderCommandEncoder().SetVisibilityResultMode(mtlpp::VisibilityResultMode::Disabled, 0);
+			}
+
+		#if PLATFORM_MAC
+			id<MTLBuffer> DebugBufferPtr = State.GetDebugBuffer().GetPtr();
+			[(id<IMTLRenderCommandEncoder>)CurrentEncoder.GetRenderCommandEncoder().GetPtr() memoryBarrierWithResources:&DebugBufferPtr count:1 afterStages:MTLRenderStageFragment beforeStages:MTLRenderStageVertex];
+			
+			CurrentEncoder.SetShaderBytes(mtlpp::FunctionType::Vertex, (uint8 const*)&DebugInfo, sizeof(DebugInfo), 0);
+			State.SetShaderBufferDirty(EAGXShaderStages::Vertex, 0);
+
+			CurrentEncoder.SetShaderBuffer(mtlpp::FunctionType::Vertex, State.GetDebugBuffer(), 0, State.GetDebugBuffer().GetLength(), 1, mtlpp::ResourceUsage::Write);
+			State.SetShaderBufferDirty(EAGXShaderStages::Vertex, 1);
+
+			CurrentEncoder.GetRenderCommandEncoder().Draw(mtlpp::PrimitiveType::Point, 0, 1);
+			
+			[(id<IMTLRenderCommandEncoder>)CurrentEncoder.GetRenderCommandEncoder().GetPtr() memoryBarrierWithResources:&DebugBufferPtr count:1 afterStages:MTLRenderStageVertex beforeStages:MTLRenderStageVertex];
+		#else
+			CurrentEncoder.GetRenderCommandEncoder().SetTileData((uint8 const*)&DebugInfo, sizeof(DebugInfo), 0);
+			CurrentEncoder.GetRenderCommandEncoder().SetTileBuffer(State.GetDebugBuffer(), 0, 1);
+			mtlpp::Size ThreadsPerTile(1, 1, 1);
+			CurrentEncoder.GetRenderCommandEncoder().DispatchThreadsPerTile(ThreadsPerTile);
+		#endif
+
+			CurrentEncoder.GetRenderCommandEncoder().SetRenderPipelineState(PSO->RenderPipelineState);
+
+			if (VisMode != mtlpp::VisibilityResultMode::Disabled)
+			{
+				CurrentEncoder.GetRenderCommandEncoder().SetVisibilityResultMode(VisMode, VisibilityOffset);
+			}
+		}
+	}
+#endif
+}
+
+void FAGXRenderPass::InsertDebugDispatch(FAGXCommandData& Data)
+{
+#if !PLATFORM_TVOS
+	if (GAGXCommandBufferDebuggingEnabled)
+	{
+		FAGXComputeShader* BoundShaderState = State.GetComputeShader();
+		
+		uint32 NumCommands = CurrentEncoder.GetMarkers().AddCommand(GetCommandBufferIndex(), GetEncoderIndex(), CmdList.GetParallelIndex(), State.GetDebugBuffer(), nullptr, BoundShaderState, Data);
+		
+		if ((NumCommands % GAGXDebugOpsCount) == 0)
+		{
+			FAGXDebugInfo DebugInfo;
+			DebugInfo.EncoderIndex = GetEncoderIndex();
+			DebugInfo.ContextIndex = CmdList.GetParallelIndex();
+			DebugInfo.CommandIndex = NumCommands;
+			DebugInfo.CmdBuffIndex = GetCommandBufferIndex();
+			DebugInfo.CommandBuffer = reinterpret_cast<uintptr_t>(CurrentEncoder.GetCommandBuffer().GetPtr());
+			DebugInfo.PSOSignature[0] = BoundShaderState->SourceLen;
+			DebugInfo.PSOSignature[1] = BoundShaderState->SourceCRC;
+			DebugInfo.PSOSignature[2] = 0;
+			DebugInfo.PSOSignature[3] = 0;
+			
+			CurrentEncoder.GetComputeCommandEncoder().SetComputePipelineState(AGXGetMetalDebugComputeState());
+			
+			id<MTLBuffer> DebugBufferPtr = State.GetDebugBuffer().GetPtr();
+			[(id<IMTLComputeCommandEncoder>)CurrentEncoder.GetComputeCommandEncoder().GetPtr() memoryBarrierWithResources:&DebugBufferPtr count:1];
+			
+			CurrentEncoder.SetShaderBytes(mtlpp::FunctionType::Kernel, (uint8 const*)&DebugInfo, sizeof(DebugInfo), 0);
+			State.SetShaderBufferDirty(EAGXShaderStages::Compute, 0);
+			
+			CurrentEncoder.SetShaderBuffer(mtlpp::FunctionType::Kernel, State.GetDebugBuffer(), 0, State.GetDebugBuffer().GetLength(), 1, mtlpp::ResourceUsage::Write);
+			State.SetShaderBufferDirty(EAGXShaderStages::Compute, 1);
+			
+			mtlpp::Size ThreadsPerTile(1, 1, 1);
+			CurrentEncoder.GetComputeCommandEncoder().DispatchThreads(ThreadsPerTile, ThreadsPerTile);
+			
+			[(id<IMTLComputeCommandEncoder>)CurrentEncoder.GetComputeCommandEncoder().GetPtr() memoryBarrierWithResources:&DebugBufferPtr count:1];
+			
+			FAGXShaderPipeline* Pipeline = BoundShaderState->GetPipeline();
+			CurrentEncoder.GetComputeCommandEncoder().SetComputePipelineState(Pipeline->ComputePipelineState);
+		}
+	}
+#endif
+}
diff --git a/Engine/Source/Runtime/Apple/AGXRHI/Private/AGXRenderPass.h b/Engine/Source/Runtime/Apple/AGXRHI/Private/AGXRenderPass.h
new file mode 100644
--- /dev/null
+++ b/Engine/Source/Runtime/Apple/AGXRHI/Private/AGXRenderPass.h
@@ -0,0 +1,212 @@
+// Copyright Epic Games, Inc. All Rights Reserved.
+
+#pragma once
+
+#include "AGXRHIPrivate.h"
+#include "AGXCommandEncoder.h"
+#include "AGXState.h"
+#include "AGXFence.h"
+
+class FAGXCommandList;
+class FAGXCommandQueue;
+
+class FAGXRenderPass
+{
+public:
+#pragma mark - Public C++ Boilerplate -
+
+	/** Default constructor */
+	FAGXRenderPass(FAGXCommandList& CmdList, FAGXStateCache& StateCache);
+	
+	/** Destructor */
+	~FAGXRenderPass(void);
+	
+#pragma mark -
+	void SetDispatchType(mtlpp::DispatchType Type);
+	
+    void Begin(FAGXFence* Fence, bool const bParallelBegin = false);
+	
+	void Wait(FAGXFence* Fence);
+
+	void Update(FAGXFence* Fence);
+	
+    void BeginParallelRenderPass(mtlpp::RenderPassDescriptor RenderPass, uint32 NumParallelContextsInPass);
+
+    void BeginRenderPass(mtlpp::RenderPassDescriptor RenderPass);
+
+    void RestartRenderPass(mtlpp::RenderPassDescriptor RenderPass);
+    
+    void DrawPrimitive(uint32 PrimitiveType, uint32 BaseVertexIndex, uint32 NumPrimitives, uint32 NumInstances);
+    
+    void DrawPrimitiveIndirect(uint32 PrimitiveType, FAGXVertexBuffer* VertexBuffer, uint32 ArgumentOffset);
+    
+    void DrawIndexedPrimitive(FAGXBuffer const& IndexBuffer, uint32 IndexStride, uint32 PrimitiveType, int32 BaseVertexIndex, uint32 FirstInstance,
+                         uint32 NumVertices, uint32 StartIndex, uint32 NumPrimitives, uint32 NumInstances);
+    
+    void DrawIndexedIndirect(FAGXIndexBuffer* IndexBufferRHI, uint32 PrimitiveType, FAGXStructuredBuffer* VertexBufferRHI, int32 DrawArgumentsIndex, uint32 NumInstances);
+    
+    void DrawIndexedPrimitiveIndirect(uint32 PrimitiveType,FAGXIndexBuffer* IndexBufferRHI,FAGXVertexBuffer* VertexBufferRHI,uint32 ArgumentOffset);
+	
+    void Dispatch(uint32 ThreadGroupCountX, uint32 ThreadGroupCountY, uint32 ThreadGroupCountZ);
+    
+    void DispatchIndirect(FAGXVertexBuffer* ArgumentBufferRHI, uint32 ArgumentOffset);
+    
+    TRefCountPtr<FAGXFence> const& EndRenderPass(void);
+    
+    void CopyFromTextureToBuffer(FAGXTexture const& Texture, uint32 sourceSlice, uint32 sourceLevel, mtlpp::Origin sourceOrigin, mtlpp::Size sourceSize, FAGXBuffer const& toBuffer, uint32 destinationOffset, uint32 destinationBytesPerRow, uint32 destinationBytesPerImage, mtlpp::BlitOption options);
+    
+    void CopyFromBufferToTexture(FAGXBuffer const& Buffer, uint32 sourceOffset, uint32 sourceBytesPerRow, uint32 sourceBytesPerImage, mtlpp::Size sourceSize, FAGXTexture const& toTexture, uint32 destinationSlice, uint32 destinationLevel, mtlpp::Origin destinationOrigin, mtlpp::BlitOption options);
+    
+    void CopyFromTextureToTexture(FAGXTexture const& Texture, uint32 sourceSlice, uint32 sourceLevel, mtlpp::Origin sourceOrigin, mtlpp::Size sourceSize, FAGXTexture const& toTexture, uint32 destinationSlice, uint32 destinationLevel, mtlpp::Origin destinationOrigin);
+	
+	void CopyFromBufferToBuffer(FAGXBuffer const& SourceBuffer, NSUInteger SourceOffset, FAGXBuffer const& DestinationBuffer, NSUInteger DestinationOffset, NSUInteger Size);
+	
+	void PresentTexture(FAGXTexture const& Texture, uint32 sourceSlice, uint32 sourceLevel, mtlpp::Origin sourceOrigin, mtlpp::Size sourceSize, FAGXTexture const& toTexture, uint32 destinationSlice, uint32 destinationLevel, mtlpp::Origin destinationOrigin);
+    
+    void SynchronizeTexture(FAGXTexture const& Texture, uint32 Slice, uint32 Level);
+    
+	void SynchroniseResource(mtlpp::Resource const& Resource);
+    
+	void FillBuffer(FAGXBuffer const& Buffer, ns::Range Range, uint8 Value);
+	
+	bool AsyncCopyFromBufferToTexture(FAGXBuffer const& Buffer, uint32 sourceOffset, uint32 sourceBytesPerRow, uint32 sourceBytesPerImage, mtlpp::Size sourceSize, FAGXTexture const& toTexture, uint32 destinationSlice, uint32 destinationLevel, mtlpp::Origin destinationOrigin, mtlpp::BlitOption options);
+	
+	bool AsyncCopyFromTextureToTexture(FAGXTexture const& Texture, uint32 sourceSlice, uint32 sourceLevel, mtlpp::Origin sourceOrigin, mtlpp::Size sourceSize, FAGXTexture const& toTexture, uint32 destinationSlice, uint32 destinationLevel, mtlpp::Origin destinationOrigin);
+	
+	bool CanAsyncCopyToBuffer(FAGXBuffer const& DestinationBuffer);
+	
+	void AsyncCopyFromBufferToBuffer(FAGXBuffer const& SourceBuffer, NSUInteger SourceOffset, FAGXBuffer const& DestinationBuffer, NSUInteger DestinationOffset, NSUInteger Size);
+	
+	FAGXBuffer AllocateTemporyBufferForCopy(FAGXBuffer const& DestinationBuffer, NSUInteger Size, NSUInteger Align);
+	
+	void AsyncGenerateMipmapsForTexture(FAGXTexture const& Texture);
+	
+    TRefCountPtr<FAGXFence> const& Submit(EAGXSubmitFlags SubmissionFlags);
+    
+    TRefCountPtr<FAGXFence> const& End(void);
+	
+	void InsertCommandBufferFence(FAGXCommandBufferFence& Fence, mtlpp::CommandBufferHandler Handler);
+	
+	void AddCompletionHandler(mtlpp::CommandBufferHandler Handler);
+	
+	void AddAsyncCommandBufferHandlers(mtlpp::CommandBufferHandler Scheduled, mtlpp::CommandBufferHandler Completion);
+	
+	void TransitionResources(mtlpp::Resource const& Resource);
+
+#pragma mark - Public Debug Support -
+	
+    /*
+     * Inserts a debug compute encoder into the command buffer. This is how we generate a timestamp when no encoder exists.
+     */
+    void InsertDebugEncoder();
+	
+	/*
+	 * Inserts a debug string into the command buffer.  This does not change any API behavior, but can be useful when debugging.
+	 * @param string The name of the signpost. 
+	 */
+	void InsertDebugSignpost(ns::String const& String);
+	
+	/*
+	 * Push a new named string onto a stack of string labels.
+	 * @param string The name of the debug group. 
+	 */
+	void PushDebugGroup(ns::String const& String);
+	
+	/* Pop the latest named string off of the stack. */
+	void PopDebugGroup(void);
+    
+#pragma mark - Public Accessors -
+	
+	/*
+	 * Get the current internal command buffer.
+	 * @returns The current command buffer.
+	 */
+	mtlpp::CommandBuffer const& GetCurrentCommandBuffer(void) const;
+	mtlpp::CommandBuffer& GetCurrentCommandBuffer(void);
+	
+    /*
+     * Get the internal current command-encoder.
+     * @returns The current command encoder.
+     */
+	inline FAGXCommandEncoder& GetCurrentCommandEncoder(void) { return CurrentEncoder; }
+	
+	/*
+	 * Get the internal ring-buffer used for temporary allocations.
+	 * @returns The temporary allocation buffer for the command-pass.
+	 */
+	FAGXSubBufferRing& GetRingBuffer(void);
+	
+	/*
+	 * Attempts to shrink the ring-buffers so we don't keep very large allocations when we don't need them.
+	 */
+	void ShrinkRingBuffers(void);
+	
+	/*
+	 * Whether the render-pass is within a parallel rendering pass.
+	 * @returns True if and only if within a parallel rendering pass, otherwise false.
+	 */
+	bool IsWithinParallelPass(void);
+
+	/*
+	 * Get a child render command-encoder and the parent parallel command-encoder when within a parallel pass.
+	 * @returns A valid render command encoder or nil.
+	 */
+    mtlpp::RenderCommandEncoder GetParallelRenderCommandEncoder(uint32 Index, mtlpp::ParallelRenderCommandEncoder& ParallelEncoder);
+	
+	void InsertTextureBarrier();
+	
+private:
+    void ConditionalSwitchToRender(void);
+    void ConditionalSwitchToCompute(void);
+	void ConditionalSwitchToBlit(void);
+	void ConditionalSwitchToAsyncBlit(void);
+	void ConditionalSwitchToAsyncCompute(void);
+	
+    void PrepareToRender(uint32 PrimType);
+    void PrepareToDispatch(void);
+	void PrepareToAsyncDispatch(void);
+
+    void CommitRenderResourceTables(void);
+    void CommitDispatchResourceTables(void);
+	void CommitAsyncDispatchResourceTables(void);
+    
+    void ConditionalSubmit();
+	
+	uint32 GetEncoderIndex(void) const;
+	uint32 GetCommandBufferIndex(void) const;
+	
+	void InsertDebugDraw(FAGXCommandData& Data);
+	void InsertDebugDispatch(FAGXCommandData& Data);
+
+private:
+	FAGXCommandList& CmdList;
+    FAGXStateCache& State;
+    
+    // Which of the buffers/textures/sampler slots are bound
+    // The state cache is responsible for ensuring we bind the correct 
+	FAGXTextureMask BoundTextures[EAGXShaderStages::Num];
+    uint32 BoundBuffers[EAGXShaderStages::Num];
+    uint16 BoundSamplers[EAGXShaderStages::Num];
+    
+    FAGXCommandEncoder CurrentEncoder;
+    FAGXCommandEncoder PrologueEncoder;
+	
+	// To ensure that buffer uploads aren't overwritten before they are used track what is in flight
+	// Disjoint ranges *are* permitted!
+	TMap<id<MTLBuffer>, TArray<NSRange>> OutstandingBufferUploads;
+
+	// Fences for the current command encoder chain
+	TRefCountPtr<FAGXFence> PassStartFence;
+	TRefCountPtr<FAGXFence> CurrentEncoderFence;
+	TRefCountPtr<FAGXFence> ParallelPassEndFence;
+
+	// Fences for the prologue command encoder chain
+	TRefCountPtr<FAGXFence> PrologueStartEncoderFence;
+	TRefCountPtr<FAGXFence> PrologueEncoderFence;
+    
+    mtlpp::RenderPassDescriptor RenderPassDesc;
+    
+	mtlpp::DispatchType ComputeDispatchType;
+    uint32 NumOutstandingOps;
+    bool bWithinRenderPass;
+};
