Initial import of AGXRHI

diff --git a/Engine/Source/Runtime/Apple/AGXRHI/Private/AGXRenderTarget.cpp b/Engine/Source/Runtime/Apple/AGXRHI/Private/AGXRenderTarget.cpp
new file mode 100644
--- /dev/null
+++ b/Engine/Source/Runtime/Apple/AGXRHI/Private/AGXRenderTarget.cpp
@@ -0,0 +1,528 @@
+// Copyright Epic Games, Inc. All Rights Reserved.
+
+/*=============================================================================
+	AGXRenderTarget.cpp: AGX RHI render target implementation.
+=============================================================================*/
+
+#include "AGXRHIPrivate.h"
+#include "ScreenRendering.h"
+#include "AGXProfiler.h"
+#include "AGXCommandBuffer.h"
+#include "ResolveShader.h"
+#include "PipelineStateCache.h"
+#include "Math/PackedVector.h"
+#include "RHISurfaceDataConversion.h"
+
+static FResolveRect GetDefaultRect(const FResolveRect& Rect, uint32 DefaultWidth, uint32 DefaultHeight)
+{
+	if (Rect.X1 >= 0 && Rect.X2 >= 0 && Rect.Y1 >= 0 && Rect.Y2 >= 0)
+	{
+		return Rect;
+	}
+	else
+	{
+		return FResolveRect(0, 0, DefaultWidth, DefaultHeight);
+	}
+}
+
+int32 GAGXUseTexGetBytes = 1;
+static FAutoConsoleVariableRef CVarAGXUseTexGetBytes(
+								TEXT("rhi.AGX.UseTexGetBytes"),
+								GAGXUseTexGetBytes,
+								TEXT("If true prefer using -[MTLTexture getBytes:...] to retreive texture data, creating a temporary shared/managed texture to copy from private texture storage when required, rather than using a temporary MTLBuffer. This works around data alignment bugs on some GPU vendor's drivers and may be more appropriate on iOS. (Default: True)"),
+								ECVF_RenderThreadSafe
+								);
+
+void FAGXRHICommandContext::RHICopyToResolveTarget(FRHITexture* SourceTextureRHI, FRHITexture* DestTextureRHI, const FResolveParams& ResolveParams)
+{
+	@autoreleasepool {
+	if (!SourceTextureRHI || !DestTextureRHI)
+	{
+		// nothing to do if one of the textures is null!
+		return;
+	}
+	if(SourceTextureRHI != DestTextureRHI)
+	{
+		FAGXSurface* Source = AGXGetMetalSurfaceFromRHITexture(SourceTextureRHI);
+		FAGXSurface* Destination = AGXGetMetalSurfaceFromRHITexture(DestTextureRHI);
+		
+		switch (Source->Type)
+		{
+			case RRT_Texture2D:
+				break;
+			case RRT_TextureCube:
+				check(Source->SizeZ == 6); // Arrays might not work yet.
+				break;
+			default:
+				check(false); // Only Tex2D & Cube are tested to work so far!
+				break;
+		}
+		switch (Destination->Type)
+		{
+			case RRT_Texture2D:
+				break;
+			case RRT_TextureCube:
+				check(Destination->SizeZ == 6); // Arrays might not work yet.
+				break;
+			default:
+				check(false); // Only Tex2D & Cube are tested to work so far!
+				break;
+		}
+		
+		mtlpp::Origin Origin(0, 0, 0);
+		mtlpp::Size Size(0, 0, 1);
+		if (ResolveParams.Rect.IsValid())
+		{
+			// Partial copy
+			Origin.x = ResolveParams.Rect.X1;
+			Origin.y = ResolveParams.Rect.Y1;
+			Size.width = ResolveParams.Rect.X2 - ResolveParams.Rect.X1;
+			Size.height = ResolveParams.Rect.Y2 - ResolveParams.Rect.Y1;
+		}
+		else
+		{
+			// Whole of source copy
+			Origin.x = 0;
+			Origin.y = 0;
+			
+			Size.width = FMath::Max<uint32>(1, Source->SizeX >> ResolveParams.MipIndex);
+			Size.height = FMath::Max<uint32>(1, Source->SizeY >> ResolveParams.MipIndex);
+		}
+		
+		const bool bSrcCubemap  = Source->bIsCubemap;
+		const bool bDestCubemap = Destination->bIsCubemap;
+		
+		uint32 DestIndex = ResolveParams.DestArrayIndex * (bDestCubemap ? 6 : 1) + (bDestCubemap ? uint32(ResolveParams.CubeFace) : 0);
+		uint32 SrcIndex  = ResolveParams.SourceArrayIndex * (bSrcCubemap ? 6 : 1) + (bSrcCubemap ? uint32(ResolveParams.CubeFace) : 0);
+		
+		if(Profiler)
+		{
+			Profiler->RegisterGPUWork();
+		}
+
+		const bool bMSAASource = Source->MSAATexture;
+        const bool bMSAADest = Destination->MSAATexture;
+        const bool bDepthStencil = Source->PixelFormat == PF_DepthStencil;
+		if (bMSAASource && !bMSAADest)
+		{
+			// Resolve required - Device must support this - Using Shader for resolve not supported amd NumSamples should be 1
+			const bool bSupportsMSAADepthResolve = GetAGXDeviceContext().SupportsFeature(EAGXFeaturesMSAADepthResolve);
+			const bool bSupportsMSAAStoreAndResolve = GetAGXDeviceContext().SupportsFeature(EAGXFeaturesMSAAStoreAndResolve);
+			check( (!bDepthStencil && bSupportsMSAAStoreAndResolve) || (bDepthStencil && bSupportsMSAADepthResolve) );
+			
+			Context->CopyFromTextureToTexture(Source->MSAAResolveTexture, SrcIndex, ResolveParams.MipIndex, Origin, Size, Destination->Texture, DestIndex, ResolveParams.MipIndex, Origin);
+		}
+		else
+		{
+			Context->CopyFromTextureToTexture(Source->Texture, SrcIndex, ResolveParams.MipIndex, Origin, Size, Destination->Texture, DestIndex, ResolveParams.MipIndex, Origin);
+		}
+
+#if PLATFORM_MAC
+		if((Destination->GPUReadback & FAGXSurface::EAGXGPUReadbackFlags::ReadbackRequested) != 0)
+		{
+			Context->GetCurrentRenderPass().SynchronizeTexture(Destination->Texture, DestIndex, ResolveParams.MipIndex);
+		}
+#endif
+		
+	}
+	}
+}
+
+/** Helper for accessing R10G10B10A2 colors. */
+struct FAGXR10G10B10A2
+{
+	uint32 R : 10;
+	uint32 G : 10;
+	uint32 B : 10;
+	uint32 A : 2;
+};
+
+/** Helper for accessing R16G16 colors. */
+struct FAGXRG16
+{
+	uint16 R;
+	uint16 G;
+};
+
+/** Helper for accessing R16G16B16A16 colors. */
+struct FAGXRGBA16
+{
+	uint16 R;
+	uint16 G;
+	uint16 B;
+	uint16 A;
+};
+
+void FAGXDynamicRHI::RHIReadSurfaceData(FRHITexture* TextureRHI, FIntRect InRect, TArray<FLinearColor>& OutData, FReadSurfaceDataFlags InFlags)
+{
+	// Use our current surface read implementation and convert to linear - should refactor to make optimal
+	TArray<FColor> OutDataUnConverted;
+	RHIReadSurfaceData(TextureRHI, InRect, OutDataUnConverted, InFlags);
+
+	OutData.Empty();
+	OutData.AddUninitialized(OutDataUnConverted.Num());
+
+	for (uint32 i = 0; i < OutDataUnConverted.Num(); ++i)
+	{
+		OutData[i] = OutDataUnConverted[i].ReinterpretAsLinear();
+	}
+}
+
+static void ConvertSurfaceDataToFColor(EPixelFormat Format, uint32 Width, uint32 Height, uint8 *In, uint32 SrcPitch, FColor* Out, FReadSurfaceDataFlags InFlags)
+{
+	bool bLinearToGamma = InFlags.GetLinearToGamma();
+	if (Format == PF_G16 || Format == PF_R16_UINT || Format == PF_R16_SINT)
+	{
+		ConvertRawR16DataToFColor(Width, Height, In, SrcPitch, Out);
+	}
+	else if (Format == PF_R8G8B8A8)
+	{
+		ConvertRawR8G8B8A8DataToFColor(Width, Height, In, SrcPitch, Out);
+	}
+	else if (Format == PF_B8G8R8A8)
+	{
+		ConvertRawB8G8R8A8DataToFColor(Width, Height, In, SrcPitch, Out);
+	}
+	else if (Format == PF_A2B10G10R10)
+	{
+		ConvertRawR10G10B10A2DataToFColor(Width, Height, In, SrcPitch, Out);
+	}
+	else if (Format == PF_FloatRGBA)
+	{
+		ConvertRawR16G16B16A16FDataToFColor(Width, Height, In, SrcPitch, Out, bLinearToGamma);
+	}
+	else if (Format == PF_FloatR11G11B10)
+	{
+		ConvertRawR11G11B10DataToFColor(Width, Height, In, SrcPitch, Out, bLinearToGamma);
+	}
+	else if (Format == PF_A32B32G32R32F)
+	{
+		ConvertRawR32G32B32A32DataToFColor(Width, Height, In, SrcPitch, Out, bLinearToGamma);
+	}
+	else if (Format == PF_A16B16G16R16)
+	{
+		ConvertRawR16G16B16A16DataToFColor(Width, Height, In, SrcPitch, Out);
+	}
+	else if (Format == PF_G16R16)
+	{
+		ConvertRawR16G16DataToFColor(Width, Height, In, SrcPitch, Out);
+	}
+	else if (Format == PF_DepthStencil)
+	{
+		ConvertRawD32S8DataToFColor(Width, Height, In, SrcPitch, Out, InFlags);
+	}
+	else
+	{
+		// not supported yet
+		NOT_SUPPORTED("RHIReadSurfaceData Format");
+	}
+}
+
+void FAGXDynamicRHI::RHIReadSurfaceData(FRHITexture* TextureRHI, FIntRect Rect, TArray<FColor>& OutData, FReadSurfaceDataFlags InFlags)
+{
+	@autoreleasepool {
+	if (!ensure(TextureRHI))
+	{
+		OutData.Empty();
+		OutData.AddZeroed(Rect.Width() * Rect.Height());
+		return;
+	}
+
+	FAGXSurface* Surface = AGXGetMetalSurfaceFromRHITexture(TextureRHI);
+
+	// allocate output space
+	const uint32 SizeX = Rect.Width();
+	const uint32 SizeY = Rect.Height();
+	OutData.Empty();
+	OutData.AddUninitialized(SizeX * SizeY);
+	
+	FColor* OutDataPtr = OutData.GetData();
+	mtlpp::Region Region(Rect.Min.X, Rect.Min.Y, SizeX, SizeY);
+    
+	FAGXTexture Texture = Surface->Texture;
+    if(!Texture && (Surface->Flags & TexCreate_Presentable))
+    {
+        Texture = Surface->GetCurrentTexture();
+    }
+    if(!Texture)
+    {
+        UE_LOG(LogRHI, Error, TEXT("Trying to read from an uninitialised texture."));
+        return;
+    }
+
+	if (GAGXUseTexGetBytes && Surface->PixelFormat != PF_DepthStencil && Surface->PixelFormat != PF_ShadowDepth)
+	{
+		SCOPE_CYCLE_COUNTER(STAT_AGXTexturePageOffTime);
+		
+		FAGXTexture TempTexture = nil;
+		if (Texture.GetStorageMode() == mtlpp::StorageMode::Private)
+		{
+#if PLATFORM_MAC
+			mtlpp::StorageMode StorageMode = mtlpp::StorageMode::Managed;
+#else
+			mtlpp::StorageMode StorageMode = mtlpp::StorageMode::Shared;
+#endif
+			mtlpp::PixelFormat MetalFormat = (mtlpp::PixelFormat)GPixelFormats[Surface->PixelFormat].PlatformFormat;
+			mtlpp::TextureDescriptor Desc;
+			Desc.SetTextureType(Texture.GetTextureType());
+			Desc.SetPixelFormat(Texture.GetPixelFormat());
+			Desc.SetWidth(SizeX);
+			Desc.SetHeight(SizeY);
+			Desc.SetDepth(1);
+			Desc.SetMipmapLevelCount(1); // Only consider a single subresource and not the whole texture (like in the other RHIs)
+			Desc.SetSampleCount(Texture.GetSampleCount());
+			Desc.SetArrayLength(Texture.GetArrayLength());
+			
+			mtlpp::ResourceOptions GeneralResourceOption = (mtlpp::ResourceOptions)FAGXCommandQueue::GetCompatibleResourceOptions(mtlpp::ResourceOptions(((NSUInteger)Texture.GetCpuCacheMode() << mtlpp::ResourceCpuCacheModeShift) | ((NSUInteger)StorageMode << mtlpp::ResourceStorageModeShift) | mtlpp::ResourceOptions::HazardTrackingModeUntracked));
+			Desc.SetResourceOptions(GeneralResourceOption);
+			
+			Desc.SetCpuCacheMode(Texture.GetCpuCacheMode());
+			Desc.SetStorageMode(StorageMode);
+			Desc.SetUsage(Texture.GetUsage());
+			
+			TempTexture = GetAGXDeviceContext().GetDevice().NewTexture(Desc);
+			
+			ImmediateContext.Context->CopyFromTextureToTexture(Texture, 0, InFlags.GetMip(), mtlpp::Origin(Region.origin), mtlpp::Size(Region.size), TempTexture, 0, 0, mtlpp::Origin(0, 0, 0));
+			
+			Texture = TempTexture;
+			Region = mtlpp::Region(0, 0, SizeX, SizeY);
+		}
+#if PLATFORM_MAC
+		if(Texture.GetStorageMode() == mtlpp::StorageMode::Managed)
+		{
+			// Synchronise the texture with the CPU
+			ImmediateContext.Context->SynchronizeTexture(Texture, 0, InFlags.GetMip());
+		}
+#endif
+
+		//kick the current command buffer.
+		ImmediateContext.Context->SubmitCommandBufferAndWait();
+		
+		const uint32 Stride = GPixelFormats[Surface->PixelFormat].BlockBytes * SizeX;
+		const uint32 BytesPerImage = Stride * SizeY;
+
+		TArray<uint8> Data;
+		Data.AddUninitialized(BytesPerImage);
+		
+		Texture.GetBytes(Data.GetData(), Stride, BytesPerImage, Region, 0, 0);
+		
+		ConvertSurfaceDataToFColor(Surface->PixelFormat, SizeX, SizeY, (uint8*)Data.GetData(), Stride, OutDataPtr, InFlags);
+		
+		if (TempTexture)
+		{
+			AGXSafeReleaseMetalTexture(TempTexture);
+		}
+	}
+	else
+	{
+		uint32 BytesPerPixel = (Surface->PixelFormat != PF_DepthStencil || !InFlags.GetOutputStencil()) ? GPixelFormats[Surface->PixelFormat].BlockBytes : 1;
+		const uint32 Stride = BytesPerPixel * SizeX;
+		const uint32 Alignment = PLATFORM_MAC ? 1u : 64u; // Mac permits natural row alignment (tightly-packed) but iOS does not.
+		const uint32 AlignedStride = ((Stride - 1) & ~(Alignment - 1)) + Alignment;
+		const uint32 BytesPerImage = AlignedStride * SizeY;
+		FAGXBuffer Buffer = ((FAGXDeviceContext*)ImmediateContext.Context)->CreatePooledBuffer(FAGXPooledBufferArgs(ImmediateContext.Context->GetDevice(), BytesPerImage, BUF_Dynamic, mtlpp::StorageMode::Shared));
+		{
+			// Synchronise the texture with the CPU
+			SCOPE_CYCLE_COUNTER(STAT_AGXTexturePageOffTime);
+			
+			if (Surface->PixelFormat != PF_DepthStencil)
+			{
+				ImmediateContext.Context->CopyFromTextureToBuffer(Texture, 0, InFlags.GetMip(), Region.origin, Region.size, Buffer, 0, AlignedStride, BytesPerImage, mtlpp::BlitOption::None);
+			}
+			else
+			{
+				if (!InFlags.GetOutputStencil())
+				{
+					ImmediateContext.Context->CopyFromTextureToBuffer(Texture, 0, InFlags.GetMip(), Region.origin, Region.size, Buffer, 0, AlignedStride, BytesPerImage, mtlpp::BlitOption::DepthFromDepthStencil);
+				}
+				else
+				{
+					ImmediateContext.Context->CopyFromTextureToBuffer(Texture, 0, InFlags.GetMip(), Region.origin, Region.size, Buffer, 0, AlignedStride, BytesPerImage, mtlpp::BlitOption::StencilFromDepthStencil);
+				}
+			}
+			
+			//kick the current command buffer.
+			ImmediateContext.Context->SubmitCommandBufferAndWait();
+			
+			ConvertSurfaceDataToFColor(Surface->PixelFormat, SizeX, SizeY, (uint8*)Buffer.GetContents(), AlignedStride, OutDataPtr, InFlags);
+		}
+		((FAGXDeviceContext*)ImmediateContext.Context)->ReleaseBuffer(Buffer);
+	}
+	}
+}
+
+void FAGXDynamicRHI::RHIMapStagingSurface(FRHITexture* TextureRHI, FRHIGPUFence* FenceRHI, void*& OutData, int32& OutWidth, int32& OutHeight, uint32 GPUIndex)
+{
+	@autoreleasepool {
+    FAGXSurface* Surface = AGXGetMetalSurfaceFromRHITexture(TextureRHI);
+    FAGXTexture2D* Texture = (FAGXTexture2D*)TextureRHI->GetTexture2D();
+    
+#if PLATFORM_MAC
+	uint16 FencePoll = FenceRHI && FenceRHI->Poll() ? 1 : 0;
+    Surface->GPUReadback |= FencePoll << FAGXSurface::EAGXGPUReadbackFlags::ReadbackFenceCompleteShift;
+#endif
+    
+    uint32 Stride = 0;
+    OutWidth = Texture->GetSizeX();
+    OutHeight = Texture->GetSizeY();
+    OutData = Surface->Lock(0, 0, RLM_ReadOnly, Stride);
+    
+#if PLATFORM_MAC
+	Surface->GPUReadback = (Surface->Texture.GetPtr() && Surface->Texture.GetStorageMode() == mtlpp::StorageMode::Managed) ? FAGXSurface::EAGXGPUReadbackFlags::ReadbackRequested : 0;
+#endif
+	}
+}
+
+void FAGXDynamicRHI::RHIUnmapStagingSurface(FRHITexture* TextureRHI, uint32 GPUIndex)
+{
+	@autoreleasepool {
+    	FAGXSurface* Surface = AGXGetMetalSurfaceFromRHITexture(TextureRHI);
+	
+    	Surface->Unlock(0, 0, false);
+	}
+}
+
+void FAGXDynamicRHI::RHIReadSurfaceFloatData(FRHITexture* TextureRHI, FIntRect Rect, TArray<FFloat16Color>& OutData, ECubeFace CubeFace,int32 ArrayIndex,int32 MipIndex)
+{
+	@autoreleasepool {
+	FAGXSurface* Surface = AGXGetMetalSurfaceFromRHITexture(TextureRHI);
+	
+    FAGXTexture Texture = Surface->Texture;
+    if(!Texture && (Surface->Flags & TexCreate_Presentable))
+    {
+		Texture = Surface->GetCurrentTexture();
+    }
+    if(!Texture)
+    {
+        UE_LOG(LogRHI, Error, TEXT("Trying to read from an uninitialised texture."));
+        return;
+    }
+    
+	// verify the input image format (but don't crash)
+	if (Surface->PixelFormat != PF_FloatRGBA)
+	{
+		UE_LOG(LogRHI, Log, TEXT("Trying to read non-FloatRGBA surface."));
+	}
+
+	if (TextureRHI->GetTextureCube())
+	{
+		// adjust index to account for cubemaps as texture arrays
+		ArrayIndex *= CubeFace_MAX;
+		ArrayIndex += GetMetalCubeFace(CubeFace);
+	}
+	
+	// allocate output space
+	const uint32 SizeX = Rect.Width();
+	const uint32 SizeY = Rect.Height();
+	OutData.Empty();
+	OutData.AddUninitialized(SizeX * SizeY);
+	
+	mtlpp::Region Region = mtlpp::Region(Rect.Min.X, Rect.Min.Y, SizeX, SizeY);
+	
+	// function wants details about the destination, not the source
+	const uint32 Stride = GPixelFormats[Surface->PixelFormat].BlockBytes * SizeX;
+	const uint32 Alignment = PLATFORM_MAC ? 1u : 64u; // Mac permits natural row alignment (tightly-packed) but iOS does not.
+	const uint32 AlignedStride = ((Stride - 1) & ~(Alignment - 1)) + Alignment;
+	const uint32 BytesPerImage = AlignedStride  * SizeY;
+	int32 FloatBGRADataSize = BytesPerImage;
+	FAGXBuffer Buffer = ((FAGXDeviceContext*)ImmediateContext.Context)->CreatePooledBuffer(FAGXPooledBufferArgs(ImmediateContext.Context->GetDevice(), FloatBGRADataSize, BUF_Dynamic, mtlpp::StorageMode::Shared));
+	{
+		// Synchronise the texture with the CPU
+		SCOPE_CYCLE_COUNTER(STAT_AGXTexturePageOffTime);
+		
+		ImmediateContext.Context->CopyFromTextureToBuffer(Texture, ArrayIndex, MipIndex, Region.origin, Region.size, Buffer, 0, AlignedStride, BytesPerImage, mtlpp::BlitOption::None);
+		
+		//kick the current command buffer.
+		ImmediateContext.Context->SubmitCommandBufferAndWait();
+	}
+	
+	uint8* DataPtr = (uint8*)Buffer.GetContents();
+	FFloat16Color* OutDataPtr = OutData.GetData();
+	if (Alignment > 1u)
+	{
+		for (uint32 Row = 0; Row < SizeY; Row++)
+		{
+			FFloat16Color* FloatBGRAData = (FFloat16Color*)DataPtr;
+			FMemory::Memcpy(OutDataPtr, FloatBGRAData, Stride);
+			DataPtr += AlignedStride;
+			OutDataPtr += SizeX;
+		}
+	}
+	else
+	{
+		FFloat16Color* FloatBGRAData = (FFloat16Color*)DataPtr;
+		FMemory::Memcpy(OutDataPtr, FloatBGRAData, FloatBGRADataSize);
+	}
+	
+	((FAGXDeviceContext*)ImmediateContext.Context)->ReleaseBuffer(Buffer);
+	}
+}
+
+void FAGXDynamicRHI::RHIRead3DSurfaceFloatData(FRHITexture* TextureRHI,FIntRect InRect,FIntPoint ZMinMax,TArray<FFloat16Color>& OutData)
+{
+	@autoreleasepool {
+	FAGXSurface* Surface = AGXGetMetalSurfaceFromRHITexture(TextureRHI);
+	
+	FAGXTexture Texture = Surface->Texture;
+	if(!Texture)
+	{
+		UE_LOG(LogRHI, Error, TEXT("Trying to read from an uninitialised texture."));
+		return;
+	}
+	
+	// verify the input image format (but don't crash)
+	if (Surface->PixelFormat != PF_FloatRGBA)
+	{
+		UE_LOG(LogRHI, Log, TEXT("Trying to read non-FloatRGBA surface."));
+	}
+	
+	// allocate output space
+	const uint32 SizeX = InRect.Width();
+	const uint32 SizeY = InRect.Height();
+	const uint32 SizeZ = ZMinMax.Y - ZMinMax.X;
+	OutData.Empty();
+	OutData.AddUninitialized(SizeX * SizeY * SizeZ);
+	
+	mtlpp::Region Region = mtlpp::Region(InRect.Min.X, InRect.Min.Y, ZMinMax.X, SizeX, SizeY, SizeZ);
+	
+	// function wants details about the destination, not the source
+	const uint32 Stride = GPixelFormats[Surface->PixelFormat].BlockBytes * SizeX;
+	const uint32 Alignment = PLATFORM_MAC ? 1u : 64u; // Mac permits natural row alignment (tightly-packed) but iOS does not.
+	const uint32 AlignedStride = ((Stride - 1) & ~(Alignment - 1)) + Alignment;
+	const uint32 BytesPerImage = AlignedStride  * SizeY;
+	int32 FloatBGRADataSize = BytesPerImage * SizeZ;
+	FAGXBuffer Buffer = ((FAGXDeviceContext*)ImmediateContext.Context)->CreatePooledBuffer(FAGXPooledBufferArgs(ImmediateContext.Context->GetDevice(), FloatBGRADataSize, BUF_Dynamic, mtlpp::StorageMode::Shared));
+	{
+		// Synchronise the texture with the CPU
+		SCOPE_CYCLE_COUNTER(STAT_AGXTexturePageOffTime);
+		
+		ImmediateContext.Context->CopyFromTextureToBuffer(Texture, 0, 0, Region.origin, Region.size, Buffer, 0, AlignedStride, BytesPerImage, mtlpp::BlitOption::None);
+		
+		//kick the current command buffer.
+		ImmediateContext.Context->SubmitCommandBufferAndWait();
+	}
+	
+	uint8* DataPtr = (uint8*)Buffer.GetContents();
+	FFloat16Color* OutDataPtr = OutData.GetData();
+	if (Alignment > 1u)
+	{
+		for (uint32 Image = 0; Image < SizeZ; Image++)
+		{
+			for (uint32 Row = 0; Row < SizeY; Row++)
+			{
+				FFloat16Color* FloatBGRAData = (FFloat16Color*)DataPtr;
+				FMemory::Memcpy(OutDataPtr, FloatBGRAData, Stride);
+				DataPtr += AlignedStride;
+				OutDataPtr += SizeX;
+			}
+		}
+	}
+	else
+	{
+		FFloat16Color* FloatBGRAData = (FFloat16Color*)DataPtr;
+		FMemory::Memcpy(OutDataPtr, FloatBGRAData, FloatBGRADataSize);
+	}
+	
+	((FAGXDeviceContext*)ImmediateContext.Context)->ReleaseBuffer(Buffer);
+	}
+}
diff --git a/Engine/Source/Runtime/Apple/AGXRHI/Private/AGXShaderTypes.h b/Engine/Source/Runtime/Apple/AGXRHI/Private/AGXShaderTypes.h
new file mode 100644
--- /dev/null
+++ b/Engine/Source/Runtime/Apple/AGXRHI/Private/AGXShaderTypes.h
@@ -0,0 +1,13 @@
+// Copyright Epic Games, Inc. All Rights Reserved.
+
+/*=============================================================================
+	AGXShaderTypes.h: AGX RHI Shader Types.
+=============================================================================*/
+
+#pragma once
+
+#include "Shaders/Types/Templates/AGXBaseShader.h"
+#include "Shaders/Types/AGXVertexShader.h"
+#include "Shaders/Types/AGXGeometryShader.h"
+#include "Shaders/Types/AGXPixelShader.h"
+#include "Shaders/Types/AGXComputeShader.h"
diff --git a/Engine/Source/Runtime/Apple/AGXRHI/Private/AGXShaders.cpp b/Engine/Source/Runtime/Apple/AGXRHI/Private/AGXShaders.cpp
new file mode 100644
--- /dev/null
+++ b/Engine/Source/Runtime/Apple/AGXRHI/Private/AGXShaders.cpp
@@ -0,0 +1,128 @@
+// Copyright Epic Games, Inc. All Rights Reserved.
+
+/*=============================================================================
+	AGXShaders.cpp: AGX RHI shader implementation.
+=============================================================================*/
+
+#include "AGXRHIPrivate.h"
+
+#include "Shaders/Debugging/AGXShaderDebugCache.h"
+#include "Shaders/AGXCompiledShaderKey.h"
+#include "Shaders/AGXCompiledShaderCache.h"
+#include "Shaders/AGXShaderLibrary.h"
+
+#include "HAL/FileManager.h"
+#include "HAL/PlatformFileManager.h"
+#include "Misc/Paths.h"
+#include "MetalShaderResources.h"
+#include "AGXResources.h"
+#include "AGXProfiler.h"
+#include "AGXCommandBuffer.h"
+#include "Serialization/MemoryReader.h"
+#include "Misc/FileHelper.h"
+#include "Misc/ScopeRWLock.h"
+#include "Misc/Compression.h"
+#include "Misc/MessageDialog.h"
+
+#define SHADERCOMPILERCOMMON_API
+#	include "Developer/ShaderCompilerCommon/Public/ShaderCompilerCommon.h"
+#undef SHADERCOMPILERCOMMON_API
+
+
+NSString* AGXDecodeMetalSourceCode(uint32 CodeSize, TArray<uint8> const& CompressedSource)
+{
+	NSString* GlslCodeNSString = nil;
+	if (CodeSize && CompressedSource.Num())
+	{
+		TArray<ANSICHAR> UncompressedCode;
+		UncompressedCode.AddZeroed(CodeSize+1);
+		bool bSucceed = FCompression::UncompressMemory(NAME_Zlib, UncompressedCode.GetData(), CodeSize, CompressedSource.GetData(), CompressedSource.Num());
+		if (bSucceed)
+		{
+			GlslCodeNSString = [[NSString stringWithUTF8String:UncompressedCode.GetData()] retain];
+		}
+	}
+	return GlslCodeNSString;
+}
+
+mtlpp::LanguageVersion AGXValidateVersion(uint8 Version)
+{
+	static uint32 MetalMacOSVersions[][3] = {
+		{10,11,6},
+		{10,11,6},
+		{10,12,6},
+		{10,13,0},
+		{10,14,0},
+		{10,15,0},
+		{11, 0,0},
+	};
+	static uint32 MetaliOSVersions[][3] = {
+		{ 8,0,0},
+		{ 9,0,0},
+		{10,0,0},
+		{11,0,0},
+		{12,0,0},
+		{13,0,0},
+		{14,0,0},
+	};
+	static TCHAR const* StandardNames[] =
+	{
+		TEXT("Metal 1.0"),
+		TEXT("Metal 1.1"),
+		TEXT("Metal 1.2"),
+		TEXT("Metal 2.0"),
+		TEXT("Metal 2.1"),
+		TEXT("Metal 2.2"),
+		TEXT("Metal 2.3"),
+	};
+	
+	mtlpp::LanguageVersion Result = mtlpp::LanguageVersion::Version1_1;
+	switch(Version)
+	{
+		case 6:
+			Result = mtlpp::LanguageVersion::Version2_3;
+			break;
+		case 5:
+			Result = mtlpp::LanguageVersion::Version2_2;
+			break;
+		case 4:
+			Result = mtlpp::LanguageVersion::Version2_1;
+			break;
+		case 3:
+			Result = mtlpp::LanguageVersion::Version2_0;
+			break;
+		case 2:
+			Result = mtlpp::LanguageVersion::Version1_2;
+			break;
+		case 1:
+			Result = mtlpp::LanguageVersion::Version1_1;
+			break;
+		case 0:
+		default:
+#if PLATFORM_MAC
+			Result = mtlpp::LanguageVersion::Version1_1;
+#else
+			Result = mtlpp::LanguageVersion::Version1_0;
+#endif
+			break;
+	}
+	
+	if (!FApplePlatformMisc::IsOSAtLeastVersion(MetalMacOSVersions[Version], MetaliOSVersions[Version], MetaliOSVersions[Version]))
+	{
+		FFormatNamedArguments Args;
+		Args.Add(TEXT("ShaderVersion"), FText::FromString(FString(StandardNames[Version])));
+#if PLATFORM_MAC
+		Args.Add(TEXT("RequiredOS"), FText::FromString(FString::Printf(TEXT("macOS %d.%d.%d"), MetalMacOSVersions[Version][0], MetalMacOSVersions[Version][1], MetalMacOSVersions[Version][2])));
+#else
+		Args.Add(TEXT("RequiredOS"), FText::FromString(FString::Printf(TEXT("%d.%d.%d"), MetaliOSVersions[Version][0], MetaliOSVersions[Version][1], MetaliOSVersions[Version][2])));
+#endif
+		FText LocalizedMsg = FText::Format(NSLOCTEXT("AGXRHI", "ShaderVersionUnsupported", "The current OS version does not support {ShaderVersion} required by the project. You must upgrade to {RequiredOS} to run this project."),Args);
+		
+		FText Title = NSLOCTEXT("AGXRHI", "ShaderVersionUnsupportedTitle", "Shader Version Unsupported");
+		FMessageDialog::Open(EAppMsgType::Ok, LocalizedMsg, &Title);
+		
+		FPlatformMisc::RequestExit(true);
+	}
+	
+	return Result;
+}
diff --git a/Engine/Source/Runtime/Apple/AGXRHI/Private/AGXState.cpp b/Engine/Source/Runtime/Apple/AGXRHI/Private/AGXState.cpp
new file mode 100644
--- /dev/null
+++ b/Engine/Source/Runtime/Apple/AGXRHI/Private/AGXState.cpp
@@ -0,0 +1,524 @@
+// Copyright Epic Games, Inc. All Rights Reserved.
+
+/*=============================================================================
+	AGXState.cpp: AGX RHI state implementation.
+=============================================================================*/
+
+#include "AGXRHIPrivate.h"
+#include "AGXProfiler.h"
+
+static uint32 GetMetalMaxAnisotropy(ESamplerFilter Filter, uint32 MaxAniso)
+{
+	switch (Filter)
+	{
+		case SF_AnisotropicPoint:
+		case SF_AnisotropicLinear:	return ComputeAnisotropyRT(MaxAniso);
+		default:					return 1;
+	}
+}
+
+static mtlpp::SamplerMinMagFilter TranslateZFilterMode(ESamplerFilter Filter)
+{	switch (Filter)
+	{
+		case SF_Point:				return mtlpp::SamplerMinMagFilter::Nearest;
+		case SF_AnisotropicPoint:	return mtlpp::SamplerMinMagFilter::Nearest;
+		case SF_AnisotropicLinear:	return mtlpp::SamplerMinMagFilter::Linear;
+		default:					return mtlpp::SamplerMinMagFilter::Linear;
+	}
+}
+
+static mtlpp::SamplerAddressMode TranslateWrapMode(ESamplerAddressMode AddressMode)
+{
+	switch (AddressMode)
+	{
+		case AM_Clamp:	return mtlpp::SamplerAddressMode::ClampToEdge;
+		case AM_Mirror: return mtlpp::SamplerAddressMode::MirrorRepeat;
+		case AM_Border: return mtlpp::SamplerAddressMode::ClampToEdge;
+		default:		return mtlpp::SamplerAddressMode::Repeat;
+	}
+}
+
+static mtlpp::CompareFunction TranslateCompareFunction(ECompareFunction CompareFunction)
+{
+	switch(CompareFunction)
+	{
+		case CF_Less:			return mtlpp::CompareFunction::Less;
+		case CF_LessEqual:		return mtlpp::CompareFunction::LessEqual;
+		case CF_Greater:		return mtlpp::CompareFunction::Greater;
+		case CF_GreaterEqual:	return mtlpp::CompareFunction::GreaterEqual;
+		case CF_Equal:			return mtlpp::CompareFunction::Equal;
+		case CF_NotEqual:		return mtlpp::CompareFunction::NotEqual;
+		case CF_Never:			return mtlpp::CompareFunction::Never;
+		default:				return mtlpp::CompareFunction::Always;
+	};
+}
+
+static mtlpp::CompareFunction TranslateSamplerCompareFunction(ESamplerCompareFunction SamplerComparisonFunction)
+{
+	switch(SamplerComparisonFunction)
+	{
+		case SCF_Less:		return mtlpp::CompareFunction::Less;
+		case SCF_Never:		return mtlpp::CompareFunction::Never;
+		default:			return mtlpp::CompareFunction::Never;
+	};
+}
+
+static mtlpp::StencilOperation TranslateStencilOp(EStencilOp StencilOp)
+{
+	switch(StencilOp)
+	{
+		case SO_Zero:				return mtlpp::StencilOperation::Zero;
+		case SO_Replace:			return mtlpp::StencilOperation::Replace;
+		case SO_SaturatedIncrement:	return mtlpp::StencilOperation::IncrementClamp;
+		case SO_SaturatedDecrement:	return mtlpp::StencilOperation::DecrementClamp;
+		case SO_Invert:				return mtlpp::StencilOperation::Invert;
+		case SO_Increment:			return mtlpp::StencilOperation::IncrementWrap;
+		case SO_Decrement:			return mtlpp::StencilOperation::DecrementWrap;
+		default:					return mtlpp::StencilOperation::Keep;
+	};
+}
+
+static mtlpp::BlendOperation TranslateBlendOp(EBlendOperation BlendOp)
+{
+	switch(BlendOp)
+	{
+		case BO_Subtract:	return mtlpp::BlendOperation::Subtract;
+		case BO_Min:		return mtlpp::BlendOperation::Min;
+		case BO_Max:		return mtlpp::BlendOperation::Max;
+		default:			return mtlpp::BlendOperation::Add;
+	};
+}
+
+
+static mtlpp::BlendFactor TranslateBlendFactor(EBlendFactor BlendFactor)
+{
+	switch(BlendFactor)
+	{
+		case BF_One:					return mtlpp::BlendFactor::One;
+		case BF_SourceColor:			return mtlpp::BlendFactor::SourceColor;
+		case BF_InverseSourceColor:		return mtlpp::BlendFactor::OneMinusSourceColor;
+		case BF_SourceAlpha:			return mtlpp::BlendFactor::SourceAlpha;
+		case BF_InverseSourceAlpha:		return mtlpp::BlendFactor::OneMinusSourceAlpha;
+		case BF_DestAlpha:				return mtlpp::BlendFactor::DestinationAlpha;
+		case BF_InverseDestAlpha:		return mtlpp::BlendFactor::OneMinusDestinationAlpha;
+		case BF_DestColor:				return mtlpp::BlendFactor::DestinationColor;
+		case BF_InverseDestColor:		return mtlpp::BlendFactor::OneMinusDestinationColor;
+		case BF_Source1Color:			return mtlpp::BlendFactor::Source1Color;
+		case BF_InverseSource1Color:	return mtlpp::BlendFactor::OneMinusSource1Color;
+		case BF_Source1Alpha:			return mtlpp::BlendFactor::Source1Alpha;
+		case BF_InverseSource1Alpha:	return mtlpp::BlendFactor::OneMinusSource1Alpha;
+		default:						return mtlpp::BlendFactor::Zero;
+	};
+}
+
+static mtlpp::ColorWriteMask TranslateWriteMask(EColorWriteMask WriteMask)
+{
+	uint32 Result = 0;
+	Result |= (WriteMask & CW_RED) ? (mtlpp::ColorWriteMask::Red) : 0;
+	Result |= (WriteMask & CW_GREEN) ? (mtlpp::ColorWriteMask::Green) : 0;
+	Result |= (WriteMask & CW_BLUE) ? (mtlpp::ColorWriteMask::Blue) : 0;
+	Result |= (WriteMask & CW_ALPHA) ? (mtlpp::ColorWriteMask::Alpha) : 0;
+	
+	return (mtlpp::ColorWriteMask)Result;
+}
+
+static EBlendOperation TranslateBlendOp(MTLBlendOperation BlendOp)
+{
+	switch(BlendOp)
+	{
+		case MTLBlendOperationSubtract:		return BO_Subtract;
+		case MTLBlendOperationMin:			return BO_Min;
+		case MTLBlendOperationMax:			return BO_Max;
+		case MTLBlendOperationAdd: default:	return BO_Add;
+	};
+}
+
+
+static EBlendFactor TranslateBlendFactor(MTLBlendFactor BlendFactor)
+{
+	switch(BlendFactor)
+	{
+		case MTLBlendFactorOne:							return BF_One;
+		case MTLBlendFactorSourceColor:					return BF_SourceColor;
+		case MTLBlendFactorOneMinusSourceColor:			return BF_InverseSourceColor;
+		case MTLBlendFactorSourceAlpha:					return BF_SourceAlpha;
+		case MTLBlendFactorOneMinusSourceAlpha:			return BF_InverseSourceAlpha;
+		case MTLBlendFactorDestinationAlpha:			return BF_DestAlpha;
+		case MTLBlendFactorOneMinusDestinationAlpha:	return BF_InverseDestAlpha;
+		case MTLBlendFactorDestinationColor:			return BF_DestColor;
+		case MTLBlendFactorOneMinusDestinationColor:	return BF_InverseDestColor;
+		case MTLBlendFactorSource1Color:				return BF_Source1Color;
+		case MTLBlendFactorOneMinusSource1Color:		return BF_InverseSource1Color;
+		case MTLBlendFactorSource1Alpha:				return BF_Source1Alpha;
+		case MTLBlendFactorOneMinusSource1Alpha:		return BF_InverseSource1Alpha;
+		case MTLBlendFactorZero: default:				return BF_Zero;
+	};
+}
+
+static EColorWriteMask TranslateWriteMask(MTLColorWriteMask WriteMask)
+{
+	uint32 Result = 0;
+	Result |= (WriteMask & MTLColorWriteMaskRed) ? (CW_RED) : 0;
+	Result |= (WriteMask & MTLColorWriteMaskGreen) ? (CW_GREEN) : 0;
+	Result |= (WriteMask & MTLColorWriteMaskBlue) ? (CW_BLUE) : 0;
+	Result |= (WriteMask & MTLColorWriteMaskAlpha) ? (CW_ALPHA) : 0;
+	
+	return (EColorWriteMask)Result;
+}
+
+template <typename InitializerType, typename StateType>
+class FAGXStateObjectCache
+{
+public:
+	FAGXStateObjectCache()
+	{
+		
+	}
+	
+	~FAGXStateObjectCache()
+	{
+		
+	}
+	
+	StateType Find(InitializerType Init)
+	{
+		if(IsRunningRHIInSeparateThread())
+		{
+			Mutex.ReadLock();
+		}
+		
+		StateType* State = Cache.Find(Init);
+		
+		if(IsRunningRHIInSeparateThread())
+		{
+			Mutex.ReadUnlock();
+		}
+		
+		return State ? *State : StateType(nullptr);
+	}
+	
+	void Add(InitializerType Init, StateType const& State)
+	{
+		if(IsRunningRHIInSeparateThread())
+		{
+			Mutex.WriteLock();
+		}
+		
+		Cache.Add(Init, State);
+		
+		if(IsRunningRHIInSeparateThread())
+		{
+			Mutex.WriteUnlock();
+		}
+	}
+	
+private:
+	TMap<InitializerType, StateType> Cache;
+	FRWLock Mutex;
+};
+
+static FAGXStateObjectCache<FSamplerStateInitializerRHI, FAGXSampler> Samplers;
+
+static FAGXSampler FindOrCreateSamplerState(mtlpp::Device Device, const FSamplerStateInitializerRHI& Initializer)
+{
+	FAGXSampler State = Samplers.Find(Initializer);
+	if (!State.GetPtr())
+	{
+		mtlpp::SamplerDescriptor Desc;
+		switch(Initializer.Filter)
+		{
+			case SF_AnisotropicLinear:
+			case SF_AnisotropicPoint:
+				Desc.SetMinFilter(mtlpp::SamplerMinMagFilter::Linear);
+				Desc.SetMagFilter(mtlpp::SamplerMinMagFilter::Linear);
+				Desc.SetMipFilter(mtlpp::SamplerMipFilter::Linear);
+				break;
+			case SF_Trilinear:
+				Desc.SetMinFilter(mtlpp::SamplerMinMagFilter::Linear);
+				Desc.SetMagFilter(mtlpp::SamplerMinMagFilter::Linear);
+				Desc.SetMipFilter(mtlpp::SamplerMipFilter::Linear);
+				break;
+			case SF_Bilinear:
+				Desc.SetMinFilter(mtlpp::SamplerMinMagFilter::Linear);
+				Desc.SetMagFilter(mtlpp::SamplerMinMagFilter::Linear);
+				Desc.SetMipFilter(mtlpp::SamplerMipFilter::Nearest);
+				break;
+			case SF_Point:
+				Desc.SetMinFilter(mtlpp::SamplerMinMagFilter::Nearest);
+				Desc.SetMagFilter(mtlpp::SamplerMinMagFilter::Nearest);
+				Desc.SetMipFilter(mtlpp::SamplerMipFilter::Nearest);
+				break;
+		}
+		Desc.SetMaxAnisotropy(GetMetalMaxAnisotropy(Initializer.Filter, Initializer.MaxAnisotropy));
+		Desc.SetSAddressMode(TranslateWrapMode(Initializer.AddressU));
+		Desc.SetTAddressMode(TranslateWrapMode(Initializer.AddressV));
+		Desc.SetRAddressMode(TranslateWrapMode(Initializer.AddressW));
+		Desc.SetLodMinClamp(Initializer.MinMipLevel);
+		Desc.SetLodMaxClamp(Initializer.MaxMipLevel);
+#if PLATFORM_TVOS
+		Desc.SetCompareFunction(mtlpp::CompareFunction::Never);	
+#elif PLATFORM_IOS
+		Desc.SetCompareFunction(Device.SupportsFeatureSet(mtlpp::FeatureSet::iOS_GPUFamily3_v1) ? TranslateSamplerCompareFunction(Initializer.SamplerComparisonFunction) : mtlpp::CompareFunction::Never);
+#else
+		Desc.SetCompareFunction(TranslateSamplerCompareFunction(Initializer.SamplerComparisonFunction));
+#endif
+#if PLATFORM_MAC
+		Desc.SetBorderColor(Initializer.BorderColor == 0 ? mtlpp::SamplerBorderColor::TransparentBlack : mtlpp::SamplerBorderColor::OpaqueWhite);
+#endif
+		if (FAGXCommandQueue::SupportsFeature(EAGXFeaturesIABs))
+		{
+			Desc.SetSupportArgumentBuffers(true);
+		}
+		
+		State = Device.NewSamplerState(Desc);
+		
+		Samplers.Add(Initializer, State);
+	}
+	return State;
+}
+
+FAGXSamplerState::FAGXSamplerState(mtlpp::Device Device, const FSamplerStateInitializerRHI& Initializer)
+{
+	State = FindOrCreateSamplerState(Device, Initializer);
+#if !PLATFORM_MAC
+	if (GetMetalMaxAnisotropy(Initializer.Filter, Initializer.MaxAnisotropy))
+	{
+		FSamplerStateInitializerRHI Init = Initializer;
+		Init.MaxAnisotropy = 1;
+		NoAnisoState = FindOrCreateSamplerState(Device, Init);
+	}
+#endif
+}
+
+FAGXSamplerState::~FAGXSamplerState()
+{
+}
+
+FAGXRasterizerState::FAGXRasterizerState(const FRasterizerStateInitializerRHI& Initializer)
+{
+	State = Initializer;
+}
+
+FAGXRasterizerState::~FAGXRasterizerState()
+{
+	
+}
+
+bool FAGXRasterizerState::GetInitializer(FRasterizerStateInitializerRHI& Init)
+{
+	Init = State;
+	return true;
+}
+
+static FAGXStateObjectCache<FDepthStencilStateInitializerRHI, mtlpp::DepthStencilState> DepthStencilStates;
+
+FAGXDepthStencilState::FAGXDepthStencilState(mtlpp::Device Device, const FDepthStencilStateInitializerRHI& InInitializer)
+{
+	Initializer = InInitializer;
+
+	State = DepthStencilStates.Find(Initializer);
+	if (!State.GetPtr())
+	{
+		mtlpp::DepthStencilDescriptor Desc;
+		
+		Desc.SetDepthCompareFunction(TranslateCompareFunction(Initializer.DepthTest));
+		Desc.SetDepthWriteEnabled(Initializer.bEnableDepthWrite);
+		
+		if (Initializer.bEnableFrontFaceStencil)
+		{
+			// set up front face stencil operations
+			mtlpp::StencilDescriptor Stencil;
+			Stencil.SetStencilCompareFunction(TranslateCompareFunction(Initializer.FrontFaceStencilTest));
+			Stencil.SetStencilFailureOperation(TranslateStencilOp(Initializer.FrontFaceStencilFailStencilOp));
+			Stencil.SetDepthFailureOperation(TranslateStencilOp(Initializer.FrontFaceDepthFailStencilOp));
+			Stencil.SetDepthStencilPassOperation(TranslateStencilOp(Initializer.FrontFacePassStencilOp));
+			Stencil.SetReadMask(Initializer.StencilReadMask);
+			Stencil.SetWriteMask(Initializer.StencilWriteMask);
+			Desc.SetFrontFaceStencil(Stencil);
+		}
+		
+		if (Initializer.bEnableBackFaceStencil)
+		{
+			// set up back face stencil operations
+			mtlpp::StencilDescriptor Stencil;
+			Stencil.SetStencilCompareFunction(TranslateCompareFunction(Initializer.BackFaceStencilTest));
+			Stencil.SetStencilFailureOperation(TranslateStencilOp(Initializer.BackFaceStencilFailStencilOp));
+			Stencil.SetDepthFailureOperation(TranslateStencilOp(Initializer.BackFaceDepthFailStencilOp));
+			Stencil.SetDepthStencilPassOperation(TranslateStencilOp(Initializer.BackFacePassStencilOp));
+			Stencil.SetReadMask(Initializer.StencilReadMask);
+			Stencil.SetWriteMask(Initializer.StencilWriteMask);
+			Desc.SetBackFaceStencil(Stencil);
+		}
+		else if(Initializer.bEnableFrontFaceStencil)
+		{
+			// set up back face stencil operations to front face in single-face mode
+			mtlpp::StencilDescriptor Stencil;
+			Stencil.SetStencilCompareFunction(TranslateCompareFunction(Initializer.FrontFaceStencilTest));
+			Stencil.SetStencilFailureOperation(TranslateStencilOp(Initializer.FrontFaceStencilFailStencilOp));
+			Stencil.SetDepthFailureOperation(TranslateStencilOp(Initializer.FrontFaceDepthFailStencilOp));
+			Stencil.SetDepthStencilPassOperation(TranslateStencilOp(Initializer.FrontFacePassStencilOp));
+			Stencil.SetReadMask(Initializer.StencilReadMask);
+			Stencil.SetWriteMask(Initializer.StencilWriteMask);
+			Desc.SetBackFaceStencil(Stencil);
+		}
+		
+		// bake out the descriptor
+		State = Device.NewDepthStencilState(Desc);
+		
+		DepthStencilStates.Add(Initializer, State);
+	}
+	
+	// cache some pipeline state info
+	bIsDepthWriteEnabled = Initializer.bEnableDepthWrite;
+	bIsStencilWriteEnabled = Initializer.bEnableFrontFaceStencil || Initializer.bEnableBackFaceStencil;
+}
+
+FAGXDepthStencilState::~FAGXDepthStencilState()
+{
+}
+
+bool FAGXDepthStencilState::GetInitializer(FDepthStencilStateInitializerRHI& Init)
+{
+	Init = Initializer;
+	return true;
+}
+
+
+
+// statics
+static FAGXStateObjectCache<FBlendStateInitializerRHI::FRenderTarget, mtlpp::RenderPipelineColorAttachmentDescriptor> BlendStates;
+TMap<uint32, uint8> FAGXBlendState::BlendSettingsToUniqueKeyMap;
+uint8 FAGXBlendState::NextKey = 0;
+FCriticalSection FAGXBlendState::Mutex;
+
+FAGXBlendState::FAGXBlendState(const FBlendStateInitializerRHI& Initializer)
+{
+	bUseIndependentRenderTargetBlendStates = Initializer.bUseIndependentRenderTargetBlendStates;
+	bUseAlphaToCoverage = Initializer.bUseAlphaToCoverage;
+	for(uint32 RenderTargetIndex = 0;RenderTargetIndex < MaxSimultaneousRenderTargets; ++RenderTargetIndex)
+	{
+		// which initializer to use
+		const FBlendStateInitializerRHI::FRenderTarget& Init =
+			Initializer.bUseIndependentRenderTargetBlendStates
+				? Initializer.RenderTargets[RenderTargetIndex]
+				: Initializer.RenderTargets[0];
+
+		// make a new blend state
+		mtlpp::RenderPipelineColorAttachmentDescriptor BlendState = BlendStates.Find(Init);
+
+		if (!BlendState.GetPtr())
+		{
+			BlendState = mtlpp::RenderPipelineColorAttachmentDescriptor();
+			
+			// set values
+			BlendState.SetBlendingEnabled(
+				Init.ColorBlendOp != BO_Add || Init.ColorDestBlend != BF_Zero || Init.ColorSrcBlend != BF_One ||
+				Init.AlphaBlendOp != BO_Add || Init.AlphaDestBlend != BF_Zero || Init.AlphaSrcBlend != BF_One);
+			BlendState.SetSourceRgbBlendFactor(TranslateBlendFactor(Init.ColorSrcBlend));
+			BlendState.SetDestinationRgbBlendFactor(TranslateBlendFactor(Init.ColorDestBlend));
+			BlendState.SetRgbBlendOperation(TranslateBlendOp(Init.ColorBlendOp));
+			BlendState.SetSourceAlphaBlendFactor(TranslateBlendFactor(Init.AlphaSrcBlend));
+			BlendState.SetDestinationAlphaBlendFactor(TranslateBlendFactor(Init.AlphaDestBlend));
+			BlendState.SetAlphaBlendOperation(TranslateBlendOp(Init.AlphaBlendOp));
+			BlendState.SetWriteMask(TranslateWriteMask(Init.ColorWriteMask));
+			
+			BlendStates.Add(Init, BlendState);
+		}
+		
+		RenderTargetStates[RenderTargetIndex].BlendState = BlendState;
+
+		// get the unique key
+		uint32 BlendBitMask =
+			((uint32)BlendState.GetSourceRgbBlendFactor() << 0) | ((uint32)BlendState.GetDestinationRgbBlendFactor() << 4) | ((uint32)BlendState.GetRgbBlendOperation() << 8) |
+			((uint32)BlendState.GetSourceAlphaBlendFactor() << 11) | ((uint32)BlendState.GetDestinationAlphaBlendFactor() << 15) | ((uint32)BlendState.GetAlphaBlendOperation() << 19) |
+			((uint32)BlendState.GetWriteMask() << 22);
+		
+		
+		if(IsRunningRHIInSeparateThread())
+		{
+			Mutex.Lock();
+		}
+		uint8* Key = BlendSettingsToUniqueKeyMap.Find(BlendBitMask);
+		if (Key == NULL)
+		{
+			Key = &BlendSettingsToUniqueKeyMap.Add(BlendBitMask, NextKey++);
+
+			// only giving limited bits to the key, since we need to pack 8 of them into a key
+			checkf(NextKey < (1 << NumBits_BlendState), TEXT("Too many unique blend states to fit into the PipelineStateHash [%d allowed]"), 1 << NumBits_BlendState);
+		}
+		// set the key
+		RenderTargetStates[RenderTargetIndex].BlendStateKey = *Key;
+		if(IsRunningRHIInSeparateThread())
+		{
+			Mutex.Unlock();
+		}
+	}
+}
+
+FAGXBlendState::~FAGXBlendState()
+{
+}
+
+bool FAGXBlendState::GetInitializer(FBlendStateInitializerRHI& Initializer)
+{
+	Initializer.bUseIndependentRenderTargetBlendStates = bUseIndependentRenderTargetBlendStates;
+	Initializer.bUseAlphaToCoverage = bUseAlphaToCoverage;
+	for(uint32 RenderTargetIndex = 0;RenderTargetIndex < MaxSimultaneousRenderTargets; ++RenderTargetIndex)
+	{
+		// which initializer to use
+		FBlendStateInitializerRHI::FRenderTarget& Init = Initializer.RenderTargets[RenderTargetIndex];
+		MTLRenderPipelineColorAttachmentDescriptor* CurrentState = RenderTargetStates[RenderTargetIndex].BlendState;
+		
+		if (CurrentState)
+		{
+			Init.ColorSrcBlend = TranslateBlendFactor(CurrentState.sourceRGBBlendFactor);
+			Init.ColorDestBlend = TranslateBlendFactor(CurrentState.destinationRGBBlendFactor);
+			Init.ColorBlendOp = TranslateBlendOp(CurrentState.rgbBlendOperation);
+			Init.AlphaSrcBlend = TranslateBlendFactor(CurrentState.sourceAlphaBlendFactor);
+			Init.AlphaDestBlend = TranslateBlendFactor(CurrentState.destinationAlphaBlendFactor);
+			Init.AlphaBlendOp = TranslateBlendOp(CurrentState.alphaBlendOperation);
+			Init.ColorWriteMask = TranslateWriteMask(CurrentState.writeMask);
+		}
+		
+		if (!bUseIndependentRenderTargetBlendStates)
+		{
+			break;
+		}
+	}
+	
+	return true;
+}
+
+
+
+
+
+FSamplerStateRHIRef FAGXDynamicRHI::RHICreateSamplerState(const FSamplerStateInitializerRHI& Initializer)
+{
+    @autoreleasepool {
+	return new FAGXSamplerState(ImmediateContext.Context->GetDevice(), Initializer);
+	}
+}
+
+FRasterizerStateRHIRef FAGXDynamicRHI::RHICreateRasterizerState(const FRasterizerStateInitializerRHI& Initializer)
+{
+	@autoreleasepool {
+    return new FAGXRasterizerState(Initializer);
+	}
+}
+
+FDepthStencilStateRHIRef FAGXDynamicRHI::RHICreateDepthStencilState(const FDepthStencilStateInitializerRHI& Initializer)
+{
+	@autoreleasepool {
+	return new FAGXDepthStencilState(ImmediateContext.Context->GetDevice(), Initializer);
+	}
+}
+
+
+FBlendStateRHIRef FAGXDynamicRHI::RHICreateBlendState(const FBlendStateInitializerRHI& Initializer)
+{
+	@autoreleasepool {
+	return new FAGXBlendState(Initializer);
+	}
+}
+
