Initial import of AGXRHI

diff --git a/Engine/Source/Runtime/Apple/AGXRHI/Private/AGXStateCache.cpp b/Engine/Source/Runtime/Apple/AGXRHI/Private/AGXStateCache.cpp
new file mode 100644
--- /dev/null
+++ b/Engine/Source/Runtime/Apple/AGXRHI/Private/AGXStateCache.cpp
@@ -0,0 +1,2302 @@
+// Copyright Epic Games, Inc. All Rights Reserved.
+
+
+#include "AGXRHIPrivate.h"
+#include "AGXRHIRenderQuery.h"
+#include "AGXShaderTypes.h"
+#include "AGXGraphicsPipelineState.h"
+#include "AGXStateCache.h"
+#include "AGXProfiler.h"
+#include "AGXCommandBuffer.h"
+
+#if PLATFORM_MAC
+	#ifndef UINT128_MAX
+		#define UINT128_MAX (((__uint128_t)1 << 127) - (__uint128_t)1 + ((__uint128_t)1 << 127))
+	#endif
+	#define FMETALTEXTUREMASK_MAX UINT128_MAX
+#else
+	#define FMETALTEXTUREMASK_MAX UINT32_MAX
+#endif
+
+static mtlpp::TriangleFillMode TranslateFillMode(ERasterizerFillMode FillMode)
+{
+	switch (FillMode)
+	{
+		case FM_Wireframe:	return mtlpp::TriangleFillMode::Lines;
+		case FM_Point:		return mtlpp::TriangleFillMode::Fill;
+		default:			return mtlpp::TriangleFillMode::Fill;
+	};
+}
+
+static mtlpp::CullMode TranslateCullMode(ERasterizerCullMode CullMode)
+{
+	switch (CullMode)
+	{
+		case CM_CCW:	return mtlpp::CullMode::Front;
+		case CM_CW:		return mtlpp::CullMode::Back;
+		default:		return mtlpp::CullMode::None;
+	}
+}
+
+FORCEINLINE mtlpp::StoreAction GetMetalRTStoreAction(ERenderTargetStoreAction StoreAction)
+{
+	switch(StoreAction)
+	{
+		case ERenderTargetStoreAction::ENoAction: return mtlpp::StoreAction::DontCare;
+		case ERenderTargetStoreAction::EStore: return mtlpp::StoreAction::Store;
+		//default store action in the desktop renderers needs to be mtlpp::StoreAction::StoreAndMultisampleResolve.  Trying to express the renderer by the requested maxrhishaderplatform
+        //because we may render to the same MSAA target twice in two separate passes.  BasePass, then some stuff, then translucency for example and we need to not lose the prior MSAA contents to do this properly.
+		case ERenderTargetStoreAction::EMultisampleResolve:
+		{
+            static bool bNoMSAA = FParse::Param(FCommandLine::Get(), TEXT("nomsaa"));
+			static bool bSupportsMSAAStoreResolve = FAGXCommandQueue::SupportsFeature(EAGXFeaturesMSAAStoreAndResolve) && (GMaxRHIFeatureLevel >= ERHIFeatureLevel::SM5);
+            if (bNoMSAA)
+            {
+                return mtlpp::StoreAction::Store;
+            }
+			else if (bSupportsMSAAStoreResolve)
+			{
+				return mtlpp::StoreAction::StoreAndMultisampleResolve;
+			}
+			else
+			{
+				return mtlpp::StoreAction::MultisampleResolve;
+			}
+		}
+		default: return mtlpp::StoreAction::DontCare;
+	}
+}
+
+FORCEINLINE mtlpp::StoreAction GetConditionalMetalRTStoreAction(bool bMSAATarget)
+{
+	if (bMSAATarget)
+	{
+		//this func should only be getting called when an encoder had to abnormally break.  In this case we 'must' do StoreAndResolve because the encoder will be restarted later
+		//with the original MSAA rendertarget and the original data must still be there to continue the render properly.
+		check(FAGXCommandQueue::SupportsFeature(EAGXFeaturesMSAAStoreAndResolve));
+		return mtlpp::StoreAction::StoreAndMultisampleResolve;
+	}
+	else
+	{
+		return mtlpp::StoreAction::Store;
+	}	
+}
+
+class FAGXRenderPassDescriptorPool
+{
+public:
+	FAGXRenderPassDescriptorPool()
+	{
+		
+	}
+	
+	~FAGXRenderPassDescriptorPool()
+	{
+		
+	}
+	
+	mtlpp::RenderPassDescriptor CreateDescriptor()
+	{
+		MTLRenderPassDescriptor* Desc = Cache.Pop();
+		if (!Desc)
+		{
+			Desc = [MTLRenderPassDescriptor new];
+		}
+		return mtlpp::RenderPassDescriptor(Desc);
+	}
+	
+	void ReleaseDescriptor(mtlpp::RenderPassDescriptor& Desc)
+	{
+		mtlpp::Texture EmptyTex;
+		
+		ns::Array<mtlpp::RenderPassColorAttachmentDescriptor> Attachements = Desc.GetColorAttachments();
+		for (uint32 i = 0; i < MaxSimultaneousRenderTargets; i++)
+		{
+			mtlpp::RenderPassColorAttachmentDescriptor Color = Attachements[i];
+			Color.SetTexture(EmptyTex);
+			Color.SetResolveTexture(EmptyTex);
+			Color.SetStoreAction(mtlpp::StoreAction::Store);
+		}
+		
+		mtlpp::RenderPassDepthAttachmentDescriptor Depth = Desc.GetDepthAttachment();
+		Depth.SetTexture(EmptyTex);
+		Depth.SetResolveTexture(EmptyTex);
+		Depth.SetStoreAction(mtlpp::StoreAction::Store);
+
+		mtlpp::RenderPassStencilAttachmentDescriptor Stencil = Desc.GetStencilAttachment();
+		Stencil.SetTexture(EmptyTex);
+		Stencil.SetResolveTexture(EmptyTex);
+		Stencil.SetStoreAction(mtlpp::StoreAction::Store);
+
+		mtlpp::Buffer Empty;
+		Desc.SetVisibilityResultBuffer(Empty);
+		
+#if PLATFORM_MAC
+		Desc.SetRenderTargetArrayLength(1);
+#endif
+		
+		Cache.Push(Desc.GetPtr());
+	}
+	
+	static FAGXRenderPassDescriptorPool& Get()
+	{
+		static FAGXRenderPassDescriptorPool sSelf;
+		return sSelf;
+	}
+	
+private:
+	TLockFreePointerListLIFO<MTLRenderPassDescriptor> Cache;
+};
+
+void AGXSafeReleaseMetalRenderPassDescriptor(mtlpp::RenderPassDescriptor& Desc)
+{
+	if (Desc.GetPtr())
+	{
+		FAGXRenderPassDescriptorPool::Get().ReleaseDescriptor(Desc);
+	}
+}
+
+FAGXStateCache::FAGXStateCache(bool const bInImmediate)
+: DepthStore(mtlpp::StoreAction::Unknown)
+, StencilStore(mtlpp::StoreAction::Unknown)
+, VisibilityResults(nullptr)
+, VisibilityMode(mtlpp::VisibilityResultMode::Disabled)
+, VisibilityOffset(0)
+, VisibilityWritten(0)
+, DepthStencilState(nullptr)
+, RasterizerState(nullptr)
+, StencilRef(0)
+, BlendFactor(FLinearColor::Transparent)
+, FrameBufferSize(CGSizeMake(0.0, 0.0))
+, RenderTargetArraySize(1)
+, RenderPassDesc(nil)
+, RasterBits(0)
+, PipelineBits(0)
+, bIsRenderTargetActive(false)
+, bHasValidRenderTarget(false)
+, bHasValidColorTarget(false)
+, bScissorRectEnabled(false)
+, bCanRestartRenderPass(false)
+, bImmediate(bInImmediate)
+, bFallbackDepthStencilBound(false)
+{
+	FMemory::Memzero(Viewport);
+	FMemory::Memzero(Scissor);
+	
+	ActiveViewports = 0;
+	ActiveScissors = 0;
+	
+	for (uint32 i = 0; i < MaxSimultaneousRenderTargets; i++)
+	{
+		ColorStore[i] = mtlpp::StoreAction::Unknown;
+	}
+	
+	FMemory::Memzero(RenderPassInfo);
+	FMemory::Memzero(DirtyUniformBuffers);
+}
+
+FAGXStateCache::~FAGXStateCache()
+{
+	RenderPassDesc = nil;
+	
+	for (uint32 i = 0; i < MaxVertexElementCount; i++)
+	{
+		VertexBuffers[i].Buffer = nil;
+		VertexBuffers[i].Bytes = nil;
+		VertexBuffers[i].Length = 0;
+		VertexBuffers[i].Offset = 0;
+	}
+	for (uint32 Frequency = 0; Frequency < EAGXShaderStages::Num; Frequency++)
+	{
+		ShaderSamplers[Frequency].Bound = 0;
+		for (uint32 i = 0; i < ML_MaxSamplers; i++)
+		{
+			ShaderSamplers[Frequency].Samplers[i] = nil;
+		}
+		for (uint32 i = 0; i < ML_MaxBuffers; i++)
+		{
+			BoundUniformBuffers[Frequency][i] = nullptr;
+			ShaderBuffers[Frequency].Buffers[i].Buffer = nil;
+			ShaderBuffers[Frequency].Buffers[i].Bytes = nil;
+			ShaderBuffers[Frequency].Buffers[i].Length = 0;
+			ShaderBuffers[Frequency].Buffers[i].Offset = 0;
+			ShaderBuffers[Frequency].Buffers[i].Usage = mtlpp::ResourceUsage(0);
+			ShaderBuffers[Frequency].Formats[i] = PF_Unknown;
+		}
+		ShaderBuffers[Frequency].Bound = 0;
+		for (uint32 i = 0; i < ML_MaxTextures; i++)
+		{
+			ShaderTextures[Frequency].Textures[i] = nil;
+			ShaderTextures[Frequency].Usage[i] = mtlpp::ResourceUsage(0);
+		}
+		ShaderTextures[Frequency].Bound = 0;
+	}
+	
+	VisibilityResults = nil;
+}
+
+void FAGXStateCache::Reset(void)
+{
+	SampleCount = 0;
+	
+	FMemory::Memzero(Viewport);
+	FMemory::Memzero(Scissor);
+	
+	ActiveViewports = 0;
+	ActiveScissors = 0;
+	
+	FMemory::Memzero(RenderPassInfo);
+	bIsRenderTargetActive = false;
+	bHasValidRenderTarget = false;
+	bHasValidColorTarget = false;
+	bScissorRectEnabled = false;
+	
+	FMemory::Memzero(DirtyUniformBuffers);
+	FMemory::Memzero(BoundUniformBuffers);
+	ActiveUniformBuffers.Empty();
+	
+	for (uint32 i = 0; i < MaxVertexElementCount; i++)
+	{
+		VertexBuffers[i].Buffer = nil;
+		VertexBuffers[i].Bytes = nil;
+		VertexBuffers[i].Length = 0;
+		VertexBuffers[i].Offset = 0;
+	}
+	for (uint32 Frequency = 0; Frequency < EAGXShaderStages::Num; Frequency++)
+	{
+		ShaderSamplers[Frequency].Bound = 0;
+		for (uint32 i = 0; i < ML_MaxSamplers; i++)
+		{
+			ShaderSamplers[Frequency].Samplers[i] = nil;
+		}
+		for (uint32 i = 0; i < ML_MaxBuffers; i++)
+		{
+			ShaderBuffers[Frequency].Buffers[i].Buffer = nil;
+			ShaderBuffers[Frequency].Buffers[i].Bytes = nil;
+			ShaderBuffers[Frequency].Buffers[i].Length = 0;
+			ShaderBuffers[Frequency].Buffers[i].Offset = 0;
+			ShaderBuffers[Frequency].Formats[i] = PF_Unknown;
+		}
+		ShaderBuffers[Frequency].Bound = 0;
+		for (uint32 i = 0; i < ML_MaxTextures; i++)
+		{
+			ShaderTextures[Frequency].Textures[i] = nil;
+			ShaderTextures[Frequency].Usage[i] = mtlpp::ResourceUsage(0);
+		}
+		ShaderTextures[Frequency].Bound = 0;
+	}
+	
+	VisibilityResults = nil;
+	VisibilityMode = mtlpp::VisibilityResultMode::Disabled;
+	VisibilityOffset = 0;
+	VisibilityWritten = 0;
+	
+	DepthStencilState.SafeRelease();
+	RasterizerState.SafeRelease();
+	GraphicsPSO.SafeRelease();
+	ComputeShader.SafeRelease();
+	DepthStencilSurface.SafeRelease();
+	StencilRef = 0;
+	
+	RenderPassDesc = nil;
+	
+	for (uint32 i = 0; i < MaxSimultaneousRenderTargets; i++)
+	{
+		ColorStore[i] = mtlpp::StoreAction::Unknown;
+	}
+	DepthStore = mtlpp::StoreAction::Unknown;
+	StencilStore = mtlpp::StoreAction::Unknown;
+	
+	BlendFactor = FLinearColor::Transparent;
+	FrameBufferSize = CGSizeMake(0.0, 0.0);
+	RenderTargetArraySize = 0;
+    bCanRestartRenderPass = false;
+    
+    RasterBits = EAGXRenderFlagMask;
+    PipelineBits = EAGXPipelineFlagMask;
+}
+
+static bool MTLScissorRectEqual(mtlpp::ScissorRect const& Left, mtlpp::ScissorRect const& Right)
+{
+	return Left.x == Right.x && Left.y == Right.y && Left.width == Right.width && Left.height == Right.height;
+}
+
+void FAGXStateCache::SetScissorRect(bool const bEnable, mtlpp::ScissorRect const& Rect)
+{
+	if (bScissorRectEnabled != bEnable || !MTLScissorRectEqual(Scissor[0], Rect))
+	{
+		bScissorRectEnabled = bEnable;
+		if (bEnable)
+		{
+			Scissor[0] = Rect;
+		}
+		else
+		{
+			Scissor[0].x = Viewport[0].originX;
+			Scissor[0].y = Viewport[0].originY;
+			Scissor[0].width = Viewport[0].width;
+			Scissor[0].height = Viewport[0].height;
+		}
+		
+		// Clamp to framebuffer size - Metal doesn't allow scissor to be larger.
+		Scissor[0].x = Scissor[0].x;
+		Scissor[0].y = Scissor[0].y;
+		Scissor[0].width = FMath::Max((Scissor[0].x + Scissor[0].width <= FMath::RoundToInt(FrameBufferSize.width)) ? Scissor[0].width : FMath::RoundToInt(FrameBufferSize.width) - Scissor[0].x, (NSUInteger)1u);
+		Scissor[0].height = FMath::Max((Scissor[0].y + Scissor[0].height <= FMath::RoundToInt(FrameBufferSize.height)) ? Scissor[0].height : FMath::RoundToInt(FrameBufferSize.height) - Scissor[0].y, (NSUInteger)1u);
+		
+		RasterBits |= EAGXRenderFlagScissorRect;
+	}
+	
+	ActiveScissors = 1;
+}
+
+void FAGXStateCache::SetBlendFactor(FLinearColor const& InBlendFactor)
+{
+	if(BlendFactor != InBlendFactor)
+	{
+		BlendFactor = InBlendFactor;
+		RasterBits |= EAGXRenderFlagBlendColor;
+	}
+}
+
+void FAGXStateCache::SetStencilRef(uint32 const InStencilRef)
+{
+	if(StencilRef != InStencilRef)
+	{
+		StencilRef = InStencilRef;
+		RasterBits |= EAGXRenderFlagStencilReferenceValue;
+	}
+}
+
+void FAGXStateCache::SetDepthStencilState(FAGXDepthStencilState* InDepthStencilState)
+{
+	if(DepthStencilState != InDepthStencilState)
+	{
+		DepthStencilState = InDepthStencilState;
+		RasterBits |= EAGXRenderFlagDepthStencilState;
+	}
+}
+
+void FAGXStateCache::SetRasterizerState(FAGXRasterizerState* InRasterizerState)
+{
+	if(RasterizerState != InRasterizerState)
+	{
+		RasterizerState = InRasterizerState;
+		RasterBits |= EAGXRenderFlagFrontFacingWinding|EAGXRenderFlagCullMode|EAGXRenderFlagDepthBias|EAGXRenderFlagTriangleFillMode;
+	}
+}
+
+void FAGXStateCache::SetComputeShader(FAGXComputeShader* InComputeShader)
+{
+	if(ComputeShader != InComputeShader)
+	{
+		ComputeShader = InComputeShader;
+		
+		PipelineBits |= EAGXPipelineFlagComputeShader;
+		
+		DirtyUniformBuffers[EAGXShaderStages::Compute] = 0xffffffff;
+
+		for (uint32 Index = 0; Index < ML_MaxTextures; ++Index)
+		{
+			ShaderTextures[EAGXShaderStages::Compute].Textures[Index] = nil;
+			ShaderTextures[EAGXShaderStages::Compute].Usage[Index] = mtlpp::ResourceUsage(0);
+		}
+		ShaderTextures[EAGXShaderStages::Compute].Bound = 0;
+
+		for (const auto& PackedGlobalArray : InComputeShader->Bindings.PackedGlobalArrays)
+		{
+			ShaderParameters[EAGXShaderStages::Compute].PrepareGlobalUniforms(CrossCompiler::PackedTypeNameToTypeIndex(PackedGlobalArray.TypeName), PackedGlobalArray.Size);
+		}
+	}
+}
+
+bool FAGXStateCache::SetRenderPassInfo(FRHIRenderPassInfo const& InRenderTargets, FAGXQueryBuffer* QueryBuffer, bool const bRestart)
+{
+	bool bNeedsSet = false;
+	
+	// see if our new Info matches our previous Info
+	if (NeedsToSetRenderTarget(InRenderTargets))
+	{
+		bool bNeedsClear = false;
+		
+		//Create local store action states if we support deferred store
+		mtlpp::StoreAction NewColorStore[MaxSimultaneousRenderTargets];
+		for (uint32 i = 0; i < MaxSimultaneousRenderTargets; ++i)
+		{
+			NewColorStore[i] = mtlpp::StoreAction::Unknown;
+		}
+		
+		mtlpp::StoreAction NewDepthStore = mtlpp::StoreAction::Unknown;
+		mtlpp::StoreAction NewStencilStore = mtlpp::StoreAction::Unknown;
+		
+		// back this up for next frame
+		RenderPassInfo = InRenderTargets;
+		
+		// at this point, we need to fully set up an encoder/command buffer, so make a new one (autoreleased)
+		mtlpp::RenderPassDescriptor RenderPass = FAGXRenderPassDescriptorPool::Get().CreateDescriptor();
+	
+		// if we need to do queries, write to the supplied query buffer
+		if (IsFeatureLevelSupported(GMaxRHIShaderPlatform, ERHIFeatureLevel::ES3_1))
+		{
+			VisibilityResults = QueryBuffer;
+			RenderPass.SetVisibilityResultBuffer(QueryBuffer ? QueryBuffer->Buffer : nil);
+		}
+		else
+		{
+			VisibilityResults = NULL;
+		}
+		
+		if (QueryBuffer != VisibilityResults)
+		{
+			VisibilityOffset = 0;
+			VisibilityWritten = 0;
+		}
+	
+		// default to non-msaa
+	    int32 OldCount = SampleCount;
+		SampleCount = 0;
+	
+		bIsRenderTargetActive = false;
+		bHasValidRenderTarget = false;
+		bHasValidColorTarget = false;
+		
+		bFallbackDepthStencilBound = false;
+		
+		uint8 ArrayTargets = 0;
+		uint8 BoundTargets = 0;
+		uint32 ArrayRenderLayers = UINT_MAX;
+		
+		bool bFramebufferSizeSet = false;
+		FrameBufferSize = CGSizeMake(0.f, 0.f);
+		
+		bCanRestartRenderPass = true;
+		
+		ns::Array<mtlpp::RenderPassColorAttachmentDescriptor> Attachements = RenderPass.GetColorAttachments();
+		
+		uint32 NumColorRenderTargets = RenderPassInfo.GetNumColorRenderTargets();
+		
+		for (uint32 RenderTargetIndex = 0; RenderTargetIndex < MaxSimultaneousRenderTargets; RenderTargetIndex++)
+		{
+			// default to invalid
+			uint8 FormatKey = 0;
+			// only try to set it if it was one that was set (ie less than RenderPassInfo.NumColorRenderTargets)
+			if (RenderTargetIndex < NumColorRenderTargets && RenderPassInfo.ColorRenderTargets[RenderTargetIndex].RenderTarget != nullptr)
+			{
+				const FRHIRenderPassInfo::FColorEntry& RenderTargetView = RenderPassInfo.ColorRenderTargets[RenderTargetIndex];
+				ColorTargets[RenderTargetIndex] = RenderTargetView.RenderTarget;
+				ResolveTargets[RenderTargetIndex] = RenderTargetView.ResolveTarget;
+				
+				FAGXSurface& Surface = *AGXGetMetalSurfaceFromRHITexture(RenderTargetView.RenderTarget);
+				FormatKey = Surface.FormatKey;
+				
+				uint32 Width = FMath::Max((uint32)(Surface.SizeX >> RenderTargetView.MipIndex), (uint32)1);
+				uint32 Height = FMath::Max((uint32)(Surface.SizeY >> RenderTargetView.MipIndex), (uint32)1);
+				if(!bFramebufferSizeSet)
+				{
+					bFramebufferSizeSet = true;
+					FrameBufferSize.width = Width;
+					FrameBufferSize.height = Height;
+				}
+				else
+				{
+					FrameBufferSize.width = FMath::Min(FrameBufferSize.width, (CGFloat)Width);
+					FrameBufferSize.height = FMath::Min(FrameBufferSize.height, (CGFloat)Height);
+				}
+	
+				// if this is the back buffer, make sure we have a usable drawable
+				ConditionalUpdateBackBuffer(Surface);
+	
+				BoundTargets |= 1 << RenderTargetIndex;
+            
+#if !PLATFORM_MAC
+                if (Surface.Texture.GetPtr() == nil)
+                {
+                    SampleCount = OldCount;
+                    bCanRestartRenderPass &= (OldCount <= 1);
+                    return true;
+                }
+#endif
+				
+				// The surface cannot be nil - we have to have a valid render-target array after this call.
+				check (Surface.Texture);
+	
+				// user code generally passes -1 as a default, but we need 0
+				uint32 ArraySliceIndex = RenderTargetView.ArraySlice == 0xFFFFFFFF ? 0 : RenderTargetView.ArraySlice;
+				if (Surface.bIsCubemap)
+				{
+					ArraySliceIndex = GetMetalCubeFace((ECubeFace)ArraySliceIndex);
+				}
+				
+				switch(Surface.Type)
+				{
+					case RRT_Texture2DArray:
+					case RRT_Texture3D:
+					case RRT_TextureCube:
+						if(RenderTargetView.ArraySlice == 0xFFFFFFFF)
+						{
+							ArrayTargets |= (1 << RenderTargetIndex);
+							ArrayRenderLayers = FMath::Min(ArrayRenderLayers, Surface.GetNumFaces());
+						}
+						else
+						{
+							ArrayRenderLayers = 1;
+						}
+						break;
+					default:
+						ArrayRenderLayers = 1;
+						break;
+				}
+	
+				mtlpp::RenderPassColorAttachmentDescriptor ColorAttachment = Attachements[RenderTargetIndex];
+	
+				ERenderTargetStoreAction HighLevelStoreAction = GetStoreAction(RenderTargetView.Action);
+				ERenderTargetLoadAction HighLevelLoadAction = GetLoadAction(RenderTargetView.Action);
+				
+				// on iOS with memory-less MSAA textures we can't load them
+                // in case high level code wants to load and render to MSAA target, set attachment to a resolved texture
+				bool bUseResolvedTexture = false;
+#if PLATFORM_IOS
+				bUseResolvedTexture = (
+					Surface.MSAATexture && 
+					Surface.MSAATexture.GetStorageMode() == mtlpp::StorageMode::Memoryless && 
+					HighLevelLoadAction == ERenderTargetLoadAction::ELoad);
+#endif
+				
+				if (Surface.MSAATexture && !bUseResolvedTexture)
+				{
+					bool bMemoryless = false;
+#if PLATFORM_IOS
+					if (Surface.MSAATexture.GetStorageMode() == mtlpp::StorageMode::Memoryless)
+					{
+						bMemoryless = true;
+						HighLevelLoadAction = ERenderTargetLoadAction::EClear;
+					}
+#endif
+					// set up an MSAA attachment
+					ColorAttachment.SetTexture(Surface.MSAATexture);
+					NewColorStore[RenderTargetIndex] = GetMetalRTStoreAction(ERenderTargetStoreAction::EMultisampleResolve);
+					ColorAttachment.SetStoreAction(!bMemoryless && GRHIDeviceId > 2 ? mtlpp::StoreAction::Unknown : NewColorStore[RenderTargetIndex]);
+					ColorAttachment.SetResolveTexture(Surface.MSAAResolveTexture ? Surface.MSAAResolveTexture : Surface.Texture);
+					SampleCount = Surface.MSAATexture.GetSampleCount();
+					// only allow one MRT with msaa
+					checkf(NumColorRenderTargets == 1, TEXT("Only expected one MRT when using MSAA"));
+				}
+				else
+				{
+					bool bMemoryless = false;
+#if PLATFORM_IOS
+					if (Surface.Texture.GetStorageMode() == mtlpp::StorageMode::Memoryless)
+					{
+						bMemoryless = true;
+						HighLevelStoreAction = ERenderTargetStoreAction::ENoAction;
+						HighLevelLoadAction = ERenderTargetLoadAction::EClear;
+					}
+#endif
+					// set up non-MSAA attachment
+					ColorAttachment.SetTexture(Surface.Texture);
+					NewColorStore[RenderTargetIndex] = GetMetalRTStoreAction(HighLevelStoreAction);
+					ColorAttachment.SetStoreAction(!bMemoryless ? mtlpp::StoreAction::Unknown : NewColorStore[RenderTargetIndex]);
+                    SampleCount = 1;
+				}
+				
+				ColorAttachment.SetLevel(RenderTargetView.MipIndex);
+				if(Surface.Type == RRT_Texture3D)
+				{
+					ColorAttachment.SetSlice(0);
+					ColorAttachment.SetDepthPlane(ArraySliceIndex);
+				}
+				else
+				{
+					ColorAttachment.SetSlice(ArraySliceIndex);
+				}
+				
+				ColorAttachment.SetLoadAction((Surface.Written || !bImmediate || bRestart) ? GetMetalRTLoadAction(HighLevelLoadAction) : mtlpp::LoadAction::Clear);
+				FPlatformAtomics::InterlockedExchange(&Surface.Written, 1);
+				
+				bNeedsClear |= (ColorAttachment.GetLoadAction() == mtlpp::LoadAction::Clear);
+				
+				const FClearValueBinding& ClearValue = RenderPassInfo.ColorRenderTargets[RenderTargetIndex].RenderTarget->GetClearBinding();
+				if (ClearValue.ColorBinding == EClearBinding::EColorBound)
+				{
+					const FLinearColor& ClearColor = ClearValue.GetClearColor();
+					ColorAttachment.SetClearColor(mtlpp::ClearColor(ClearColor.R, ClearColor.G, ClearColor.B, ClearColor.A));
+				}
+
+				bCanRestartRenderPass &= (SampleCount <= 1) && (ColorAttachment.GetLoadAction() == mtlpp::LoadAction::Load) && (HighLevelStoreAction == ERenderTargetStoreAction::EStore);
+	
+				bHasValidRenderTarget = true;
+				bHasValidColorTarget = true;
+			}
+			else
+			{
+				ColorTargets[RenderTargetIndex].SafeRelease();
+				ResolveTargets[RenderTargetIndex].SafeRelease();
+			}
+		}
+		
+		RenderTargetArraySize = 1;
+		
+		if(ArrayTargets)
+		{
+			if (!GetAGXDeviceContext().SupportsFeature(EAGXFeaturesLayeredRendering))
+			{
+				METAL_FATAL_ASSERT(ArrayRenderLayers != 1, TEXT("Layered rendering is unsupported on this device (%d)."), ArrayRenderLayers);
+			}
+#if PLATFORM_MAC
+			else
+			{
+				METAL_FATAL_ASSERT(ArrayTargets == BoundTargets, TEXT("All color render targets must be layered when performing multi-layered rendering under Metal (%d != %d)."), ArrayTargets, BoundTargets);
+					RenderTargetArraySize = ArrayRenderLayers;
+					RenderPass.SetRenderTargetArrayLength(ArrayRenderLayers);
+			}
+#endif
+		}
+	
+		// default to invalid
+		uint8 DepthFormatKey = 0;
+		uint8 StencilFormatKey = 0;
+		
+		// setup depth and/or stencil
+		if (RenderPassInfo.DepthStencilRenderTarget.DepthStencilTarget != nullptr)
+		{
+			FAGXSurface& Surface = *AGXGetMetalSurfaceFromRHITexture(RenderPassInfo.DepthStencilRenderTarget.DepthStencilTarget);
+			
+			switch(Surface.Type)
+			{
+				case RRT_Texture2DArray:
+				case RRT_Texture3D:
+				case RRT_TextureCube:
+					ArrayRenderLayers = Surface.GetNumFaces();
+					break;
+				default:
+					ArrayRenderLayers = 1;
+					break;
+			}
+			if(!ArrayTargets && ArrayRenderLayers > 1)
+			{
+				METAL_FATAL_ASSERT(GetAGXDeviceContext().SupportsFeature(EAGXFeaturesLayeredRendering), TEXT("Layered rendering is unsupported on this device (%d)."), ArrayRenderLayers);
+#if PLATFORM_MAC
+					RenderTargetArraySize = ArrayRenderLayers;
+					RenderPass.SetRenderTargetArrayLength(ArrayRenderLayers);
+#endif
+			}
+			
+			if(!bFramebufferSizeSet)
+			{
+				bFramebufferSizeSet = true;
+				FrameBufferSize.width = Surface.SizeX;
+				FrameBufferSize.height = Surface.SizeY;
+			}
+			else
+			{
+				FrameBufferSize.width = FMath::Min(FrameBufferSize.width, (CGFloat)Surface.SizeX);
+				FrameBufferSize.height = FMath::Min(FrameBufferSize.height, (CGFloat)Surface.SizeY);
+			}
+			
+			EPixelFormat DepthStencilPixelFormat = RenderPassInfo.DepthStencilRenderTarget.DepthStencilTarget->GetFormat();
+			
+			FAGXTexture DepthTexture = nil;
+			FAGXTexture StencilTexture = nil;
+			
+            const bool bSupportSeparateMSAAResolve = FAGXCommandQueue::SupportsSeparateMSAAAndResolveTarget();
+			uint32 DepthSampleCount = (Surface.MSAATexture ? Surface.MSAATexture.GetSampleCount() : Surface.Texture.GetSampleCount());
+            bool bDepthStencilSampleCountMismatchFixup = false;
+            DepthTexture = Surface.MSAATexture ? Surface.MSAATexture : Surface.Texture;
+			if (SampleCount == 0)
+			{
+				SampleCount = DepthSampleCount;
+			}
+			else if (SampleCount != DepthSampleCount)
+            {
+				static bool bLogged = false;
+				if (!bSupportSeparateMSAAResolve)
+				{
+					//in the case of NOT support separate MSAA resolve the high level may legitimately cause a mismatch which we need to handle by binding the resolved target which we normally wouldn't do.
+					DepthTexture = Surface.Texture;
+					bDepthStencilSampleCountMismatchFixup = true;
+					DepthSampleCount = 1;
+				}
+				else if (!bLogged)
+				{
+					UE_LOG(LogAGX, Error, TEXT("If we support separate targets the high level should always give us matching counts"));
+					bLogged = true;
+				}
+            }
+
+			switch (DepthStencilPixelFormat)
+			{
+				case PF_X24_G8:
+				case PF_DepthStencil:
+				case PF_D24:
+				{
+					mtlpp::PixelFormat DepthStencilFormat = Surface.Texture ? (mtlpp::PixelFormat)Surface.Texture.GetPixelFormat() : mtlpp::PixelFormat::Invalid;
+					
+					switch(DepthStencilFormat)
+					{
+						case mtlpp::PixelFormat::Depth32Float:
+							StencilTexture =  nil;
+							break;
+						case mtlpp::PixelFormat::Stencil8:
+							StencilTexture = DepthTexture;
+							break;
+						case mtlpp::PixelFormat::Depth32Float_Stencil8:
+							StencilTexture = DepthTexture;
+							break;
+#if PLATFORM_MAC
+						case mtlpp::PixelFormat::Depth24Unorm_Stencil8:
+							StencilTexture = DepthTexture;
+							break;
+#endif
+						default:
+							break;
+					}
+					
+					break;
+				}
+				case PF_ShadowDepth:
+				{
+					break;
+				}
+				default:
+					break;
+			}
+			
+			float DepthClearValue = 0.0f;
+			uint32 StencilClearValue = 0;
+			const FClearValueBinding& ClearValue = RenderPassInfo.DepthStencilRenderTarget.DepthStencilTarget->GetClearBinding();
+			if (ClearValue.ColorBinding == EClearBinding::EDepthStencilBound)
+			{
+				ClearValue.GetDepthStencil(DepthClearValue, StencilClearValue);
+			}
+			else if(!ArrayTargets && ArrayRenderLayers > 1)
+			{
+				DepthClearValue = 1.0f;
+			}
+
+           bool const bCombinedDepthStencilUsingStencil = (DepthTexture && (mtlpp::PixelFormat)DepthTexture.GetPixelFormat() != mtlpp::PixelFormat::Depth32Float && RenderPassInfo.DepthStencilRenderTarget.ExclusiveDepthStencil.IsUsingStencil());
+			bool const bUsingDepth = (RenderPassInfo.DepthStencilRenderTarget.ExclusiveDepthStencil.IsUsingDepth() || (bCombinedDepthStencilUsingStencil));
+			if (DepthTexture && bUsingDepth)
+			{
+				mtlpp::RenderPassDepthAttachmentDescriptor DepthAttachment;
+				
+				DepthFormatKey = Surface.FormatKey;
+				
+				ERenderTargetActions DepthActions = GetDepthActions(RenderPassInfo.DepthStencilRenderTarget.Action);
+				ERenderTargetLoadAction DepthLoadAction = GetLoadAction(DepthActions);
+				ERenderTargetStoreAction DepthStoreAction = GetStoreAction(DepthActions);
+
+				// set up the depth attachment
+				DepthAttachment.SetTexture(DepthTexture);
+				DepthAttachment.SetLoadAction(GetMetalRTLoadAction(DepthLoadAction));
+				
+				bNeedsClear |= (DepthAttachment.GetLoadAction() == mtlpp::LoadAction::Clear);
+				
+				ERenderTargetStoreAction HighLevelStoreAction = (Surface.MSAATexture && !bDepthStencilSampleCountMismatchFixup) ? ERenderTargetStoreAction::EMultisampleResolve : DepthStoreAction;
+				if (bUsingDepth && (HighLevelStoreAction == ERenderTargetStoreAction::ENoAction || bDepthStencilSampleCountMismatchFixup))
+				{
+					if (DepthSampleCount > 1)
+					{
+						HighLevelStoreAction = ERenderTargetStoreAction::EMultisampleResolve;
+					}
+					else
+					{
+						HighLevelStoreAction = ERenderTargetStoreAction::EStore;
+					}
+				}
+				
+				const bool bSupportsMSAADepthResolve = GetAGXDeviceContext().SupportsFeature(EAGXFeaturesMSAADepthResolve);
+				bool bDepthTextureMemoryless = false;
+#if PLATFORM_IOS
+				bDepthTextureMemoryless = DepthTexture.GetStorageMode() == mtlpp::StorageMode::Memoryless;
+				if (bDepthTextureMemoryless)
+				{
+					DepthAttachment.SetLoadAction(mtlpp::LoadAction::Clear);
+					
+					if (bSupportsMSAADepthResolve && Surface.MSAATexture && DepthStoreAction == ERenderTargetStoreAction::EMultisampleResolve)
+					{
+						HighLevelStoreAction = ERenderTargetStoreAction::EMultisampleResolve;
+					}
+					else
+					{
+						HighLevelStoreAction = ERenderTargetStoreAction::ENoAction;
+					}
+				}
+                else
+                {
+                	HighLevelStoreAction = DepthStoreAction;
+                }
+#endif
+                //needed to quiet the metal validation that runs when you end renderpass. (it requires some kind of 'resolve' for an msaa target)
+				//But with deferredstore we don't set the real one until submit time.
+				NewDepthStore = !Surface.MSAATexture || bSupportsMSAADepthResolve ? GetMetalRTStoreAction(HighLevelStoreAction) : mtlpp::StoreAction::DontCare;
+				DepthAttachment.SetStoreAction(!bDepthTextureMemoryless && Surface.MSAATexture && GRHIDeviceId > 2 ? mtlpp::StoreAction::Unknown : NewDepthStore);
+				DepthAttachment.SetClearDepth(DepthClearValue);
+				check(SampleCount > 0);
+
+				if (Surface.MSAATexture && bSupportsMSAADepthResolve && DepthAttachment.GetStoreAction() != mtlpp::StoreAction::DontCare)
+				{
+                    if (!bDepthStencilSampleCountMismatchFixup)
+                    {
+                        DepthAttachment.SetResolveTexture(Surface.MSAAResolveTexture ? Surface.MSAAResolveTexture : Surface.Texture);
+                    }
+#if PLATFORM_MAC
+					//would like to assert and do manual custom resolve, but that is causing some kind of weird corruption.
+					//checkf(false, TEXT("Depth resolves need to do 'max' for correctness.  MacOS does not expose this yet unless the spec changed."));
+#else
+					DepthAttachment.SetDepthResolveFilter(mtlpp::MultisampleDepthResolveFilter::Max);
+#endif
+				}
+				
+				bHasValidRenderTarget = true;
+				bFallbackDepthStencilBound = (RenderPassInfo.DepthStencilRenderTarget.DepthStencilTarget == FallbackDepthStencilSurface);
+
+				bCanRestartRenderPass &= (SampleCount <= 1) && ((RenderPassInfo.DepthStencilRenderTarget.DepthStencilTarget == FallbackDepthStencilSurface) || ((DepthAttachment.GetLoadAction() == mtlpp::LoadAction::Load) && (!RenderPassInfo.DepthStencilRenderTarget.ExclusiveDepthStencil.IsDepthWrite() || (DepthStoreAction == ERenderTargetStoreAction::EStore))));
+				
+				// and assign it
+				RenderPass.SetDepthAttachment(DepthAttachment);
+			}
+	
+            //if we're dealing with a samplecount mismatch we just bail on stencil entirely as stencil
+            //doesn't have an autoresolve target to use.
+			
+			bool const bCombinedDepthStencilUsingDepth = (StencilTexture && StencilTexture.GetPixelFormat() != mtlpp::PixelFormat::Stencil8 && RenderPassInfo.DepthStencilRenderTarget.ExclusiveDepthStencil.IsUsingDepth());
+			bool const bUsingStencil = RenderPassInfo.DepthStencilRenderTarget.ExclusiveDepthStencil.IsUsingStencil() || (bCombinedDepthStencilUsingDepth);
+			if (StencilTexture && bUsingStencil)
+			{
+				mtlpp::RenderPassStencilAttachmentDescriptor StencilAttachment;
+				
+				StencilFormatKey = Surface.FormatKey;
+				
+				ERenderTargetActions StencilActions = GetStencilActions(RenderPassInfo.DepthStencilRenderTarget.Action);
+				ERenderTargetLoadAction StencilLoadAction = GetLoadAction(StencilActions);
+				ERenderTargetStoreAction StencilStoreAction = GetStoreAction(StencilActions);
+	
+				// set up the stencil attachment
+				StencilAttachment.SetTexture(StencilTexture);
+				StencilAttachment.SetLoadAction(GetMetalRTLoadAction(StencilLoadAction));
+				
+				bNeedsClear |= (StencilAttachment.GetLoadAction() == mtlpp::LoadAction::Clear);
+				
+				ERenderTargetStoreAction HighLevelStoreAction = StencilStoreAction;
+				if (bUsingStencil && (HighLevelStoreAction == ERenderTargetStoreAction::ENoAction || bDepthStencilSampleCountMismatchFixup))
+				{
+					HighLevelStoreAction = ERenderTargetStoreAction::EStore;
+				}
+				
+				bool bStencilMemoryless = false;
+#if PLATFORM_IOS
+				if (StencilTexture.GetStorageMode() == mtlpp::StorageMode::Memoryless)
+				{
+					bStencilMemoryless = true;
+					HighLevelStoreAction = ERenderTargetStoreAction::ENoAction;
+					StencilAttachment.SetLoadAction(mtlpp::LoadAction::Clear);
+				}
+				else
+				{
+					HighLevelStoreAction = StencilStoreAction;
+				}
+#endif
+				
+				// For the case where Depth+Stencil is MSAA we can't Resolve depth and Store stencil - we can only Resolve + DontCare or StoreResolve + Store (on newer H/W and iOS).
+				// We only allow use of StoreResolve in the Desktop renderers as the mobile renderer does not and should not assume hardware support for it.
+				NewStencilStore = (StencilTexture.GetSampleCount() == 1  || GetMetalRTStoreAction(ERenderTargetStoreAction::EMultisampleResolve) == mtlpp::StoreAction::StoreAndMultisampleResolve) ? GetMetalRTStoreAction(HighLevelStoreAction) : mtlpp::StoreAction::DontCare;
+				StencilAttachment.SetStoreAction(!bStencilMemoryless && StencilTexture.GetSampleCount() > 1 && GRHIDeviceId > 2 ? mtlpp::StoreAction::Unknown : NewStencilStore);
+				StencilAttachment.SetClearStencil(StencilClearValue);
+
+				if (SampleCount == 0)
+				{
+					SampleCount = StencilAttachment.GetTexture().GetSampleCount();
+				}
+				
+				bHasValidRenderTarget = true;
+				
+				// @todo Stencil writes that need to persist must use ERenderTargetStoreAction::EStore on iOS.
+				// We should probably be using deferred store actions so that we can safely lazily instantiate encoders.
+				bCanRestartRenderPass &= (SampleCount <= 1) && ((RenderPassInfo.DepthStencilRenderTarget.DepthStencilTarget == FallbackDepthStencilSurface) || ((StencilAttachment.GetLoadAction() == mtlpp::LoadAction::Load) && (1 || !RenderPassInfo.DepthStencilRenderTarget.ExclusiveDepthStencil.IsStencilWrite() || (StencilStoreAction == ERenderTargetStoreAction::EStore))));
+				
+				// and assign it
+				RenderPass.SetStencilAttachment(StencilAttachment);
+			}
+		}
+		
+		//Update deferred store states if required otherwise they're already set directly on the Metal Attachement Descriptors
+		{
+			for (uint32 i = 0; i < MaxSimultaneousRenderTargets; ++i)
+			{
+				ColorStore[i] = NewColorStore[i];
+			}
+			DepthStore = NewDepthStore;
+			StencilStore = NewStencilStore;
+		}
+		
+		if (SampleCount == 0)
+		{
+			SampleCount = 1;
+		}
+		
+		bIsRenderTargetActive = bHasValidRenderTarget;
+		
+		// Only start encoding if the render target state is valid
+		if (bHasValidRenderTarget)
+		{
+			// Retain and/or release the depth-stencil surface in case it is a temporary surface for a draw call that writes to depth without a depth/stencil buffer bound.
+			DepthStencilSurface = RenderPassInfo.DepthStencilRenderTarget.DepthStencilTarget;
+			DepthStencilResolve = RenderPassInfo.DepthStencilRenderTarget.ResolveTarget;
+		}
+		else
+		{
+			DepthStencilSurface.SafeRelease();
+			DepthStencilResolve.SafeRelease();
+		}
+		
+		RenderPassDesc = RenderPass;
+		
+		bNeedsSet = true;
+	}
+
+	return bNeedsSet;
+}
+
+void FAGXStateCache::InvalidateRenderTargets(void)
+{
+	bHasValidRenderTarget = false;
+	bHasValidColorTarget = false;
+	bIsRenderTargetActive = false;
+}
+
+void FAGXStateCache::SetRenderTargetsActive(bool const bActive)
+{
+	bIsRenderTargetActive = bActive;
+}
+
+static bool MTLViewportEqual(mtlpp::Viewport const& Left, mtlpp::Viewport const& Right)
+{
+	return FMath::IsNearlyEqual(Left.originX, Right.originX) &&
+			FMath::IsNearlyEqual(Left.originY, Right.originY) &&
+			FMath::IsNearlyEqual(Left.width, Right.width) &&
+			FMath::IsNearlyEqual(Left.height, Right.height) &&
+			FMath::IsNearlyEqual(Left.znear, Right.znear) &&
+			FMath::IsNearlyEqual(Left.zfar, Right.zfar);
+}
+
+void FAGXStateCache::SetViewport(const mtlpp::Viewport& InViewport)
+{
+	if (!MTLViewportEqual(Viewport[0], InViewport))
+	{
+		Viewport[0] = InViewport;
+	
+		RasterBits |= EAGXRenderFlagViewport;
+	}
+	
+	ActiveViewports = 1;
+	
+	if (!bScissorRectEnabled)
+	{
+		mtlpp::ScissorRect Rect;
+		Rect.x = InViewport.originX;
+		Rect.y = InViewport.originY;
+		Rect.width = InViewport.width;
+		Rect.height = InViewport.height;
+		SetScissorRect(false, Rect);
+	}
+}
+
+void FAGXStateCache::SetViewport(uint32 Index, const mtlpp::Viewport& InViewport)
+{
+	check(Index < ML_MaxViewports);
+	
+	if (!MTLViewportEqual(Viewport[Index], InViewport))
+	{
+		Viewport[Index] = InViewport;
+		
+		RasterBits |= EAGXRenderFlagViewport;
+	}
+	
+	// There may not be gaps in the viewport array.
+	ActiveViewports = Index + 1;
+	
+	// This always sets the scissor rect because the RHI doesn't bother to expose proper scissor states for multiple viewports.
+	// This will have to change if we want to guarantee correctness in the mid to long term.
+	{
+		mtlpp::ScissorRect Rect;
+		Rect.x = InViewport.originX;
+		Rect.y = InViewport.originY;
+		Rect.width = InViewport.width;
+		Rect.height = InViewport.height;
+		SetScissorRect(Index, false, Rect);
+	}
+}
+
+void FAGXStateCache::SetScissorRect(uint32 Index, bool const bEnable, mtlpp::ScissorRect const& Rect)
+{
+	check(Index < ML_MaxViewports);
+	if (!MTLScissorRectEqual(Scissor[Index], Rect))
+	{
+		// There's no way we can setup the bounds correctly - that must be done by the caller or incorrect rendering & crashes will ensue.
+		Scissor[Index] = Rect;
+		RasterBits |= EAGXRenderFlagScissorRect;
+	}
+	
+	ActiveScissors = Index + 1;
+}
+
+void FAGXStateCache::SetViewports(const mtlpp::Viewport InViewport[], uint32 Count)
+{
+	check(Count >= 1 && Count < ML_MaxViewports);
+	
+	// Check if the count has changed first & if so mark for a rebind
+	if (ActiveViewports != Count)
+	{
+		RasterBits |= EAGXRenderFlagViewport;
+		RasterBits |= EAGXRenderFlagScissorRect;
+	}
+	
+	for (uint32 i = 0; i < Count; i++)
+	{
+		SetViewport(i, InViewport[i]);
+	}
+	
+	ActiveViewports = Count;
+}
+
+void FAGXStateCache::SetVertexStream(uint32 const Index, FAGXBuffer* Buffer, FAGXBufferData* Bytes, uint32 const Offset, uint32 const Length)
+{
+	check(Index < MaxVertexElementCount);
+	check(UNREAL_TO_METAL_BUFFER_INDEX(Index) < MaxMetalStreams);
+
+	if (Buffer)
+	{
+		VertexBuffers[Index].Buffer = *Buffer;
+	}
+	else
+	{
+		VertexBuffers[Index].Buffer = nil;
+	}
+	VertexBuffers[Index].Offset = 0;
+	VertexBuffers[Index].Bytes = Bytes;
+	VertexBuffers[Index].Length = Length;
+	
+	SetShaderBuffer(EAGXShaderStages::Vertex, VertexBuffers[Index].Buffer, Bytes, Offset, Length, UNREAL_TO_METAL_BUFFER_INDEX(Index), mtlpp::ResourceUsage::Read);
+}
+
+uint32 FAGXStateCache::GetVertexBufferSize(uint32 const Index)
+{
+	check(Index < MaxVertexElementCount);
+	check(UNREAL_TO_METAL_BUFFER_INDEX(Index) < MaxMetalStreams);
+	return VertexBuffers[Index].Length;
+}
+
+void FAGXStateCache::SetGraphicsPipelineState(FAGXGraphicsPipelineState* State)
+{
+	if (GraphicsPSO != State)
+	{
+		GraphicsPSO = State;
+				
+		DirtyUniformBuffers[EAGXShaderStages::Vertex] = 0xffffffff;
+		DirtyUniformBuffers[EAGXShaderStages::Pixel] = 0xffffffff;
+#if PLATFORM_SUPPORTS_GEOMETRY_SHADERS
+		DirtyUniformBuffers[EAGXShaderStages::Geometry] = 0xffffffff;
+#endif
+		
+		PipelineBits |= EAGXPipelineFlagPipelineState;
+		
+        if (AGXSafeGetRuntimeDebuggingLevel() >= EAGXDebugLevelResetOnBind)
+        {
+            for (uint32 i = 0; i < EAGXShaderStages::Num; i++)
+            {
+                ShaderBuffers[i].Bound = UINT32_MAX;
+                ShaderTextures[i].Bound = FMETALTEXTUREMASK_MAX;
+                ShaderSamplers[i].Bound = UINT16_MAX;
+            }
+        }
+		
+		SetDepthStencilState(State->DepthStencilState);
+		SetRasterizerState(State->RasterizerState);
+
+		for (const auto& PackedGlobalArray : State->VertexShader->Bindings.PackedGlobalArrays)
+		{
+			ShaderParameters[EAGXShaderStages::Vertex].PrepareGlobalUniforms(CrossCompiler::PackedTypeNameToTypeIndex(PackedGlobalArray.TypeName), PackedGlobalArray.Size);
+		}
+
+		if (State->PixelShader)
+		{
+			for (const auto& PackedGlobalArray : State->PixelShader->Bindings.PackedGlobalArrays)
+			{
+				ShaderParameters[EAGXShaderStages::Pixel].PrepareGlobalUniforms(CrossCompiler::PackedTypeNameToTypeIndex(PackedGlobalArray.TypeName), PackedGlobalArray.Size);
+			}
+		}
+	}
+}
+
+FAGXShaderPipeline* FAGXStateCache::GetPipelineState() const
+{
+	return GraphicsPSO->GetPipeline();
+}
+
+EPrimitiveType FAGXStateCache::GetPrimitiveType()
+{
+	check(IsValidRef(GraphicsPSO));
+	return GraphicsPSO->GetPrimitiveType();
+}
+
+void FAGXStateCache::BindUniformBuffer(EAGXShaderStages const Freq, uint32 const BufferIndex, FRHIUniformBuffer* BufferRHI)
+{
+	check(BufferIndex < ML_MaxBuffers);
+	if (BoundUniformBuffers[Freq][BufferIndex] != BufferRHI)
+	{
+		ActiveUniformBuffers.Add(BufferRHI);
+		BoundUniformBuffers[Freq][BufferIndex] = BufferRHI;
+		DirtyUniformBuffers[Freq] |= 1 << BufferIndex;
+	}
+}
+
+void FAGXStateCache::SetDirtyUniformBuffers(EAGXShaderStages const Freq, uint32 const Dirty)
+{
+	DirtyUniformBuffers[Freq] = Dirty;
+}
+
+void FAGXStateCache::SetVisibilityResultMode(mtlpp::VisibilityResultMode const Mode, NSUInteger const Offset)
+{
+	if (VisibilityMode != Mode || VisibilityOffset != Offset)
+	{
+		VisibilityMode = Mode;
+		VisibilityOffset = Offset;
+		
+		RasterBits |= EAGXRenderFlagVisibilityResultMode;
+	}
+}
+
+void FAGXStateCache::ConditionalUpdateBackBuffer(FAGXSurface& Surface)
+{
+	// are we setting the back buffer? if so, make sure we have the drawable
+	if ((Surface.Flags & TexCreate_Presentable))
+	{
+		// update the back buffer texture the first time used this frame
+		if (Surface.Texture.GetPtr() == nil)
+		{
+			// set the texture into the backbuffer
+			Surface.GetDrawableTexture();
+		}
+#if PLATFORM_MAC
+		check (Surface.Texture);
+#endif
+	}
+}
+
+bool FAGXStateCache::NeedsToSetRenderTarget(const FRHIRenderPassInfo& InRenderPassInfo)
+{
+	// see if our new Info matches our previous Info
+	uint32 CurrentNumColorRenderTargets = RenderPassInfo.GetNumColorRenderTargets();
+	uint32 NewNumColorRenderTargets = InRenderPassInfo.GetNumColorRenderTargets();
+	
+	// basic checks
+	bool bAllChecksPassed = GetHasValidRenderTarget() && bIsRenderTargetActive && CurrentNumColorRenderTargets == NewNumColorRenderTargets &&
+		(InRenderPassInfo.DepthStencilRenderTarget.DepthStencilTarget == RenderPassInfo.DepthStencilRenderTarget.DepthStencilTarget);
+
+	// now check each color target if the basic tests passe
+	if (bAllChecksPassed)
+	{
+		for (int32 RenderTargetIndex = 0; RenderTargetIndex < NewNumColorRenderTargets; RenderTargetIndex++)
+		{
+			const FRHIRenderPassInfo::FColorEntry& RenderTargetView = InRenderPassInfo.ColorRenderTargets[RenderTargetIndex];
+			const FRHIRenderPassInfo::FColorEntry& PreviousRenderTargetView = RenderPassInfo.ColorRenderTargets[RenderTargetIndex];
+
+			// handle simple case of switching textures or mip/slice
+			if (RenderTargetView.RenderTarget != PreviousRenderTargetView.RenderTarget ||
+				RenderTargetView.ResolveTarget != PreviousRenderTargetView.ResolveTarget ||
+				RenderTargetView.MipIndex != PreviousRenderTargetView.MipIndex ||
+				RenderTargetView.ArraySlice != PreviousRenderTargetView.ArraySlice)
+			{
+				bAllChecksPassed = false;
+				break;
+			}
+			
+			// it's non-trivial when we need to switch based on load/store action:
+			// LoadAction - it only matters what we are switching to in the new one
+			//    If we switch to Load, no need to switch as we can re-use what we already have
+			//    If we switch to Clear, we have to always switch to a new RT to force the clear
+			//    If we switch to DontCare, there's definitely no need to switch
+			//    If we switch *from* Clear then we must change target as we *don't* want to clear again.
+            if (GetLoadAction(RenderTargetView.Action) == ERenderTargetLoadAction::EClear)
+            {
+                bAllChecksPassed = false;
+                break;
+            }
+            // StoreAction - this matters what the previous one was **In Spirit**
+            //    If we come from Store, we need to switch to a new RT to force the store
+            //    If we come from DontCare, then there's no need to switch
+            //    @todo metal: However, we basically only use Store now, and don't
+            //        care about intermediate results, only final, so we don't currently check the value
+            //			if (PreviousRenderTargetView.StoreAction == ERenderTTargetStoreAction::EStore)
+            //			{
+            //				bAllChecksPassed = false;
+            //				break;
+            //			}
+        }
+        
+        if (InRenderPassInfo.DepthStencilRenderTarget.DepthStencilTarget && (GetLoadAction(GetDepthActions(InRenderPassInfo.DepthStencilRenderTarget.Action)) == ERenderTargetLoadAction::EClear || GetLoadAction(GetStencilActions(InRenderPassInfo.DepthStencilRenderTarget.Action)) == ERenderTargetLoadAction::EClear))
+        {
+            bAllChecksPassed = false;
+		}
+		
+		if (InRenderPassInfo.DepthStencilRenderTarget.DepthStencilTarget && (GetStoreAction(GetDepthActions(InRenderPassInfo.DepthStencilRenderTarget.Action)) > GetStoreAction(GetDepthActions(RenderPassInfo.DepthStencilRenderTarget.Action)) || GetStoreAction(GetStencilActions(InRenderPassInfo.DepthStencilRenderTarget.Action)) > GetStoreAction(GetStencilActions(RenderPassInfo.DepthStencilRenderTarget.Action))))
+		{
+			// Don't break the encoder if we can just change the store actions.
+			mtlpp::StoreAction NewDepthStore = DepthStore;
+			mtlpp::StoreAction NewStencilStore = StencilStore;
+			if (GetStoreAction(GetDepthActions(InRenderPassInfo.DepthStencilRenderTarget.Action)) > GetStoreAction(GetDepthActions(RenderPassInfo.DepthStencilRenderTarget.Action)))
+			{
+				if (RenderPassDesc.GetDepthAttachment().GetTexture())
+				{
+					FAGXSurface& Surface = *AGXGetMetalSurfaceFromRHITexture(RenderPassInfo.DepthStencilRenderTarget.DepthStencilTarget);
+					
+					const uint32 DepthSampleCount = (Surface.MSAATexture ? Surface.MSAATexture.GetSampleCount() : Surface.Texture.GetSampleCount());
+					bool const bDepthStencilSampleCountMismatchFixup = (SampleCount != DepthSampleCount);
+
+					ERenderTargetStoreAction HighLevelStoreAction = (Surface.MSAATexture && !bDepthStencilSampleCountMismatchFixup) ? ERenderTargetStoreAction::EMultisampleResolve : GetStoreAction(GetDepthActions(RenderPassInfo.DepthStencilRenderTarget.Action));
+					
+#if PLATFORM_IOS
+					FAGXTexture& Tex = Surface.MSAATexture ? Surface.MSAATexture : Surface.Texture;
+					if (Tex.GetStorageMode() == mtlpp::StorageMode::Memoryless)
+					{
+						HighLevelStoreAction = ERenderTargetStoreAction::ENoAction;
+					}
+#endif
+					
+					NewDepthStore = GetMetalRTStoreAction(HighLevelStoreAction);
+				}
+				else
+				{
+					bAllChecksPassed = false;
+				}
+			}
+			
+			if (GetStoreAction(GetStencilActions(InRenderPassInfo.DepthStencilRenderTarget.Action)) > GetStoreAction(GetStencilActions(RenderPassInfo.DepthStencilRenderTarget.Action)))
+			{
+				if (RenderPassDesc.GetStencilAttachment().GetTexture())
+				{
+					NewStencilStore = GetMetalRTStoreAction(GetStoreAction(GetStencilActions(RenderPassInfo.DepthStencilRenderTarget.Action)));
+#if PLATFORM_IOS
+					if (RenderPassDesc.GetStencilAttachment().GetTexture().GetStorageMode() == mtlpp::StorageMode::Memoryless)
+					{
+						NewStencilStore = GetMetalRTStoreAction(ERenderTargetStoreAction::ENoAction);
+					}
+#endif
+				}
+				else
+				{
+					bAllChecksPassed = false;
+				}
+			}
+			
+			if (bAllChecksPassed)
+			{
+				DepthStore = NewDepthStore;
+				StencilStore = NewStencilStore;
+			}
+		}
+	}
+
+	// if we are setting them to nothing, then this is probably end of frame, and we can't make a framebuffer
+	// with nothng, so just abort this (only need to check on single MRT case)
+	if (NewNumColorRenderTargets == 1 && InRenderPassInfo.ColorRenderTargets[0].RenderTarget == nullptr &&
+		InRenderPassInfo.DepthStencilRenderTarget.DepthStencilTarget == nullptr)
+	{
+		bAllChecksPassed = true;
+	}
+
+	return bAllChecksPassed == false;
+}
+
+void FAGXStateCache::SetShaderBuffer(EAGXShaderStages const Frequency, FAGXBuffer const& Buffer, FAGXBufferData* const Bytes, NSUInteger const Offset, NSUInteger const Length, NSUInteger const Index, mtlpp::ResourceUsage const Usage, EPixelFormat const Format)
+{
+	check(Frequency < EAGXShaderStages::Num);
+	check(Index < ML_MaxBuffers);
+	
+	if (ShaderBuffers[Frequency].Buffers[Index].Buffer != Buffer ||
+		ShaderBuffers[Frequency].Buffers[Index].Bytes != Bytes ||
+		ShaderBuffers[Frequency].Buffers[Index].Offset != Offset ||
+		ShaderBuffers[Frequency].Buffers[Index].Length != Length ||
+		!(ShaderBuffers[Frequency].Buffers[Index].Usage & Usage) ||
+		ShaderBuffers[Frequency].Formats[Index] != Format)
+	{
+		ShaderBuffers[Frequency].Buffers[Index].Buffer = Buffer;
+		ShaderBuffers[Frequency].Buffers[Index].Bytes = Bytes;
+		ShaderBuffers[Frequency].Buffers[Index].Offset = Offset;
+		ShaderBuffers[Frequency].Buffers[Index].Length = Length;
+		ShaderBuffers[Frequency].Buffers[Index].Usage = Usage;
+		
+		ShaderBuffers[Frequency].Formats[Index] = Format;
+		
+		if (Buffer || Bytes)
+		{
+			ShaderBuffers[Frequency].Bound |= (1 << Index);
+		}
+		else
+		{
+			ShaderBuffers[Frequency].Bound &= ~(1 << Index);
+		}
+	}
+}
+
+void FAGXStateCache::SetShaderTexture(EAGXShaderStages const Frequency, FAGXTexture const& Texture, NSUInteger const Index, mtlpp::ResourceUsage const Usage)
+{
+	check(Frequency < EAGXShaderStages::Num);
+	check(Index < ML_MaxTextures);
+
+#if (PLATFORM_IOS || PLATFORM_TVOS)
+    UE_CLOG([Texture.GetPtr() storageMode] == MTLStorageModeMemoryless, LogAGX, Fatal, TEXT("FATAL: Attempting to bind a memoryless texture. Stage %u Index %u Texture %@"), Frequency, Index, Texture.GetPtr());
+#endif
+	
+	if (ShaderTextures[Frequency].Textures[Index] != Texture
+		|| ShaderTextures[Frequency].Usage[Index] != Usage)
+	{
+		ShaderTextures[Frequency].Textures[Index] = Texture;
+		ShaderTextures[Frequency].Usage[Index] = Usage;
+		
+		if (Texture)
+		{
+			ShaderTextures[Frequency].Bound |= (FAGXTextureMask(1) << FAGXTextureMask(Index));
+		}
+		else
+		{
+			ShaderTextures[Frequency].Bound &= ~(FAGXTextureMask(1) << FAGXTextureMask(Index));
+		}
+	}
+}
+
+void FAGXStateCache::SetShaderSamplerState(EAGXShaderStages const Frequency, FAGXSamplerState* const Sampler, NSUInteger const Index)
+{
+	check(Frequency < EAGXShaderStages::Num);
+	check(Index < ML_MaxSamplers);
+	
+	if (ShaderSamplers[Frequency].Samplers[Index].GetPtr() != (Sampler ? Sampler->State.GetPtr() : nil))
+	{
+		if (Sampler)
+		{
+#if !PLATFORM_MAC
+			ShaderSamplers[Frequency].Samplers[Index] = ((Frequency == EAGXShaderStages::Vertex || Frequency == EAGXShaderStages::Compute) && Sampler->NoAnisoState) ? Sampler->NoAnisoState : Sampler->State;
+#else
+			ShaderSamplers[Frequency].Samplers[Index] = Sampler->State;
+#endif
+			ShaderSamplers[Frequency].Bound |= (1 << Index);
+		}
+		else
+		{
+			ShaderSamplers[Frequency].Samplers[Index] = nil;
+			ShaderSamplers[Frequency].Bound &= ~(1 << Index);
+		}
+	}
+}
+
+void FAGXStateCache::SetResource(uint32 ShaderStage, uint32 BindIndex, FRHITexture* RESTRICT TextureRHI, float CurrentTime)
+{
+	FAGXSurface* Surface = AGXGetMetalSurfaceFromRHITexture(TextureRHI);
+	ns::AutoReleased<FAGXTexture> Texture;
+	mtlpp::ResourceUsage Usage = (mtlpp::ResourceUsage)0;
+	if (Surface != nullptr)
+	{
+		TextureRHI->SetLastRenderTime(CurrentTime);
+		Texture = Surface->Texture;
+		Usage = mtlpp::ResourceUsage(mtlpp::ResourceUsage::Read|mtlpp::ResourceUsage::Sample);
+	}
+	
+	switch (ShaderStage)
+	{
+		case CrossCompiler::SHADER_STAGE_PIXEL:
+			SetShaderTexture(EAGXShaderStages::Pixel, Texture, BindIndex, Usage);
+			break;
+			
+		case CrossCompiler::SHADER_STAGE_VERTEX:
+			SetShaderTexture(EAGXShaderStages::Vertex, Texture, BindIndex, Usage);
+			break;
+			
+		case CrossCompiler::SHADER_STAGE_COMPUTE:
+			SetShaderTexture(EAGXShaderStages::Compute, Texture, BindIndex, Usage);
+			break;
+			
+		default:
+			check(0);
+			break;
+	}
+}
+
+void FAGXStateCache::SetShaderResourceView(FAGXContext* Context, EAGXShaderStages ShaderStage, uint32 BindIndex, FAGXShaderResourceView* RESTRICT SRV)
+{
+	if (SRV)
+	{
+		FRHITexture* Texture = SRV->SourceTexture.GetReference();
+		FAGXVertexBuffer* VB = SRV->SourceVertexBuffer.GetReference();
+		FAGXIndexBuffer* IB = SRV->SourceIndexBuffer.GetReference();
+		FAGXStructuredBuffer* SB = SRV->SourceStructuredBuffer.GetReference();
+		if (Texture)
+		{
+			FAGXSurface* Surface = SRV->TextureView;
+			if (Surface != nullptr)
+			{
+				SetShaderTexture(ShaderStage, Surface->Texture, BindIndex, mtlpp::ResourceUsage(mtlpp::ResourceUsage::Read|mtlpp::ResourceUsage::Sample));
+			}
+			else
+			{
+				SetShaderTexture(ShaderStage, nil, BindIndex, mtlpp::ResourceUsage(0));
+			}
+		}
+		else if (IsLinearBuffer(ShaderStage, BindIndex) && SRV->GetLinearTexture(false))
+		{
+			ns::AutoReleased<FAGXTexture> Tex;
+			Tex = SRV->GetLinearTexture(false);
+			
+			SetShaderTexture(ShaderStage, Tex, BindIndex, mtlpp::ResourceUsage(mtlpp::ResourceUsage::Read|mtlpp::ResourceUsage::Sample));
+			if (VB)
+            {
+				SetShaderBuffer(ShaderStage, VB->GetCurrentBufferOrNil(), VB->Data, SRV->Offset, VB->GetSize(), BindIndex, mtlpp::ResourceUsage::Read, (EPixelFormat)SRV->Format);
+            }
+            else if (IB)
+            {
+                SetShaderBuffer(ShaderStage, IB->GetCurrentBufferOrNil(), nil, SRV->Offset, IB->GetSize(), BindIndex, mtlpp::ResourceUsage::Read, (EPixelFormat)SRV->Format);
+            }
+		}
+		else if (VB)
+		{
+			SetShaderBuffer(ShaderStage, VB->GetCurrentBufferOrNil(), VB->Data, SRV->Offset, VB->GetSize(), BindIndex, mtlpp::ResourceUsage::Read, (EPixelFormat)SRV->Format);
+		}
+		else if (IB)
+		{
+			SetShaderBuffer(ShaderStage, IB->GetCurrentBufferOrNil(), nil, SRV->Offset, IB->GetSize(), BindIndex, mtlpp::ResourceUsage::Read, (EPixelFormat)SRV->Format);
+		}
+		else if (SB)
+		{
+			SetShaderBuffer(ShaderStage, SB->GetCurrentBufferOrNil(), nil, SRV->Offset, SB->GetSize(), BindIndex, mtlpp::ResourceUsage::Read);
+		}
+	}
+}
+
+bool FAGXStateCache::IsLinearBuffer(EAGXShaderStages ShaderStage, uint32 BindIndex)
+{
+    switch (ShaderStage)
+    {
+        case EAGXShaderStages::Vertex:
+        {
+            return (GraphicsPSO->VertexShader->Bindings.LinearBuffer & (1 << BindIndex)) != 0;
+            break;
+        }
+        case EAGXShaderStages::Pixel:
+        {
+            return (GraphicsPSO->PixelShader->Bindings.LinearBuffer & (1 << BindIndex)) != 0;
+            break;
+        }
+        case EAGXShaderStages::Compute:
+        {
+            return (ComputeShader->Bindings.LinearBuffer & (1 << BindIndex)) != 0;
+        }
+        default:
+        {
+            check(false);
+            return false;
+        }
+    }
+}
+
+void FAGXStateCache::SetShaderUnorderedAccessView(EAGXShaderStages ShaderStage, uint32 BindIndex, FAGXUnorderedAccessView* RESTRICT UAV)
+{
+	if (UAV)
+	{
+		// figure out which one of the resources we need to set
+		FAGXStructuredBuffer* StructuredBuffer = UAV->SourceView->SourceStructuredBuffer.GetReference();
+		FAGXVertexBuffer* VertexBuffer = UAV->SourceView->SourceVertexBuffer.GetReference();
+		FAGXIndexBuffer* IndexBuffer = UAV->SourceView->SourceIndexBuffer.GetReference();
+		FRHITexture* Texture = UAV->SourceView->SourceTexture.GetReference();
+		FAGXSurface* Surface = UAV->SourceView->TextureView;
+		if (StructuredBuffer)
+		{
+			SetShaderBuffer(ShaderStage, StructuredBuffer->GetCurrentBufferOrNil(), nil, 0, StructuredBuffer->GetSize(), BindIndex, mtlpp::ResourceUsage(mtlpp::ResourceUsage::Read | mtlpp::ResourceUsage::Write));
+		}
+		else if (VertexBuffer)
+		{
+			check(!VertexBuffer->Data && VertexBuffer->GetCurrentBufferOrNil());
+			if (IsLinearBuffer(ShaderStage, BindIndex) && UAV->SourceView->GetLinearTexture(true))
+			{
+				ns::AutoReleased<FAGXTexture> Tex;
+				Tex = UAV->SourceView->GetLinearTexture(true);
+				SetShaderTexture(ShaderStage, Tex, BindIndex, mtlpp::ResourceUsage(mtlpp::ResourceUsage::Read | mtlpp::ResourceUsage::Write));
+                
+                SetShaderBuffer(ShaderStage, VertexBuffer->GetCurrentBufferOrNil(), VertexBuffer->Data, 0, VertexBuffer->GetSize(), BindIndex, mtlpp::ResourceUsage(mtlpp::ResourceUsage::Read | mtlpp::ResourceUsage::Write), (EPixelFormat)UAV->SourceView->Format);
+			}
+			else
+			{
+				SetShaderBuffer(ShaderStage, VertexBuffer->GetCurrentBufferOrNil(), VertexBuffer->Data, 0, VertexBuffer->GetSize(), BindIndex, mtlpp::ResourceUsage(mtlpp::ResourceUsage::Read | mtlpp::ResourceUsage::Write), (EPixelFormat)UAV->SourceView->Format);
+			}
+		}
+		else if (IndexBuffer)
+		{
+			check(IndexBuffer->GetCurrentBufferOrNil());
+			if (IsLinearBuffer(ShaderStage, BindIndex) && UAV->SourceView->GetLinearTexture(true))
+			{
+				ns::AutoReleased<FAGXTexture> Tex;
+				Tex = UAV->SourceView->GetLinearTexture(true);
+				SetShaderTexture(ShaderStage, Tex, BindIndex, mtlpp::ResourceUsage(mtlpp::ResourceUsage::Read | mtlpp::ResourceUsage::Write));
+				
+				SetShaderBuffer(ShaderStage, IndexBuffer->GetCurrentBufferOrNil(), nullptr, 0, IndexBuffer->GetSize(), BindIndex, mtlpp::ResourceUsage(mtlpp::ResourceUsage::Read | mtlpp::ResourceUsage::Write), (EPixelFormat)UAV->SourceView->Format);
+			}
+			else
+			{
+				SetShaderBuffer(ShaderStage, IndexBuffer->GetCurrentBufferOrNil(), nullptr, 0, IndexBuffer->GetSize(), BindIndex, mtlpp::ResourceUsage(mtlpp::ResourceUsage::Read | mtlpp::ResourceUsage::Write), (EPixelFormat)UAV->SourceView->Format);
+			}
+		}
+		else if (Texture)
+		{
+			if (!Surface)
+			{
+				Surface = AGXGetMetalSurfaceFromRHITexture(Texture);
+			}
+			if (Surface != nullptr)
+			{
+				FAGXSurface* Source = AGXGetMetalSurfaceFromRHITexture(Texture);
+				
+				FPlatformAtomics::InterlockedExchange(&Surface->Written, 1);
+				FPlatformAtomics::InterlockedExchange(&Source->Written, 1);
+				
+				SetShaderTexture(ShaderStage, Surface->Texture, BindIndex, mtlpp::ResourceUsage(mtlpp::ResourceUsage::Read | mtlpp::ResourceUsage::Write));
+				
+				if ((Source->Flags & (TexCreate_UAV|TexCreate_NoTiling)) == (TexCreate_UAV|TexCreate_NoTiling) && Surface->Texture.GetBuffer())
+				{
+					FAGXBuffer Buffer(Surface->Texture.GetBuffer(), false);
+					SetShaderBuffer(ShaderStage, Buffer, nullptr, 0, Surface->Texture.GetBuffer().GetLength(), BindIndex, mtlpp::ResourceUsage(mtlpp::ResourceUsage::Read | mtlpp::ResourceUsage::Write), PF_MAX);
+				}
+			}
+			else
+			{
+				SetShaderTexture(ShaderStage, nil, BindIndex, mtlpp::ResourceUsage(0));
+			}
+		}
+	}
+}
+
+void FAGXStateCache::SetResource(uint32 ShaderStage, uint32 BindIndex, FAGXShaderResourceView* RESTRICT SRV, float CurrentTime)
+{
+	switch (ShaderStage)
+	{
+		case CrossCompiler::SHADER_STAGE_PIXEL:
+			SetShaderResourceView(nullptr, EAGXShaderStages::Pixel, BindIndex, SRV);
+			break;
+			
+		case CrossCompiler::SHADER_STAGE_VERTEX:
+			SetShaderResourceView(nullptr, EAGXShaderStages::Vertex, BindIndex, SRV);
+			break;
+			
+		case CrossCompiler::SHADER_STAGE_COMPUTE:
+			SetShaderResourceView(nullptr, EAGXShaderStages::Compute, BindIndex, SRV);
+			break;
+			
+		default:
+			check(0);
+			break;
+	}
+}
+
+void FAGXStateCache::SetResource(uint32 ShaderStage, uint32 BindIndex, FAGXSamplerState* RESTRICT SamplerState, float CurrentTime)
+{
+	check(SamplerState->State);
+	switch (ShaderStage)
+	{
+		case CrossCompiler::SHADER_STAGE_PIXEL:
+			SetShaderSamplerState(EAGXShaderStages::Pixel, SamplerState, BindIndex);
+			break;
+			
+		case CrossCompiler::SHADER_STAGE_VERTEX:
+			SetShaderSamplerState(EAGXShaderStages::Vertex, SamplerState, BindIndex);
+			break;
+			
+		case CrossCompiler::SHADER_STAGE_COMPUTE:
+			SetShaderSamplerState(EAGXShaderStages::Compute, SamplerState, BindIndex);
+			break;
+			
+		default:
+			check(0);
+			break;
+	}
+}
+
+void FAGXStateCache::SetResource(uint32 ShaderStage, uint32 BindIndex, FAGXUnorderedAccessView* RESTRICT UAV, float CurrentTime)
+{
+	switch (ShaderStage)
+	{
+		case CrossCompiler::SHADER_STAGE_PIXEL:
+			SetShaderUnorderedAccessView(EAGXShaderStages::Pixel, BindIndex, UAV);
+			break;
+			
+		case CrossCompiler::SHADER_STAGE_VERTEX:
+			SetShaderUnorderedAccessView(EAGXShaderStages::Vertex, BindIndex, UAV);
+			break;
+			
+		case CrossCompiler::SHADER_STAGE_COMPUTE:
+			SetShaderUnorderedAccessView(EAGXShaderStages::Compute, BindIndex, UAV);
+			break;
+			
+		default:
+			check(0);
+			break;
+	}
+}
+
+
+template <typename MetalResourceType>
+inline int32 FAGXStateCache::SetShaderResourcesFromBuffer(uint32 ShaderStage, FAGXUniformBuffer* RESTRICT Buffer, const uint32* RESTRICT ResourceMap, int32 BufferIndex, float CurrentTime)
+{
+	const TRefCountPtr<FRHIResource>* RESTRICT Resources = Buffer->ResourceTable.GetData();
+	int32 NumSetCalls = 0;
+	uint32 BufferOffset = ResourceMap[BufferIndex];
+	if (BufferOffset > 0)
+	{
+		const uint32* RESTRICT ResourceInfos = &ResourceMap[BufferOffset];
+		uint32 ResourceInfo = *ResourceInfos++;
+		do
+		{
+			checkSlow(FRHIResourceTableEntry::GetUniformBufferIndex(ResourceInfo) == BufferIndex);
+			const uint16 ResourceIndex = FRHIResourceTableEntry::GetResourceIndex(ResourceInfo);
+			const uint8 BindIndex = FRHIResourceTableEntry::GetBindIndex(ResourceInfo);
+			
+			MetalResourceType* ResourcePtr = (MetalResourceType*)Resources[ResourceIndex].GetReference();
+			
+			// todo: could coalesce adjacent bound resources.
+			SetResource(ShaderStage, BindIndex, ResourcePtr, CurrentTime);
+			
+			NumSetCalls++;
+			ResourceInfo = *ResourceInfos++;
+		} while (FRHIResourceTableEntry::GetUniformBufferIndex(ResourceInfo) == BufferIndex);
+	}
+	return NumSetCalls;
+}
+
+template <class ShaderType>
+void FAGXStateCache::SetResourcesFromTables(ShaderType Shader, uint32 ShaderStage)
+{
+	checkSlow(Shader);
+	
+	EAGXShaderStages Frequency;
+	switch(ShaderStage)
+	{
+		case CrossCompiler::SHADER_STAGE_VERTEX:
+			Frequency = EAGXShaderStages::Vertex;
+			break;
+		case CrossCompiler::SHADER_STAGE_PIXEL:
+			Frequency = EAGXShaderStages::Pixel;
+			break;
+		case CrossCompiler::SHADER_STAGE_COMPUTE:
+			Frequency = EAGXShaderStages::Compute;
+			break;
+		default:
+			Frequency = EAGXShaderStages::Num; //Silence a compiler warning/error
+			check(false);
+			break;
+	}
+
+	float CurrentTime = FPlatformTime::Seconds();
+
+	// Mask the dirty bits by those buffers from which the shader has bound resources.
+	uint32 DirtyBits = Shader->Bindings.ShaderResourceTable.ResourceTableBits & GetDirtyUniformBuffers(Frequency);
+	while (DirtyBits)
+	{
+		// Scan for the lowest set bit, compute its index, clear it in the set of dirty bits.
+		const uint32 LowestBitMask = (DirtyBits)& (-(int32)DirtyBits);
+		const int32 BufferIndex = FMath::FloorLog2(LowestBitMask); // todo: This has a branch on zero, we know it could never be zero...
+		DirtyBits ^= LowestBitMask;
+		FAGXUniformBuffer* Buffer = (FAGXUniformBuffer*)GetBoundUniformBuffers(Frequency)[BufferIndex];
+		if (Buffer && !FAGXCommandQueue::SupportsFeature(EAGXFeaturesIABs))
+		{
+			check(BufferIndex < Shader->Bindings.ShaderResourceTable.ResourceTableLayoutHashes.Num());
+			check(Buffer->GetLayout().GetHash() == Shader->Bindings.ShaderResourceTable.ResourceTableLayoutHashes[BufferIndex]);
+			
+			// todo: could make this two pass: gather then set
+			SetShaderResourcesFromBuffer<FRHITexture>(ShaderStage, Buffer, Shader->Bindings.ShaderResourceTable.TextureMap.GetData(), BufferIndex, CurrentTime);
+			SetShaderResourcesFromBuffer<FAGXShaderResourceView>(ShaderStage, Buffer, Shader->Bindings.ShaderResourceTable.ShaderResourceViewMap.GetData(), BufferIndex, CurrentTime);
+			SetShaderResourcesFromBuffer<FAGXSamplerState>(ShaderStage, Buffer, Shader->Bindings.ShaderResourceTable.SamplerMap.GetData(), BufferIndex, CurrentTime);
+			SetShaderResourcesFromBuffer<FAGXUnorderedAccessView>(ShaderStage, Buffer, Shader->Bindings.ShaderResourceTable.UnorderedAccessViewMap.GetData(), BufferIndex, CurrentTime);
+		}
+	}
+	SetDirtyUniformBuffers(Frequency, 0);
+}
+
+void FAGXStateCache::CommitRenderResources(FAGXCommandEncoder* Raster)
+{
+	check(IsValidRef(GraphicsPSO));
+    
+    SetResourcesFromTables(GraphicsPSO->VertexShader, CrossCompiler::SHADER_STAGE_VERTEX);
+    GetShaderParameters(EAGXShaderStages::Vertex).CommitPackedGlobals(this, Raster, EAGXShaderStages::Vertex, GraphicsPSO->VertexShader->Bindings);
+	
+    if (IsValidRef(GraphicsPSO->PixelShader))
+    {
+    	SetResourcesFromTables(GraphicsPSO->PixelShader, CrossCompiler::SHADER_STAGE_PIXEL);
+        GetShaderParameters(EAGXShaderStages::Pixel).CommitPackedGlobals(this, Raster, EAGXShaderStages::Pixel, GraphicsPSO->PixelShader->Bindings);
+    }
+}
+
+void FAGXStateCache::CommitComputeResources(FAGXCommandEncoder* Compute)
+{
+	check(IsValidRef(ComputeShader));
+	SetResourcesFromTables(ComputeShader, CrossCompiler::SHADER_STAGE_COMPUTE);
+	
+	GetShaderParameters(EAGXShaderStages::Compute).CommitPackedGlobals(this, Compute, EAGXShaderStages::Compute, ComputeShader->Bindings);
+}
+
+bool FAGXStateCache::PrepareToRestart(bool const bCurrentApplied)
+{
+	if(CanRestartRenderPass())
+	{
+		return true;
+	}
+	else
+	{
+		if (SampleCount <= 1)
+		{
+			FRHIRenderPassInfo Info = GetRenderPassInfo();
+			
+			ERenderTargetActions DepthActions = GetDepthActions(Info.DepthStencilRenderTarget.Action);
+			ERenderTargetActions StencilActions = GetStencilActions(Info.DepthStencilRenderTarget.Action);
+			ERenderTargetLoadAction DepthLoadAction = GetLoadAction(DepthActions);
+			ERenderTargetStoreAction DepthStoreAction = GetStoreAction(DepthActions);
+			ERenderTargetLoadAction StencilLoadAction = GetLoadAction(StencilActions);
+			ERenderTargetStoreAction StencilStoreAction = GetStoreAction(StencilActions);
+
+			if (Info.DepthStencilRenderTarget.DepthStencilTarget)
+			{
+				if (bCurrentApplied || DepthLoadAction != ERenderTargetLoadAction::EClear)
+				{
+					DepthLoadAction = ERenderTargetLoadAction::ELoad;
+				}
+				if (Info.DepthStencilRenderTarget.ExclusiveDepthStencil.IsDepthWrite())
+				{
+					DepthStoreAction = ERenderTargetStoreAction::EStore;
+				}
+
+				if (bCurrentApplied || StencilLoadAction != ERenderTargetLoadAction::EClear)
+				{
+					StencilLoadAction = ERenderTargetLoadAction::ELoad;
+				}
+				if (Info.DepthStencilRenderTarget.ExclusiveDepthStencil.IsStencilWrite())
+				{
+					StencilStoreAction = ERenderTargetStoreAction::EStore;
+				}
+				
+				DepthActions = MakeRenderTargetActions(DepthLoadAction, DepthStoreAction);
+				StencilActions = MakeRenderTargetActions(StencilLoadAction, StencilStoreAction);
+				Info.DepthStencilRenderTarget.Action = MakeDepthStencilTargetActions(DepthActions, StencilActions);
+			}
+			
+			for (int32 RenderTargetIndex = 0; RenderTargetIndex < Info.GetNumColorRenderTargets(); RenderTargetIndex++)
+			{
+				FRHIRenderPassInfo::FColorEntry& RenderTargetView = Info.ColorRenderTargets[RenderTargetIndex];
+				ERenderTargetLoadAction LoadAction = GetLoadAction(RenderTargetView.Action);
+				ERenderTargetStoreAction StoreAction = GetStoreAction(RenderTargetView.Action);
+				
+				if (!bCurrentApplied && LoadAction == ERenderTargetLoadAction::EClear)
+				{
+					StoreAction == ERenderTargetStoreAction::EStore;
+				}
+				else
+				{
+					LoadAction = ERenderTargetLoadAction::ELoad;
+				}
+				RenderTargetView.Action = MakeRenderTargetActions(LoadAction, StoreAction);
+				check(RenderTargetView.RenderTarget == nil || GetStoreAction(RenderTargetView.Action) == ERenderTargetStoreAction::EStore);
+			}
+			
+			InvalidateRenderTargets();
+			return SetRenderPassInfo(Info, GetVisibilityResultsBuffer(), true) && CanRestartRenderPass();
+		}
+		else
+		{
+			return false;
+		}
+	}
+}
+
+void FAGXStateCache::SetStateDirty(void)
+{	
+	RasterBits = UINT32_MAX;
+    PipelineBits = EAGXPipelineFlagMask;
+	for (uint32 i = 0; i < EAGXShaderStages::Num; i++)
+	{
+		ShaderBuffers[i].Bound = UINT32_MAX;
+		ShaderTextures[i].Bound = FMETALTEXTUREMASK_MAX;
+		ShaderSamplers[i].Bound = UINT16_MAX;
+	}
+}
+
+void FAGXStateCache::SetShaderBufferDirty(EAGXShaderStages const Frequency, NSUInteger const Index)
+{
+	ShaderBuffers[Frequency].Bound |= (1 << Index);
+}
+
+void FAGXStateCache::SetRenderStoreActions(FAGXCommandEncoder& CommandEncoder, bool const bConditionalSwitch)
+{
+	check(CommandEncoder.IsRenderCommandEncoderActive())
+	{
+		if (bConditionalSwitch)
+		{
+			ns::Array<mtlpp::RenderPassColorAttachmentDescriptor> ColorAttachments = RenderPassDesc.GetColorAttachments();
+			for (int32 RenderTargetIndex = 0; RenderTargetIndex < RenderPassInfo.GetNumColorRenderTargets(); RenderTargetIndex++)
+			{
+				FRHIRenderPassInfo::FColorEntry& RenderTargetView = RenderPassInfo.ColorRenderTargets[RenderTargetIndex];
+				if(RenderTargetView.RenderTarget != nil)
+				{
+					const bool bMultiSampled = (ColorAttachments[RenderTargetIndex].GetTexture().GetSampleCount() > 1);
+					ColorStore[RenderTargetIndex] = GetConditionalMetalRTStoreAction(bMultiSampled);
+				}
+			}
+			
+			if (RenderPassInfo.DepthStencilRenderTarget.DepthStencilTarget)
+			{
+				const bool bMultiSampled = RenderPassDesc.GetDepthAttachment().GetTexture() && (RenderPassDesc.GetDepthAttachment().GetTexture().GetSampleCount() > 1);
+				DepthStore = GetConditionalMetalRTStoreAction(bMultiSampled);
+				StencilStore = GetConditionalMetalRTStoreAction(false);
+			}
+		}
+		CommandEncoder.SetRenderPassStoreActions(ColorStore, DepthStore, StencilStore);
+	}
+}
+
+void FAGXStateCache::FlushVisibilityResults(FAGXCommandEncoder& CommandEncoder)
+{
+#if PLATFORM_MAC
+	if(VisibilityResults && VisibilityResults->Buffer && VisibilityResults->Buffer.GetStorageMode() == mtlpp::StorageMode::Managed && VisibilityWritten && CommandEncoder.IsRenderCommandEncoderActive())
+	{
+		TRefCountPtr<FAGXFence> Fence = CommandEncoder.EndEncoding();
+		
+        CommandEncoder.BeginBlitCommandEncoding();
+		CommandEncoder.WaitForFence(Fence);
+		
+		mtlpp::BlitCommandEncoder& Encoder = CommandEncoder.GetBlitCommandEncoder();
+
+		// METAL_GPUPROFILE(FAGXProfiler::GetProfiler()->EncodeBlit(CommandEncoder.GetCommandBufferStats(), __FUNCTION__));
+		MTLPP_VALIDATE(mtlpp::BlitCommandEncoder, Encoder, AGXSafeGetRuntimeDebuggingLevel() >= EAGXDebugLevelValidation, Synchronize(VisibilityResults->Buffer));
+		METAL_DEBUG_LAYER(EAGXDebugLevelFastValidation, CommandEncoder.GetBlitCommandEncoderDebugging().Synchronize(VisibilityResults->Buffer));
+		
+		VisibilityWritten = 0;
+	}
+#endif
+}
+
+void FAGXStateCache::SetRenderState(FAGXCommandEncoder& CommandEncoder, FAGXCommandEncoder* PrologueEncoder)
+{
+	SCOPE_CYCLE_COUNTER(STAT_AGXSetRenderStateTime);
+	
+	if (RasterBits)
+	{
+		if (RasterBits & EAGXRenderFlagViewport)
+		{
+			CommandEncoder.SetViewport(Viewport, ActiveViewports);
+		}
+		if (RasterBits & EAGXRenderFlagFrontFacingWinding)
+		{
+			CommandEncoder.SetFrontFacingWinding(mtlpp::Winding::CounterClockwise);
+		}
+		if (RasterBits & EAGXRenderFlagCullMode)
+		{
+			check(IsValidRef(RasterizerState));
+			CommandEncoder.SetCullMode(TranslateCullMode(RasterizerState->State.CullMode));
+		}
+		if (RasterBits & EAGXRenderFlagDepthBias)
+		{
+			check(IsValidRef(RasterizerState));
+			CommandEncoder.SetDepthBias(RasterizerState->State.DepthBias, RasterizerState->State.SlopeScaleDepthBias, FLT_MAX);
+		}
+		if ((RasterBits & EAGXRenderFlagScissorRect))
+		{
+			CommandEncoder.SetScissorRect(Scissor, ActiveScissors);
+		}
+		if (RasterBits & EAGXRenderFlagTriangleFillMode)
+		{
+			check(IsValidRef(RasterizerState));
+			CommandEncoder.SetTriangleFillMode(TranslateFillMode(RasterizerState->State.FillMode));
+		}
+		if (RasterBits & EAGXRenderFlagBlendColor)
+		{
+			CommandEncoder.SetBlendColor(BlendFactor.R, BlendFactor.G, BlendFactor.B, BlendFactor.A);
+		}
+		if (RasterBits & EAGXRenderFlagDepthStencilState)
+		{
+			check(IsValidRef(DepthStencilState));
+            
+            if (DepthStencilState && RenderPassDesc && AGXSafeGetRuntimeDebuggingLevel() >= EAGXDebugLevelFastValidation)
+            {
+                METAL_FATAL_ASSERT(DepthStencilState->bIsDepthWriteEnabled == false || (RenderPassDesc.GetDepthAttachment() && RenderPassDesc.GetDepthAttachment().GetTexture()) , TEXT("Attempting to set a depth-stencil state that writes depth but no depth texture is configured!\nState: %s\nRender Pass: %s"), *FString([DepthStencilState->State.GetPtr() description]), *FString([RenderPassDesc.GetPtr() description]));
+                METAL_FATAL_ASSERT(DepthStencilState->bIsStencilWriteEnabled == false || (RenderPassDesc.GetStencilAttachment() && RenderPassDesc.GetStencilAttachment().GetTexture()), TEXT("Attempting to set a depth-stencil state that writes stencil but no stencil texture is configured!\nState: %s\nRender Pass: %s"), *FString([DepthStencilState->State.GetPtr() description]), *FString([RenderPassDesc.GetPtr() description]));
+            }
+            
+			CommandEncoder.SetDepthStencilState(DepthStencilState ? DepthStencilState->State : nil);
+		}
+		if (RasterBits & EAGXRenderFlagStencilReferenceValue)
+		{
+			CommandEncoder.SetStencilReferenceValue(StencilRef);
+		}
+		if (RasterBits & EAGXRenderFlagVisibilityResultMode)
+		{
+			CommandEncoder.SetVisibilityResultMode(VisibilityMode, VisibilityOffset);
+			if (VisibilityMode != mtlpp::VisibilityResultMode::Disabled)
+			{
+            	VisibilityWritten = VisibilityOffset + FAGXQueryBufferPool::EQueryResultMaxSize;
+			}
+        }
+		RasterBits = 0;
+	}
+}
+
+void FAGXStateCache::EnsureTextureAndType(EAGXShaderStages Stage, uint32 Index, const TMap<uint8, uint8>& TexTypes) const
+{
+#if UE_BUILD_DEBUG || UE_BUILD_DEVELOPMENT
+	if (ShaderTextures[Stage].Textures[Index])
+	{
+		if (ShaderTextures[Stage].Textures[Index].GetTextureType() != (mtlpp::TextureType)TexTypes.FindRef(Index))
+		{
+			ensureMsgf(0, TEXT("Mismatched texture type: EAGXShaderStages %d, Index %d, ShaderTextureType %d != TexTypes %d"), (uint32)Stage, Index, (uint32)ShaderTextures[Stage].Textures[Index].GetTextureType(), (uint32)TexTypes.FindRef(Index));
+		}
+	}
+	else
+	{
+		ensureMsgf(0, TEXT("NULL texture: EAGXShaderStages %d, Index %d"), (uint32)Stage, Index);
+	}
+#endif
+}
+
+void FAGXStateCache::SetRenderPipelineState(FAGXCommandEncoder& CommandEncoder, FAGXCommandEncoder* PrologueEncoder)
+{
+	SCOPE_CYCLE_COUNTER(STAT_AGXSetRenderPipelineStateTime);
+	
+    if ((PipelineBits & EAGXPipelineFlagRasterMask) != 0)
+    {
+    	// Some Intel drivers need RenderPipeline state to be set after DepthStencil state to work properly
+    	FAGXShaderPipeline* Pipeline = GetPipelineState();
+
+		check(Pipeline);
+        CommandEncoder.SetRenderPipelineState(Pipeline);
+        if (Pipeline->ComputePipelineState)
+        {
+            check(PrologueEncoder);
+            PrologueEncoder->SetComputePipelineState(Pipeline);
+        }
+        
+        PipelineBits &= EAGXPipelineFlagComputeMask;
+    }
+	
+#if METAL_DEBUG_OPTIONS
+	if (AGXSafeGetRuntimeDebuggingLevel() >= EAGXDebugLevelFastValidation)
+	{
+		FAGXShaderPipeline* Pipeline = GetPipelineState();
+		EAGXShaderStages VertexStage = EAGXShaderStages::Vertex;
+		
+		FAGXDebugShaderResourceMask VertexMask = Pipeline->ResourceMask[EAGXShaderVertex];
+		TArray<uint32>& MinVertexBufferSizes = Pipeline->BufferDataSizes[EAGXShaderVertex];
+		const TMap<uint8, uint8>& VertexTexTypes = Pipeline->TextureTypes[EAGXShaderVertex];
+		while(VertexMask.BufferMask)
+		{
+			uint32 Index = __builtin_ctz(VertexMask.BufferMask);
+			VertexMask.BufferMask &= ~(1 << Index);
+			
+			if (VertexStage == EAGXShaderStages::Vertex)
+			{
+				FAGXBufferBinding const& Binding = ShaderBuffers[VertexStage].Buffers[Index];
+				ensure(Binding.Buffer || Binding.Bytes);
+				ensure(MinVertexBufferSizes.Num() > Index);
+				ensure(Binding.Length >= MinVertexBufferSizes[Index]);
+			}
+		}
+#if PLATFORM_MAC
+		{
+			uint64 LoTextures = (uint64)VertexMask.TextureMask;
+			while(LoTextures)
+			{
+				uint32 Index = __builtin_ctzll(LoTextures);
+				LoTextures &= ~(uint64(1) << uint64(Index));
+				EnsureTextureAndType(VertexStage, Index, VertexTexTypes);
+			}
+			
+			uint64 HiTextures = (uint64)(VertexMask.TextureMask >> FAGXTextureMask(64));
+			while(HiTextures)
+			{
+				uint32 Index = __builtin_ctzll(HiTextures);
+				HiTextures &= ~(uint64(1) << uint64(Index));
+				EnsureTextureAndType(VertexStage, Index + 64, VertexTexTypes);
+			}
+		}
+#else
+		while(VertexMask.TextureMask)
+		{
+			uint32 Index = __builtin_ctz(VertexMask.TextureMask);
+			VertexMask.TextureMask &= ~(1 << Index);
+			
+			EnsureTextureAndType(VertexStage, Index, VertexTexTypes);
+		}
+#endif
+		while(VertexMask.SamplerMask)
+		{
+			uint32 Index = __builtin_ctz(VertexMask.SamplerMask);
+			VertexMask.SamplerMask &= ~(1 << Index);
+			ensure(ShaderSamplers[VertexStage].Samplers[Index]);
+		}
+		
+		FAGXDebugShaderResourceMask FragmentMask = Pipeline->ResourceMask[EAGXShaderFragment];
+		TArray<uint32>& MinFragmentBufferSizes = Pipeline->BufferDataSizes[EAGXShaderFragment];
+		const TMap<uint8, uint8>& FragmentTexTypes = Pipeline->TextureTypes[EAGXShaderFragment];
+		while(FragmentMask.BufferMask)
+		{
+			uint32 Index = __builtin_ctz(FragmentMask.BufferMask);
+			FragmentMask.BufferMask &= ~(1 << Index);
+			
+			FAGXBufferBinding const& Binding = ShaderBuffers[EAGXShaderStages::Pixel].Buffers[Index];
+			ensure(Binding.Buffer || Binding.Bytes);
+			ensure(MinFragmentBufferSizes.Num() > Index);
+			ensure(Binding.Length >= MinFragmentBufferSizes[Index]);
+		}
+#if PLATFORM_MAC
+		{
+			uint64 LoTextures = (uint64)FragmentMask.TextureMask;
+			while(LoTextures)
+			{
+				uint32 Index = __builtin_ctzll(LoTextures);
+				LoTextures &= ~(uint64(1) << uint64(Index));
+				EnsureTextureAndType(EAGXShaderStages::Pixel, Index, FragmentTexTypes);
+			}
+			
+			uint64 HiTextures = (uint64)(FragmentMask.TextureMask >> FAGXTextureMask(64));
+			while(HiTextures)
+			{
+				uint32 Index = __builtin_ctzll(HiTextures);
+				HiTextures &= ~(uint64(1) << uint64(Index));
+				EnsureTextureAndType(EAGXShaderStages::Pixel, Index + 64, FragmentTexTypes);
+			}
+		}
+#else
+		while(FragmentMask.TextureMask)
+		{
+			uint32 Index = __builtin_ctz(FragmentMask.TextureMask);
+			FragmentMask.TextureMask &= ~(1 << Index);
+			
+			EnsureTextureAndType(EAGXShaderStages::Pixel, Index, FragmentTexTypes);
+		}
+#endif
+		while(FragmentMask.SamplerMask)
+		{
+			uint32 Index = __builtin_ctz(FragmentMask.SamplerMask);
+			FragmentMask.SamplerMask &= ~(1 << Index);
+			ensure(ShaderSamplers[EAGXShaderStages::Pixel].Samplers[Index]);
+		}
+	}
+#endif // METAL_DEBUG_OPTIONS
+}
+
+void FAGXStateCache::SetComputePipelineState(FAGXCommandEncoder& CommandEncoder)
+{
+	if ((PipelineBits & EAGXPipelineFlagComputeMask) != 0)
+	{
+		FAGXShaderPipeline* Pipeline = ComputeShader->GetPipeline();
+	    check(Pipeline);
+	    CommandEncoder.SetComputePipelineState(Pipeline);
+        
+        PipelineBits &= EAGXPipelineFlagRasterMask;
+    }
+	
+	if (AGXSafeGetRuntimeDebuggingLevel() >= EAGXDebugLevelFastValidation)
+	{
+		FAGXShaderPipeline* Pipeline = ComputeShader->GetPipeline();
+		check(Pipeline);
+		
+		FAGXDebugShaderResourceMask ComputeMask = Pipeline->ResourceMask[EAGXShaderCompute];
+		TArray<uint32>& MinComputeBufferSizes = Pipeline->BufferDataSizes[EAGXShaderCompute];
+		const TMap<uint8, uint8>& ComputeTexTypes = Pipeline->TextureTypes[EAGXShaderCompute];
+		while(ComputeMask.BufferMask)
+		{
+			uint32 Index = __builtin_ctz(ComputeMask.BufferMask);
+			ComputeMask.BufferMask &= ~(1 << Index);
+			
+			FAGXBufferBinding const& Binding = ShaderBuffers[EAGXShaderStages::Compute].Buffers[Index];
+			ensure(Binding.Buffer || Binding.Bytes);
+			ensure(MinComputeBufferSizes.Num() > Index);
+			ensure(Binding.Length >= MinComputeBufferSizes[Index]);
+		}
+#if PLATFORM_MAC
+		{
+			uint64 LoTextures = (uint64)ComputeMask.TextureMask;
+			while(LoTextures)
+			{
+				uint32 Index = __builtin_ctzll(LoTextures);
+				LoTextures &= ~(uint64(1) << uint64(Index));
+				EnsureTextureAndType(EAGXShaderStages::Compute, Index, ComputeTexTypes);
+			}
+			
+			uint64 HiTextures = (uint64)(ComputeMask.TextureMask >> FAGXTextureMask(64));
+			while(HiTextures)
+			{
+				uint32 Index = __builtin_ctzll(HiTextures);
+				HiTextures &= ~(uint64(1) << uint64(Index));
+				EnsureTextureAndType(EAGXShaderStages::Compute, Index + 64, ComputeTexTypes);
+			}
+		}
+#else
+		while(ComputeMask.TextureMask)
+		{
+			uint32 Index = __builtin_ctz(ComputeMask.TextureMask);
+			ComputeMask.TextureMask &= ~(1 << Index);
+			
+			EnsureTextureAndType(EAGXShaderStages::Compute, Index, ComputeTexTypes);
+		}
+#endif
+		while(ComputeMask.SamplerMask)
+		{
+			uint32 Index = __builtin_ctz(ComputeMask.SamplerMask);
+			ComputeMask.SamplerMask &= ~(1 << Index);
+			ensure(ShaderSamplers[EAGXShaderStages::Compute].Samplers[Index]);
+		}
+	}
+}
+
+void FAGXStateCache::CommitResourceTable(EAGXShaderStages const Frequency, mtlpp::FunctionType const Type, FAGXCommandEncoder& CommandEncoder)
+{
+	FAGXBufferBindings& BufferBindings = ShaderBuffers[Frequency];
+	while(BufferBindings.Bound)
+	{
+		uint32 Index = __builtin_ctz(BufferBindings.Bound);
+		BufferBindings.Bound &= ~(1 << Index);
+		
+		if (Index < ML_MaxBuffers)
+		{
+			FAGXBufferBinding& Binding = BufferBindings.Buffers[Index];
+			if (Binding.Buffer)
+			{
+				CommandEncoder.SetShaderBuffer(Type, Binding.Buffer, Binding.Offset, Binding.Length, Index, Binding.Usage, BufferBindings.Formats[Index]);
+				
+				if (Binding.Buffer.IsSingleUse())
+				{
+					Binding.Buffer = nil;
+				}
+			}
+			else if (Binding.Bytes)
+			{
+				CommandEncoder.SetShaderData(Type, Binding.Bytes, Binding.Offset, Index, BufferBindings.Formats[Index]);
+			}
+		}
+	}
+	
+	FAGXTextureBindings& TextureBindings = ShaderTextures[Frequency];
+#if PLATFORM_MAC
+	uint64 LoTextures = (uint64)TextureBindings.Bound;
+	while(LoTextures)
+	{
+		uint32 Index = __builtin_ctzll(LoTextures);
+		LoTextures &= ~(uint64(1) << uint64(Index));
+		
+		if (Index < ML_MaxTextures && TextureBindings.Textures[Index])
+		{
+			CommandEncoder.SetShaderTexture(Type, TextureBindings.Textures[Index], Index, TextureBindings.Usage[Index]);
+		}
+	}
+	
+	uint64 HiTextures = (uint64)(TextureBindings.Bound >> FAGXTextureMask(64));
+	while(HiTextures)
+	{
+		uint32 Index = __builtin_ctzll(HiTextures);
+		HiTextures &= ~(uint64(1) << uint64(Index));
+		Index += 64;
+		
+		if (Index < ML_MaxTextures && TextureBindings.Textures[Index])
+		{
+			CommandEncoder.SetShaderTexture(Type, TextureBindings.Textures[Index], Index, TextureBindings.Usage[Index]);
+		}
+	}
+	
+	TextureBindings.Bound = FAGXTextureMask(LoTextures) | (FAGXTextureMask(HiTextures) << FAGXTextureMask(64));
+	check(TextureBindings.Bound == 0);
+#else
+	while(TextureBindings.Bound)
+	{
+		uint32 Index = __builtin_ctz(TextureBindings.Bound);
+		TextureBindings.Bound &= ~(FAGXTextureMask(FAGXTextureMask(1) << FAGXTextureMask(Index)));
+		
+		if (Index < ML_MaxTextures && TextureBindings.Textures[Index])
+		{
+			CommandEncoder.SetShaderTexture(Type, TextureBindings.Textures[Index], Index, TextureBindings.Usage[Index]);
+		}
+	}
+#endif
+	
+    FAGXSamplerBindings& SamplerBindings = ShaderSamplers[Frequency];
+	while(SamplerBindings.Bound)
+	{
+		uint32 Index = __builtin_ctz(SamplerBindings.Bound);
+		SamplerBindings.Bound &= ~(1 << Index);
+		
+		if (Index < ML_MaxSamplers && SamplerBindings.Samplers[Index])
+		{
+			CommandEncoder.SetShaderSamplerState(Type, SamplerBindings.Samplers[Index], Index);
+		}
+	}
+}
+
+FAGXBuffer& FAGXStateCache::GetDebugBuffer()
+{
+    if (!DebugBuffer)
+    {
+        // Assume worst case tiling (16x16) and render-target size (4096x4096) on iOS for now
+        uint32 NumTiles = PLATFORM_MAC ? 1 : 65536;
+        DebugBuffer = GetAGXDeviceContext().CreatePooledBuffer(FAGXPooledBufferArgs(GetAGXDeviceContext().GetDevice(), NumTiles * sizeof(FAGXDebugInfo), BUF_Dynamic, mtlpp::StorageMode::Shared));
+    }
+    return DebugBuffer;
+}
+
+FTexture2DRHIRef FAGXStateCache::CreateFallbackDepthStencilSurface(uint32 Width, uint32 Height)
+{
+#if PLATFORM_MAC
+	if (!IsValidRef(FallbackDepthStencilSurface) || FallbackDepthStencilSurface->GetSizeX() < Width || FallbackDepthStencilSurface->GetSizeY() < Height)
+#else
+	if (!IsValidRef(FallbackDepthStencilSurface) || FallbackDepthStencilSurface->GetSizeX() != Width || FallbackDepthStencilSurface->GetSizeY() != Height)
+#endif
+	{
+		FRHIResourceCreateInfo TexInfo(TEXT("FallbackDepthStencilSurface"));
+		FallbackDepthStencilSurface = RHICreateTexture2D(Width, Height, PF_DepthStencil, 1, 1, TexCreate_DepthStencilTargetable, TexInfo);
+	}
+	check(IsValidRef(FallbackDepthStencilSurface));
+	return FallbackDepthStencilSurface;
+}
+
+void FAGXStateCache::DiscardRenderTargets(bool Depth, bool Stencil, uint32 ColorBitMask)
+{
+	if (Depth)
+	{
+		switch (DepthStore)
+		{
+			case mtlpp::StoreAction::Unknown:
+			case mtlpp::StoreAction::Store:
+				DepthStore = mtlpp::StoreAction::DontCare;
+				break;
+			case mtlpp::StoreAction::StoreAndMultisampleResolve:
+				DepthStore = mtlpp::StoreAction::MultisampleResolve;
+				break;
+			default:
+				break;
+		}
+	}
+
+	if (Stencil)
+	{
+		StencilStore = mtlpp::StoreAction::DontCare;
+	}
+
+	for (uint32 Index = 0; Index < MaxSimultaneousRenderTargets; ++Index)
+	{
+		if ((ColorBitMask & (1u << Index)) != 0)
+		{
+			switch (ColorStore[Index])
+			{
+				case mtlpp::StoreAction::Unknown:
+				case mtlpp::StoreAction::Store:
+					ColorStore[Index] = mtlpp::StoreAction::DontCare;
+					break;
+				case mtlpp::StoreAction::StoreAndMultisampleResolve:
+					ColorStore[Index] = mtlpp::StoreAction::MultisampleResolve;
+					break;
+				default:
+					break;
+			}
+		}
+	}
+}
diff --git a/Engine/Source/Runtime/Apple/AGXRHI/Private/AGXStateCache.h b/Engine/Source/Runtime/Apple/AGXRHI/Private/AGXStateCache.h
new file mode 100755
--- /dev/null
+++ b/Engine/Source/Runtime/Apple/AGXRHI/Private/AGXStateCache.h
@@ -0,0 +1,299 @@
+// Copyright Epic Games, Inc. All Rights Reserved.
+
+#pragma once
+
+#include "AGXRHIPrivate.h"
+#include "AGXUniformBuffer.h"
+#include "Shaders/AGXShaderParameterCache.h"
+#include "AGXCommandEncoder.h"
+#include "AGXPipeline.h"
+
+class FAGXGraphicsPipelineState;
+class FAGXQueryBuffer;
+
+enum EAGXPipelineFlags
+{
+	EAGXPipelineFlagPipelineState = 1 << 0,
+    EAGXPipelineFlagComputeShader = 1 << 5,
+    EAGXPipelineFlagRasterMask = 0xF,
+    EAGXPipelineFlagComputeMask = 0x30,
+    EAGXPipelineFlagMask = 0x3F
+};
+
+enum EAGXRenderFlags
+{
+    EAGXRenderFlagViewport = 1 << 0,
+    EAGXRenderFlagFrontFacingWinding = 1 << 1,
+    EAGXRenderFlagCullMode = 1 << 2,
+    EAGXRenderFlagDepthBias = 1 << 3,
+    EAGXRenderFlagScissorRect = 1 << 4,
+    EAGXRenderFlagTriangleFillMode = 1 << 5,
+    EAGXRenderFlagBlendColor = 1 << 6,
+    EAGXRenderFlagDepthStencilState = 1 << 7,
+    EAGXRenderFlagStencilReferenceValue = 1 << 8,
+    EAGXRenderFlagVisibilityResultMode = 1 << 9,
+    EAGXRenderFlagMask = 0x1FF
+};
+
+class FAGXStateCache
+{
+public:
+	FAGXStateCache(bool const bInImmediate);
+	~FAGXStateCache();
+	
+	/** Reset cached state for reuse */
+	void Reset(void);
+
+	void SetScissorRect(bool const bEnable, mtlpp::ScissorRect const& Rect);
+	void SetBlendFactor(FLinearColor const& InBlendFactor);
+	void SetStencilRef(uint32 const InStencilRef);
+	void SetComputeShader(FAGXComputeShader* InComputeShader);
+	bool SetRenderPassInfo(FRHIRenderPassInfo const& InRenderTargets, FAGXQueryBuffer* QueryBuffer, bool const bRestart);
+	void InvalidateRenderTargets(void);
+	void SetRenderTargetsActive(bool const bActive);
+	void SetViewport(const mtlpp::Viewport& InViewport);
+	void SetViewports(const mtlpp::Viewport InViewport[], uint32 Count);
+	void SetVertexStream(uint32 const Index, FAGXBuffer* Buffer, FAGXBufferData* Bytes, uint32 const Offset, uint32 const Length);
+	void SetGraphicsPipelineState(FAGXGraphicsPipelineState* State);
+	void BindUniformBuffer(EAGXShaderStages const Freq, uint32 const BufferIndex, FRHIUniformBuffer* BufferRHI);
+	void SetDirtyUniformBuffers(EAGXShaderStages const Freq, uint32 const Dirty);
+	
+	/*
+	 * Monitor if samples pass the depth and stencil tests.
+	 * @param Mode Controls if the counter is disabled or moniters passing samples.
+	 * @param Offset The offset relative to the occlusion query buffer provided when the command encoder was created.  offset must be a multiple of 8.
+	 */
+	void SetVisibilityResultMode(mtlpp::VisibilityResultMode const Mode, NSUInteger const Offset);
+	
+#pragma mark - Public Shader Resource Mutators -
+	/*
+	 * Set a global buffer for the specified shader frequency at the given bind point index.
+	 * @param Frequency The shader frequency to modify.
+	 * @param Buffer The buffer to bind or nil to clear.
+	 * @param Bytes The FAGXBufferData to bind or nil to clear.
+	 * @param Offset The offset in the buffer or 0 when Buffer is nil.
+	 * @param Offset The length of data (caller accounts for Offset) in the buffer or 0 when Buffer is nil.
+	 * @param Index The index to modify.
+	 * @param Usage The resource usage flags.
+	 * @param Format The UAV pixel format.
+	 */
+	void SetShaderBuffer(EAGXShaderStages const Frequency, FAGXBuffer const& Buffer, FAGXBufferData* const Bytes, NSUInteger const Offset, NSUInteger const Length, NSUInteger const Index, mtlpp::ResourceUsage const Usage, EPixelFormat const Format = PF_Unknown);
+	
+	/*
+	 * Set a global texture for the specified shader frequency at the given bind point index.
+	 * @param Frequency The shader frequency to modify.
+	 * @param Texture The texture to bind or nil to clear.
+	 * @param Index The index to modify.
+	 * @param Usage The resource usage flags.
+	 */
+	void SetShaderTexture(EAGXShaderStages const Frequency, FAGXTexture const& Texture, NSUInteger const Index, mtlpp::ResourceUsage const Usage);
+	
+	/*
+	 * Set a global sampler for the specified shader frequency at the given bind point index.
+	 * @param Frequency The shader frequency to modify.
+	 * @param Sampler The sampler state to bind or nil to clear.
+	 * @param Index The index to modify.
+	 */
+	void SetShaderSamplerState(EAGXShaderStages const Frequency, FAGXSamplerState* const Sampler, NSUInteger const Index);
+
+	void SetShaderResourceView(FAGXContext* Context, EAGXShaderStages ShaderStage, uint32 BindIndex, FAGXShaderResourceView* RESTRICT SRV);
+	
+	void SetShaderUnorderedAccessView(EAGXShaderStages ShaderStage, uint32 BindIndex, FAGXUnorderedAccessView* RESTRICT UAV);
+
+	void SetStateDirty(void);
+	
+	void SetShaderBufferDirty(EAGXShaderStages const Frequency, NSUInteger const Index);
+	
+	void SetRenderStoreActions(FAGXCommandEncoder& CommandEncoder, bool const bConditionalSwitch);
+	
+	void SetRenderState(FAGXCommandEncoder& CommandEncoder, FAGXCommandEncoder* PrologueEncoder);
+
+	void CommitRenderResources(FAGXCommandEncoder* Raster);
+
+	void CommitComputeResources(FAGXCommandEncoder* Compute);
+	
+	void CommitResourceTable(EAGXShaderStages const Frequency, mtlpp::FunctionType const Type, FAGXCommandEncoder& CommandEncoder);
+	
+	bool PrepareToRestart(bool const bCurrentApplied);
+	
+	FAGXBuffer& GetDebugBuffer();
+	FAGXShaderParameterCache& GetShaderParameters(EAGXShaderStages const Stage) { return ShaderParameters[Stage]; }
+	FLinearColor const& GetBlendFactor() const { return BlendFactor; }
+	uint32 GetStencilRef() const { return StencilRef; }
+	FAGXDepthStencilState* GetDepthStencilState() const { return DepthStencilState; }
+	FAGXRasterizerState* GetRasterizerState() const { return RasterizerState; }
+	FAGXGraphicsPipelineState* GetGraphicsPSO() const { return GraphicsPSO; }
+	FAGXComputeShader* GetComputeShader() const { return ComputeShader; }
+	CGSize GetFrameBufferSize() const { return FrameBufferSize; }
+	FRHIRenderPassInfo const& GetRenderPassInfo() const { return RenderPassInfo; }
+	int32 GetNumRenderTargets() { return bHasValidColorTarget ? RenderPassInfo.GetNumColorRenderTargets() : -1; }
+	bool GetHasValidRenderTarget() const { return bHasValidRenderTarget; }
+	bool GetHasValidColorTarget() const { return bHasValidColorTarget; }
+	const mtlpp::Viewport& GetViewport(uint32 const Index) const { check(Index < ML_MaxViewports); return Viewport[Index]; }
+	uint32 GetVertexBufferSize(uint32 const Index);
+	uint32 GetRenderTargetArraySize() const { return RenderTargetArraySize; }
+	const FRHIUniformBuffer** GetBoundUniformBuffers(EAGXShaderStages const Freq) { return (const FRHIUniformBuffer**)&BoundUniformBuffers[Freq][0]; }
+	uint32 GetDirtyUniformBuffers(EAGXShaderStages const Freq) const { return DirtyUniformBuffers[Freq]; }
+	FAGXQueryBuffer* GetVisibilityResultsBuffer() const { return VisibilityResults; }
+	bool GetScissorRectEnabled() const { return bScissorRectEnabled; }
+	bool NeedsToSetRenderTarget(const FRHIRenderPassInfo& RenderPassInfo);
+	bool HasValidDepthStencilSurface() const { return IsValidRef(DepthStencilSurface); }
+    bool CanRestartRenderPass() const { return bCanRestartRenderPass; }
+	mtlpp::RenderPassDescriptor GetRenderPassDescriptor(void) const { return RenderPassDesc; }
+	uint32 GetSampleCount(void) const { return SampleCount; }
+    bool IsLinearBuffer(EAGXShaderStages ShaderStage, uint32 BindIndex);
+	FAGXShaderPipeline* GetPipelineState() const;
+	EPrimitiveType GetPrimitiveType();
+	mtlpp::VisibilityResultMode GetVisibilityResultMode() { return VisibilityMode; }
+	uint32 GetVisibilityResultOffset() { return VisibilityOffset; }
+	
+	FTexture2DRHIRef CreateFallbackDepthStencilSurface(uint32 Width, uint32 Height);
+	bool GetFallbackDepthStencilBound(void) const { return bFallbackDepthStencilBound; }
+	
+	void SetRenderPipelineState(FAGXCommandEncoder& CommandEncoder, FAGXCommandEncoder* PrologueEncoder);
+    void SetComputePipelineState(FAGXCommandEncoder& CommandEncoder);
+	void FlushVisibilityResults(FAGXCommandEncoder& CommandEncoder);
+
+	void DiscardRenderTargets(bool Depth, bool Stencil, uint32 ColorBitMask);
+private:
+	void ConditionalUpdateBackBuffer(FAGXSurface& Surface);
+	
+	void SetDepthStencilState(FAGXDepthStencilState* InDepthStencilState);
+	void SetRasterizerState(FAGXRasterizerState* InRasterizerState);
+
+	FORCEINLINE void SetResource(uint32 ShaderStage, uint32 BindIndex, FRHITexture* RESTRICT TextureRHI, float CurrentTime);
+	
+	FORCEINLINE void SetResource(uint32 ShaderStage, uint32 BindIndex, FAGXShaderResourceView* RESTRICT SRV, float CurrentTime);
+	
+	FORCEINLINE void SetResource(uint32 ShaderStage, uint32 BindIndex, FAGXSamplerState* RESTRICT SamplerState, float CurrentTime);
+	
+	FORCEINLINE void SetResource(uint32 ShaderStage, uint32 BindIndex, FAGXUnorderedAccessView* RESTRICT UAV, float CurrentTime);
+	
+	template <typename MetalResourceType>
+	inline int32 SetShaderResourcesFromBuffer(uint32 ShaderStage, FAGXUniformBuffer* RESTRICT Buffer, const uint32* RESTRICT ResourceMap, int32 BufferIndex, float CurrentTime);
+	
+	template <class ShaderType>
+	void SetResourcesFromTables(ShaderType Shader, uint32 ShaderStage);
+	
+	void SetViewport(uint32 Index, const mtlpp::Viewport& InViewport);
+	void SetScissorRect(uint32 Index, bool const bEnable, mtlpp::ScissorRect const& Rect);
+
+private:
+
+	void EnsureTextureAndType(EAGXShaderStages Stage, uint32 Index, const TMap<uint8, uint8>& TexTypes) const;
+	
+private:
+#pragma mark - Private Type Declarations -
+	struct FAGXBufferBinding
+	{
+		FAGXBufferBinding() : Bytes(nil), Offset(0), Length(0), Usage((mtlpp::ResourceUsage)0) {}
+		/** The bound buffers or nil. */
+		ns::AutoReleased<FAGXBuffer> Buffer;
+		/** Optional bytes buffer used instead of an FAGXBuffer */
+		FAGXBufferData* Bytes;
+		/** The bound buffer offsets or 0. */
+		NSUInteger Offset;
+		/** The bound buffer lengths or 0. */
+		NSUInteger Length;
+		/** The bound buffer usage or 0 */
+		mtlpp::ResourceUsage Usage;
+	};
+	
+	/** A structure of arrays for the current buffer binding settings. */
+	struct FAGXBufferBindings
+	{
+		FAGXBufferBindings() : Bound(0) {}
+		/** The bound buffers/bytes or nil. */
+		FAGXBufferBinding Buffers[ML_MaxBuffers];
+		/** The pixel formats for buffers bound so that we emulate [RW]Buffer<T> type conversion */
+		EPixelFormat Formats[ML_MaxBuffers];
+		/** A bitmask for which buffers were bound by the application where a bit value of 1 is bound and 0 is unbound. */
+		uint32 Bound;
+	};
+	
+	/** A structure of arrays for the current texture binding settings. */
+	struct FAGXTextureBindings
+	{
+		FAGXTextureBindings() : Bound(0) { FMemory::Memzero(Usage); }
+		/** The bound textures or nil. */
+		ns::AutoReleased<FAGXTexture> Textures[ML_MaxTextures];
+		/** The bound texture usage or 0 */
+		mtlpp::ResourceUsage Usage[ML_MaxTextures];
+		/** A bitmask for which textures were bound by the application where a bit value of 1 is bound and 0 is unbound. */
+		FAGXTextureMask Bound;
+	};
+	
+	/** A structure of arrays for the current sampler binding settings. */
+	struct FAGXSamplerBindings
+	{
+		FAGXSamplerBindings() : Bound(0) {}
+		/** The bound sampler states or nil. */
+		ns::AutoReleased<FAGXSampler> Samplers[ML_MaxSamplers];
+		/** A bitmask for which samplers were bound by the application where a bit value of 1 is bound and 0 is unbound. */
+		uint16 Bound;
+	};
+    
+private:
+	FAGXShaderParameterCache ShaderParameters[EAGXShaderStages::Num];
+
+	uint32 SampleCount;
+
+	TSet<TRefCountPtr<FRHIUniformBuffer>> ActiveUniformBuffers;
+	FRHIUniformBuffer* BoundUniformBuffers[EAGXShaderStages::Num][ML_MaxBuffers];
+	
+	/** Bitfield for which uniform buffers are dirty */
+	uint32 DirtyUniformBuffers[EAGXShaderStages::Num];
+	
+	/** Vertex attribute buffers */
+	FAGXBufferBinding VertexBuffers[MaxVertexElementCount];
+	
+	/** Bound shader resource tables. */
+	FAGXBufferBindings ShaderBuffers[EAGXShaderStages::Num];
+	FAGXTextureBindings ShaderTextures[EAGXShaderStages::Num];
+	FAGXSamplerBindings ShaderSamplers[EAGXShaderStages::Num];
+	
+	mtlpp::StoreAction ColorStore[MaxSimultaneousRenderTargets];
+	mtlpp::StoreAction DepthStore;
+	mtlpp::StoreAction StencilStore;
+
+	FAGXQueryBuffer* VisibilityResults;
+	mtlpp::VisibilityResultMode VisibilityMode;
+	NSUInteger VisibilityOffset;
+	NSUInteger VisibilityWritten;
+
+	TRefCountPtr<FAGXDepthStencilState> DepthStencilState;
+	TRefCountPtr<FAGXRasterizerState> RasterizerState;
+	TRefCountPtr<FAGXGraphicsPipelineState> GraphicsPSO;
+	TRefCountPtr<FAGXComputeShader> ComputeShader;
+	uint32 StencilRef;
+	
+	FLinearColor BlendFactor;
+	CGSize FrameBufferSize;
+	
+	uint32 RenderTargetArraySize;
+
+	mtlpp::Viewport Viewport[ML_MaxViewports];
+	mtlpp::ScissorRect Scissor[ML_MaxViewports];
+	
+	uint32 ActiveViewports;
+	uint32 ActiveScissors;
+	
+	FRHIRenderPassInfo RenderPassInfo;
+	FTextureRHIRef ColorTargets[MaxSimultaneousRenderTargets];
+	FTextureRHIRef ResolveTargets[MaxSimultaneousRenderTargets];
+	FTextureRHIRef DepthStencilSurface;
+	FTextureRHIRef DepthStencilResolve;
+	/** A fallback depth-stencil surface for draw calls that write to depth without a depth-stencil surface bound. */
+	FTexture2DRHIRef FallbackDepthStencilSurface;
+	FAGXBuffer DebugBuffer;
+	mtlpp::RenderPassDescriptor RenderPassDesc;
+	uint32 RasterBits;
+    uint8 PipelineBits;
+	bool bIsRenderTargetActive;
+	bool bHasValidRenderTarget;
+	bool bHasValidColorTarget;
+	bool bScissorRectEnabled;
+    bool bCanRestartRenderPass;
+    bool bImmediate;
+	bool bFallbackDepthStencilBound;
+};
