Initial import of AGXRHI

diff --git a/Engine/Source/Runtime/Apple/AGXRHI/Private/AGXTexture.cpp b/Engine/Source/Runtime/Apple/AGXRHI/Private/AGXTexture.cpp
new file mode 100644
--- /dev/null
+++ b/Engine/Source/Runtime/Apple/AGXRHI/Private/AGXTexture.cpp
@@ -0,0 +1,2791 @@
+// Copyright Epic Games, Inc. All Rights Reserved.
+
+/*=============================================================================
+	AGXTexture.cpp: AGX texture RHI implementation.
+ =============================================================================*/
+
+#include "AGXRHIPrivate.h"
+#include "AGXProfiler.h" // for STAT_AGXTexturePageOffTime
+#include "AGXCommandBuffer.h"
+#include "RenderUtils.h"
+#include "Containers/ResourceArray.h"
+#include "Misc/ScopeRWLock.h"
+#include "AGXLLM.h"
+
+volatile int64 FAGXSurface::ActiveUploads = 0;
+
+int32 GAGXMaxOutstandingAsyncTexUploads = 100 * 1024 * 1024;
+FAutoConsoleVariableRef CVarAGXMaxOutstandingAsyncTexUploads(TEXT("rhi.AGX.MaxOutstandingAsyncTexUploads"),
+															 GAGXMaxOutstandingAsyncTexUploads,
+															 TEXT("The maximum number of outstanding asynchronous texture uploads allowed to be pending in Metal. After the limit is reached the next upload will wait for all outstanding operations to complete and purge the waiting free-lists in order to reduce peak memory consumption. Defaults to 0 (infinite), set to a value > 0 limit the number."),
+															 ECVF_ReadOnly | ECVF_RenderThreadSafe);
+
+int32 GAGXForceIOSTexturesShared = 1;
+FAutoConsoleVariableRef CVarAGXForceIOSTexturesShared(TEXT("rhi.AGX.ForceIOSTexturesShared"),
+													  GAGXForceIOSTexturesShared,
+													  TEXT("If true, forces all textures to be Shared on iOS"),
+													  ECVF_RenderThreadSafe);
+
+/** Given a pointer to a RHI texture that was created by the AGX RHI, returns a pointer to the FAGXTextureBase it encapsulates. */
+FAGXSurface* AGXGetMetalSurfaceFromRHITexture(FRHITexture* Texture)
+{
+	if (!Texture)
+	{
+		return NULL;
+	}
+	else
+	{
+		FAGXSurface* Surface = (FAGXSurface*)Texture->GetTextureBaseRHI();
+		return Surface;
+	}
+}
+
+static bool IsRenderTarget(ETextureCreateFlags Flags)
+{
+	return (Flags & (TexCreate_RenderTargetable | TexCreate_ResolveTargetable | TexCreate_DepthStencilTargetable | TexCreate_DepthStencilResolveTarget)) != 0;
+}
+
+static mtlpp::TextureUsage ConvertFlagsToUsage(ETextureCreateFlags Flags)
+{
+	NSUInteger Usage = mtlpp::TextureUsage::Unknown;
+    if(Flags & (TexCreate_ShaderResource|TexCreate_ResolveTargetable|TexCreate_DepthStencilTargetable))
+	{
+		Usage |= mtlpp::TextureUsage::ShaderRead;
+		Usage |= mtlpp::TextureUsage::PixelFormatView;
+	}
+	
+	if (Flags & TexCreate_UAV)
+	{
+		Usage |= mtlpp::TextureUsage::ShaderRead;
+		Usage |= mtlpp::TextureUsage::ShaderWrite;
+		Usage |= mtlpp::TextureUsage::PixelFormatView;
+	}
+	
+	// offline textures are normal shader read textures
+	if (Flags & TexCreate_OfflineProcessed)
+	{
+		Usage |= mtlpp::TextureUsage::ShaderRead;
+	}
+
+	//if the high level is doing manual resolves then the textures specifically markes as resolve targets
+	//are likely to be used in a manual shader resolve by the high level and must be bindable as rendertargets.
+	const bool bSeparateResolveTargets = FAGXCommandQueue::SupportsSeparateMSAAAndResolveTarget();
+	const bool bResolveTarget = (Flags & TexCreate_ResolveTargetable);
+	if ((Flags & (TexCreate_RenderTargetable|TexCreate_DepthStencilTargetable|TexCreate_DepthStencilResolveTarget)) || (bResolveTarget && bSeparateResolveTargets))
+	{
+		Usage |= mtlpp::TextureUsage::RenderTarget;
+		Usage |= mtlpp::TextureUsage::ShaderRead;
+	}
+	return (mtlpp::TextureUsage)Usage;
+}
+
+static bool IsPixelFormatCompressed(EPixelFormat Format)
+{
+	switch (Format)
+	{
+		case PF_DXT1:
+		case PF_DXT3:
+		case PF_DXT5:
+		case PF_PVRTC2:
+		case PF_PVRTC4:
+		case PF_BC4:
+		case PF_BC5:
+		case PF_ETC2_RGB:
+		case PF_ETC2_RGBA:
+		case PF_ASTC_4x4:
+		case PF_ASTC_6x6:
+		case PF_ASTC_8x8:
+		case PF_ASTC_10x10:
+		case PF_ASTC_12x12:
+		case PF_BC6H:
+		case PF_BC7:
+			return true;
+		default:
+			return false;
+	}
+}
+
+static bool IsPixelFormatASTCCompressed(EPixelFormat Format)
+{
+	switch (Format)
+	{
+		case PF_ASTC_4x4:
+		case PF_ASTC_6x6:
+		case PF_ASTC_8x8:
+		case PF_ASTC_10x10:
+		case PF_ASTC_12x12:
+			return true;
+		default:
+			return false;
+	}
+}
+
+static bool IsPixelFormatPVRTCCompressed(EPixelFormat Format)
+{
+	switch (Format)
+	{
+		case PF_PVRTC2:
+		case PF_PVRTC4:
+		case PF_ETC2_RGB:
+		case PF_ETC2_RGBA:
+			return true;
+		default:
+			return false;
+	}
+}
+
+void AGXSafeReleaseMetalTexture(FAGXSurface* Surface, FAGXTexture& Texture, bool bTextureView, bool bAVFoundationTexture)
+{
+	if(GIsAGXInitialized && GDynamicRHI)
+	{
+		if (!bAVFoundationTexture && !bTextureView)
+		{
+			GetAGXDeviceContext().ReleaseTexture(Surface, Texture);
+		}
+		else
+		{
+			AGXSafeReleaseMetalObject([Texture.GetPtr() retain]);
+		}
+	}
+}
+
+void AGXSafeReleaseMetalTexture(FAGXSurface* Surface, FAGXTexture& Texture)
+{
+	if(GIsAGXInitialized && GDynamicRHI)
+	{
+		GetAGXDeviceContext().ReleaseTexture(Surface, Texture);
+	}
+}
+
+#if !PLATFORM_IOS
+mtlpp::PixelFormat ToSRGBFormat_NonAppleMacGPU(mtlpp::PixelFormat MTLFormat)
+{
+	switch (MTLFormat)
+	{
+		case mtlpp::PixelFormat::RGBA8Unorm:
+			MTLFormat = mtlpp::PixelFormat::RGBA8Unorm_sRGB;
+			break;
+		case mtlpp::PixelFormat::BGRA8Unorm:
+			MTLFormat = mtlpp::PixelFormat::BGRA8Unorm_sRGB;
+			break;
+		case mtlpp::PixelFormat::BC1_RGBA:
+			MTLFormat = mtlpp::PixelFormat::BC1_RGBA_sRGB;
+			break;
+		case mtlpp::PixelFormat::BC2_RGBA:
+			MTLFormat = mtlpp::PixelFormat::BC2_RGBA_sRGB;
+			break;
+		case mtlpp::PixelFormat::BC3_RGBA:
+			MTLFormat = mtlpp::PixelFormat::BC3_RGBA_sRGB;
+			break;
+		case mtlpp::PixelFormat::BC7_RGBAUnorm:
+			MTLFormat = mtlpp::PixelFormat::BC7_RGBAUnorm_sRGB;
+			break;
+		default:
+			break;
+	}
+	return MTLFormat;
+}
+#endif
+
+mtlpp::PixelFormat ToSRGBFormat_AppleGPU(mtlpp::PixelFormat MTLFormat)
+{
+	switch (MTLFormat)
+	{
+		case mtlpp::PixelFormat::RGBA8Unorm:
+			MTLFormat = mtlpp::PixelFormat::RGBA8Unorm_sRGB;
+			break;
+		case mtlpp::PixelFormat::BGRA8Unorm:
+			MTLFormat = mtlpp::PixelFormat::BGRA8Unorm_sRGB;
+			break;
+		case mtlpp::PixelFormat::R8Unorm:
+			MTLFormat = mtlpp::PixelFormat::R8Unorm_sRGB;
+			break;
+		case mtlpp::PixelFormat::PVRTC_RGBA_2BPP:
+			MTLFormat = mtlpp::PixelFormat::PVRTC_RGBA_2BPP_sRGB;
+			break;
+		case mtlpp::PixelFormat::PVRTC_RGBA_4BPP:
+			MTLFormat = mtlpp::PixelFormat::PVRTC_RGBA_4BPP_sRGB;
+			break;
+		case mtlpp::PixelFormat::ASTC_4x4_LDR:
+			MTLFormat = mtlpp::PixelFormat::ASTC_4x4_sRGB;
+			break;
+		case mtlpp::PixelFormat::ASTC_6x6_LDR:
+			MTLFormat = mtlpp::PixelFormat::ASTC_6x6_sRGB;
+			break;
+		case mtlpp::PixelFormat::ASTC_8x8_LDR:
+			MTLFormat = mtlpp::PixelFormat::ASTC_8x8_sRGB;
+			break;
+		case mtlpp::PixelFormat::ASTC_10x10_LDR:
+			MTLFormat = mtlpp::PixelFormat::ASTC_10x10_sRGB;
+			break;
+		case mtlpp::PixelFormat::ASTC_12x12_LDR:
+			MTLFormat = mtlpp::PixelFormat::ASTC_12x12_sRGB;
+			break;
+		default:
+			break;
+	}
+	return MTLFormat;
+}
+
+mtlpp::PixelFormat AGXToSRGBFormat(mtlpp::PixelFormat MTLFormat)
+{
+	if([GetAGXDeviceContext().GetDevice().GetPtr() supportsFamily:MTLGPUFamilyApple1])
+	{
+		MTLFormat = ToSRGBFormat_AppleGPU(MTLFormat);
+	}
+#if !PLATFORM_IOS
+	else if([GetAGXDeviceContext().GetDevice().GetPtr() supportsFamily:MTLGPUFamilyMac1])
+	{
+		MTLFormat = ToSRGBFormat_NonAppleMacGPU(MTLFormat);
+	}
+#endif
+	
+	return MTLFormat;
+}
+
+void FAGXSurface::PrepareTextureView()
+{
+	// Recreate the texture to enable MTLTextureUsagePixelFormatView which must be off unless we definitely use this feature or we are throwing ~4% performance vs. Windows on the floor.
+	mtlpp::TextureUsage Usage = (mtlpp::TextureUsage)Texture.GetUsage();
+	bool bMemoryLess = false;
+#if PLATFORM_IOS
+	bMemoryLess = (Texture.GetStorageMode() == mtlpp::StorageMode::Memoryless);
+#endif
+	if(!(Usage & mtlpp::TextureUsage::PixelFormatView) && !bMemoryLess)
+	{
+		check(!bTextureView);
+		check(ImageSurfaceRef == nullptr);
+		
+		check(Texture);
+		const bool bMSAATextureIsTexture = MSAATexture == Texture;
+		const bool bMSAAResolveTextureIsTexture = MSAAResolveTexture == Texture;
+		if (MSAATexture && !bMSAATextureIsTexture)
+		{
+			FAGXTexture OldMSAATexture = MSAATexture;
+			MSAATexture = Reallocate(MSAATexture, mtlpp::TextureUsage::PixelFormatView);
+			AGXSafeReleaseMetalTexture(this, OldMSAATexture, bTextureView, ImageSurfaceRef != nullptr);
+		}
+		if (MSAAResolveTexture && !bMSAAResolveTextureIsTexture)
+		{
+			FAGXTexture OldMSAAResolveTexture = MSAAResolveTexture;
+			MSAAResolveTexture = Reallocate(MSAAResolveTexture, mtlpp::TextureUsage::PixelFormatView);
+			AGXSafeReleaseMetalTexture(this, OldMSAAResolveTexture, bTextureView, ImageSurfaceRef != nullptr);
+		}
+		
+		FAGXTexture OldTexture = Texture;
+		Texture = Reallocate(Texture, mtlpp::TextureUsage::PixelFormatView);
+		AGXSafeReleaseMetalTexture(this, OldTexture, bTextureView, ImageSurfaceRef != nullptr);
+		
+		if (bMSAATextureIsTexture)
+		{
+			MSAATexture = Texture;
+		}
+		if (bMSAAResolveTextureIsTexture)
+		{
+			MSAAResolveTexture = Texture;
+		}
+	}
+}
+
+FAGXTexture FAGXSurface::Reallocate(FAGXTexture InTexture, mtlpp::TextureUsage UsageModifier)
+{
+	mtlpp::TextureDescriptor Desc;
+	Desc.SetTextureType(InTexture.GetTextureType());
+	Desc.SetPixelFormat(InTexture.GetPixelFormat());
+	Desc.SetWidth(InTexture.GetWidth());
+	Desc.SetHeight(InTexture.GetHeight());
+	Desc.SetDepth(InTexture.GetDepth());
+	Desc.SetMipmapLevelCount(InTexture.GetMipmapLevelCount());
+	Desc.SetSampleCount(InTexture.GetSampleCount());
+	Desc.SetArrayLength(InTexture.GetArrayLength());
+	
+	static mtlpp::ResourceOptions GeneralResourceOption = (mtlpp::ResourceOptions)FAGXCommandQueue::GetCompatibleResourceOptions(mtlpp::ResourceOptions::HazardTrackingModeUntracked);
+	
+	Desc.SetResourceOptions(mtlpp::ResourceOptions(((NSUInteger)InTexture.GetCpuCacheMode() << mtlpp::ResourceCpuCacheModeShift) | ((NSUInteger)Texture.GetStorageMode() << mtlpp::ResourceStorageModeShift) | GeneralResourceOption));
+	Desc.SetCpuCacheMode(InTexture.GetCpuCacheMode());
+	Desc.SetStorageMode(InTexture.GetStorageMode());
+	Desc.SetUsage(mtlpp::TextureUsage(InTexture.GetUsage() | UsageModifier));
+	
+	FAGXTexture NewTex = GetAGXDeviceContext().CreateTexture(this, Desc);
+	check(NewTex);
+	return NewTex;
+}
+
+void FAGXSurface::ReplaceTexture(FAGXContext& Context, FAGXTexture CurrentTexture, FAGXTexture NewTexture)
+{
+	check(NewTexture);
+	{
+		// Copy texture content...
+		const uint32 BlockSizeX = GPixelFormats[PixelFormat].BlockSizeX;
+		const uint32 BlockSizeY = GPixelFormats[PixelFormat].BlockSizeY;
+		const uint32 BlockSizeZ = GPixelFormats[PixelFormat].BlockSizeZ;
+		
+		mtlpp::Origin Origin(0,0,0);
+		
+		// DXT/BC formats on Mac actually do have mip-tails that are smaller than the block size, they end up being uncompressed.
+		bool const bPixelFormatASTC = IsPixelFormatASTCCompressed(PixelFormat);
+		
+		for (uint32 SliceIndex = 0; SliceIndex < CurrentTexture.GetArrayLength(); SliceIndex++)
+		{
+			for (uint32 MipIndex = 0; MipIndex < CurrentTexture.GetMipmapLevelCount(); ++MipIndex)
+			{
+				const uint32 UnalignedMipSizeX = FMath::Max<uint32>(1, CurrentTexture.GetWidth() >> (MipIndex));
+				const uint32 UnalignedMipSizeY = FMath::Max<uint32>(1, CurrentTexture.GetHeight() >> (MipIndex));
+				const uint32 UnalignedMipSizeZ = FMath::Max<uint32>(1, CurrentTexture.GetDepth() >> (MipIndex));
+				const uint32 MipSizeX = (bPixelFormatASTC) ? AlignArbitrary(UnalignedMipSizeX, BlockSizeX) : UnalignedMipSizeX;
+				const uint32 MipSizeY = (bPixelFormatASTC) ? AlignArbitrary(UnalignedMipSizeY, BlockSizeY) : UnalignedMipSizeY;
+				const uint32 MipSizeZ = (bPixelFormatASTC) ? AlignArbitrary(UnalignedMipSizeZ, BlockSizeZ) : UnalignedMipSizeZ;
+				
+				Context.CopyFromTextureToTexture(CurrentTexture, SliceIndex, MipIndex, Origin, mtlpp::Size(MipSizeX, MipSizeY, MipSizeZ), NewTexture, SliceIndex, MipIndex, Origin);
+			}
+		}
+		
+	}
+	
+	FAGXTexture OldTexture = Texture;
+	if (Texture && Texture == CurrentTexture)
+	{
+		AGXSafeReleaseMetalTexture(OldTexture);
+		Texture = NewTexture;
+	}
+	if (MSAATexture && MSAATexture == CurrentTexture)
+	{
+		if (MSAATexture && (MSAATexture != OldTexture))
+		{
+			FAGXTexture OldMSAATexture = MSAATexture;
+			AGXSafeReleaseMetalTexture(OldMSAATexture);
+			MSAATexture = NewTexture;
+		}
+		else if(MSAATexture)
+		{
+			MSAATexture = NewTexture;
+		}
+		
+		if (MSAAResolveTexture && MSAAResolveTexture == CurrentTexture)
+		{
+			if (MSAAResolveTexture && (MSAAResolveTexture != OldTexture))
+			{
+				FAGXTexture OldMSAAResolveTexture = MSAAResolveTexture;
+				AGXSafeReleaseMetalTexture(OldMSAAResolveTexture);
+				MSAAResolveTexture = NewTexture;
+			}
+			else if (MSAAResolveTexture)
+			{
+				MSAAResolveTexture = NewTexture;
+			}
+		}
+	}
+	
+	for (FAGXShaderResourceView* SRV : SRVs)
+	{
+		if(OldTexture != SRV->TextureView->Texture)
+		{
+			AGXSafeReleaseMetalTexture(nullptr, SRV->TextureView->Texture, true, false);
+		}
+		SRV->TextureView->Texture = nil;
+		
+		SRV->TextureView->MSAATexture = nil;
+		
+		if (SRV->Format == PF_Unknown)
+		{
+			SRV->TextureView->Init(*this, NSMakeRange(SRV->MipLevel, SRV->NumMips));
+		}
+		else
+		{
+			SRV->TextureView->Init(*this, NSMakeRange(SRV->MipLevel, SRV->NumMips), (EPixelFormat)SRV->Format, SRV->bSRGBForceDisable);
+		}
+	}
+}
+
+void FAGXSurface::MakeAliasable(void)
+{
+	check(!bTextureView);
+	check(ImageSurfaceRef == nullptr);
+	
+	static bool bSupportsHeaps = GetAGXDeviceContext().SupportsFeature(EAGXFeaturesHeaps);
+	if (bSupportsHeaps && Texture.GetStorageMode() == mtlpp::StorageMode::Private && Texture.GetHeap())
+	{
+		if (MSAATexture && (MSAATexture != Texture) && !MSAATexture.IsAliasable())
+		{
+			MSAATexture.MakeAliasable();
+#if STATS || ENABLE_LOW_LEVEL_MEM_TRACKER
+			AGXLLM::LogAliasTexture(MSAATexture);
+#endif
+		}
+		if (!Texture.IsAliasable())
+		{
+			Texture.MakeAliasable();
+#if STATS || ENABLE_LOW_LEVEL_MEM_TRACKER
+			AGXLLM::LogAliasTexture(Texture);
+#endif
+		}
+	}
+}
+
+void FAGXSurface::MakeUnAliasable(void)
+{
+	check(!bTextureView);
+	check(ImageSurfaceRef == nullptr);
+	
+	static bool bSupportsHeaps = GetAGXDeviceContext().SupportsFeature(EAGXFeaturesHeaps);
+	if (bSupportsHeaps && Texture.GetStorageMode() == mtlpp::StorageMode::Private && Texture.GetHeap() && Texture.IsAliasable())
+	{
+		FAGXTexture OldTexture = Texture;
+		Texture = Reallocate(Texture, mtlpp::TextureUsage::Unknown);
+		AGXSafeReleaseMetalTexture(this, OldTexture);
+		if (MSAATexture && (MSAATexture != OldTexture))
+		{
+			FAGXTexture OldMSAATexture = MSAATexture;
+			MSAATexture = Reallocate(MSAATexture, mtlpp::TextureUsage::Unknown);
+			AGXSafeReleaseMetalTexture(this, OldMSAATexture);
+		}
+		else if(MSAATexture)
+		{
+			MSAATexture = Texture;
+		}
+		
+		for (FAGXShaderResourceView* SRV : SRVs)
+		{
+			if(OldTexture != SRV->TextureView->Texture)
+			{
+				AGXSafeReleaseMetalTexture(this, SRV->TextureView->Texture, true, false);
+			}
+			SRV->TextureView->Texture = nil;
+			
+			SRV->TextureView->MSAATexture = nil;
+			
+			if (SRV->Format == PF_Unknown)
+			{
+				SRV->TextureView->Init(*this, NSMakeRange(SRV->MipLevel, SRV->NumMips));
+			}
+			else
+			{
+				SRV->TextureView->Init(*this, NSMakeRange(SRV->MipLevel, SRV->NumMips), (EPixelFormat)SRV->Format, SRV->bSRGBForceDisable);
+			}
+		}
+	}
+}
+
+void FAGXSurface::Init(FAGXSurface& Source, NSRange MipRange)
+{
+#if PLATFORM_IOS
+	// Mmeory;ess targets can't have texture views (SRVs or UAVs)
+	if (Source.Texture.GetStorageMode() == mtlpp::StorageMode::Memoryless)
+	{
+		return;
+	}
+#endif
+	auto TextureViewType = [&Source]()
+	{
+		return Source.bIsCubemap ? mtlpp::TextureType::Texture2DArray : Source.Texture.GetTextureType();
+	};
+
+	mtlpp::PixelFormat MetalFormat = (mtlpp::PixelFormat)GPixelFormats[PixelFormat].PlatformFormat;
+	if (Flags & TexCreate_SRGB)
+	{
+		// Ensure we have the correct sRGB target format if we create a new texture view rather than using the source texture
+		MetalFormat = AGXToSRGBFormat(MetalFormat);
+	}
+	
+	bool const bUseSourceTex = (Source.PixelFormat != PF_DepthStencil) && MipRange.location == 0 && MipRange.length == Source.Texture.GetMipmapLevelCount();
+	
+	// Recreate the texture to enable MTLTextureUsagePixelFormatView which must be off unless we definitely use this feature or we are throwing ~4% performance vs. Windows on the floor.
+	mtlpp::TextureUsage Usage = Source.Texture.GetUsage();
+	if(!(Usage & mtlpp::TextureUsage::PixelFormatView) && (Source.PixelFormat != PF_DepthStencil) && !bUseSourceTex)
+	{
+		Source.PrepareTextureView();
+	}
+	
+	ns::Range Slices(0, Source.Texture.GetArrayLength() * (bIsCubemap ? 6 : 1));
+	// Stencil requires a format conversion, so this will access depth only, consequently there are no mip levels other than 0, so this path requires no real special casing.
+	if(Source.PixelFormat != PF_DepthStencil && !bUseSourceTex)
+	{
+		Texture = Source.Texture.NewTextureView(MetalFormat, TextureViewType(), ns::Range(MipRange.location, MipRange.length), Slices);
+	}
+	else
+	{
+		Texture = Source.Texture;
+	}
+	
+	const uint32 BlockSizeX = GPixelFormats[PixelFormat].BlockSizeX;
+	const uint32 BlockSizeY = GPixelFormats[PixelFormat].BlockSizeY;
+	const uint32 BlockBytes = GPixelFormats[PixelFormat].BlockBytes;
+	SizeX = FMath::Max(SizeX >> MipRange.location, BlockSizeX);
+	SizeY = FMath::Max(SizeY >> MipRange.location, BlockSizeY);
+	SizeZ = (Type != RRT_Texture3D) ? SizeZ : FMath::Max(SizeZ >> MipRange.location, 1u);
+}
+
+void FAGXSurface::Init(FAGXSurface& Source, NSRange MipRange, EPixelFormat Format, bool bSRGBForceDisable)
+{
+	check(!Source.MSAATexture || Format == PF_X24_G8);
+#if PLATFORM_IOS
+	// Mmeory;ess targets can't have texture views (SRVs or UAVs)
+	if (Source.Texture.GetStorageMode() == mtlpp::StorageMode::Memoryless)
+	{
+		return;
+	}
+#endif
+	
+	mtlpp::PixelFormat MetalFormat = (mtlpp::PixelFormat)GPixelFormats[PixelFormat].PlatformFormat;
+	
+	bool bUseSourceTex = (Source.PixelFormat != PF_DepthStencil) && Source.PixelFormat == Format && MipRange.location == 0 && MipRange.length == Source.Texture.GetMipmapLevelCount();
+	
+	if (Flags & TexCreate_SRGB)
+	{
+		if(bSRGBForceDisable)
+		{
+			// Force no sRGB, should create new view without sRGB format
+			bUseSourceTex = false;
+		}
+		else
+		{
+			// Ensure we have the correct sRGB target format if we create a new texture view rather than using the source texture
+			MetalFormat = AGXToSRGBFormat(MetalFormat);
+		}
+	}
+	
+	// Recreate the texture to enable MTLTextureUsagePixelFormatView which must be off unless we definitely use this feature or we are throwing ~4% performance vs. Windows on the floor.
+	mtlpp::TextureUsage Usage = Source.Texture.GetUsage();
+	if(!(Usage & mtlpp::TextureUsage::PixelFormatView) && (Source.PixelFormat == PF_DepthStencil && Format == PF_X24_G8) && !bUseSourceTex)
+	{
+		Source.PrepareTextureView();
+	}
+	
+	ns::Range Slices(0, Source.Texture.GetArrayLength() * (bIsCubemap ? 6 : 1));
+	if(Source.PixelFormat != PF_DepthStencil && !bUseSourceTex)
+	{
+		Texture = Source.Texture.NewTextureView(MetalFormat, Source.Texture.GetTextureType(), ns::Range(MipRange.location, MipRange.length), Slices);
+	}
+	else if (Source.PixelFormat == PF_DepthStencil && Format == PF_X24_G8)
+	{
+		switch (Source.Texture.GetPixelFormat())
+		{
+#if PLATFORM_MAC // Currently only required on Mac as iOS uses separate textures
+			case mtlpp::PixelFormat::Depth24Unorm_Stencil8:
+				MetalFormat = mtlpp::PixelFormat::X24_Stencil8;
+				break;
+#endif
+			case mtlpp::PixelFormat::Depth32Float_Stencil8:
+				MetalFormat = mtlpp::PixelFormat::X32_Stencil8;
+				break;
+			default:
+				check(false);
+				break;
+		}
+		
+		Texture = Source.Texture.NewTextureView(MetalFormat, Source.Texture.GetTextureType(), ns::Range(MipRange.location, MipRange.length), Slices);
+	}
+	else
+	{
+		Texture = Source.Texture;
+	}
+	
+	const uint32 BlockSizeX = GPixelFormats[PixelFormat].BlockSizeX;
+	const uint32 BlockSizeY = GPixelFormats[PixelFormat].BlockSizeY;
+	const uint32 BlockBytes = GPixelFormats[PixelFormat].BlockBytes;
+	SizeX = FMath::Max(SizeX >> MipRange.location, BlockSizeX);
+	SizeY = FMath::Max(SizeY >> MipRange.location, BlockSizeY);
+	SizeZ = (Type != RRT_Texture3D) ? SizeZ : FMath::Max(SizeZ >> MipRange.location, 1u);
+}
+
+FAGXSurface::FAGXSurface(FAGXSurface& Source, NSRange MipRange)
+: Type(Source.Type)
+, PixelFormat(Source.PixelFormat)
+, Texture(nil)
+, MSAATexture(nil)
+, MSAAResolveTexture(nil)
+, SizeX(Source.SizeX)
+, SizeY(Source.SizeY)
+, SizeZ(Source.SizeZ)
+, bIsCubemap(Source.bIsCubemap)
+, Flags(Source.Flags)
+, BufferLocks(0)
+, TotalTextureSize(0)
+, Viewport(nullptr)
+, ImageSurfaceRef(nullptr)
+, bTextureView(true)
+{
+	FPlatformAtomics::InterlockedExchange(&Written, Source.Written);
+	
+	Init(Source, MipRange);
+}
+
+FAGXSurface::FAGXSurface(FAGXSurface& Source, NSRange const MipRange, EPixelFormat Format, bool bSRGBForceDisable)
+: Type(Source.Type)
+, PixelFormat(Format)
+, Texture(nil)
+, MSAATexture(nil)
+, MSAAResolveTexture(nil)
+, SizeX(Source.SizeX)
+, SizeY(Source.SizeY)
+, SizeZ(Source.SizeZ)
+, bIsCubemap(Source.bIsCubemap)
+, Flags(Source.Flags)
+, BufferLocks(0)
+, TotalTextureSize(0)
+, Viewport(nullptr)
+, ImageSurfaceRef(nullptr)
+, bTextureView(true)
+{
+	FPlatformAtomics::InterlockedExchange(&Written, Source.Written);
+	
+	Init(Source, MipRange, Format, bSRGBForceDisable);
+}
+
+static TMap<uint64, uint8>& GetMetalPixelFormatKeyMap()
+{
+	static TMap<uint64, uint8> PixelFormatKeyMap;
+	return PixelFormatKeyMap;
+}
+
+uint8 AGXGetMetalPixelFormatKey(mtlpp::PixelFormat Format)
+{
+	static FRWLock Mutex;
+	static uint8 NextKey = 1; // 0 is reserved for mtlpp::PixelFormat::Invalid
+	FRWScopeLock Lock(Mutex, SLT_ReadOnly);
+	uint8* Key = GetMetalPixelFormatKeyMap().Find((uint64)Format);
+	if (Key == NULL)
+	{
+		Lock.ReleaseReadOnlyLockAndAcquireWriteLock_USE_WITH_CAUTION();
+		Key = GetMetalPixelFormatKeyMap().Find((uint64)Format);
+		if (Key == NULL)
+		{
+			Key = &GetMetalPixelFormatKeyMap().Add((uint64)Format, NextKey++);
+			// only giving 6 bits to the key
+			checkf(NextKey < 64, TEXT("Too many unique pixel formats to fit into the PipelineStateHash"));
+		}
+	}
+	return *Key;
+}
+
+FAGXSurface::FAGXSurface(ERHIResourceType ResourceType, EPixelFormat Format, uint32 InSizeX, uint32 InSizeY, uint32 InSizeZ, uint32 NumSamples, bool bArray, uint32 ArraySize, uint32 NumMips, ETextureCreateFlags InFlags, FResourceBulkDataInterface* BulkData)
+: Type(ResourceType)
+, PixelFormat(Format)
+, Texture(nil)
+, MSAATexture(nil)
+, MSAAResolveTexture(nil)
+, SizeX(InSizeX)
+, SizeY(InSizeY)
+, SizeZ(InSizeZ)
+, bIsCubemap(false)
+, Flags(InFlags)
+, BufferLocks(0)
+, TotalTextureSize(0)
+, Viewport(nullptr)
+, ImageSurfaceRef(nullptr)
+, bTextureView(false)
+{
+	// get a unique key for this surface's format
+	TMap<uint64, uint8>& PixelFormatKeyMap = GetMetalPixelFormatKeyMap();
+	if (PixelFormatKeyMap.Num() == 0)
+	{
+		// Add depth stencil formats first, so we don't have to use 6 bits for them in the pipeline hash
+		AGXGetMetalPixelFormatKey(mtlpp::PixelFormat::Depth32Float);
+		AGXGetMetalPixelFormatKey(mtlpp::PixelFormat::Stencil8);
+		AGXGetMetalPixelFormatKey(mtlpp::PixelFormat::Depth32Float_Stencil8);
+#if PLATFORM_MAC
+		AGXGetMetalPixelFormatKey(mtlpp::PixelFormat::Depth24Unorm_Stencil8);
+		AGXGetMetalPixelFormatKey(mtlpp::PixelFormat::Depth16Unorm);
+#endif
+	}
+	
+	FPlatformAtomics::InterlockedExchange(&Written, 0);
+	mtlpp::PixelFormat MTLFormat = (mtlpp::PixelFormat)GPixelFormats[Format].PlatformFormat;
+	
+	if (Flags & TexCreate_SRGB)
+	{
+		MTLFormat = AGXToSRGBFormat(MTLFormat);
+	}
+	
+	// set the key
+	FormatKey = AGXGetMetalPixelFormatKey(MTLFormat);
+	
+	
+	// the special back buffer surface will be updated in GetAGXDeviceContext().BeginDrawingViewport - no need to set the texture here
+	if (Flags & TexCreate_Presentable)
+	{
+		return;
+	}
+	
+	bool bIsRenderTarget = IsRenderTarget(Flags);
+	mtlpp::TextureDescriptor Desc;
+	
+	if (ResourceType == RRT_TextureCube)
+	{
+		Desc = mtlpp::TextureDescriptor::TextureCubeDescriptor(MTLFormat, SizeX, (NumMips > 1));
+		bIsCubemap = true;
+	}
+	else if (ResourceType == RRT_Texture3D)
+	{
+		Desc.SetTextureType(mtlpp::TextureType::Texture3D);
+		Desc.SetWidth(SizeX);
+		Desc.SetHeight(SizeY);
+		Desc.SetDepth(SizeZ);
+		Desc.SetPixelFormat(MTLFormat);
+		Desc.SetArrayLength(1);
+		Desc.SetMipmapLevelCount(1);
+		Desc.SetSampleCount(1);
+	}
+	else
+	{
+		Desc = mtlpp::TextureDescriptor::Texture2DDescriptor(MTLFormat, SizeX, SizeY, (NumMips > 1));
+		Desc.SetDepth(SizeZ);
+	}
+	
+	// flesh out the descriptor
+	if (bArray)
+	{
+		Desc.SetArrayLength(ArraySize);
+		if (bIsCubemap)
+		{
+			if (FAGXCommandQueue::SupportsFeature(EAGXFeaturesCubemapArrays))
+			{
+				Desc.SetTextureType(mtlpp::TextureType::TextureCubeArray);
+			}
+			else
+			{
+				Desc.SetTextureType(mtlpp::TextureType::Texture2DArray);
+				Desc.SetArrayLength(ArraySize * 6);
+			}
+		}
+		else
+		{
+			Desc.SetTextureType(mtlpp::TextureType::Texture2DArray);
+		}
+	}
+	Desc.SetMipmapLevelCount(NumMips);
+	
+	{
+		Desc.SetUsage(ConvertFlagsToUsage(Flags));
+		
+		if((Flags & TexCreate_CPUReadback) && !(Flags & (TexCreate_RenderTargetable|TexCreate_DepthStencilTargetable|TexCreate_FastVRAM)))
+		{
+			Desc.SetCpuCacheMode(mtlpp::CpuCacheMode::DefaultCache);
+#if PLATFORM_MAC
+			Desc.SetStorageMode(mtlpp::StorageMode::Managed);
+			Desc.SetResourceOptions((mtlpp::ResourceOptions)(mtlpp::ResourceOptions::CpuCacheModeDefaultCache|mtlpp::ResourceOptions::StorageModeManaged));
+#else
+			Desc.SetStorageMode(mtlpp::StorageMode::Shared);
+			Desc.SetResourceOptions((mtlpp::ResourceOptions)(mtlpp::ResourceOptions::CpuCacheModeDefaultCache|mtlpp::ResourceOptions::StorageModeShared));
+#endif
+		}
+		else if(((Flags & (TexCreate_NoTiling)) && !(Flags & (TexCreate_FastVRAM|TexCreate_DepthStencilTargetable|TexCreate_RenderTargetable|TexCreate_UAV))))
+		{
+#if PLATFORM_MAC
+			Desc.SetCpuCacheMode(mtlpp::CpuCacheMode::WriteCombined);
+			Desc.SetStorageMode(mtlpp::StorageMode::Managed);
+			Desc.SetResourceOptions((mtlpp::ResourceOptions)(mtlpp::ResourceOptions::CpuCacheModeWriteCombined|mtlpp::ResourceOptions::StorageModeManaged));
+#else
+			Desc.SetCpuCacheMode(mtlpp::CpuCacheMode::DefaultCache);
+			Desc.SetStorageMode(mtlpp::StorageMode::Shared);
+			Desc.SetResourceOptions((mtlpp::ResourceOptions)(mtlpp::ResourceOptions::CpuCacheModeDefaultCache|mtlpp::ResourceOptions::StorageModeShared));
+#endif
+		}
+		else if (Flags & (TexCreate_RenderTargetable|TexCreate_DepthStencilTargetable|TexCreate_ResolveTargetable|TexCreate_DepthStencilResolveTarget))
+		{
+			check(!(Flags & TexCreate_CPUReadback));
+			Desc.SetCpuCacheMode(mtlpp::CpuCacheMode::DefaultCache);
+#if PLATFORM_MAC
+			Desc.SetStorageMode(mtlpp::StorageMode::Private);
+			Desc.SetResourceOptions((mtlpp::ResourceOptions)(mtlpp::ResourceOptions::CpuCacheModeDefaultCache|mtlpp::ResourceOptions::StorageModePrivate));
+#else
+			if(GAGXForceIOSTexturesShared)
+			{
+				Desc.SetStorageMode(mtlpp::StorageMode::Shared);
+				Desc.SetResourceOptions((mtlpp::ResourceOptions)(mtlpp::ResourceOptions::CpuCacheModeDefaultCache|mtlpp::ResourceOptions::StorageModeShared));
+			}
+			else
+			{
+				Desc.SetStorageMode(mtlpp::StorageMode::Private);
+				Desc.SetResourceOptions((mtlpp::ResourceOptions)(mtlpp::ResourceOptions::CpuCacheModeDefaultCache|mtlpp::ResourceOptions::StorageModePrivate));
+			}
+#endif
+		}
+		else
+		{
+			check(!(Flags & TexCreate_CPUReadback));
+#if PLATFORM_MAC
+			Desc.SetCpuCacheMode(mtlpp::CpuCacheMode::WriteCombined);
+			Desc.SetStorageMode(mtlpp::StorageMode::Private);
+			Desc.SetResourceOptions((mtlpp::ResourceOptions)(mtlpp::ResourceOptions::CpuCacheModeWriteCombined|mtlpp::ResourceOptions::StorageModePrivate));
+#else
+			Desc.SetCpuCacheMode(mtlpp::CpuCacheMode::DefaultCache);
+			if(GAGXForceIOSTexturesShared)
+			{
+				Desc.SetStorageMode(mtlpp::StorageMode::Shared);
+				Desc.SetResourceOptions((mtlpp::ResourceOptions)(mtlpp::ResourceOptions::CpuCacheModeDefaultCache|mtlpp::ResourceOptions::StorageModeShared));
+			}
+			// No private storage for PVRTC as it messes up the blit-encoder usage.
+			// note: this is set to always be on and will be re-addressed in a future release
+			else
+			{
+				if (IsPixelFormatPVRTCCompressed(Format))
+				{
+					Desc.SetStorageMode(mtlpp::StorageMode::Shared);
+					Desc.SetResourceOptions((mtlpp::ResourceOptions)(mtlpp::ResourceOptions::CpuCacheModeDefaultCache|mtlpp::ResourceOptions::StorageModeShared));
+				}
+				else
+				{
+					Desc.SetStorageMode(mtlpp::StorageMode::Private);
+					Desc.SetResourceOptions((mtlpp::ResourceOptions)(mtlpp::ResourceOptions::CpuCacheModeDefaultCache|mtlpp::ResourceOptions::StorageModePrivate));
+				}
+			}
+#endif
+		}
+		
+#if PLATFORM_IOS
+		if (Flags & TexCreate_Memoryless)
+		{
+			ensure(Flags & (TexCreate_RenderTargetable | TexCreate_DepthStencilTargetable));
+			ensure(!(Flags & (TexCreate_CPUReadback | TexCreate_CPUWritable)));
+			ensure(!(Flags & TexCreate_UAV));
+			Desc.SetStorageMode(mtlpp::StorageMode::Memoryless);
+			Desc.SetResourceOptions((mtlpp::ResourceOptions)(mtlpp::ResourceOptions::CpuCacheModeDefaultCache | mtlpp::ResourceOptions::StorageModeMemoryless));
+		}
+#endif
+
+		static mtlpp::ResourceOptions GeneralResourceOption = FAGXCommandQueue::GetCompatibleResourceOptions(mtlpp::ResourceOptions::HazardTrackingModeUntracked);
+		Desc.SetResourceOptions((mtlpp::ResourceOptions)(Desc.GetResourceOptions() | GeneralResourceOption));
+	}
+	
+	if (BulkData != nullptr)
+	{
+		if (BulkData->GetResourceType() == FResourceBulkDataInterface::EBulkDataType::MediaTexture)
+		{
+			checkf(NumMips == 1 && ArraySize == 1, TEXT("Only handling bulk data with 1 mip and 1 array length"));
+			ImageSurfaceRef = (CFTypeRef)BulkData->GetResourceBulkData();
+			CFRetain(ImageSurfaceRef);
+			{
+#if !COREVIDEO_SUPPORTS_METAL
+				Texture = MTLPP_VALIDATE(mtlpp::Device, GetAGXDeviceContext().GetDevice(), AGXSafeGetRuntimeDebuggingLevel() >= EAGXDebugLevelValidation, NewTexture(Desc, CVPixelBufferGetIOSurface((CVPixelBufferRef)ImageSurfaceRef), 0));
+#else
+				Texture = CVMetalTextureGetTexture((CVMetalTextureRef)ImageSurfaceRef);
+#endif
+				METAL_FATAL_ASSERT(Texture, TEXT("Failed to create texture, desc %s"), *FString([Desc description]));
+			}
+			
+			BulkData->Discard();
+		}
+#if PLATFORM_MAC
+		else if (BulkData->GetResourceType() == FResourceBulkDataInterface::EBulkDataType::VREyeBuffer)
+		{
+			ImageSurfaceRef = (CFTypeRef)BulkData->GetResourceBulkData();
+			CFRetain(ImageSurfaceRef);
+			
+			Desc.SetStorageMode(mtlpp::StorageMode::Managed);
+			Desc.SetResourceOptions((mtlpp::ResourceOptions)((Desc.GetResourceOptions() & ~(mtlpp::ResourceStorageModeMask)) | mtlpp::ResourceOptions::StorageModeManaged));
+			
+			Texture = [GetAGXDeviceContext().GetDevice() newTextureWithDescriptor:Desc iosurface : (IOSurfaceRef)ImageSurfaceRef plane : 0];
+			
+			METAL_FATAL_ASSERT(Texture, TEXT("Failed to create texture, desc %s"), *FString([Desc description]));
+			
+			BulkData->Discard();
+		}
+#endif
+	}
+	else
+	{
+		const bool bBufferCompatibleOption = (Desc.GetTextureType() == mtlpp::TextureType::Texture2D || Desc.GetTextureType() == mtlpp::TextureType::TextureBuffer) && NumMips == 1;
+
+		if (!bBufferCompatibleOption || ((Flags & (TexCreate_UAV|TexCreate_NoTiling)) != (TexCreate_UAV|TexCreate_NoTiling)))
+		{
+			Texture = GetAGXDeviceContext().CreateTexture(this, Desc);
+		}
+		else
+		{
+			mtlpp::Device Device = GetAGXDeviceContext().GetDevice();
+			mtlpp::SizeAndAlign SizeAlign = Device.HeapTextureSizeAndAlign(Desc);
+
+			// Backing buffer resource options must match the texture we are going to create from it
+			FAGXPooledBufferArgs Args(Device, SizeAlign.Size, BUF_Dynamic, mtlpp::StorageMode::Private, Desc.GetCpuCacheMode());
+			FAGXBuffer Buffer = GetAGXDeviceContext().CreatePooledBuffer(Args);
+
+			Texture = Buffer.NewTexture(Desc, 0, Align(SizeAlign.Size / SizeY, 256));
+		}
+		
+		METAL_FATAL_ASSERT(Texture, TEXT("Failed to create texture, desc %s"), *FString([Desc description]));
+		
+		// upload existing bulkdata
+		if (BulkData)
+		{
+			UE_LOG(LogAGX, Display, TEXT("Got a bulk data texture, with %d mips"), NumMips);
+			checkf(NumMips == 1&& ArraySize == 1, TEXT("Only handling bulk data with 1 mip and 1 array length"));
+			uint32 Stride;
+			
+			// lock, copy, unlock
+			void* LockedData = Lock(0, 0, RLM_WriteOnly, Stride);
+			FMemory::Memcpy(LockedData, BulkData->GetResourceBulkData(), BulkData->GetResourceBulkDataSize());
+			Unlock(0, 0, false);
+			
+			// bulk data can be unloaded now
+			BulkData->Discard();
+		}
+	}
+	
+	// calculate size of the texture
+	TotalTextureSize = GetMemorySize();
+	
+	const bool bSupportsMSAADepthResolve = GetAGXDeviceContext().SupportsFeature(EAGXFeaturesMSAADepthResolve);
+	if (!FParse::Param(FCommandLine::Get(), TEXT("nomsaa")))
+	{
+		if (NumSamples > 1)
+		{
+			check(bIsRenderTarget);
+			Desc.SetTextureType(mtlpp::TextureType::Texture2DMultisample);
+			
+			// allow commandline to override
+			FParse::Value(FCommandLine::Get(), TEXT("msaa="), NumSamples);
+			Desc.SetSampleCount(NumSamples);
+			
+			bool bMemoryless = false;
+#if PLATFORM_IOS
+			if (GMaxRHIFeatureLevel < ERHIFeatureLevel::SM5)
+			{
+				bMemoryless = true;
+				Desc.SetStorageMode(mtlpp::StorageMode::Memoryless);
+				Desc.SetResourceOptions(mtlpp::ResourceOptions::StorageModeMemoryless);
+			}
+#endif
+			
+			MSAATexture = GetAGXDeviceContext().CreateTexture(this, Desc);
+			
+			//device doesn't support HW depth resolve.  This case only valid on mobile renderer or
+			//on Mac where RHISupportsSeparateMSAAAndResolveTextures is true.
+			const bool bDepthButNoResolveSupported = Format == PF_DepthStencil && !bSupportsMSAADepthResolve;
+			if (bDepthButNoResolveSupported)
+			{
+				Texture = MSAATexture;
+				
+				// we don't have the resolve texture, so we just update the memory size with the MSAA size
+				TotalTextureSize = TotalTextureSize * NumSamples;
+			}
+			else if (!bMemoryless)
+			{
+				// an MSAA render target takes NumSamples more space, in addition to the resolve texture
+				TotalTextureSize += TotalTextureSize * NumSamples;
+			}
+			
+			if (MSAATexture != Texture)
+			{
+				check(!MSAAResolveTexture);
+				
+				//if bSupportsSeparateMSAAAndResolve then the high level expect to binds the MSAA when binding shader params.
+				const bool bSupportsSeparateMSAAAndResolve = FAGXCommandQueue::SupportsSeparateMSAAAndResolveTarget();
+				if (bSupportsSeparateMSAAAndResolve)
+				{
+					MSAAResolveTexture = Texture;
+					
+					Texture = MSAATexture;
+				}
+				else
+				{
+					MSAAResolveTexture = Texture;
+				}
+			}
+			
+			//we always require an MSAAResolveTexture if MSAATexture is active.
+			check(!MSAATexture || MSAAResolveTexture || bDepthButNoResolveSupported);
+			
+			NSLog(@"Creating %dx MSAA %d x %d %s surface", (int32)Desc.GetSampleCount(), SizeX, SizeY, (Flags & TexCreate_RenderTargetable) ? "Color" : "Depth");
+			if (MSAATexture.GetPtr() == nil)
+			{
+				NSLog(@"Failed to create texture, desc  %@", Desc.GetPtr());
+			}
+		}
+	}
+	
+	// create a stencil buffer if needed
+	if (Format == PF_DepthStencil)
+	{
+		// 1 byte per texel
+		TotalTextureSize += SizeX * SizeY;
+	}
+	
+	// track memory usage
+	if (bIsRenderTarget)
+	{
+		GCurrentRendertargetMemorySize += Align(TotalTextureSize, 1024) / 1024;
+	}
+	else
+	{
+		GCurrentTextureMemorySize += Align(TotalTextureSize, 1024) / 1024;
+	}
+	
+#if STATS
+	if (ResourceType == RRT_TextureCube)
+	{
+		if (bIsRenderTarget)
+		{
+			INC_MEMORY_STAT_BY(STAT_RenderTargetMemoryCube, TotalTextureSize);
+		}
+		else
+		{
+			INC_MEMORY_STAT_BY(STAT_TextureMemoryCube, TotalTextureSize);
+		}
+	}
+	else if (ResourceType == RRT_Texture3D)
+	{
+		if (bIsRenderTarget)
+		{
+			INC_MEMORY_STAT_BY(STAT_RenderTargetMemory3D, TotalTextureSize);
+		}
+		else
+		{
+			INC_MEMORY_STAT_BY(STAT_TextureMemory3D, TotalTextureSize);
+		}
+	}
+	else
+	{
+		if (bIsRenderTarget)
+		{
+			INC_MEMORY_STAT_BY(STAT_RenderTargetMemory2D, TotalTextureSize);
+		}
+		else
+		{
+			INC_MEMORY_STAT_BY(STAT_TextureMemory2D, TotalTextureSize);
+		}
+	}
+#endif
+}
+
+@interface FAGXDeferredStats : FApplePlatformObject
+{
+@public
+	uint64 TextureSize;
+	ERHIResourceType TexType;
+	bool bIsRenderTarget;
+}
+@end
+
+@implementation FAGXDeferredStats
+APPLE_PLATFORM_OBJECT_ALLOC_OVERRIDES(FAGXDeferredStats)
+-(void)dealloc
+{
+#if STATS
+	if (TexType == RRT_TextureCube)
+	{
+		if (bIsRenderTarget)
+		{
+			DEC_MEMORY_STAT_BY(STAT_RenderTargetMemoryCube, TextureSize);
+		}
+		else
+		{
+			DEC_MEMORY_STAT_BY(STAT_TextureMemoryCube, TextureSize);
+		}
+	}
+	else if (TexType == RRT_Texture3D)
+	{
+		if (bIsRenderTarget)
+		{
+			DEC_MEMORY_STAT_BY(STAT_RenderTargetMemory3D, TextureSize);
+		}
+		else
+		{
+			DEC_MEMORY_STAT_BY(STAT_TextureMemory3D, TextureSize);
+		}
+	}
+	else
+	{
+		if (bIsRenderTarget)
+		{
+			DEC_MEMORY_STAT_BY(STAT_RenderTargetMemory2D, TextureSize);
+		}
+		else
+		{
+			DEC_MEMORY_STAT_BY(STAT_TextureMemory2D, TextureSize);
+		}
+	}
+#endif
+	if (bIsRenderTarget)
+	{
+		GCurrentRendertargetMemorySize -= Align(TextureSize, 1024) / 1024;
+	}
+	else
+	{
+		GCurrentTextureMemorySize -= Align(TextureSize, 1024) / 1024;
+	}
+	[super dealloc];
+}
+@end
+
+FAGXSurface::~FAGXSurface()
+{
+	bool const bIsRenderTarget = IsRenderTarget(Flags);
+	
+	if (MSAATexture.GetPtr())
+	{
+		if (Texture.GetPtr() != MSAATexture.GetPtr())
+		{
+			AGXSafeReleaseMetalTexture(this, MSAATexture, bTextureView, false);
+		}
+	}
+	
+	
+	//do the same as above.  only do a [release] if it'st he same as texture.
+	if (MSAAResolveTexture.GetPtr())
+	{
+		if (Texture.GetPtr() != MSAAResolveTexture.GetPtr())
+		{
+			AGXSafeReleaseMetalTexture(this, MSAAResolveTexture, bTextureView, false);
+		}
+	}
+	
+	if (!(Flags & TexCreate_Presentable) && Texture.GetPtr())
+	{
+		AGXSafeReleaseMetalTexture(this, Texture, bTextureView, (ImageSurfaceRef != nullptr));
+	}
+	
+	MSAATexture = nil;
+	MSAAResolveTexture = nil;
+	Texture = nil;
+	
+	// track memory usage
+	FAGXDeferredStats* Block = [FAGXDeferredStats new];
+	Block->TexType = Type;
+	Block->TextureSize = TotalTextureSize;
+	Block->bIsRenderTarget = bIsRenderTarget;
+	AGXSafeReleaseMetalObject(Block);
+	
+	if(ImageSurfaceRef)
+	{
+		// CFArray can contain CFType objects and is toll-free bridged with NSArray
+		CFArrayRef Temp = CFArrayCreate(kCFAllocatorSystemDefault, &ImageSurfaceRef, 1, &kCFTypeArrayCallBacks);
+		AGXSafeReleaseMetalObject((NSArray*)Temp);
+		CFRelease(ImageSurfaceRef);
+	}
+	
+	ImageSurfaceRef = nullptr;
+}
+
+id <MTLBuffer> FAGXSurface::AllocSurface(uint32 MipIndex, uint32 ArrayIndex, EResourceLockMode LockMode, uint32& DestStride, bool SingleLayer /*= false*/)
+{
+	check(IsInRenderingThread());
+
+	// get size and stride
+	uint32 MipBytes = GetMipSize(MipIndex, &DestStride, SingleLayer);
+	
+	// allocate some temporary memory
+	// This should really be pooled and texture transfers should be their own pool
+	id <MTLDevice> Device = GetAGXDeviceContext().GetDevice();
+	id <MTLBuffer> Buffer = [Device newBufferWithLength:MipBytes options:MTLResourceStorageModeShared];
+	Buffer.label = @"Temporary Surface Backing";
+	
+	// Note: while the lock is active, this map owns the backing store.
+	GRHILockTracker.Lock(this, Buffer, MipIndex, 0, LockMode, false);
+	
+#if PLATFORM_MAC
+	// Expand R8_sRGB into RGBA8_sRGB for Mac.
+	if (PixelFormat == PF_G8 && (Flags & TexCreate_SRGB) && Type == RRT_Texture2D && LockMode == RLM_WriteOnly)
+	{
+		DestStride = FMath::Max<uint32>(SizeX >> MipIndex, 1);
+	}
+#endif
+	
+	check(Buffer);
+	
+	return Buffer;
+}
+
+void FAGXSurface::UpdateSurfaceAndDestroySourceBuffer(id <MTLBuffer> SourceBuffer, uint32 MipIndex, uint32 ArrayIndex)
+{
+#if STATS
+	uint64 Start = FPlatformTime::Cycles64();
+#endif
+	check(SourceBuffer);
+	
+	uint32 Stride;
+	uint32 BytesPerImage = GetMipSize(MipIndex, &Stride, true);
+	
+	mtlpp::Region Region;
+	if (SizeZ <= 1 || bIsCubemap)
+	{
+		// upload the texture to the texture slice
+		Region = mtlpp::Region(0, 0, FMath::Max<uint32>(SizeX >> MipIndex, 1), FMath::Max<uint32>(SizeY >> MipIndex, 1));
+	}
+	else
+	{
+		// upload the texture to the texture slice
+		Region = mtlpp::Region(0, 0, 0, FMath::Max<uint32>(SizeX >> MipIndex, 1), FMath::Max<uint32>(SizeY >> MipIndex, 1), FMath::Max<uint32>(SizeZ >> MipIndex, 1));
+	}
+#if PLATFORM_MAC
+	// Expand R8_sRGB into RGBA8_sRGB for Mac.
+	if (PixelFormat == PF_G8 && (Flags & TexCreate_SRGB) && Type == RRT_Texture2D)
+	{
+		TArray<uint8> Data;
+		uint8* ExpandedMem = (uint8*) SourceBuffer.contents;
+		check(ExpandedMem);
+		Data.Append(ExpandedMem, BytesPerImage);
+		uint32 SrcStride = FMath::Max<uint32>(SizeX >> MipIndex, 1);
+		for(uint y = 0; y < FMath::Max<uint32>(SizeY >> MipIndex, 1); y++)
+		{
+			uint8* RowDest = ExpandedMem;
+			for(uint x = 0; x < FMath::Max<uint32>(SizeX >> MipIndex, 1); x++)
+			{
+				*(RowDest++) = Data[(y * SrcStride) + x];
+				*(RowDest++) = Data[(y * SrcStride) + x];
+				*(RowDest++) = Data[(y * SrcStride) + x];
+				*(RowDest++) = Data[(y * SrcStride) + x];
+			}
+			ExpandedMem = (ExpandedMem + Stride);
+		}
+	}
+#endif
+	
+	if(Texture.GetStorageMode() == mtlpp::StorageMode::Private)
+	{
+		SCOPED_AUTORELEASE_POOL;
+		
+		FAGXBuffer Buffer(SourceBuffer);
+		
+		int64 Size = BytesPerImage * Region.size.depth * FMath::Max(1u, ArrayIndex);
+		
+		int64 Count = FPlatformAtomics::InterlockedAdd(&ActiveUploads, Size);
+		
+		bool const bWait = ((GetAGXDeviceContext().GetNumActiveContexts() == 1) && (GAGXMaxOutstandingAsyncTexUploads > 0) && (Count >= GAGXMaxOutstandingAsyncTexUploads));
+		
+		mtlpp::BlitOption Options = mtlpp::BlitOption::None;
+#if !PLATFORM_MAC
+		if (Texture.GetPixelFormat() >= mtlpp::PixelFormat::PVRTC_RGB_2BPP && Texture.GetPixelFormat() <= mtlpp::PixelFormat::PVRTC_RGBA_4BPP_sRGB)
+		{
+			Options = mtlpp::BlitOption::RowLinearPVRTC;
+		}
+#endif
+		
+		if(GetAGXDeviceContext().AsyncCopyFromBufferToTexture(Buffer, 0, Stride, BytesPerImage, Region.size, Texture, ArrayIndex, MipIndex, Region.origin, Options))
+		{
+			mtlpp::CommandBufferHandler ScheduledHandler = nil;
+	#if STATS
+			int64* Cycles = new int64;
+			FPlatformAtomics::InterlockedExchange(Cycles, 0);
+			ScheduledHandler = [Cycles](mtlpp::CommandBuffer const&)
+			{
+				FPlatformAtomics::InterlockedExchange(Cycles, FPlatformTime::Cycles64());
+			};
+			mtlpp::CommandBufferHandler CompletionHandler = [SourceBuffer, Size, Cycles](mtlpp::CommandBuffer const&)
+	#else
+			mtlpp::CommandBufferHandler CompletionHandler = [SourceBuffer, Size](mtlpp::CommandBuffer const&)
+	#endif
+			{
+				FPlatformAtomics::InterlockedAdd(&ActiveUploads, -Size);
+	#if STATS
+				int64 Taken = FPlatformTime::Cycles64() - *Cycles;
+				delete Cycles;
+				FPlatformAtomics::InterlockedAdd(&GAGXTexturePageOnTime, Taken);
+	#endif
+				[SourceBuffer release];
+			};
+			GetAGXDeviceContext().SubmitAsyncCommands(ScheduledHandler, CompletionHandler, bWait);
+			
+		}
+		else
+		{
+			mtlpp::CommandBufferHandler CompletionHandler = [SourceBuffer, Size](mtlpp::CommandBuffer const&)
+			{
+				FPlatformAtomics::InterlockedAdd(&ActiveUploads, -Size);
+				[SourceBuffer release];
+			};
+			GetAGXDeviceContext().GetCurrentRenderPass().AddCompletionHandler(CompletionHandler);
+		}
+		
+		INC_DWORD_STAT_BY(STAT_AGXTextureMemUpdate, Size);
+		
+		if (bWait)
+		{
+			GetAGXDeviceContext().ClearFreeList();
+		}
+	}
+	else
+	{
+#if !PLATFORM_MAC
+		if (Texture.GetPixelFormat() >= mtlpp::PixelFormat::PVRTC_RGB_2BPP && Texture.GetPixelFormat() <= mtlpp::PixelFormat::PVRTC_RGBA_4BPP_sRGB) // @todo Calculate correct strides and byte-counts
+		{
+			Stride = 0;
+			BytesPerImage = 0;
+		}
+#endif
+		
+		MTLPP_VALIDATE(mtlpp::Texture, Texture, AGXSafeGetRuntimeDebuggingLevel() >= EAGXDebugLevelValidation, Replace(Region, MipIndex, ArrayIndex, SourceBuffer.contents, Stride, BytesPerImage));
+		[SourceBuffer release];
+		
+		INC_DWORD_STAT_BY(STAT_AGXTextureMemUpdate, BytesPerImage);
+	}
+	
+	FPlatformAtomics::InterlockedExchange(&Written, 1);
+	
+#if STATS
+	GAGXTexturePageOnTime += (FPlatformTime::Cycles64() - Start);
+#endif
+}
+
+void* FAGXSurface::Lock(uint32 MipIndex, uint32 ArrayIndex, EResourceLockMode LockMode, uint32& DestStride, bool SingleLayer /*= false*/)
+{
+	// get size and stride
+	uint32 MipBytes = GetMipSize(MipIndex, &DestStride, false);
+	
+	// allocate some temporary memory
+	id <MTLBuffer> Buffer = AllocSurface(MipIndex, ArrayIndex, LockMode, DestStride, SingleLayer);
+	FAGXBuffer SourceData(Buffer);
+	
+	switch(LockMode)
+	{
+		case RLM_ReadOnly:
+		{
+			SCOPE_CYCLE_COUNTER(STAT_AGXTexturePageOffTime);
+			
+			FRHICommandListImmediate& RHICmdList = FRHICommandListExecutor::GetImmediateCommandList();
+			const bool bIssueImmediateCommands = RHICmdList.Bypass() || IsInRHIThread();
+			
+			mtlpp::Region Region;
+			if (SizeZ <= 1 || bIsCubemap)
+			{
+				// upload the texture to the texture slice
+				Region = mtlpp::Region(0, 0, FMath::Max<uint32>(SizeX >> MipIndex, 1), FMath::Max<uint32>(SizeY >> MipIndex, 1));
+			}
+			else
+			{
+				// upload the texture to the texture slice
+				Region = mtlpp::Region(0, 0, 0, FMath::Max<uint32>(SizeX >> MipIndex, 1), FMath::Max<uint32>(SizeY >> MipIndex, 1), FMath::Max<uint32>(SizeZ >> MipIndex, 1));
+			}
+			
+			if (Texture.GetStorageMode() == mtlpp::StorageMode::Private)
+			{
+				// If we are running with command lists or the RHI thread is enabled we have to execute GFX commands in that context.
+				auto CopyTexToBuf =
+				[this, &ArrayIndex, &MipIndex, &Region, &SourceData, &DestStride, &MipBytes](FRHICommandListImmediate& RHICmdList)
+				{
+					GetAGXDeviceContext().CopyFromTextureToBuffer(this->Texture, ArrayIndex, MipIndex, Region.origin, Region.size, SourceData, 0, DestStride, MipBytes, mtlpp::BlitOption::None);
+					//kick the current command buffer.
+					GetAGXDeviceContext().SubmitCommandBufferAndWait();
+				};
+				
+				if (bIssueImmediateCommands)
+				{
+					CopyTexToBuf(RHICmdList);
+				}
+				else
+				{
+					RHICmdList.EnqueueLambda(MoveTemp(CopyTexToBuf));
+					RHICmdList.ImmediateFlush(EImmediateFlushType::FlushRHIThread);
+				}
+			}
+			else
+			{
+#if PLATFORM_MAC
+				if((GPUReadback & EAGXGPUReadbackFlags::ReadbackRequestedAndComplete) != EAGXGPUReadbackFlags::ReadbackRequestedAndComplete)
+				{
+					// A previous texture sync has not been done, need the data now, request texture sync and kick the current command buffer.
+					auto SyncReadbackToCPU =
+					[this, &ArrayIndex, &MipIndex](FRHICommandListImmediate& RHICmdList)
+					{
+						GetAGXDeviceContext().SynchronizeTexture(this->Texture, ArrayIndex, MipIndex);
+						GetAGXDeviceContext().SubmitCommandBufferAndWait();
+					};
+					
+					// Similar to above. If we are in a context where we have command lists or the RHI thread we must execute
+					// commands there. Otherwise we can just do this directly.
+					if (bIssueImmediateCommands)
+					{
+						SyncReadbackToCPU(RHICmdList);
+					}
+					else
+					{
+						RHICmdList.EnqueueLambda(MoveTemp(SyncReadbackToCPU));
+						RHICmdList.ImmediateFlush(EImmediateFlushType::FlushRHIThread);
+					}
+				}
+#endif
+				
+				// This block breaks the texture atlas system in Ocean, which depends on nonzero strides coming back from compressed textures. Turning off.
+#if 0
+				if (PixelFormat == PF_PVRTC2 || PixelFormat == PF_PVRTC4)
+				{
+					// for compressed textures metal debug RT expects 0 for rowBytes and imageBytes.
+					DestStride = 0;
+					MipBytes = 0;
+				}
+#endif
+				uint32 BytesPerRow = DestStride;
+				if (PixelFormat == PF_PVRTC2 || PixelFormat == PF_PVRTC4)
+				{
+					// for compressed textures metal debug RT expects 0 for rowBytes and imageBytes.
+					BytesPerRow = 0;
+					MipBytes = 0;
+				}
+				MTLPP_VALIDATE(mtlpp::Texture, Texture, AGXSafeGetRuntimeDebuggingLevel() >= EAGXDebugLevelValidation, GetBytes(MTLPP_VALIDATE(mtlpp::Buffer, SourceData, AGXSafeGetRuntimeDebuggingLevel() >= EAGXDebugLevelValidation, GetContents()), BytesPerRow, MipBytes, Region, MipIndex, ArrayIndex));
+			}
+			
+#if PLATFORM_MAC
+			// Pack RGBA8_sRGB into R8_sRGB for Mac.
+			if (PixelFormat == PF_G8 && (Flags & TexCreate_SRGB) && Type == RRT_Texture2D)
+			{
+				TArray<uint8> Data;
+				uint8* ExpandedMem = (uint8*)SourceData.GetContents();
+				Data.Append(ExpandedMem, MipBytes);
+				uint32 SrcStride = DestStride;
+				DestStride = FMath::Max<uint32>(SizeX >> MipIndex, 1);
+				for(uint y = 0; y < FMath::Max<uint32>(SizeY >> MipIndex, 1); y++)
+				{
+					uint8* RowDest = ExpandedMem;
+					for(uint x = 0; x < FMath::Max<uint32>(SizeX >> MipIndex, 1); x++)
+					{
+						*(RowDest++) = Data[(y * SrcStride) + (x * 4)];
+					}
+					ExpandedMem = (ExpandedMem + DestStride);
+				}
+			}
+#endif
+			
+			break;
+		}
+		case RLM_WriteOnly:
+		{
+			BufferLocks |= 1 << MipIndex;
+			break;
+		}
+		default:
+			check(false);
+			break;
+	}
+	
+	return SourceData.GetContents();
+}
+
+void FAGXSurface::Unlock(uint32 MipIndex, uint32 ArrayIndex, bool bTryAsync)
+{
+	check(IsInRenderingThread());
+	BufferLocks &= ~(1 << MipIndex);
+	
+	FRHILockTracker::FLockParams Params = GRHILockTracker.Unlock(this, MipIndex);
+	id <MTLBuffer> SourceData = (id <MTLBuffer>) Params.Buffer;
+	
+	if(bTryAsync)
+	{
+		AsyncUnlock(SourceData, MipIndex, ArrayIndex);
+	}
+	else
+	{
+		UpdateSurfaceAndDestroySourceBuffer(SourceData, MipIndex, ArrayIndex);
+	}
+}
+
+void* FAGXSurface::AsyncLock(class FRHICommandListImmediate& RHICmdList, uint32 MipIndex, uint32 ArrayIndex, EResourceLockMode LockMode, uint32& DestStride, bool bNeedsDefaultRHIFlush)
+{
+	bool bDirectLock = (LockMode == RLM_ReadOnly || !GIsRHIInitialized);
+	
+	void* BufferData = nullptr;
+	
+	// Never flush for writing, it is unnecessary
+	if (bDirectLock)
+	{
+		if (bNeedsDefaultRHIFlush)
+		{
+			// @todo Not all read locks need to flush either, but that'll require resource use tracking
+			QUICK_SCOPE_CYCLE_COUNTER(STAT_RHIMETHOD_LockTexture2D_Flush);
+			RHICmdList.ImmediateFlush(EImmediateFlushType::FlushRHIThread);
+		}
+		BufferData = Lock(MipIndex, ArrayIndex, LockMode, DestStride);
+	}
+	else
+	{
+		id <MTLBuffer> Buffer = AllocSurface(MipIndex, 0, LockMode, DestStride);
+		check(Buffer);
+		
+		BufferData = Buffer.contents;
+	}
+	
+	check(BufferData);
+	
+	return BufferData;
+}
+
+struct FAGXRHICommandUnlockTextureUpdate final : public FRHICommand<FAGXRHICommandUnlockTextureUpdate>
+{
+	FAGXSurface* Surface;
+	id <MTLBuffer> UpdateData;
+	uint32 MipIndex;
+	
+	FORCEINLINE_DEBUGGABLE FAGXRHICommandUnlockTextureUpdate(FAGXSurface* InSurface, id <MTLBuffer> InUpdateData, uint32 InMipIndex)
+	: Surface(InSurface)
+	, UpdateData(InUpdateData)
+	, MipIndex(InMipIndex)
+	{
+		[UpdateData retain];
+	}
+	
+	void Execute(FRHICommandListBase& CmdList)
+	{
+		Surface->UpdateSurfaceAndDestroySourceBuffer(UpdateData, MipIndex, 0);
+	}
+	
+	virtual ~FAGXRHICommandUnlockTextureUpdate()
+	{
+		[UpdateData release];
+	}
+};
+
+void FAGXSurface::AsyncUnlock(id <MTLBuffer> SourceData, uint32 MipIndex, uint32 ArrayIndex)
+{
+	FRHICommandListImmediate& RHICmdList = FRHICommandListExecutor::GetImmediateCommandList();
+	if (RHICmdList.Bypass() || !IsRunningRHIInSeparateThread())
+	{
+		UpdateSurfaceAndDestroySourceBuffer(SourceData, MipIndex, ArrayIndex);
+	}
+	else
+	{
+		new (RHICmdList.AllocCommand<FAGXRHICommandUnlockTextureUpdate>()) FAGXRHICommandUnlockTextureUpdate(this, SourceData, MipIndex);
+	}
+}
+
+uint32 FAGXSurface::GetMipSize(uint32 MipIndex, uint32* Stride, bool bSingleLayer)
+{
+	// DXT/BC formats on Mac actually do have mip-tails that are smaller than the block size, they end up being uncompressed.
+	bool const bPixelFormatASTC = IsPixelFormatASTCCompressed(PixelFormat);
+	
+	// Calculate the dimensions of the mip-map.
+	const uint32 BlockSizeX = GPixelFormats[PixelFormat].BlockSizeX;
+	const uint32 BlockSizeY = GPixelFormats[PixelFormat].BlockSizeY;
+	const uint32 BlockBytes = GPixelFormats[PixelFormat].BlockBytes;
+	const uint32 Alignment = 1u; // Apparently we always want natural row alignment (tightly-packed) even though the docs say iOS doesn't support it - this may be because we don't upload texture data from one contiguous buffer.
+	const uint32 UnalignedMipSizeX = FMath::Max(SizeX >> MipIndex, BlockSizeX);
+	const uint32 UnalignedMipSizeY = FMath::Max(SizeY >> MipIndex, BlockSizeY);
+	const uint32 MipSizeX = (bPixelFormatASTC) ? AlignArbitrary(UnalignedMipSizeX, BlockSizeX) : UnalignedMipSizeX;
+	const uint32 MipSizeY = (bPixelFormatASTC) ? AlignArbitrary(UnalignedMipSizeY, BlockSizeY) : UnalignedMipSizeY;
+	
+	const uint32 MipSizeZ = bSingleLayer ? 1 : FMath::Max(SizeZ >> MipIndex, 1u);
+	uint32 NumBlocksX = (MipSizeX + BlockSizeX - 1) / BlockSizeX;
+	uint32 NumBlocksY = (MipSizeY + BlockSizeY - 1) / BlockSizeY;
+	if (PixelFormat == PF_PVRTC2 || PixelFormat == PF_PVRTC4)
+	{
+		// PVRTC has minimum 2 blocks width and height
+		NumBlocksX = FMath::Max<uint32>(NumBlocksX, 2);
+		NumBlocksY = FMath::Max<uint32>(NumBlocksY, 2);
+	}
+#if PLATFORM_MAC
+	else if (PixelFormat == PF_G8 && (Flags & TexCreate_SRGB))
+	{
+		// RGBA_sRGB is the closest match - so expand the data.
+		NumBlocksX *= 4;
+	}
+#endif
+	
+	const uint32 MipStride = NumBlocksX * BlockBytes;
+	const uint32 AlignedStride = ((MipStride - 1) & ~(Alignment - 1)) + Alignment;
+	
+	const uint32 MipBytes = AlignedStride * NumBlocksY * MipSizeZ;
+	
+	if (Stride)
+	{
+		*Stride = AlignedStride;
+	}
+	
+	return MipBytes;
+}
+
+uint32 FAGXSurface::GetMemorySize()
+{
+	// if already calculated, no need to do it again
+	if (TotalTextureSize != 0)
+	{
+		return TotalTextureSize;
+	}
+	
+	if (Texture.GetPtr() == nil)
+	{
+		return 0;
+	}
+	
+	uint32 TotalSize = 0;
+	for (uint32 MipIndex = 0; MipIndex < Texture.GetMipmapLevelCount(); MipIndex++)
+	{
+		TotalSize += GetMipSize(MipIndex, NULL, false);
+	}
+	
+	return TotalSize;
+}
+
+uint32 FAGXSurface::GetNumFaces()
+{
+	switch (Type)
+	{
+		case RRT_Texture2DArray:
+		case RRT_Texture3D:
+		case RRT_TextureCube:
+			return SizeZ * Texture.GetArrayLength();
+			
+		case RRT_Texture:
+		case RRT_Texture2D:
+		default:
+			return 1;
+	}
+}
+
+FAGXTexture FAGXSurface::GetDrawableTexture()
+{
+	if (!Texture && (Flags & TexCreate_Presentable))
+	{
+		check(Viewport);
+		Texture = Viewport->GetDrawableTexture(EAGXViewportAccessRHI);
+	}
+	return Texture;
+}
+
+ns::AutoReleased<FAGXTexture> FAGXSurface::GetCurrentTexture()
+{
+	ns::AutoReleased<FAGXTexture> Tex;
+	if (Viewport && (Flags & TexCreate_Presentable))
+	{
+		check(Viewport);
+		Tex = Viewport->GetCurrentTexture(EAGXViewportAccessRHI);
+	}
+	return Tex;
+}
+
+
+/*-----------------------------------------------------------------------------
+ Texture allocator support.
+ -----------------------------------------------------------------------------*/
+
+void FAGXDynamicRHI::RHIGetTextureMemoryStats(FTextureMemoryStats& OutStats)
+{
+	if(MemoryStats.TotalGraphicsMemory > 0)
+	{
+		OutStats.DedicatedVideoMemory = MemoryStats.DedicatedVideoMemory;
+		OutStats.DedicatedSystemMemory = MemoryStats.DedicatedSystemMemory;
+		OutStats.SharedSystemMemory = MemoryStats.SharedSystemMemory;
+		OutStats.TotalGraphicsMemory = MemoryStats.TotalGraphicsMemory;
+	}
+	else
+	{
+		OutStats.DedicatedVideoMemory = 0;
+		OutStats.DedicatedSystemMemory = 0;
+		OutStats.SharedSystemMemory = 0;
+		OutStats.TotalGraphicsMemory = 0;
+	}
+	
+	OutStats.AllocatedMemorySize = int64(GCurrentTextureMemorySize) * 1024;
+	OutStats.LargestContiguousAllocation = OutStats.AllocatedMemorySize;
+	OutStats.TexturePoolSize = GTexturePoolSize;
+	OutStats.PendingMemoryAdjustment = 0;
+}
+
+bool FAGXDynamicRHI::RHIGetTextureMemoryVisualizeData( FColor* /*TextureData*/, int32 /*SizeX*/, int32 /*SizeY*/, int32 /*Pitch*/, int32 /*PixelSize*/ )
+{
+	NOT_SUPPORTED("RHIGetTextureMemoryVisualizeData");
+	return false;
+}
+
+uint32 FAGXDynamicRHI::RHIComputeMemorySize(FRHITexture* TextureRHI)
+{
+	@autoreleasepool {
+		if(!TextureRHI)
+		{
+			return 0;
+		}
+		
+		return AGXGetMetalSurfaceFromRHITexture(TextureRHI)->GetMemorySize();
+	}
+}
+
+/*-----------------------------------------------------------------------------
+ 2D texture support.
+ -----------------------------------------------------------------------------*/
+
+FTexture2DRHIRef FAGXDynamicRHI::RHICreateTexture2D(uint32 SizeX, uint32 SizeY, uint8 Format, uint32 NumMips, uint32 NumSamples, ETextureCreateFlags Flags, ERHIAccess InResourceState, FRHIResourceCreateInfo& CreateInfo)
+{
+	@autoreleasepool {
+		return new FAGXTexture2D((EPixelFormat)Format, SizeX, SizeY, NumMips, NumSamples, Flags, CreateInfo.BulkData, CreateInfo.ClearValueBinding);
+	}
+}
+
+FTexture2DRHIRef FAGXDynamicRHI::RHICreateTextureExternal2D(uint32 SizeX, uint32 SizeY, uint8 Format, uint32 NumMips, uint32 NumSamples, ETextureCreateFlags Flags, ERHIAccess InResourceState, FRHIResourceCreateInfo& CreateInfo)
+{
+	@autoreleasepool {
+		return new FAGXTexture2D((EPixelFormat)Format, SizeX, SizeY, NumMips, NumSamples, Flags, CreateInfo.BulkData, CreateInfo.ClearValueBinding);
+	}
+}
+
+FTexture2DRHIRef FAGXDynamicRHI::RHIAsyncCreateTexture2D(uint32 SizeX, uint32 SizeY, uint8 Format, uint32 NumMips, ETextureCreateFlags Flags, ERHIAccess InResourceState, void** InitialMipData, uint32 NumInitialMips)
+{
+	UE_LOG(LogAGX, Fatal, TEXT("RHIAsyncCreateTexture2D is not supported"));
+	return FTexture2DRHIRef();
+}
+
+void FAGXDynamicRHI::RHICopySharedMips(FRHITexture2D* DestTexture2D, FRHITexture2D* SrcTexture2D)
+{
+	NOT_SUPPORTED("RHICopySharedMips");
+}
+
+FTexture2DArrayRHIRef FAGXDynamicRHI::RHICreateTexture2DArray(uint32 SizeX, uint32 SizeY, uint32 SizeZ, uint8 Format, uint32 NumMips, uint32 NumSamples, ETextureCreateFlags Flags, ERHIAccess InResourceState, FRHIResourceCreateInfo& CreateInfo)
+{
+	@autoreleasepool {
+		return new FAGXTexture2DArray((EPixelFormat)Format, SizeX, SizeY, SizeZ, NumMips, Flags, CreateInfo.BulkData, CreateInfo.ClearValueBinding);
+	}
+}
+
+FTexture3DRHIRef FAGXDynamicRHI::RHICreateTexture3D(uint32 SizeX, uint32 SizeY, uint32 SizeZ, uint8 Format, uint32 NumMips, ETextureCreateFlags Flags, ERHIAccess InResourceState, FRHIResourceCreateInfo& CreateInfo)
+{
+	@autoreleasepool {
+		return new FAGXTexture3D((EPixelFormat)Format, SizeX, SizeY, SizeZ, NumMips, Flags, CreateInfo.BulkData, CreateInfo.ClearValueBinding);
+	}
+}
+
+void FAGXDynamicRHI::RHIGetResourceInfo(FRHITexture* Ref, FRHIResourceInfo& OutInfo)
+{
+	// @todo Needed for visualisation!!
+	// NOT_SUPPORTED("RHIGetResourceInfo");
+}
+
+void FAGXDynamicRHI::RHIGenerateMips(FRHITexture* SourceSurfaceRHI)
+{
+	@autoreleasepool {
+		FAGXSurface* Surf = AGXGetMetalSurfaceFromRHITexture(SourceSurfaceRHI);
+		if (Surf && Surf->Texture)
+		{
+			ImmediateContext.GetInternalContext().AsyncGenerateMipmapsForTexture(Surf->Texture);
+		}
+	}
+}
+
+// deferred unlock commands
+struct FAGXRHICommandAsyncReallocateTexture2D final : public FRHICommand<FAGXRHICommandAsyncReallocateTexture2D>
+{
+	FAGXContext& Context;
+	FAGXTexture2D* OldTexture;
+	FAGXTexture2D* NewTexture;
+	int32 NewMipCount;
+	int32 NewSizeX;
+	int32 NewSizeY;
+	FThreadSafeCounter* RequestStatus;
+	
+	FORCEINLINE_DEBUGGABLE FAGXRHICommandAsyncReallocateTexture2D(FAGXContext& InContext, FAGXTexture2D* InOldTexture, FAGXTexture2D* InNewTexture, int32 InNewMipCount, int32 InNewSizeX, int32 InNewSizeY, FThreadSafeCounter* InRequestStatus)
+	: Context(InContext)
+	, OldTexture(InOldTexture)
+	, NewTexture(InNewTexture)
+	, NewMipCount(InNewMipCount)
+	, NewSizeX(InNewSizeX)
+	, NewSizeY(InNewSizeY)
+	, RequestStatus(InRequestStatus)
+	{
+	}
+	
+	void Execute(FRHICommandListBase& CmdList)
+	{
+		CopyMips(Context, OldTexture, NewTexture, NewMipCount, NewSizeX, NewSizeY, RequestStatus);
+	}
+	
+	static void CopyMips(FAGXContext& Context, FAGXTexture2D* OldTexture, FAGXTexture2D* NewTexture, int32 NewMipCount, int32 NewSizeX, int32 NewSizeY, FThreadSafeCounter* RequestStatus)
+	{
+		// figure out what mips to schedule
+		const uint32 NumSharedMips = FMath::Min(OldTexture->GetNumMips(), NewTexture->GetNumMips());
+		const uint32 SourceMipOffset = OldTexture->GetNumMips() - NumSharedMips;
+		const uint32 DestMipOffset = NewTexture->GetNumMips() - NumSharedMips;
+		
+		const uint32 BlockSizeX = GPixelFormats[OldTexture->GetFormat()].BlockSizeX;
+		const uint32 BlockSizeY = GPixelFormats[OldTexture->GetFormat()].BlockSizeY;
+		
+		// only handling straight 2D textures here
+		uint32 SliceIndex = 0;
+		mtlpp::Origin Origin(0,0,0);
+		
+		FAGXTexture Tex = OldTexture->Surface.Texture;
+		
+		// DXT/BC formats on Mac actually do have mip-tails that are smaller than the block size, they end up being uncompressed.
+		bool const bPixelFormatASTC = IsPixelFormatASTCCompressed(OldTexture->GetFormat());
+		
+		bool bAsync = true;
+		for (uint32 MipIndex = 0; MipIndex < NumSharedMips; ++MipIndex)
+		{
+			const uint32 UnalignedMipSizeX = FMath::Max<uint32>(1, NewSizeX >> (MipIndex + DestMipOffset));
+			const uint32 UnalignedMipSizeY = FMath::Max<uint32>(1, NewSizeY >> (MipIndex + DestMipOffset));
+			const uint32 MipSizeX = FMath::Max<uint32>(1, NewSizeX >> (MipIndex + DestMipOffset));
+			const uint32 MipSizeY = FMath::Max<uint32>(1, NewSizeY >> (MipIndex + DestMipOffset));
+			
+			bAsync &= Context.AsyncCopyFromTextureToTexture(OldTexture->Surface.Texture, SliceIndex, MipIndex + SourceMipOffset, Origin, mtlpp::Size(MipSizeX, MipSizeY, 1), NewTexture->Surface.Texture, SliceIndex, MipIndex + DestMipOffset, Origin);
+		}
+		
+		// when done, decrement the counter to indicate it's safe
+		mtlpp::CommandBufferHandler CompletionHandler = [Tex](mtlpp::CommandBuffer const&)
+		{
+		};
+		
+		if(bAsync)
+		{
+	    	// kck it off!
+			Context.SubmitAsyncCommands(nil, CompletionHandler, false);
+		}
+		else
+		{
+			Context.GetCurrentRenderPass().AddCompletionHandler(CompletionHandler);
+		}
+		
+		// Like D3D mark this as complete immediately.
+		RequestStatus->Decrement();
+		
+		FAGXSurface* Source = AGXGetMetalSurfaceFromRHITexture(OldTexture);
+		Source->MakeAliasable();
+	}
+};
+
+FTexture2DRHIRef FAGXDynamicRHI::AsyncReallocateTexture2D_RenderThread(class FRHICommandListImmediate& RHICmdList, FRHITexture2D* Texture2D, int32 NewMipCount, int32 NewSizeX, int32 NewSizeY, FThreadSafeCounter* RequestStatus)
+{
+	@autoreleasepool {
+		FTexture2DRHIRef Result;
+		
+		if (RHICmdList.Bypass() || !IsRunningRHIInSeparateThread())
+		{
+			Result = GDynamicRHI->RHIAsyncReallocateTexture2D(Texture2D, NewMipCount, NewSizeX, NewSizeY, RequestStatus);
+		}
+		else
+		{
+			FAGXTexture2D* OldTexture = ResourceCast(Texture2D);
+			
+			FAGXTexture2D* NewTexture = new FAGXTexture2D(OldTexture->GetFormat(), NewSizeX, NewSizeY, NewMipCount, OldTexture->GetNumSamples(), OldTexture->GetFlags(), NULL, Texture2D->GetClearBinding());
+			
+			Result = NewTexture;
+			
+			new (RHICmdList.AllocCommand<FAGXRHICommandAsyncReallocateTexture2D>()) FAGXRHICommandAsyncReallocateTexture2D(ImmediateContext.GetInternalContext(), OldTexture, NewTexture, NewMipCount, NewSizeX, NewSizeY, RequestStatus);
+		}
+		
+		return Result;
+	}
+}
+
+ETextureReallocationStatus FAGXDynamicRHI::FinalizeAsyncReallocateTexture2D_RenderThread(class FRHICommandListImmediate& RHICmdList, FRHITexture2D* Texture2D, bool bBlockUntilCompleted)
+{
+	// No need to flush - does nothing
+	return GDynamicRHI->RHIFinalizeAsyncReallocateTexture2D(Texture2D, bBlockUntilCompleted);
+}
+
+ETextureReallocationStatus FAGXDynamicRHI::CancelAsyncReallocateTexture2D_RenderThread(class FRHICommandListImmediate& RHICmdList, FRHITexture2D* Texture2D, bool bBlockUntilCompleted)
+{
+	// No need to flush - does nothing
+	return GDynamicRHI->RHICancelAsyncReallocateTexture2D(Texture2D, bBlockUntilCompleted);
+}
+
+
+FTexture2DRHIRef FAGXDynamicRHI::RHIAsyncReallocateTexture2D(FRHITexture2D* OldTextureRHI, int32 NewMipCount, int32 NewSizeX, int32 NewSizeY, FThreadSafeCounter* RequestStatus)
+{
+	@autoreleasepool {
+		FAGXTexture2D* OldTexture = ResourceCast(OldTextureRHI);
+		
+		FAGXTexture2D* NewTexture = new FAGXTexture2D(OldTexture->GetFormat(), NewSizeX, NewSizeY, NewMipCount, OldTexture->GetNumSamples(), OldTexture->GetFlags(), NULL, OldTextureRHI->GetClearBinding());
+		
+		FAGXRHICommandAsyncReallocateTexture2D::CopyMips(ImmediateContext.GetInternalContext(), OldTexture, NewTexture, NewMipCount, NewSizeX, NewSizeY, RequestStatus);
+		
+		return NewTexture;
+	}
+}
+
+ETextureReallocationStatus FAGXDynamicRHI::RHIFinalizeAsyncReallocateTexture2D(FRHITexture2D* Texture2D, bool bBlockUntilCompleted )
+{
+	return TexRealloc_Succeeded;
+}
+
+ETextureReallocationStatus FAGXDynamicRHI::RHICancelAsyncReallocateTexture2D(FRHITexture2D* Texture2D, bool bBlockUntilCompleted )
+{
+	return TexRealloc_Failed;
+}
+
+void* FAGXDynamicRHI::LockTexture2D_RenderThread(class FRHICommandListImmediate& RHICmdList, FRHITexture2D* Texture, uint32 MipIndex, EResourceLockMode LockMode, uint32& DestStride, bool bLockWithinMiptail, bool bNeedsDefaultRHIFlush)
+{
+	@autoreleasepool {
+		check(IsInRenderingThread());
+		
+		FAGXTexture2D* TextureMTL = ResourceCast(Texture);
+		
+		void* BufferData = TextureMTL->Surface.AsyncLock(RHICmdList, MipIndex, 0, LockMode, DestStride, bNeedsDefaultRHIFlush);
+		
+		return BufferData;
+	}
+}
+
+void FAGXDynamicRHI::UnlockTexture2D_RenderThread(class FRHICommandListImmediate& RHICmdList, FRHITexture2D* Texture, uint32 MipIndex, bool bLockWithinMiptail, bool bNeedsDefaultRHIFlush)
+{
+	@autoreleasepool
+	{
+		check(IsInRenderingThread());
+		
+		FAGXTexture2D* TextureMTL = ResourceCast(Texture);
+		TextureMTL->Surface.Unlock(MipIndex, 0, true);
+	}
+}
+
+
+void* FAGXDynamicRHI::RHILockTexture2D(FRHITexture2D* TextureRHI,uint32 MipIndex,EResourceLockMode LockMode,uint32& DestStride,bool bLockWithinMiptail)
+{
+	@autoreleasepool {
+		FAGXTexture2D* Texture = ResourceCast(TextureRHI);
+		return Texture->Surface.Lock(MipIndex, 0, LockMode, DestStride);
+	}
+}
+
+void FAGXDynamicRHI::RHIUnlockTexture2D(FRHITexture2D* TextureRHI,uint32 MipIndex,bool bLockWithinMiptail)
+{
+	@autoreleasepool {
+		FAGXTexture2D* Texture = ResourceCast(TextureRHI);
+		Texture->Surface.Unlock(MipIndex, 0, false);
+	}
+}
+
+void* FAGXDynamicRHI::RHILockTexture2DArray(FRHITexture2DArray* TextureRHI,uint32 TextureIndex,uint32 MipIndex,EResourceLockMode LockMode,uint32& DestStride,bool bLockWithinMiptail)
+{
+	@autoreleasepool {
+		FAGXTexture2DArray* Texture = ResourceCast(TextureRHI);
+		return Texture->Surface.Lock(MipIndex, TextureIndex, LockMode, DestStride);
+	}
+}
+
+void FAGXDynamicRHI::RHIUnlockTexture2DArray(FRHITexture2DArray* TextureRHI,uint32 TextureIndex,uint32 MipIndex,bool bLockWithinMiptail)
+{
+	@autoreleasepool {
+		FAGXTexture2DArray* Texture = ResourceCast(TextureRHI);
+		Texture->Surface.Unlock(MipIndex, TextureIndex, false);
+	}
+}
+
+#if PLATFORM_MAC
+static void InternalExpandR8ToStandardRGBA(uint32* pDest, const struct FUpdateTextureRegion2D& UpdateRegion, uint32& InOutSourcePitch, const uint8* pSrc)
+{
+	const uint32 ExpandedPitch = UpdateRegion.Width * sizeof(uint32);
+	
+	for(uint y = 0; y < UpdateRegion.Height; y++)
+	{
+		for(uint x = 0; x < UpdateRegion.Width; x++)
+		{
+			uint8 Value = pSrc[(y * InOutSourcePitch) + x];
+			*(pDest++) = (Value | (Value << 8) | (Value << 16) | (Value << 24));
+		}
+	}
+	
+	InOutSourcePitch = ExpandedPitch;
+}
+#endif
+
+static FAGXBuffer InternalCopyTexture2DUpdateRegion(FRHITexture2D* TextureRHI, const struct FUpdateTextureRegion2D& UpdateRegion, uint32& InOutSourcePitch, const uint8* SourceData)
+{	
+	const uint32 InSourcePitch = InOutSourcePitch;
+
+	FAGXBuffer Buffer;
+
+	FAGXTexture2D* Texture = ResourceCast(TextureRHI);	
+
+#if PLATFORM_MAC
+	// Expand R8_sRGB into RGBA8_sRGB for Mac.
+	if(Texture->GetFormat() == PF_G8 && (Texture->GetFlags() & TexCreate_SRGB))
+	{
+		const uint32 BufferSize = UpdateRegion.Height * UpdateRegion.Width * sizeof(uint32);		
+		Buffer = GetAGXDeviceContext().CreatePooledBuffer(FAGXPooledBufferArgs(GetAGXDeviceContext().GetDevice(), BufferSize, BUF_Dynamic, mtlpp::StorageMode::Shared));
+		InternalExpandR8ToStandardRGBA((uint32*)Buffer.GetContents(), UpdateRegion, InOutSourcePitch, SourceData);
+	}
+
+	if(Buffer.GetPtr() == nil)
+#endif
+	{
+		const uint32 BufferSize = UpdateRegion.Height * InSourcePitch;
+		Buffer = GetAGXDeviceContext().CreatePooledBuffer(FAGXPooledBufferArgs(GetAGXDeviceContext().GetDevice(), BufferSize, BUF_Dynamic, mtlpp::StorageMode::Shared));
+
+		// Limit copy to line by line by update region pitch otherwise we can go off the end of source data on the last row
+		uint8* pDestRow = (uint8*)Buffer.GetContents();
+		uint8* pSourceRow = (uint8*)SourceData;
+		uint32 CopyPitch = FMath::DivideAndRoundUp(UpdateRegion.Width, (uint32)GPixelFormats[Texture->GetFormat()].BlockSizeX) * GPixelFormats[Texture->GetFormat()].BlockBytes;
+		check(CopyPitch <= InSourcePitch);
+		
+		for (uint32 i = 0;i < UpdateRegion.Height;++i)
+		{
+			FMemory::Memcpy(pDestRow, pSourceRow, CopyPitch);
+			pSourceRow += InSourcePitch;
+			pDestRow += InSourcePitch;
+		}
+	}
+
+	return Buffer;
+}
+
+static void InternalUpdateTexture2D(FAGXContext& Context, FRHITexture2D* TextureRHI, uint32 MipIndex, FUpdateTextureRegion2D const& UpdateRegion, uint32 SourcePitch, FAGXBuffer Buffer)
+{
+	FAGXTexture2D* Texture = ResourceCast(TextureRHI);
+	FAGXTexture Tex = Texture->Surface.Texture;
+	
+	mtlpp::Region Region = mtlpp::Region(UpdateRegion.DestX, UpdateRegion.DestY, UpdateRegion.Width, UpdateRegion.Height);
+	
+	if(Tex.GetStorageMode() == mtlpp::StorageMode::Private)
+	{
+		SCOPED_AUTORELEASE_POOL;
+		
+		uint32 BytesPerImage = SourcePitch * UpdateRegion.Height;
+		
+		mtlpp::BlitOption Options = mtlpp::BlitOption::None;
+#if !PLATFORM_MAC
+		if (Tex.GetPixelFormat() >= mtlpp::PixelFormat::PVRTC_RGB_2BPP && Tex.GetPixelFormat() <= mtlpp::PixelFormat::PVRTC_RGBA_4BPP_sRGB)
+		{
+			Options = mtlpp::BlitOption::RowLinearPVRTC;
+		}
+#endif
+		if(Context.AsyncCopyFromBufferToTexture(Buffer, 0, SourcePitch, BytesPerImage, Region.size, Tex, 0, MipIndex, Region.origin, Options))
+		{
+			Context.SubmitAsyncCommands(nil, nil, false);
+		}
+	}
+	else
+	{
+		MTLPP_VALIDATE(mtlpp::Texture, Tex, AGXSafeGetRuntimeDebuggingLevel() >= EAGXDebugLevelValidation, Replace(Region, MipIndex, 0, MTLPP_VALIDATE(mtlpp::Buffer, Buffer, AGXSafeGetRuntimeDebuggingLevel() >= EAGXDebugLevelValidation, GetContents()), SourcePitch, 0));
+	}
+
+	FPlatformAtomics::InterlockedExchange(&Texture->Surface.Written, 1);
+}
+
+struct FAGXRHICommandUpdateTexture2D final : public FRHICommand<FAGXRHICommandUpdateTexture2D>
+{
+	FAGXContext& Context;
+	FRHITexture2D* Texture;
+	uint32 MipIndex;
+	FUpdateTextureRegion2D UpdateRegion;
+	uint32 SourcePitch;
+	FAGXBuffer SourceBuffer;
+
+	FORCEINLINE_DEBUGGABLE FAGXRHICommandUpdateTexture2D(FAGXContext& InContext, FRHITexture2D* InTexture, uint32 InMipIndex, FUpdateTextureRegion2D InUpdateRegion, uint32 InSourcePitch, FAGXBuffer InSourceBuffer)
+	: Context(InContext)
+	, Texture(InTexture)
+	, MipIndex(InMipIndex)
+	, UpdateRegion(InUpdateRegion)
+	, SourcePitch(InSourcePitch)
+	, SourceBuffer(InSourceBuffer)
+	{
+	}
+	
+	void Execute(FRHICommandListBase& CmdList)
+	{		
+		InternalUpdateTexture2D(Context, Texture, MipIndex, UpdateRegion, SourcePitch, SourceBuffer);
+		GetAGXDeviceContext().ReleaseBuffer(SourceBuffer);
+		INC_DWORD_STAT_BY(STAT_AGXTextureMemUpdate, UpdateRegion.Height * SourcePitch);
+	}
+};
+
+void FAGXDynamicRHI::UpdateTexture2D_RenderThread(class FRHICommandListImmediate& RHICmdList, FRHITexture2D* Texture, uint32 MipIndex, const struct FUpdateTextureRegion2D& UpdateRegion, uint32 SourcePitch, const uint8* SourceData)
+{
+	@autoreleasepool {
+		if (RHICmdList.Bypass() || !IsRunningRHIInSeparateThread())
+		{
+			this->RHIUpdateTexture2D(Texture, MipIndex, UpdateRegion, SourcePitch, SourceData);
+		}
+		else
+		{
+			FAGXBuffer Buffer = InternalCopyTexture2DUpdateRegion(Texture, UpdateRegion, SourcePitch, SourceData);
+			new (RHICmdList.AllocCommand<FAGXRHICommandUpdateTexture2D>()) FAGXRHICommandUpdateTexture2D(ImmediateContext.GetInternalContext(), Texture, MipIndex, UpdateRegion, SourcePitch, Buffer);
+		}
+	}
+}
+
+void FAGXDynamicRHI::RHIUpdateTexture2D(FRHITexture2D* TextureRHI, uint32 MipIndex, const struct FUpdateTextureRegion2D& UpdateRegion, uint32 SourcePitch, const uint8* SourceData)
+{
+	@autoreleasepool 
+	{
+		FAGXTexture2D* Texture = ResourceCast(TextureRHI);
+		FAGXTexture Tex = Texture->Surface.Texture;
+		bool const bUseIntermediateMetalBuffer = Tex.GetStorageMode() == mtlpp::StorageMode::Private;
+		
+		if(bUseIntermediateMetalBuffer)
+		{
+			FAGXBuffer Buffer = InternalCopyTexture2DUpdateRegion(TextureRHI, UpdateRegion, SourcePitch, SourceData);
+			InternalUpdateTexture2D(ImmediateContext.GetInternalContext(), TextureRHI, MipIndex, UpdateRegion, SourcePitch, Buffer);
+			GetAGXDeviceContext().ReleaseBuffer(Buffer);
+		}
+		else
+		{
+#if PLATFORM_MAC
+			TArray<uint32> ExpandedData;
+			if(Texture->GetFormat() == PF_G8 && (Texture->GetFlags() & TexCreate_SRGB))
+			{
+				ExpandedData.AddZeroed(UpdateRegion.Height * UpdateRegion.Width);
+				InternalExpandR8ToStandardRGBA((uint32*)ExpandedData.GetData(), UpdateRegion, SourcePitch, SourceData);
+				SourceData = (uint8*)ExpandedData.GetData();
+			}
+#endif
+			mtlpp::Region Region = mtlpp::Region(UpdateRegion.DestX, UpdateRegion.DestY, UpdateRegion.Width, UpdateRegion.Height);
+			MTLPP_VALIDATE(mtlpp::Texture, Tex, AGXSafeGetRuntimeDebuggingLevel() >= EAGXDebugLevelValidation, Replace(Region, MipIndex, 0, SourceData, SourcePitch, 0));
+
+			FPlatformAtomics::InterlockedExchange(&Texture->Surface.Written, 1);
+		}
+		
+		INC_DWORD_STAT_BY(STAT_AGXTextureMemUpdate, UpdateRegion.Height*SourcePitch);
+	}
+}
+
+static void InternalCopyTexture3DUpdateRegionData(FRHITexture3D* TextureRHI, const struct FUpdateTextureRegion3D& UpdateRegion, uint32 SourceRowPitch, uint32 SourceDepthPitch, const uint8* SourceData, uint8* DestData)
+{
+	// Perform safe line copy
+	FAGXTexture3D* Texture = ResourceCast(TextureRHI);	
+	
+	uint8* pDest = DestData;
+	uint32 CopyPitch = FMath::DivideAndRoundUp(UpdateRegion.Width, (uint32)GPixelFormats[Texture->GetFormat()].BlockSizeX) * GPixelFormats[Texture->GetFormat()].BlockBytes;
+
+	for (uint32 i = 0; i < UpdateRegion.Depth;++i)
+	{
+		const uint8* pSourceRowData = SourceData + (SourceDepthPitch * i);
+		uint8* pDestRowData = pDest + (SourceDepthPitch * i);
+
+		for (uint32 j = 0;j < UpdateRegion.Height;++j)
+		{
+			FMemory::Memcpy(pDestRowData, pSourceRowData, CopyPitch);
+			pSourceRowData += SourceRowPitch;
+			pDestRowData += SourceRowPitch;
+		}
+	}
+}
+
+
+static void InternalUpdateTexture3D(FAGXContext& Context, FRHITexture3D* TextureRHI, uint32 MipIndex, const FUpdateTextureRegion3D& UpdateRegion, uint32 SourceRowPitch, uint32 SourceDepthPitch, FAGXBuffer Buffer)
+{
+	FAGXTexture3D* Texture = ResourceCast(TextureRHI);
+	FAGXTexture Tex = Texture->Surface.Texture;
+	
+	mtlpp::Region Region = mtlpp::Region(UpdateRegion.DestX, UpdateRegion.DestY, UpdateRegion.DestZ, UpdateRegion.Width, UpdateRegion.Height, UpdateRegion.Depth);
+	
+	if(Tex.GetStorageMode() == mtlpp::StorageMode::Private)
+	{
+		const uint32 BytesPerImage = SourceRowPitch * UpdateRegion.Height;
+		mtlpp::BlitOption Options = mtlpp::BlitOption::None;
+#if !PLATFORM_MAC
+		if (Tex.GetPixelFormat() >= mtlpp::PixelFormat::PVRTC_RGB_2BPP && Tex.GetPixelFormat() <= mtlpp::PixelFormat::PVRTC_RGBA_4BPP_sRGB)
+		{
+			Options = mtlpp::BlitOption::RowLinearPVRTC;
+		}
+#endif
+		if(Context.AsyncCopyFromBufferToTexture(Buffer, 0, SourceRowPitch, BytesPerImage, Region.size, Tex, 0, MipIndex, Region.origin, Options))
+		{
+			Context.SubmitAsyncCommands(nil, nil, false);
+		}
+	}
+	else
+	{
+		MTLPP_VALIDATE(mtlpp::Texture, Tex, AGXSafeGetRuntimeDebuggingLevel() >= EAGXDebugLevelValidation, Replace(Region, MipIndex, 0, (uint8*)Buffer.GetContents(), SourceRowPitch, SourceDepthPitch));
+	}
+
+	FPlatformAtomics::InterlockedExchange(&Texture->Surface.Written, 1);
+}
+
+struct FAGXDynamicRHIUpdateTexture3DCommand final : public FRHICommand<FAGXDynamicRHIUpdateTexture3DCommand>
+{
+	FAGXContext& Context;
+	FRHITexture3D* DestinationTexture;
+	uint32 MipIndex;
+	FUpdateTextureRegion3D UpdateRegion;
+	uint32 SourceRowPitch;
+	uint32 SourceDepthPitch;
+	FAGXBuffer Buffer;
+	
+	FORCEINLINE_DEBUGGABLE FAGXDynamicRHIUpdateTexture3DCommand(FAGXContext& InContext, FRHITexture3D* TextureRHI, uint32 InMipIndex, const struct FUpdateTextureRegion3D& InUpdateRegion, uint32 InSourceRowPitch, uint32 InSourceDepthPitch, const uint8* SourceData)
+	: Context(InContext)
+	, DestinationTexture(TextureRHI)
+	, MipIndex(InMipIndex)
+	, UpdateRegion(InUpdateRegion)
+	, SourceRowPitch(InSourceRowPitch)
+	, SourceDepthPitch(InSourceDepthPitch)
+	{
+		FAGXTexture3D* Texture = ResourceCast(TextureRHI);
+		FAGXTexture Tex = Texture->Surface.Texture;
+		const uint32 BufferSize = UpdateRegion.Height * UpdateRegion.Depth* SourceRowPitch;
+		
+		Buffer = GetAGXDeviceContext().CreatePooledBuffer(FAGXPooledBufferArgs(GetAGXDeviceContext().GetDevice(), BufferSize, BUF_Dynamic, mtlpp::StorageMode::Shared));
+		InternalCopyTexture3DUpdateRegionData(Texture, UpdateRegion, SourceRowPitch, SourceDepthPitch, SourceData, (uint8*)Buffer.GetContents());
+	}
+	
+	void Execute(FRHICommandListBase& CmdList)
+	{
+		InternalUpdateTexture3D(Context, DestinationTexture, MipIndex, UpdateRegion, SourceRowPitch, SourceDepthPitch, Buffer);
+		GetAGXDeviceContext().ReleaseBuffer(Buffer);
+		INC_DWORD_STAT_BY(STAT_AGXTextureMemUpdate, UpdateRegion.Height * UpdateRegion.Width * SourceDepthPitch);
+	}
+};
+
+void FAGXDynamicRHI::UpdateTexture3D_RenderThread(class FRHICommandListImmediate& RHICmdList, FRHITexture3D* Texture, uint32 MipIndex, const struct FUpdateTextureRegion3D& UpdateRegion, uint32 SourceRowPitch, uint32 SourceDepthPitch, const uint8* SourceData)
+{
+	if (RHICmdList.Bypass() || !IsRunningRHIInSeparateThread())
+	{
+		RHIUpdateTexture3D(Texture, MipIndex, UpdateRegion, SourceRowPitch, SourceDepthPitch, SourceData);
+	}
+	else
+	{
+		new (RHICmdList.AllocCommand<FAGXDynamicRHIUpdateTexture3DCommand>()) FAGXDynamicRHIUpdateTexture3DCommand(ImmediateContext.GetInternalContext(), Texture, MipIndex, UpdateRegion, SourceRowPitch, SourceDepthPitch, SourceData);
+	}
+}
+
+FUpdateTexture3DData FAGXDynamicRHI::BeginUpdateTexture3D_RenderThread(class FRHICommandListImmediate& RHICmdList, FRHITexture3D* Texture, uint32 MipIndex, const struct FUpdateTextureRegion3D& UpdateRegion)
+{
+	check(IsInRenderingThread());
+	
+	const int32 FormatSize = PixelFormatBlockBytes[Texture->GetFormat()];
+	const int32 RowPitch = UpdateRegion.Width * FormatSize;
+	const int32 DepthPitch = UpdateRegion.Width * UpdateRegion.Height * FormatSize;
+	
+	SIZE_T MemorySize = DepthPitch * UpdateRegion.Depth;
+	uint8* Data = (uint8*)FMemory::Malloc(MemorySize);
+	
+	return FUpdateTexture3DData(Texture, MipIndex, UpdateRegion, RowPitch, DepthPitch, Data, MemorySize, GFrameNumberRenderThread);
+}
+
+void FAGXDynamicRHI::EndUpdateTexture3D_RenderThread(class FRHICommandListImmediate& RHICmdList, FUpdateTexture3DData& UpdateData)
+{
+	check(IsInRenderingThread());
+	check(GFrameNumberRenderThread == UpdateData.FrameNumber);
+	
+	if (RHICmdList.Bypass() || !IsRunningRHIInSeparateThread())
+	{
+		GDynamicRHI->RHIUpdateTexture3D(UpdateData.Texture, UpdateData.MipIndex, UpdateData.UpdateRegion, UpdateData.RowPitch, UpdateData.DepthPitch, UpdateData.Data);
+	}
+	else
+	{
+		new (RHICmdList.AllocCommand<FAGXDynamicRHIUpdateTexture3DCommand>()) FAGXDynamicRHIUpdateTexture3DCommand(ImmediateContext.GetInternalContext(), UpdateData.Texture, UpdateData.MipIndex, UpdateData.UpdateRegion, UpdateData.RowPitch, UpdateData.DepthPitch, UpdateData.Data);
+	}
+	
+	FMemory::Free(UpdateData.Data);
+	UpdateData.Data = nullptr;
+}
+
+void FAGXDynamicRHI::RHIUpdateTexture3D(FRHITexture3D* TextureRHI,uint32 MipIndex,const FUpdateTextureRegion3D& UpdateRegion,uint32 SourceRowPitch,uint32 SourceDepthPitch, const uint8* SourceData)
+{
+	@autoreleasepool {
+	
+		FAGXTexture3D* Texture = ResourceCast(TextureRHI);
+		FAGXTexture Tex = Texture->Surface.Texture;
+		
+#if PLATFORM_MAC
+		checkf(!(Texture->GetFormat() == PF_G8 && (Texture->GetFlags() & TexCreate_SRGB)), TEXT("AGXRHI does not support PF_G8_sRGB on 3D, array or cube textures as it requires manual, CPU-side expansion to RGBA8_sRGB which is expensive!"));
+#endif
+		if(Tex.GetStorageMode() == mtlpp::StorageMode::Private)
+		{
+			SCOPED_AUTORELEASE_POOL;
+
+			const uint32 BufferSize = UpdateRegion.Height * UpdateRegion.Depth * SourceRowPitch;
+			FAGXBuffer IntermediateBuffer = GetAGXDeviceContext().CreatePooledBuffer(FAGXPooledBufferArgs(GetAGXDeviceContext().GetDevice(), BufferSize, BUF_Dynamic, mtlpp::StorageMode::Shared));
+			InternalCopyTexture3DUpdateRegionData(TextureRHI, UpdateRegion, SourceRowPitch, SourceDepthPitch, SourceData, (uint8*)IntermediateBuffer.GetContents());
+			InternalUpdateTexture3D(ImmediateContext.GetInternalContext(), TextureRHI, MipIndex, UpdateRegion, SourceRowPitch, SourceDepthPitch, IntermediateBuffer);
+			GetAGXDeviceContext().ReleaseBuffer(IntermediateBuffer);
+		}
+		else
+		{
+			mtlpp::Region Region = mtlpp::Region(UpdateRegion.DestX, UpdateRegion.DestY, UpdateRegion.DestZ, UpdateRegion.Width, UpdateRegion.Height, UpdateRegion.Depth);
+			MTLPP_VALIDATE(mtlpp::Texture, Tex, AGXSafeGetRuntimeDebuggingLevel() >= EAGXDebugLevelValidation, Replace(Region, MipIndex, 0, SourceData, SourceRowPitch, SourceDepthPitch));
+			FPlatformAtomics::InterlockedExchange(&Texture->Surface.Written, 1);
+		}
+		
+		INC_DWORD_STAT_BY(STAT_AGXTextureMemUpdate, UpdateRegion.Height * UpdateRegion.Width * SourceDepthPitch);
+	}
+}
+
+/*-----------------------------------------------------------------------------
+ Cubemap texture support.
+ -----------------------------------------------------------------------------*/
+FTextureCubeRHIRef FAGXDynamicRHI::RHICreateTextureCube(uint32 Size, uint8 Format, uint32 NumMips, ETextureCreateFlags Flags, ERHIAccess InResourceState, FRHIResourceCreateInfo& CreateInfo)
+{
+	@autoreleasepool {
+		return new FAGXTextureCube((EPixelFormat)Format, Size, false, 1, NumMips, Flags, CreateInfo.BulkData, CreateInfo.ClearValueBinding);
+	}
+}
+
+FTextureCubeRHIRef FAGXDynamicRHI::RHICreateTextureCubeArray(uint32 Size, uint32 ArraySize, uint8 Format, uint32 NumMips, ETextureCreateFlags Flags, ERHIAccess InResourceState, FRHIResourceCreateInfo& CreateInfo)
+{
+	@autoreleasepool {
+		return new FAGXTextureCube((EPixelFormat)Format, Size, true, ArraySize, NumMips, Flags, CreateInfo.BulkData, CreateInfo.ClearValueBinding);
+	}
+}
+
+void* FAGXDynamicRHI::RHILockTextureCubeFace(FRHITextureCube* TextureCubeRHI,uint32 FaceIndex,uint32 ArrayIndex,uint32 MipIndex,EResourceLockMode LockMode,uint32& DestStride,bool bLockWithinMiptail)
+{
+	@autoreleasepool {
+		FAGXTextureCube* TextureCube = ResourceCast(TextureCubeRHI);
+		uint32 MetalFace = GetMetalCubeFace((ECubeFace)FaceIndex);
+		return TextureCube->Surface.Lock(MipIndex, MetalFace + (6 * ArrayIndex), LockMode, DestStride, true);
+	}
+}
+
+void FAGXDynamicRHI::RHIUnlockTextureCubeFace(FRHITextureCube* TextureCubeRHI,uint32 FaceIndex,uint32 ArrayIndex,uint32 MipIndex,bool bLockWithinMiptail)
+{
+	@autoreleasepool {
+		FAGXTextureCube* TextureCube = ResourceCast(TextureCubeRHI);
+		uint32 MetalFace = GetMetalCubeFace((ECubeFace)FaceIndex);
+		TextureCube->Surface.Unlock(MipIndex, MetalFace + (ArrayIndex * 6), false);
+	}
+}
+
+
+FTexture2DRHIRef FAGXDynamicRHI::RHICreateTexture2D_RenderThread(class FRHICommandListImmediate& RHICmdList, uint32 SizeX, uint32 SizeY, uint8 Format, uint32 NumMips, uint32 NumSamples, ETextureCreateFlags Flags, ERHIAccess InResourceState, FRHIResourceCreateInfo& CreateInfo)
+{
+	@autoreleasepool {
+		FRHIResourceCreateInfo Info = CreateInfo;
+		bool const bIOSurfaceData = (CreateInfo.BulkData && CreateInfo.BulkData->GetResourceType() != FResourceBulkDataInterface::EBulkDataType::Default);
+		Info.BulkData = bIOSurfaceData ? CreateInfo.BulkData : nullptr;
+		FTexture2DRHIRef Result = GDynamicRHI->RHICreateTexture2D(SizeX, SizeY, Format, NumMips, NumSamples, Flags, InResourceState, Info);
+		if (!bIOSurfaceData && CreateInfo.BulkData)
+		{
+			// upload existing bulkdata
+			UE_LOG(LogAGX, Display, TEXT("Got a bulk data texture, with %d mips"), NumMips);
+			checkf(NumMips == 1, TEXT("Only handling bulk data with 1 mip and 1 array length"));
+			uint32 Stride;
+			
+			// lock, copy, unlock
+			void* LockedData = LockTexture2D_RenderThread(RHICmdList, Result.GetReference(), 0, RLM_WriteOnly, Stride, false);
+			check(LockedData);
+			FMemory::Memcpy(LockedData, CreateInfo.BulkData->GetResourceBulkData(), CreateInfo.BulkData->GetResourceBulkDataSize());
+			UnlockTexture2D_RenderThread(RHICmdList, Result.GetReference(), 0, false);
+			
+			// bulk data can be unloaded now
+			CreateInfo.BulkData->Discard();
+		}
+		return Result;
+	}
+}
+
+FTexture2DArrayRHIRef FAGXDynamicRHI::RHICreateTexture2DArray_RenderThread(class FRHICommandListImmediate& RHICmdList, uint32 SizeX, uint32 SizeY, uint32 SizeZ, uint8 Format, uint32 NumMips, uint32 NumSamples, ETextureCreateFlags Flags, ERHIAccess InResourceState, FRHIResourceCreateInfo& CreateInfo)
+{
+	@autoreleasepool {
+		FRHIResourceCreateInfo Info = CreateInfo;
+		bool const bIOSurfaceData = (CreateInfo.BulkData && CreateInfo.BulkData->GetResourceType() != FResourceBulkDataInterface::EBulkDataType::Default);
+		Info.BulkData = bIOSurfaceData ? CreateInfo.BulkData : nullptr;
+		FTexture2DArrayRHIRef Result = GDynamicRHI->RHICreateTexture2DArray(SizeX, SizeY, SizeZ, Format, NumMips, NumSamples, Flags, InResourceState, Info);
+		if (CreateInfo.BulkData)
+		{
+			// upload existing bulkdata
+			UE_LOG(LogAGX, Display, TEXT("Got a bulk data texture, with %d mips"), NumMips);
+			checkf(NumMips == 1 && SizeZ == 1, TEXT("Only handling bulk data with 1 mip and 1 array length"));
+			uint32 Stride;
+			
+			// lock, copy, unlock
+			FAGXTexture2DArray* TextureMTL = ResourceCast(Result.GetReference());
+			void* LockedData = TextureMTL->Surface.AsyncLock(RHICmdList, 0, 0, RLM_WriteOnly, Stride, false);
+			FMemory::Memcpy(LockedData, CreateInfo.BulkData->GetResourceBulkData(), CreateInfo.BulkData->GetResourceBulkDataSize());
+			TextureMTL->Surface.Unlock(0, 0, true);
+			
+			// bulk data can be unloaded now
+			CreateInfo.BulkData->Discard();
+		}
+		return Result;
+	}
+}
+
+FTexture3DRHIRef FAGXDynamicRHI::RHICreateTexture3D_RenderThread(class FRHICommandListImmediate& RHICmdList, uint32 SizeX, uint32 SizeY, uint32 SizeZ, uint8 Format, uint32 NumMips, ETextureCreateFlags Flags, ERHIAccess InResourceState, FRHIResourceCreateInfo& CreateInfo)
+{
+	@autoreleasepool {
+		FRHIResourceCreateInfo Info = CreateInfo;
+		bool const bIOSurfaceData = (CreateInfo.BulkData && CreateInfo.BulkData->GetResourceType() != FResourceBulkDataInterface::EBulkDataType::Default);
+		Info.BulkData = bIOSurfaceData ? CreateInfo.BulkData : nullptr;
+		FTexture3DRHIRef Result = GDynamicRHI->RHICreateTexture3D(SizeX, SizeY, SizeZ, Format, NumMips, Flags, InResourceState, Info);
+		if (CreateInfo.BulkData)
+		{
+			// upload existing bulkdata
+			UE_LOG(LogAGX, Display, TEXT("Got a bulk data texture, with %d mips"), NumMips);
+			checkf(NumMips == 1, TEXT("Only handling bulk data with 1 mip and 1 array length"));
+			uint32 Stride;
+			
+			// lock, copy, unlock
+			FAGXTexture3D* TextureMTL = ResourceCast(Result.GetReference());
+			void* LockedData = TextureMTL->Surface.AsyncLock(RHICmdList, 0, 0, RLM_WriteOnly, Stride, false);
+			FMemory::Memcpy(LockedData, CreateInfo.BulkData->GetResourceBulkData(), CreateInfo.BulkData->GetResourceBulkDataSize());
+			TextureMTL->Surface.Unlock(0, 0, true);
+			
+			// bulk data can be unloaded now
+			CreateInfo.BulkData->Discard();
+		}
+		return Result;
+	}
+}
+
+FTextureCubeRHIRef FAGXDynamicRHI::RHICreateTextureCube_RenderThread(class FRHICommandListImmediate& RHICmdList, uint32 Size, uint8 Format, uint32 NumMips, ETextureCreateFlags Flags, ERHIAccess InResourceState, FRHIResourceCreateInfo& CreateInfo)
+{
+	@autoreleasepool {
+		check(!CreateInfo.BulkData);
+		FTextureCubeRHIRef Result = GDynamicRHI->RHICreateTextureCube(Size, Format, NumMips, Flags, InResourceState, CreateInfo);
+		return Result;
+	}
+}
+
+FTextureCubeRHIRef FAGXDynamicRHI::RHICreateTextureCubeArray_RenderThread(class FRHICommandListImmediate& RHICmdList, uint32 Size, uint32 ArraySize, uint8 Format, uint32 NumMips, ETextureCreateFlags Flags, ERHIAccess InResourceState, FRHIResourceCreateInfo& CreateInfo)
+{
+	@autoreleasepool {
+		check(!CreateInfo.BulkData);
+		FTextureCubeRHIRef Result = GDynamicRHI->RHICreateTextureCubeArray(Size, ArraySize, Format, NumMips, Flags, InResourceState, CreateInfo);
+		return Result;
+	}
+}
+
+
+void FAGXDynamicRHI::RHIBindDebugLabelName(FRHITexture* TextureRHI, const TCHAR* Name)
+{
+	@autoreleasepool {
+		FAGXSurface* Surf = AGXGetMetalSurfaceFromRHITexture(TextureRHI);
+		if(Surf->Texture)
+		{
+			Surf->Texture.SetLabel(FString(Name).GetNSString());
+		}
+		if(Surf->MSAATexture)
+		{
+			Surf->MSAATexture.SetLabel(FString(Name).GetNSString());
+		}
+	}
+}
+
+void FAGXDynamicRHI::RHIVirtualTextureSetFirstMipInMemory(FRHITexture2D* TextureRHI, uint32 FirstMip)
+{
+	NOT_SUPPORTED("RHIVirtualTextureSetFirstMipInMemory");
+}
+
+void FAGXDynamicRHI::RHIVirtualTextureSetFirstMipVisible(FRHITexture2D* TextureRHI, uint32 FirstMip)
+{
+	NOT_SUPPORTED("RHIVirtualTextureSetFirstMipVisible");
+}
+
+struct FAGXRHICommandUnaliasTextures final : public FRHICommand<FAGXRHICommandUnaliasTextures>
+{
+	TArray<FRHITexture*> Textures;
+	
+	FORCEINLINE_DEBUGGABLE FAGXRHICommandUnaliasTextures(FRHITexture** InTextures, int32 NumTextures)
+	{
+		check(InTextures && NumTextures);
+		Textures.Append(InTextures, NumTextures);
+	}
+	
+	void Execute(FRHICommandListBase& CmdList)
+	{
+		@autoreleasepool {
+		for (int32 i = 0; i < Textures.Num(); ++i)
+		{
+			FAGXSurface* Source = AGXGetMetalSurfaceFromRHITexture(Textures[i]);
+			Source->MakeUnAliasable();
+		}
+		}
+	}
+};
+
+void FAGXDynamicRHI::RHIAcquireTransientResource_RenderThread(FRHITexture* Texture)
+{
+	@autoreleasepool {
+	FRHICommandListImmediate& RHICmdList = FRHICommandListExecutor::GetImmediateCommandList();
+	if (RHICmdList.Bypass() || !IsRunningRHIInSeparateThread())
+	{
+		FAGXSurface* Source = AGXGetMetalSurfaceFromRHITexture(Texture);
+		Source->MakeUnAliasable();
+	}
+	else
+	{
+		new (RHICmdList.AllocCommand<FAGXRHICommandUnaliasTextures>()) FAGXRHICommandUnaliasTextures(&Texture, 1);
+		RHICmdList.RHIThreadFence(true);
+	}
+	}
+}
+
+struct FAGXRHICommandAliasTextures final : public FRHICommand<FAGXRHICommandAliasTextures>
+{
+	TArray<FRHITexture*> Textures;
+	
+	FORCEINLINE_DEBUGGABLE FAGXRHICommandAliasTextures(FRHITexture** InTextures, int32 NumTextures)
+	{
+		check(InTextures && NumTextures);
+		Textures.Append(InTextures, NumTextures);
+	}
+	
+	void Execute(FRHICommandListBase& CmdList)
+	{
+		@autoreleasepool {
+		for (int32 i = 0; i < Textures.Num(); ++i)
+		{
+			FAGXSurface* Source = AGXGetMetalSurfaceFromRHITexture(Textures[i]);
+			Source->MakeAliasable();
+		}
+		}
+	}
+};
+
+void FAGXDynamicRHI::RHIDiscardTransientResource_RenderThread(FRHITexture* Texture)
+{
+	@autoreleasepool {
+	FRHICommandListImmediate& RHICmdList = FRHICommandListExecutor::GetImmediateCommandList();
+	if (RHICmdList.Bypass() || !IsRunningRHIInSeparateThread())
+	{
+		FAGXSurface* Source = AGXGetMetalSurfaceFromRHITexture(Texture);
+		Source->MakeAliasable();
+	}
+	else
+	{
+		new (RHICmdList.AllocCommand<FAGXRHICommandAliasTextures>()) FAGXRHICommandAliasTextures(&Texture, 1);
+		RHICmdList.RHIThreadFence(true);
+	}
+	}
+}
+
+void FAGXDynamicRHI::RHIAcquireTransientResource_RenderThread(FRHIBuffer* Buffer)
+{
+	// This is a no-op for the Metal RHI.
+}
+
+void FAGXDynamicRHI::RHIDiscardTransientResource_RenderThread(FRHIBuffer* Buffer)
+{
+	// This is a no-op for the Metal RHI.
+}
+
+struct FRHICopySubTextureRegion final : public FRHICommand<FRHICopySubTextureRegion>
+{
+	FRHITexture2D* SourceTexture;
+	FRHITexture2D* DestinationTexture;
+	FBox2D SourceBox;
+	FBox2D DestinationBox;
+	
+	FORCEINLINE_DEBUGGABLE FRHICopySubTextureRegion(FRHITexture2D* InSourceTexture, FRHITexture2D* InDestinationTexture, FBox2D InSourceBox, FBox2D InDestinationBox)
+	: SourceTexture(InSourceTexture)
+	, DestinationTexture(InDestinationTexture)
+	, SourceBox(InSourceBox)
+	, DestinationBox(InDestinationBox)
+	{
+	}
+	
+	void Execute(FRHICommandListBase& CmdList)
+	{
+		GDynamicRHI->RHICopySubTextureRegion(SourceTexture, DestinationTexture, SourceBox, DestinationBox);
+	}
+};
+
+void FAGXDynamicRHI::RHICopySubTextureRegion_RenderThread(class FRHICommandListImmediate& RHICmdList, FRHITexture2D* SourceTexture, FRHITexture2D* DestinationTexture, FBox2D SourceBox, FBox2D DestinationBox)
+{
+	@autoreleasepool
+	{
+		if (RHICmdList.Bypass() || !IsRunningRHIInSeparateThread())
+		{
+			GDynamicRHI->RHICopySubTextureRegion(SourceTexture, DestinationTexture, SourceBox, DestinationBox);
+		}
+		else
+		{
+			new (RHICmdList.AllocCommand<FRHICopySubTextureRegion>()) FRHICopySubTextureRegion(SourceTexture, DestinationTexture, SourceBox, DestinationBox);
+		}
+	}
+}
+
+void FAGXDynamicRHI::RHICopySubTextureRegion(FRHITexture2D* SourceTexture, FRHITexture2D* DestinationTexture, FBox2D SourceBox, FBox2D DestinationBox)
+{
+	@autoreleasepool {
+		check(SourceTexture);
+		check(DestinationTexture);
+		
+		FAGXTexture2D* MetalSrcTexture = ResourceCast(SourceTexture);
+		FAGXTexture2D* MetalDestTexture = ResourceCast(DestinationTexture);
+		
+		if(SourceTexture->GetFormat() == DestinationTexture->GetFormat())
+		{
+			FVector2D SourceSizeVector = SourceBox.GetSize();
+			FVector2D DestinatioSizeVector = DestinationBox.GetSize();
+			
+			mtlpp::Origin SourceOrigin(SourceBox.Min.X, SourceBox.Min.Y, 0);
+			mtlpp::Size SourceSize(SourceSizeVector.X, SourceSizeVector.Y, 1);
+			
+			mtlpp::Origin DestinationOrigin(DestinationBox.Min.X, DestinationBox.Min.Y, 0);
+			mtlpp::Size DestinationSize(DestinatioSizeVector.X, DestinatioSizeVector.Y, 1);
+			
+			check(DestinationSize.width == SourceSize.width);
+			check(DestinationSize.height == SourceSize.height);
+			
+			FAGXTexture SrcTexture;
+			mtlpp::TextureUsage Usage = MetalSrcTexture->Surface.Texture.GetUsage();
+			if(Usage & mtlpp::TextureUsage::PixelFormatView)
+			{
+				ns::Range Slices(0, MetalSrcTexture->Surface.Texture.GetArrayLength() * (MetalSrcTexture->Surface.bIsCubemap ? 6 : 1));
+				if(MetalSrcTexture->Surface.Texture.GetPixelFormat() != MetalDestTexture->Surface.Texture.GetPixelFormat())
+				{
+					SrcTexture = MetalSrcTexture->Surface.Texture.NewTextureView(MetalDestTexture->Surface.Texture.GetPixelFormat(), MetalSrcTexture->Surface.Texture.GetTextureType(), ns::Range(0, MetalSrcTexture->Surface.Texture.GetMipmapLevelCount()), Slices);
+				}
+			}
+			if (!SrcTexture)
+			{
+				SrcTexture = MetalSrcTexture->Surface.Texture;
+			}
+			
+			// Account for create with TexCreate_SRGB flag which could make these different
+			if(SrcTexture.GetPixelFormat() == MetalDestTexture->Surface.Texture.GetPixelFormat())
+			{
+				ImmediateContext.GetInternalContext().CopyFromTextureToTexture(SrcTexture, 0, 0, SourceOrigin,SourceSize,MetalDestTexture->Surface.Texture, 0, 0, DestinationOrigin);
+			}
+			else
+			{
+				// Linear and sRGB mismatch then try to go via metal buffer
+				// Modified clone of logic from MetalRenderTarget.cpp
+				uint32 BytesPerPixel = (MetalSrcTexture->Surface.PixelFormat != PF_DepthStencil) ? GPixelFormats[MetalSrcTexture->Surface.PixelFormat].BlockBytes : 1;
+				const uint32 Stride = BytesPerPixel * SourceSize.width;
+				const uint32 Alignment = PLATFORM_MAC ? 1u : 64u;
+				const uint32 AlignedStride = ((Stride - 1) & ~(Alignment - 1)) + Alignment;
+				const uint32 BytesPerImage = AlignedStride *  SourceSize.height;
+				
+				FAGXBuffer Buffer = GetAGXDeviceContext().CreatePooledBuffer(FAGXPooledBufferArgs(ImmediateContext.Context->GetDevice(), BytesPerImage, BUF_Dynamic, mtlpp::StorageMode::Shared));
+				
+				check(Buffer);
+				
+				mtlpp::BlitOption Options = mtlpp::BlitOption::None;
+#if !PLATFORM_MAC
+				if (MetalSrcTexture->Surface.Texture.GetPixelFormat() >= mtlpp::PixelFormat::PVRTC_RGB_2BPP && MetalSrcTexture->Surface.Texture.GetPixelFormat() <= mtlpp::PixelFormat::PVRTC_RGBA_4BPP_sRGB)
+				{
+					Options = mtlpp::BlitOption::RowLinearPVRTC;
+				}
+#endif
+				ImmediateContext.GetInternalContext().CopyFromTextureToBuffer(MetalSrcTexture->Surface.Texture, 0, 0, SourceOrigin, SourceSize, Buffer, 0, AlignedStride, BytesPerImage, Options);
+				ImmediateContext.GetInternalContext().CopyFromBufferToTexture(Buffer, 0, Stride, BytesPerImage, SourceSize, MetalDestTexture->Surface.Texture, 0, 0, DestinationOrigin, Options);
+				
+				GetAGXDeviceContext().ReleaseBuffer(Buffer);
+			}
+			
+			if (SrcTexture != MetalSrcTexture->Surface.Texture)
+			{
+				AGXSafeReleaseMetalTexture(SrcTexture);
+			}
+		}
+		else
+		{
+			UE_LOG(LogAGX, Error, TEXT("RHICopySubTextureRegion Source (UE4 %d: MTL %d) <-> Destination (UE4 %d: MTL %d) texture format mismatch"), (uint32)SourceTexture->GetFormat(), (uint32)MetalSrcTexture->Surface.Texture.GetPixelFormat(), (uint32)DestinationTexture->GetFormat(), (uint32)MetalDestTexture->Surface.Texture.GetPixelFormat());
+		}
+	}
+}
+
+inline bool AGXRHICopyTexutre_IsTextureFormatCompatible(EPixelFormat SrcFmt, EPixelFormat DstFmt)
+{
+	//
+	// For now, we only support copies between textures of mismatching
+	// formats if they are of size-compatible internal formats.  This allows us
+	// to copy from uncompressed to compressed textures, specifically in support
+	// of the runtime virtual texture system.  Note that copies of compatible
+	// formats incur the cost of an extra copy, as we must copy from the source
+	// texture to a temporary buffer and finally to the destination texture.
+	//
+	return ((SrcFmt == DstFmt) || (GPixelFormats[SrcFmt].BlockBytes == GPixelFormats[DstFmt].BlockBytes));
+}
+
+void FAGXRHICommandContext::RHICopyTexture(FRHITexture* SourceTextureRHI, FRHITexture* DestTextureRHI, const FRHICopyTextureInfo& CopyInfo)
+{
+	if (!SourceTextureRHI || !DestTextureRHI || SourceTextureRHI == DestTextureRHI)
+	{
+		// no need to do anything (silently ignored)
+		return;
+	}
+
+	@autoreleasepool {
+		check(SourceTextureRHI);
+		check(DestTextureRHI);
+		
+		FAGXSurface* MetalSrcTexture = AGXGetMetalSurfaceFromRHITexture(SourceTextureRHI);
+		FAGXSurface* MetalDestTexture = AGXGetMetalSurfaceFromRHITexture(DestTextureRHI);
+		
+		const bool TextureFormatExactMatch = (SourceTextureRHI->GetFormat() == DestTextureRHI->GetFormat());
+		const bool TextureFormatCompatible = AGXRHICopyTexutre_IsTextureFormatCompatible(SourceTextureRHI->GetFormat(), DestTextureRHI->GetFormat());
+		
+		if (TextureFormatExactMatch || TextureFormatCompatible)
+		{
+			FIntVector Size = (CopyInfo.Size != FIntVector::ZeroValue) ? CopyInfo.Size : FIntVector(MetalSrcTexture->SizeX, MetalSrcTexture->SizeY, MetalSrcTexture->SizeZ);
+			
+			mtlpp::Origin SourceOrigin(CopyInfo.SourcePosition.X, CopyInfo.SourcePosition.Y, CopyInfo.SourcePosition.Z);
+			mtlpp::Origin DestinationOrigin(CopyInfo.DestPosition.X, CopyInfo.DestPosition.Y, CopyInfo.DestPosition.Z);
+			
+			FAGXTexture SrcTexture;
+
+			if (TextureFormatExactMatch)
+			{
+				mtlpp::TextureUsage Usage = MetalSrcTexture->Texture.GetUsage();
+				if (Usage & mtlpp::TextureUsage::PixelFormatView)
+				{
+					ns::Range Slices(0, MetalSrcTexture->Texture.GetArrayLength() * (MetalSrcTexture->bIsCubemap ? 6 : 1));
+					if (MetalSrcTexture->Texture.GetPixelFormat() != MetalDestTexture->Texture.GetPixelFormat())
+					{
+						SrcTexture = MetalSrcTexture->Texture.NewTextureView(MetalDestTexture->Texture.GetPixelFormat(), MetalSrcTexture->Texture.GetTextureType(), ns::Range(0, MetalSrcTexture->Texture.GetMipmapLevelCount()), Slices);
+					}
+				}
+				if (!SrcTexture)
+				{
+					SrcTexture = MetalSrcTexture->Texture;
+				}
+			}
+			
+			for (uint32 SliceIndex = 0; SliceIndex < CopyInfo.NumSlices; ++SliceIndex)
+			{
+				uint32 SourceSliceIndex = CopyInfo.SourceSliceIndex + SliceIndex;
+				uint32 DestSliceIndex = CopyInfo.DestSliceIndex + SliceIndex;
+
+				for (uint32 MipIndex = 0; MipIndex < CopyInfo.NumMips; ++MipIndex)
+				{
+					uint32 SourceMipIndex = CopyInfo.SourceMipIndex + MipIndex;
+					uint32 DestMipIndex = CopyInfo.DestMipIndex + MipIndex;
+					mtlpp::Size SourceSize(FMath::Max(Size.X >> MipIndex, 1), FMath::Max(Size.Y >> MipIndex, 1), FMath::Max(Size.Z >> MipIndex, 1));
+					mtlpp::Size DestSize = SourceSize;
+
+					if (TextureFormatCompatible)
+					{
+						DestSize.width  *= GPixelFormats[MetalDestTexture->PixelFormat].BlockSizeX;
+						DestSize.height *= GPixelFormats[MetalDestTexture->PixelFormat].BlockSizeY;
+					}
+
+					// Account for create with TexCreate_SRGB flag which could make these different
+					if (TextureFormatExactMatch && (SrcTexture.GetPixelFormat() == MetalDestTexture->Texture.GetPixelFormat()))
+					{
+						GetInternalContext().CopyFromTextureToTexture(SrcTexture, SourceSliceIndex, SourceMipIndex, SourceOrigin,SourceSize,MetalDestTexture->Texture, DestSliceIndex, DestMipIndex, DestinationOrigin);
+					}
+					else
+					{
+						//
+						// In the case of compatible texture formats or pixel
+						// format mismatch (like linear vs. sRGB), then we must
+						// achieve the copy by going through a buffer object.
+						//
+						const bool BlockSizeMatch = (GPixelFormats[MetalSrcTexture->PixelFormat].BlockSizeX == GPixelFormats[MetalDestTexture->PixelFormat].BlockSizeX);
+						const uint32 BytesPerPixel = (MetalSrcTexture->PixelFormat != PF_DepthStencil) ? GPixelFormats[MetalSrcTexture->PixelFormat].BlockBytes : 1;
+						const uint32 Stride = BytesPerPixel * SourceSize.width;
+#if PLATFORM_MAC
+						const uint32 Alignment = 1u;
+#else
+						// don't mess with alignment if we copying between formats with a different block size
+						const uint32 Alignment = BlockSizeMatch ? 64u : 1u;
+#endif
+						const uint32 AlignedStride = ((Stride - 1) & ~(Alignment - 1)) + Alignment;
+						const uint32 BytesPerImage = AlignedStride *  SourceSize.height;
+						const uint32 DataSize = BytesPerImage * SourceSize.depth;
+						
+						FAGXBuffer Buffer = GetAGXDeviceContext().CreatePooledBuffer(FAGXPooledBufferArgs(GetInternalContext().GetDevice(), DataSize, BUF_Dynamic, mtlpp::StorageMode::Shared));
+						
+						check(Buffer);
+						
+						mtlpp::BlitOption Options = mtlpp::BlitOption::None;
+#if !PLATFORM_MAC
+						if (MetalSrcTexture->Texture.GetPixelFormat() >= mtlpp::PixelFormat::PVRTC_RGB_2BPP && MetalSrcTexture->Texture.GetPixelFormat() <= mtlpp::PixelFormat::PVRTC_RGBA_4BPP_sRGB)
+						{
+							Options = mtlpp::BlitOption::RowLinearPVRTC;
+						}
+#endif
+						GetInternalContext().CopyFromTextureToBuffer(MetalSrcTexture->Texture, SourceSliceIndex, SourceMipIndex, SourceOrigin, SourceSize, Buffer, 0, AlignedStride, BytesPerImage, Options);
+						GetInternalContext().CopyFromBufferToTexture(Buffer, 0, Stride, BytesPerImage, DestSize, MetalDestTexture->Texture, DestSliceIndex, DestMipIndex, DestinationOrigin, Options);
+						
+						GetAGXDeviceContext().ReleaseBuffer(Buffer);
+					}
+				}
+			}
+			
+			if (SrcTexture && (SrcTexture != MetalSrcTexture->Texture))
+			{
+				AGXSafeReleaseMetalTexture(SrcTexture);
+			}
+		}
+		else
+		{
+			UE_LOG(LogAGX, Error, TEXT("RHICopyTexture Source (UE4 %d: MTL %d) <-> Destination (UE4 %d: MTL %d) texture format mismatch"), (uint32)SourceTextureRHI->GetFormat(), (uint32)MetalSrcTexture->Texture.GetPixelFormat(), (uint32)DestTextureRHI->GetFormat(), (uint32)MetalDestTexture->Texture.GetPixelFormat());
+		}
+	}
+}
+
+void FAGXRHICommandContext::RHICopyBufferRegion(FRHIBuffer* DstBufferRHI, uint64 DstOffset, FRHIBuffer* SrcBufferRHI, uint64 SrcOffset, uint64 NumBytes)
+{
+	if (!DstBufferRHI || !SrcBufferRHI || DstBufferRHI == SrcBufferRHI || !NumBytes)
+	{
+		return;
+	}
+
+	@autoreleasepool {
+		FAGXResourceMultiBuffer* DstBuffer = ResourceCast(DstBufferRHI);
+		FAGXResourceMultiBuffer* SrcBuffer = ResourceCast(SrcBufferRHI);
+
+		check(DstBuffer && SrcBuffer);
+		check(!DstBuffer->Data && !SrcBuffer->Data);
+		check(DstOffset + NumBytes <= DstBufferRHI->GetSize() && SrcOffset + NumBytes <= SrcBufferRHI->GetSize());
+
+		GetInternalContext().CopyFromBufferToBuffer(SrcBuffer->GetCurrentBuffer(), SrcOffset, DstBuffer->GetCurrentBuffer(), DstOffset, NumBytes);
+	}
+}
