Initial import of AGXRHI

diff --git a/Engine/Source/Runtime/Apple/AGXRHI/Private/AGXTransitionData.cpp b/Engine/Source/Runtime/Apple/AGXRHI/Private/AGXTransitionData.cpp
new file mode 100644
--- /dev/null
+++ b/Engine/Source/Runtime/Apple/AGXRHI/Private/AGXTransitionData.cpp
@@ -0,0 +1,113 @@
+// Copyright Epic Games, Inc. All Rights Reserved.
+
+/*==============================================================================
+	AGXTransitionData.cpp: AGX RHI Resource Transition Implementation.
+==============================================================================*/
+
+#include "AGXRHIPrivate.h"
+#include "AGXTransitionData.h"
+
+
+//------------------------------------------------------------------------------
+
+#pragma mark - AGX RHI Resource Transition Data Definitions -
+
+#define UE_METAL_TRANSITION_DATA_USES_EXPLICIT_FENCING					0
+
+
+//------------------------------------------------------------------------------
+
+#pragma mark - AGX RHI Resource Transition Data Implementation -
+
+
+FAGXTransitionData::FAGXTransitionData(ERHIPipeline                         InSrcPipelines,
+										   ERHIPipeline                         InDstPipelines,
+										   ERHITransitionCreateFlags            InCreateFlags,
+										   TArrayView<const FRHITransitionInfo> InInfos)
+{
+	SrcPipelines   = InSrcPipelines;
+	DstPipelines   = InDstPipelines;
+	CreateFlags    = InCreateFlags;
+
+	bCrossPipeline = (SrcPipelines != DstPipelines);
+
+	Infos.Append(InInfos.GetData(), InInfos.Num());
+
+	// TODO: Determine whether the AGX RHI needs to create a separate, per-transition fence.
+#if UE_METAL_TRANSITION_DATA_USES_EXPLICIT_FENCING
+	if (bCrossPipeline && !EnumHasAnyFlags(CreateFlags, ERHITransitionCreateFlags::NoFence))
+	{
+		// Get the current context pointer.
+		FAGXContext* Context = GetAGXDeviceContext().GetCurrentContext();
+
+		// Get the current render pass fence.
+		TRefCountPtr<FAGXFence> const& MetalFence = Context->GetCurrentRenderPass().End();
+
+		// Write it again as we may wait on this fence in two different encoders.
+		Context->GetCurrentRenderPass().Update(MetalFence);
+
+		// Write it into the transition data.
+		Fence = MetalFence;
+		Fence->AddRef();
+
+		if (GSupportsEfficientAsyncCompute)
+		{
+			static_cast<FAGXRHICommandContext*>(RHIGetDefaultContext())->RHISubmitCommandsHint();
+		}
+	}
+#endif // UE_METAL_TRANSITION_DATA_USES_EXPLICIT_FENCING
+}
+
+void FAGXTransitionData::BeginResourceTransitions() const
+{
+	// TODO: Determine whether the AGX RHI needs to do anything here.
+#if UE_METAL_TRANSITION_DATA_USES_EXPLICIT_FENCING
+	if (Fence.IsValid())
+	{
+		FAGXContext* Context = GetAGXDeviceContext().GetCurrentContext();
+
+		if (Context->GetCurrentCommandBuffer())
+		{
+			Context->SubmitCommandsHint(EAGXSubmitFlagsNone);
+		}
+
+		Context->GetCurrentRenderPass().Begin(Fence);
+	}
+#endif // UE_METAL_TRANSITION_DATA_USES_EXPLICIT_FENCING
+}
+
+void FAGXTransitionData::EndResourceTransitions() const
+{
+	// No action necessary for same pipe transitions
+	if (SrcPipelines == DstPipelines)
+	{
+		return;
+	}
+
+	for (const auto& Info : Infos)
+	{
+		if (nullptr == Info.Resource)
+		{
+			continue;
+		}
+
+		switch (Info.Type)
+		{
+			case FRHITransitionInfo::EType::UAV:
+				GetAGXDeviceContext().TransitionResource(Info.UAV);
+				break;
+
+			case FRHITransitionInfo::EType::Buffer:
+				GetAGXDeviceContext().TransitionRHIResource(Info.Buffer);
+				break;
+
+			case FRHITransitionInfo::EType::Texture:
+				GetAGXDeviceContext().TransitionResource(Info.Texture);
+				break;
+
+			default:
+				checkNoEntry();
+				break;
+		}
+	}
+}
diff --git a/Engine/Source/Runtime/Apple/AGXRHI/Private/AGXTransitionData.h b/Engine/Source/Runtime/Apple/AGXRHI/Private/AGXTransitionData.h
new file mode 100644
--- /dev/null
+++ b/Engine/Source/Runtime/Apple/AGXRHI/Private/AGXTransitionData.h
@@ -0,0 +1,52 @@
+// Copyright Epic Games, Inc. All Rights Reserved.
+
+/*==============================================================================
+	AGXTransitionData.h: AGX RHI Resource Transition Definitions.
+==============================================================================*/
+
+#pragma once
+
+
+//------------------------------------------------------------------------------
+
+#pragma mark - AGX RHI Resource Transition Info Array Type Definition -
+
+
+typedef TArray<FRHITransitionInfo, TInlineAllocator<4> > FAGXTransitionInfoArray;
+
+
+//------------------------------------------------------------------------------
+
+#pragma mark - AGX RHI Resource Transition Data Class -
+
+
+class FAGXTransitionData
+{
+public:
+	explicit FAGXTransitionData(ERHIPipeline                         InSrcPipelines,
+								  ERHIPipeline                         InDstPipelines,
+								  ERHITransitionCreateFlags            InCreateFlags,
+								  TArrayView<const FRHITransitionInfo> InInfos);
+
+	// The default destructor is sufficient.
+	~FAGXTransitionData() = default;
+
+	// Disallow default, copy and move constructors.
+	FAGXTransitionData()                             = delete;
+	FAGXTransitionData(const FAGXTransitionData&)  = delete;
+	FAGXTransitionData(const FAGXTransitionData&&) = delete;
+
+	// Begin resource transitions.
+	void BeginResourceTransitions() const;
+
+	// End resource transitions.
+	void EndResourceTransitions() const;
+
+private:
+	ERHIPipeline              SrcPipelines   = ERHIPipeline::Num;
+	ERHIPipeline              DstPipelines   = ERHIPipeline::Num;
+	ERHITransitionCreateFlags CreateFlags    = ERHITransitionCreateFlags::None;
+	bool                      bCrossPipeline = false;
+	FAGXTransitionInfoArray Infos          = {};
+	TRefCountPtr<FAGXFence> Fence          = nullptr;
+};
diff --git a/Engine/Source/Runtime/Apple/AGXRHI/Private/AGXUAV.cpp b/Engine/Source/Runtime/Apple/AGXRHI/Private/AGXUAV.cpp
new file mode 100644
--- /dev/null
+++ b/Engine/Source/Runtime/Apple/AGXRHI/Private/AGXUAV.cpp
@@ -0,0 +1,659 @@
+// Copyright Epic Games, Inc. All Rights Reserved.
+
+
+#include "AGXRHIPrivate.h"
+#include "AGXRHIStagingBuffer.h"
+#include "AGXCommandBuffer.h"
+#include "RenderUtils.h"
+#include "ClearReplacementShaders.h"
+#include "AGXTransitionData.h"
+
+FAGXShaderResourceView::FAGXShaderResourceView()
+	: TextureView(nullptr)
+	, Offset(0)
+	, MipLevel(0)
+	, bSRGBForceDisable(0)
+	, Reserved(0)
+	, NumMips(0)
+	, Format(0)
+	, Stride(0)
+	, LinearTextureDesc(nullptr)
+{
+	// void
+}
+
+FAGXShaderResourceView::~FAGXShaderResourceView()
+{
+	if (LinearTextureDesc)
+	{
+		delete LinearTextureDesc;
+		LinearTextureDesc = nullptr;
+	}
+	
+	if (TextureView)
+	{
+		FAGXSurface* Surface = AGXGetMetalSurfaceFromRHITexture(SourceTexture);
+		if (Surface)
+		{
+			Surface->SRVs.Remove(this);
+		}
+
+		delete TextureView;
+		TextureView = nullptr;
+	}
+	
+	SourceVertexBuffer = NULL;
+	SourceTexture = NULL;
+}
+
+void FAGXShaderResourceView::InitLinearTextureDescriptor(const FAGXLinearTextureDescriptor& InLinearTextureDescriptor)
+{
+	check(!LinearTextureDesc);
+	LinearTextureDesc = new FAGXLinearTextureDescriptor(InLinearTextureDescriptor);
+	check(LinearTextureDesc);
+}
+
+ns::AutoReleased<FAGXTexture> FAGXShaderResourceView::GetLinearTexture(bool const bUAV)
+{
+	ns::AutoReleased<FAGXTexture> NewLinearTexture;
+	{
+		if (IsValidRef(SourceVertexBuffer))
+		{
+			NewLinearTexture = SourceVertexBuffer->GetLinearTexture((EPixelFormat)Format, LinearTextureDesc);
+		}
+		else if (IsValidRef(SourceIndexBuffer))
+		{
+			NewLinearTexture = SourceIndexBuffer->GetLinearTexture((EPixelFormat)Format, LinearTextureDesc);
+		}
+	}
+	return NewLinearTexture;
+}
+
+FUnorderedAccessViewRHIRef FAGXDynamicRHI::RHICreateUnorderedAccessView_RenderThread(class FRHICommandListImmediate& RHICmdList, FRHIBuffer* Buffer, bool bUseUAVCounter, bool bAppendBuffer)
+{
+	return GDynamicRHI->RHICreateUnorderedAccessView(Buffer, bUseUAVCounter, bAppendBuffer);
+}
+
+FUnorderedAccessViewRHIRef FAGXDynamicRHI::RHICreateUnorderedAccessView_RenderThread(class FRHICommandListImmediate& RHICmdList, FRHITexture* Texture, uint32 MipLevel)
+{
+	FAGXSurface* Surface = (FAGXSurface*)Texture->GetTextureBaseRHI();
+	FAGXTexture Tex = Surface->Texture;
+	if (!(Tex.GetUsage() & mtlpp::TextureUsage::PixelFormatView))
+	{
+		FScopedRHIThreadStaller StallRHIThread(RHICmdList);
+		return GDynamicRHI->RHICreateUnorderedAccessView(Texture, MipLevel);
+	}
+	else
+	{
+		return GDynamicRHI->RHICreateUnorderedAccessView(Texture, MipLevel);
+	}
+}
+
+FUnorderedAccessViewRHIRef FAGXDynamicRHI::RHICreateUnorderedAccessView_RenderThread(class FRHICommandListImmediate& RHICmdList, FRHIBuffer* Buffer, uint8 Format)
+{
+	FUnorderedAccessViewRHIRef Result = GDynamicRHI->RHICreateUnorderedAccessView(Buffer, Format);
+	if (IsRunningRHIInSeparateThread() && !RHICmdList.Bypass()) 	{ 		RHICmdList.RHIThreadFence(true); 	}
+	return Result;
+}
+
+FUnorderedAccessViewRHIRef FAGXDynamicRHI::RHICreateUnorderedAccessView(FRHIBuffer* BufferRHI, bool bUseUAVCounter, bool bAppendBuffer)
+{
+	@autoreleasepool {
+	FAGXStructuredBuffer* StructuredBuffer = ResourceCast(BufferRHI);
+	
+	FAGXShaderResourceView* SRV = new FAGXShaderResourceView;
+	SRV->SourceVertexBuffer = nullptr;
+	SRV->SourceIndexBuffer = nullptr;
+	SRV->TextureView = nullptr;
+	SRV->SourceStructuredBuffer = StructuredBuffer;
+
+	// create the UAV buffer to point to the structured buffer's memory
+	FAGXUnorderedAccessView* UAV = new FAGXUnorderedAccessView;
+	UAV->SourceView = SRV;
+	return UAV;
+	}
+}
+
+FUnorderedAccessViewRHIRef FAGXDynamicRHI::RHICreateUnorderedAccessView(FRHITexture* TextureRHI, uint32 MipLevel)
+{
+	@autoreleasepool {
+	FAGXShaderResourceView* SRV = new FAGXShaderResourceView;
+	SRV->SourceTexture = TextureRHI;
+	
+	FAGXSurface* Surface = AGXGetMetalSurfaceFromRHITexture(TextureRHI);
+	SRV->TextureView = Surface ? new FAGXSurface(*Surface, NSMakeRange(MipLevel, 1)) : nullptr;
+	
+	SRV->SourceVertexBuffer = nullptr;
+	SRV->SourceIndexBuffer = nullptr;
+	SRV->SourceStructuredBuffer = nullptr;
+	
+	SRV->MipLevel = MipLevel;
+	SRV->NumMips = 1;
+	SRV->Format = PF_Unknown;
+		
+	if (Surface)
+	{
+		Surface->SRVs.Add(SRV);
+	}
+		
+	// create the UAV buffer to point to the structured buffer's memory
+	FAGXUnorderedAccessView* UAV = new FAGXUnorderedAccessView;
+	UAV->SourceView = SRV;
+
+	return UAV;
+	}
+}
+
+FUnorderedAccessViewRHIRef FAGXDynamicRHI::RHICreateUnorderedAccessView(FRHIBuffer* BufferRHI, uint8 Format)
+{
+	@autoreleasepool {
+	FAGXResourceMultiBuffer* Buffer = ResourceCast(BufferRHI);
+	
+	FAGXShaderResourceView* SRV = new FAGXShaderResourceView;
+	SRV->SourceVertexBuffer = (BufferRHI->GetUsage() & BUF_VertexBuffer) ? Buffer : nullptr;
+	SRV->TextureView = nullptr;
+	SRV->SourceIndexBuffer = (BufferRHI->GetUsage() & BUF_IndexBuffer) ? Buffer : nullptr;
+	SRV->SourceStructuredBuffer = nullptr;
+	SRV->Format = Format;
+	{
+		check(Buffer->GetUsage() & BUF_UnorderedAccess);
+		Buffer->CreateLinearTexture((EPixelFormat)Format, Buffer);
+	}
+		
+	// create the UAV buffer to point to the structured buffer's memory
+	FAGXUnorderedAccessView* UAV = new FAGXUnorderedAccessView;
+	UAV->SourceView = SRV;
+
+	return UAV;
+	}
+}
+
+FShaderResourceViewRHIRef FAGXDynamicRHI::CreateShaderResourceView_RenderThread(class FRHICommandListImmediate& RHICmdList, FRHIBuffer* Buffer, uint32 Stride, uint8 Format)
+{
+	FShaderResourceViewRHIRef Result = GDynamicRHI->RHICreateShaderResourceView(Buffer, Stride, Format);
+	if (IsRunningRHIInSeparateThread() && !RHICmdList.Bypass()) 	{ 		RHICmdList.RHIThreadFence(true); 	}
+	return Result;
+}
+
+FShaderResourceViewRHIRef FAGXDynamicRHI::CreateShaderResourceView_RenderThread(class FRHICommandListImmediate& RHICmdList, const FShaderResourceViewInitializer& Initializer)
+{
+	FShaderResourceViewRHIRef Result = GDynamicRHI->RHICreateShaderResourceView(Initializer);
+	if (IsRunningRHIInSeparateThread() && !RHICmdList.Bypass()) { RHICmdList.RHIThreadFence(true); }
+	return Result;
+}
+
+FShaderResourceViewRHIRef FAGXDynamicRHI::CreateShaderResourceView_RenderThread(class FRHICommandListImmediate& RHICmdList, FRHIBuffer* Buffer)
+{
+	FShaderResourceViewRHIRef Result = GDynamicRHI->RHICreateShaderResourceView(Buffer);
+	if (IsRunningRHIInSeparateThread() && !RHICmdList.Bypass()) 	{ 		RHICmdList.RHIThreadFence(true); 	}
+	return Result;
+}
+
+FShaderResourceViewRHIRef FAGXDynamicRHI::RHICreateShaderResourceView_RenderThread(class FRHICommandListImmediate& RHICmdList, FRHIBuffer* Buffer, uint32 Stride, uint8 Format)
+{
+	FShaderResourceViewRHIRef Result = GDynamicRHI->RHICreateShaderResourceView(Buffer, Stride, Format);
+	if (IsRunningRHIInSeparateThread() && !RHICmdList.Bypass()) 	{ 		RHICmdList.RHIThreadFence(true); 	}
+	return Result;
+}
+
+FShaderResourceViewRHIRef FAGXDynamicRHI::RHICreateShaderResourceView_RenderThread(class FRHICommandListImmediate& RHICmdList, const FShaderResourceViewInitializer& Initializer)
+{
+	FShaderResourceViewRHIRef Result = GDynamicRHI->RHICreateShaderResourceView(Initializer);
+	if (IsRunningRHIInSeparateThread() && !RHICmdList.Bypass()) 	{ 		RHICmdList.RHIThreadFence(true); 	}
+	return Result;
+}
+
+FShaderResourceViewRHIRef FAGXDynamicRHI::RHICreateShaderResourceView_RenderThread(class FRHICommandListImmediate& RHICmdList, FRHIBuffer* Buffer)
+{
+	FShaderResourceViewRHIRef Result = GDynamicRHI->RHICreateShaderResourceView(Buffer);
+	if (IsRunningRHIInSeparateThread() && !RHICmdList.Bypass()) 	{ 		RHICmdList.RHIThreadFence(true); 	}
+	return Result;
+}
+
+FShaderResourceViewRHIRef FAGXDynamicRHI::RHICreateShaderResourceView_RenderThread(class FRHICommandListImmediate& RHICmdList, FRHITexture* Texture2DRHI, const FRHITextureSRVCreateInfo& CreateInfo)
+{
+	FAGXSurface* Surface = (FAGXSurface*)Texture2DRHI->GetTextureBaseRHI();
+	FAGXTexture Tex = Surface->Texture;
+	if (!(Tex.GetUsage() & mtlpp::TextureUsage::PixelFormatView))
+	{
+		FScopedRHIThreadStaller StallRHIThread(RHICmdList);
+		return GDynamicRHI->RHICreateShaderResourceView(Texture2DRHI, CreateInfo);
+	}
+	else
+	{
+		return GDynamicRHI->RHICreateShaderResourceView(Texture2DRHI, CreateInfo);
+	}
+}
+
+FShaderResourceViewRHIRef FAGXDynamicRHI::RHICreateShaderResourceView(FRHITexture* Texture2DRHI, const FRHITextureSRVCreateInfo& CreateInfo)
+{
+	@autoreleasepool {
+		FAGXShaderResourceView* SRV = new FAGXShaderResourceView;
+		SRV->SourceTexture = (FRHITexture*)Texture2DRHI;
+		
+		FAGXSurface* Surface = AGXGetMetalSurfaceFromRHITexture(Texture2DRHI);
+		
+		// Asking to make a SRV with PF_Unknown means to use the same format.
+		// This matches the behavior of the DX11 RHI.
+		EPixelFormat Format = (EPixelFormat) CreateInfo.Format;
+		if(Surface && Format == PF_Unknown)
+		{
+			Format = Surface->PixelFormat;
+		}
+		
+		const bool bSRGBForceDisable = (CreateInfo.SRGBOverride == SRGBO_ForceDisable);
+		
+		SRV->TextureView = Surface ? new FAGXSurface(*Surface, NSMakeRange(CreateInfo.MipLevel, CreateInfo.NumMipLevels), Format, bSRGBForceDisable) : nullptr;
+		
+		SRV->SourceVertexBuffer = nullptr;
+		SRV->SourceIndexBuffer = nullptr;
+		SRV->SourceStructuredBuffer = nullptr;
+		
+		SRV->MipLevel = CreateInfo.MipLevel;
+		SRV->bSRGBForceDisable = bSRGBForceDisable;
+		SRV->NumMips = CreateInfo.NumMipLevels;
+		SRV->Format = CreateInfo.Format;
+		
+		if (Surface)
+		{
+			Surface->SRVs.Add(SRV);
+		}
+		
+		return SRV;
+	}
+}
+
+FShaderResourceViewRHIRef FAGXDynamicRHI::RHICreateShaderResourceView(FRHIBuffer* BufferRHI)
+{
+	@autoreleasepool {
+	return RHICreateShaderResourceView(FShaderResourceViewInitializer(BufferRHI));
+	}
+}
+
+FShaderResourceViewRHIRef FAGXDynamicRHI::RHICreateShaderResourceView(FRHIBuffer* BufferRHI, uint32 Stride, uint8 Format)
+{
+	check(GPixelFormats[Format].BlockBytes == Stride);
+
+	@autoreleasepool
+	{
+		return RHICreateShaderResourceView(FShaderResourceViewInitializer(BufferRHI, EPixelFormat(Format)));
+	}
+}
+
+FShaderResourceViewRHIRef FAGXDynamicRHI::RHICreateShaderResourceView(const FShaderResourceViewInitializer& Initializer)
+{
+	@autoreleasepool
+	{
+		switch(Initializer.GetType())
+		{
+			case FShaderResourceViewInitializer::EType::VertexBufferSRV:
+			{
+				FShaderResourceViewInitializer::FVertexBufferShaderResourceViewInitializer Desc = Initializer.AsVertexBufferSRV();
+				
+				FAGXVertexBuffer* VertexBuffer = ResourceCast(Desc.Buffer);
+				
+				FAGXShaderResourceView* SRV = new FAGXShaderResourceView;
+				SRV->SourceVertexBuffer = VertexBuffer;
+				SRV->TextureView = nullptr;
+				SRV->SourceIndexBuffer = nullptr;
+				SRV->SourceStructuredBuffer = nullptr;
+				
+				SRV->Format = Desc.Format;
+				
+				if(!VertexBuffer)
+				{
+					SRV->Offset = 0;
+					SRV->Stride = 0;
+				}
+				else
+				{
+					check(VertexBuffer->GetUsage() & BUF_ShaderResource);
+					uint32 Stride = GPixelFormats[Desc.Format].BlockBytes;
+					
+					SRV->Stride = Stride;
+					SRV->Offset = Desc.StartOffsetBytes;
+					
+					FAGXLinearTextureDescriptor LinearTextureDesc(Desc.StartOffsetBytes, Desc.NumElements, Stride);
+					SRV->InitLinearTextureDescriptor(LinearTextureDesc);
+					
+					VertexBuffer->CreateLinearTexture((EPixelFormat)Desc.Format, VertexBuffer, &LinearTextureDesc);
+				}
+				
+				return SRV;
+			} // VertexBufferSRV
+				
+			case FShaderResourceViewInitializer::EType::StructuredBufferSRV:
+			{
+				FShaderResourceViewInitializer::FStructuredBufferShaderResourceViewInitializer Desc = Initializer.AsStructuredBufferSRV();
+				
+				FAGXStructuredBuffer* StructuredBuffer = ResourceCast(Desc.Buffer);
+
+				FAGXShaderResourceView* SRV = new FAGXShaderResourceView;
+				SRV->SourceVertexBuffer = nullptr;
+				SRV->SourceIndexBuffer = nullptr;
+				SRV->TextureView = nullptr;
+				SRV->SourceStructuredBuffer = StructuredBuffer;
+				
+				SRV->Offset = Desc.StartOffsetBytes;
+				SRV->Format = 0;
+				SRV->Stride = StructuredBuffer->GetStride();
+				
+				return SRV;
+			} // StructuredBufferSRV
+				
+			case FShaderResourceViewInitializer::EType::IndexBufferSRV:
+			{
+				FShaderResourceViewInitializer::FIndexBufferShaderResourceViewInitializer Desc = Initializer.AsIndexBufferSRV();
+				
+				FAGXIndexBuffer* IndexBuffer = ResourceCast(Desc.Buffer);
+				
+				FAGXShaderResourceView* SRV = new FAGXShaderResourceView;
+				SRV->SourceVertexBuffer = nullptr;
+				SRV->SourceIndexBuffer = IndexBuffer;
+				SRV->TextureView = nullptr;
+				
+				if(!IndexBuffer)
+				{
+					SRV->Format = PF_R16_UINT;
+					SRV->Stride = 0;
+					SRV->Offset = 0;
+				}
+				else
+				{
+					SRV->Format = (IndexBuffer->IndexType == mtlpp::IndexType::UInt16) ? PF_R16_UINT : PF_R32_UINT;
+					
+					const uint32 Stride = Desc.Buffer->GetStride();
+					check(Stride == ((SRV->Format == PF_R16_UINT) ? 2 : 4));
+					
+					SRV->Offset = Desc.StartOffsetBytes;
+					SRV->Stride = Stride;
+					
+					FAGXLinearTextureDescriptor LinearTextureDesc(Desc.StartOffsetBytes, Desc.NumElements, Stride);
+					SRV->InitLinearTextureDescriptor(LinearTextureDesc);
+					
+					IndexBuffer->CreateLinearTexture((EPixelFormat)SRV->Format, IndexBuffer, &LinearTextureDesc);
+				}
+					  
+				return SRV;
+			} // IndexBufferSRV
+					  
+			default:
+			{
+				checkNoEntry();
+				return nullptr;
+			}
+		}
+	}
+}
+
+void FAGXDynamicRHI::RHIUpdateShaderResourceView(FRHIShaderResourceView* SRVRHI, FRHIBuffer* VertexBufferRHI, uint32 Stride, uint8 Format)
+{
+	check(SRVRHI);
+	FAGXShaderResourceView* SRV = ResourceCast(SRVRHI);
+	if (!VertexBufferRHI)
+	{
+		SRV->SourceVertexBuffer = nullptr;
+		SRV->TextureView = nullptr;
+		SRV->SourceIndexBuffer = nullptr;
+		SRV->SourceStructuredBuffer = nullptr;
+		SRV->Offset = 0;
+		SRV->Format = Format;
+		SRV->Stride = Stride;
+	}
+	else if (SRV->SourceVertexBuffer != VertexBufferRHI)
+	{
+		FAGXVertexBuffer* VertexBuffer = ResourceCast(VertexBufferRHI);
+		SRV->SourceVertexBuffer = VertexBuffer;
+		SRV->TextureView = nullptr;
+		SRV->SourceIndexBuffer = nullptr;
+		SRV->SourceStructuredBuffer = nullptr;
+		SRV->Offset = 0;
+		SRV->Format = Format;
+		SRV->Stride = Stride;
+	}
+}
+
+void FAGXDynamicRHI::RHIUpdateShaderResourceView(FRHIShaderResourceView* SRVRHI, FRHIBuffer* IndexBufferRHI)
+{
+	check(SRVRHI);
+	FAGXShaderResourceView* SRV = ResourceCast(SRVRHI);
+	if (!IndexBufferRHI)
+	{
+		SRV->SourceVertexBuffer = nullptr;
+		SRV->TextureView = nullptr;
+		SRV->SourceIndexBuffer = nullptr;
+		SRV->SourceStructuredBuffer = nullptr;
+		SRV->Offset = 0;
+		SRV->Format = PF_R16_UINT;
+		SRV->Stride = 0;
+	}
+	else if (SRV->SourceIndexBuffer != IndexBufferRHI)
+	{
+		FAGXIndexBuffer* IndexBuffer = ResourceCast(IndexBufferRHI);
+		SRV->SourceVertexBuffer = nullptr;
+		SRV->TextureView = nullptr;
+		SRV->SourceIndexBuffer = IndexBuffer;
+		SRV->SourceStructuredBuffer = nullptr;
+		SRV->Offset = 0;
+		SRV->Format = (IndexBuffer->IndexType == mtlpp::IndexType::UInt16) ? PF_R16_UINT : PF_R32_UINT;
+		SRV->Stride = 0;
+	}
+}
+
+#if UE_METAL_RHI_SUPPORT_CLEAR_UAV_WITH_BLIT_ENCODER
+void FAGXRHICommandContext::ClearUAVWithBlitEncoder(FRHIUnorderedAccessView* UnorderedAccessViewRHI, EAGXRHIClearUAVType Type, uint32 Pattern)
+{
+	SCOPED_AUTORELEASE_POOL;
+
+	FRHICommandListImmediate& RHICmdList = FRHICommandListExecutor::GetImmediateCommandList();
+	FAGXUnorderedAccessView* UnorderedAccessView = ResourceCast(UnorderedAccessViewRHI);
+	FAGXBuffer Buffer;
+	uint32 Size = 0;
+
+	switch (Type)
+	{
+		case EAGXRHIClearUAVType::VertexBuffer:
+			check(0 != (UnorderedAccessView->SourceView->SourceVertexBuffer->GetUsage() & BUF_ByteAddressBuffer));
+			Buffer = UnorderedAccessView->SourceView->SourceVertexBuffer->GetCurrentBuffer();
+			Size = UnorderedAccessView->SourceView->SourceVertexBuffer->GetSize();
+			break;
+
+		case EAGXRHIClearUAVType::StructuredBuffer:
+			Buffer = UnorderedAccessView->SourceView->SourceStructuredBuffer->GetCurrentBuffer();
+			Size = UnorderedAccessView->SourceView->SourceStructuredBuffer->GetSize();
+			break;
+	};
+
+	uint32 AlignedSize = Align(Size, BufferOffsetAlignment);
+	FAGXPooledBufferArgs Args(GetAGXDeviceContext().GetDevice(), AlignedSize, BUF_Dynamic, mtlpp::StorageMode::Shared);
+	FAGXBuffer Temp = GetAGXDeviceContext().CreatePooledBuffer(Args);
+	uint32* ContentBytes = (uint32*)Temp.GetContents();
+	for (uint32 Element = 0; Element < (AlignedSize >> 2); ++Element)
+	{
+		ContentBytes[Element] = Pattern;
+	}
+	Context->CopyFromBufferToBuffer(Temp, 0, Buffer, 0, Size);
+	GetAGXDeviceContext().ReleaseBuffer(Temp);
+}
+#endif // UE_METAL_RHI_SUPPORT_CLEAR_UAV_WITH_BLIT_ENCODER
+
+void FAGXRHICommandContext::RHIClearUAVFloat(FRHIUnorderedAccessView* UnorderedAccessViewRHI, const FVector4& Values)
+{
+#if UE_METAL_RHI_SUPPORT_CLEAR_UAV_WITH_BLIT_ENCODER
+	FAGXUnorderedAccessView* UnorderedAccessView = ResourceCast(UnorderedAccessViewRHI);
+	if (UnorderedAccessView->SourceView->SourceStructuredBuffer)
+	{
+		ClearUAVWithBlitEncoder(UnorderedAccessViewRHI, EAGXRHIClearUAVType::StructuredBuffer, *(uint32*)&Values.X);
+	}
+	else
+#endif // UE_METAL_RHI_SUPPORT_CLEAR_UAV_WITH_BLIT_ENCODER
+	{
+		TRHICommandList_RecursiveHazardous<FAGXRHICommandContext> RHICmdList(this);
+		ClearUAV(RHICmdList, ResourceCast(UnorderedAccessViewRHI), &Values, true);
+	}
+}
+
+void FAGXRHICommandContext::RHIClearUAVUint(FRHIUnorderedAccessView* UnorderedAccessViewRHI, const FUintVector4& Values)
+{
+#if UE_METAL_RHI_SUPPORT_CLEAR_UAV_WITH_BLIT_ENCODER
+	FAGXUnorderedAccessView* UnorderedAccessView = ResourceCast(UnorderedAccessViewRHI);
+	if (UnorderedAccessView->SourceView->SourceStructuredBuffer)
+	{
+		ClearUAVWithBlitEncoder(UnorderedAccessViewRHI, EAGXRHIClearUAVType::StructuredBuffer, *(uint32*)&Values.X);
+	}
+	else
+#endif // UE_METAL_RHI_SUPPORT_CLEAR_UAV_WITH_BLIT_ENCODER
+	{
+		TRHICommandList_RecursiveHazardous<FAGXRHICommandContext> RHICmdList(this);
+		ClearUAV(RHICmdList, ResourceCast(UnorderedAccessViewRHI), &Values, false);
+	}
+}
+
+void FAGXRHICommandContext::ClearUAV(TRHICommandList_RecursiveHazardous<FAGXRHICommandContext>& RHICmdList, FAGXUnorderedAccessView* UnorderedAccessView, const void* ClearValue, bool bFloat)
+{
+	@autoreleasepool {
+		EClearReplacementValueType ValueType = bFloat ? EClearReplacementValueType::Float : EClearReplacementValueType::Uint32;
+
+		// The Metal validation layer will complain about resources with a
+		// signed format bound against an unsigned data format type as the
+		// shader parameter.
+		switch (GPixelFormats[UnorderedAccessView->SourceView->Format].UnrealFormat)
+		{
+			case PF_R32_SINT:
+			case PF_R16_SINT:
+			case PF_R16G16B16A16_SINT:
+				ValueType = EClearReplacementValueType::Int32;
+				break;
+				
+			default:
+				break;
+		}
+
+		if (UnorderedAccessView->SourceView->SourceVertexBuffer)
+		{
+#if UE_METAL_RHI_SUPPORT_CLEAR_UAV_WITH_BLIT_ENCODER
+			if (UnorderedAccessView->SourceView->SourceVertexBuffer->GetUsage() & BUF_ByteAddressBuffer)
+			{
+				ClearUAVWithBlitEncoder(UnorderedAccessView, EAGXRHIClearUAVType::VertexBuffer, *(const uint32*)ClearValue);
+			}
+			else
+#endif // UE_METAL_RHI_SUPPORT_CLEAR_UAV_WITH_BLIT_ENCODER
+			{
+				uint32 NumElements = UnorderedAccessView->SourceView->SourceVertexBuffer->GetSize() / GPixelFormats[UnorderedAccessView->SourceView->Format].BlockBytes;
+				ClearUAVShader_T<EClearReplacementResourceType::Buffer, 4, false>(RHICmdList, UnorderedAccessView, NumElements, 1, 1, ClearValue, ValueType);
+			}
+		}
+		else if (UnorderedAccessView->SourceView->SourceTexture)
+		{
+			FIntVector SizeXYZ = UnorderedAccessView->SourceView->SourceTexture->GetSizeXYZ();
+
+			if (FRHITexture2D* Texture2D = UnorderedAccessView->SourceView->SourceTexture->GetTexture2D())
+			{
+				ClearUAVShader_T<EClearReplacementResourceType::Texture2D, 4, false>(RHICmdList, UnorderedAccessView, SizeXYZ.X, SizeXYZ.Y, SizeXYZ.Z, ClearValue, ValueType);
+			}
+			else if (FRHITexture2DArray* Texture2DArray = UnorderedAccessView->SourceView->SourceTexture->GetTexture2DArray())
+			{
+				ClearUAVShader_T<EClearReplacementResourceType::Texture2DArray, 4, false>(RHICmdList, UnorderedAccessView, SizeXYZ.X, SizeXYZ.Y, SizeXYZ.Z, ClearValue, ValueType);
+			}
+			else if (FRHITexture3D* Texture3D = UnorderedAccessView->SourceView->SourceTexture->GetTexture3D())
+			{
+				ClearUAVShader_T<EClearReplacementResourceType::Texture3D, 4, false>(RHICmdList, UnorderedAccessView, SizeXYZ.X, SizeXYZ.Y, SizeXYZ.Z, ClearValue, ValueType);
+			}
+			else if (FRHITextureCube* TextureCube = UnorderedAccessView->SourceView->SourceTexture->GetTextureCube())
+			{
+				ClearUAVShader_T<EClearReplacementResourceType::Texture2DArray, 4, false>(RHICmdList, UnorderedAccessView, SizeXYZ.X, SizeXYZ.Y, SizeXYZ.Z, ClearValue, ValueType);
+			}
+			else
+			{
+				ensure(0);
+			}
+		}
+		else
+		{
+			// TODO: ensure(0);
+			UE_LOG(LogRHI, Warning, TEXT("AGX RHI ClearUAV does not yet support clearing of a UAV without a SourceView."));
+		}
+	} // @autoreleasepool
+}
+
+void FAGXRHICommandContext::RHIBeginTransitions(TArrayView<const FRHITransition*> Transitions)
+{
+	for (auto Transition : Transitions)
+	{
+		Transition->GetPrivateData<FAGXTransitionData>()->BeginResourceTransitions();
+	}
+}
+
+void FAGXRHICommandContext::RHIEndTransitions(TArrayView<const FRHITransition*> Transitions)
+{
+	for (auto Transition : Transitions)
+	{
+		Transition->GetPrivateData<FAGXTransitionData>()->EndResourceTransitions();
+	}
+}
+
+void FAGXGPUFence::WriteInternal(mtlpp::CommandBuffer& CmdBuffer)
+{
+	Fence = CmdBuffer.GetCompletionFence();
+	check(Fence);
+}
+
+void FAGXRHICommandContext::RHICopyToStagingBuffer(FRHIBuffer* SourceBufferRHI, FRHIStagingBuffer* DestinationStagingBufferRHI, uint32 Offset, uint32 NumBytes)
+{
+	@autoreleasepool {
+		check(DestinationStagingBufferRHI);
+
+		FAGXRHIStagingBuffer* AGXStagingBuffer = ResourceCast(DestinationStagingBufferRHI);
+		ensureMsgf(!AGXStagingBuffer->bIsLocked, TEXT("Attempting to Copy to a locked staging buffer. This may have undefined behavior"));
+		FAGXVertexBuffer* SourceBuffer = ResourceCast(SourceBufferRHI);
+		FAGXBuffer& ReadbackBuffer = AGXStagingBuffer->ShadowBuffer;
+
+		// Need a shadow buffer for this read. If it hasn't been allocated in our FStagingBuffer or if
+		// it's not big enough to hold our readback we need to allocate.
+		if (!ReadbackBuffer || ReadbackBuffer.GetLength() < NumBytes)
+		{
+			if (ReadbackBuffer)
+			{
+				AGXSafeReleaseMetalBuffer(ReadbackBuffer);
+			}
+			FAGXPooledBufferArgs ArgsCPU(GetAGXDeviceContext().GetDevice(), NumBytes, BUF_Dynamic, mtlpp::StorageMode::Shared);
+			ReadbackBuffer = GetAGXDeviceContext().CreatePooledBuffer(ArgsCPU);
+		}
+
+		// Inline copy from the actual buffer to the shadow
+		GetAGXDeviceContext().CopyFromBufferToBuffer(SourceBuffer->GetCurrentBuffer(), Offset, ReadbackBuffer, 0, NumBytes);
+	}
+}
+
+void FAGXRHICommandContext::RHIWriteGPUFence(FRHIGPUFence* FenceRHI)
+{
+	@autoreleasepool {
+		check(FenceRHI);
+		FAGXGPUFence* Fence = ResourceCast(FenceRHI);
+		Fence->WriteInternal(Context->GetCurrentCommandBuffer());
+	}
+}
+
+FGPUFenceRHIRef FAGXDynamicRHI::RHICreateGPUFence(const FName &Name)
+{
+	@autoreleasepool {
+	return new FAGXGPUFence(Name);
+	}
+}
+
+void FAGXGPUFence::Clear()
+{
+	Fence = mtlpp::CommandBufferFence();
+}
+
+bool FAGXGPUFence::Poll() const
+{
+	if (Fence)
+	{
+		return Fence.Wait(0);
+	}
+	else
+	{
+		return false;
+	}
+}
diff --git a/Engine/Source/Runtime/Apple/AGXRHI/Private/AGXUniformBuffer.cpp b/Engine/Source/Runtime/Apple/AGXRHI/Private/AGXUniformBuffer.cpp
new file mode 100644
--- /dev/null
+++ b/Engine/Source/Runtime/Apple/AGXRHI/Private/AGXUniformBuffer.cpp
@@ -0,0 +1,112 @@
+// Copyright Epic Games, Inc. All Rights Reserved.
+
+/*=============================================================================
+	AGXConstantBuffer.cpp: AGX RHI constant buffer implementation.
+=============================================================================*/
+
+#include "AGXRHIPrivate.h"
+#include "AGXFrameAllocator.h"
+#include "AGXUniformBuffer.h"
+#include "ShaderParameterStruct.h"
+
+#pragma mark Suballocated Uniform Buffer Implementation
+
+FAGXSuballocatedUniformBuffer::FAGXSuballocatedUniformBuffer(const FRHIUniformBufferLayout& Layout, EUniformBufferUsage Usage, EUniformBufferValidation InValidation)
+    : FRHIUniformBuffer(Layout)
+    , LastFrameUpdated(0)
+    , Offset(0)
+    , Backing(nil)
+    , Shadow(nullptr)
+    , ResourceTable()
+#if METAL_UNIFORM_BUFFER_VALIDATION
+    , Validation(InValidation)
+#endif // METAL_UNIFORM_BUFFER_VALIDATION
+{
+    // Slate can create SingleDraw uniform buffers and use them several frames later. So it must be included.
+    if (Usage == UniformBuffer_SingleDraw || Usage == UniformBuffer_MultiFrame)
+    {
+        Shadow = FMemory::Malloc(GetSize());
+    }
+}
+
+FAGXSuballocatedUniformBuffer::~FAGXSuballocatedUniformBuffer()
+{
+    if (HasShadow())
+    {
+        FMemory::Free(Shadow);
+    }
+
+    // Note: this object does NOT own a reference
+    // to the uniform buffer backing store
+}
+
+bool FAGXSuballocatedUniformBuffer::HasShadow()
+{
+    return Shadow != nullptr;
+}
+
+void FAGXSuballocatedUniformBuffer::Update(const void* Contents, TArray<TRefCountPtr<FRHIResource> > const& InResourceTable)
+{
+    if (HasShadow())
+    {
+        FMemory::Memcpy(Shadow, Contents, GetSize());
+    }
+
+	ResourceTable = InResourceTable;
+
+	PushToGPUBacking(Contents);
+}
+
+// Acquires a region in the current frame's uniform buffer and
+// pushes the data in Contents into that GPU backing store
+// The amount of data read from Contents is given by the Layout
+void FAGXSuballocatedUniformBuffer::PushToGPUBacking(const void* Contents)
+{
+    check(IsInRenderingThread() ^ IsRunningRHIInSeparateThread());
+    
+    FAGXDeviceContext& DeviceContext = GetAGXDeviceContext();
+    
+    FAGXFrameAllocator* Allocator = DeviceContext.GetUniformAllocator();
+    FAGXFrameAllocator::AllocationEntry Entry = Allocator->AcquireSpace(GetSize());
+    // copy contents into backing
+    Backing = Entry.Backing;
+    Offset = Entry.Offset;
+    uint8* ConstantSpace = reinterpret_cast<uint8*>([Backing contents]) + Entry.Offset;
+    FMemory::Memcpy(ConstantSpace, Contents, GetSize());
+    LastFrameUpdated = DeviceContext.GetFrameNumberRHIThread();
+}
+
+// Because we can create a uniform buffer on frame N and may not bind it until frame N+10
+// we need to keep a copy of the most recent data. Then when it's time to bind this
+// uniform buffer we can push the data into the GPU backing.
+void FAGXSuballocatedUniformBuffer::PrepareToBind()
+{
+    FAGXDeviceContext& DeviceContext = GetAGXDeviceContext();
+    if(Shadow && LastFrameUpdated < DeviceContext.GetFrameNumberRHIThread())
+    {
+        PushToGPUBacking(Shadow);
+    }
+}
+
+void FAGXSuballocatedUniformBuffer::CopyResourceTable_RenderThread(const void* Contents, TArray<TRefCountPtr<FRHIResource> >& OutResourceTable)
+{
+#if METAL_UNIFORM_BUFFER_VALIDATION
+	if (Validation == EUniformBufferValidation::ValidateResources)
+	{
+		ValidateShaderParameterResourcesRHI(Contents, GetLayout());
+	}
+#endif // METAL_UNIFORM_BUFFER_VALIDATION
+
+	const FRHIUniformBufferLayout& Layout = GetLayout();
+    const uint32 NumResources = Layout.Resources.Num();
+    if (NumResources > 0)
+    {
+		OutResourceTable.Empty(NumResources);
+		OutResourceTable.AddZeroed(NumResources);
+        
+        for (uint32 Index = 0; Index < NumResources; ++Index)
+		{
+			OutResourceTable[Index] = GetShaderParameterResourceRHI(Contents, Layout.Resources[Index].MemberOffset, Layout.Resources[Index].MemberType);
+		}
+    }
+}
diff --git a/Engine/Source/Runtime/Apple/AGXRHI/Private/AGXUniformBufferCommands.cpp b/Engine/Source/Runtime/Apple/AGXRHI/Private/AGXUniformBufferCommands.cpp
new file mode 100644
--- /dev/null
+++ b/Engine/Source/Runtime/Apple/AGXRHI/Private/AGXUniformBufferCommands.cpp
@@ -0,0 +1,104 @@
+// Copyright Epic Games, Inc. All Rights Reserved.
+
+#include "AGXRHIPrivate.h"
+#include "AGXShaderTypes.h"
+#include "AGXFrameAllocator.h"
+
+static void DoUpdateUniformBuffer(FAGXUniformBuffer* UB, const void* Contents)
+{
+    check(IsInRenderingThread() || IsInRHIThread());
+    
+    FRHICommandListImmediate& RHICmdList = FRHICommandListExecutor::GetImmediateCommandList();
+    // The only way we can be on the RHI thread here is if we're in the process of creating a FLocalUniformBuffer.
+    bool bUpdateImmediately = RHICmdList.Bypass() || IsInRHIThread();
+
+	TArray<TRefCountPtr<FRHIResource> > ResourceTable;
+
+	UB->CopyResourceTable_RenderThread(Contents, ResourceTable);
+
+    if(bUpdateImmediately)
+    {
+        UB->Update(Contents, ResourceTable);
+    }
+    else
+    {
+        const uint32 NumBytes = UB->GetLayout().ConstantBufferSize;
+        void* Data = RHICmdList.Alloc(NumBytes, 16);
+        FMemory::Memcpy(Data, Contents, NumBytes);
+        
+        RHICmdList.EnqueueLambda([Data, UB, NewResourceTable = MoveTemp(ResourceTable)](FRHICommandListImmediate& RHICmdList)
+        {
+            UB->Update(Data, NewResourceTable);
+        });
+        
+        RHICmdList.RHIThreadFence(true);
+    }
+}
+
+FUniformBufferRHIRef FAGXDynamicRHI::RHICreateUniformBuffer(const void* Contents, const FRHIUniformBufferLayout& Layout, EUniformBufferUsage Usage, EUniformBufferValidation Validation)
+{
+    FAGXDeviceContext& DeviceContext = (FAGXDeviceContext&)GetAGXDeviceContext();
+    FAGXFrameAllocator* UniformAllocator = DeviceContext.GetUniformAllocator();
+    
+    FAGXUniformBuffer* UB = new FAGXUniformBuffer(Layout, Usage, Validation);
+    
+    DoUpdateUniformBuffer(UB, Contents);
+    
+    return UB;
+}
+
+void FAGXDynamicRHI::RHIUpdateUniformBuffer(FRHIUniformBuffer* UniformBufferRHI, const void* Contents)
+{
+    check(IsInRenderingThread());
+    
+    FAGXUniformBuffer* UB = ResourceCast(UniformBufferRHI);
+    DoUpdateUniformBuffer(UB, Contents);
+}
+
+template<EAGXShaderStages Stage, typename RHIShaderType>
+static void SetUniformBufferInternal(FAGXContext* Context, RHIShaderType* ShaderRHI, uint32 BufferIndex, FRHIUniformBuffer* UBRHI)
+{
+    @autoreleasepool
+    {
+        auto Shader = ResourceCast(ShaderRHI);
+        Context->GetCurrentState().BindUniformBuffer(Stage, BufferIndex, UBRHI);
+        
+        auto& Bindings = Shader->Bindings;
+        if((Bindings.ConstantBuffers) & (1 << BufferIndex))
+        {
+            FAGXUniformBuffer* UB = ResourceCast(UBRHI);
+            UB->PrepareToBind();
+            
+            FAGXBuffer Buf(UB->Backing, ns::Ownership::AutoRelease);
+            Context->GetCurrentState().SetShaderBuffer(Stage, Buf, nil, UB->Offset, UB->GetSize(), BufferIndex, mtlpp::ResourceUsage::Read);
+        }
+    }
+}
+
+void FAGXRHICommandContext::RHISetShaderUniformBuffer(FRHIGraphicsShader* Shader, uint32 BufferIndex, FRHIUniformBuffer* Buffer)
+{
+	switch (Shader->GetFrequency())
+	{
+		case SF_Vertex:
+			SetUniformBufferInternal<EAGXShaderStages::Vertex, FRHIVertexShader>(Context, static_cast<FRHIVertexShader*>(Shader), BufferIndex, Buffer);
+			break;
+
+		case SF_Geometry:
+			NOT_SUPPORTED("RHISetShaderUniformBuffer-Geometry");
+			break;
+
+		case SF_Pixel:
+			SetUniformBufferInternal<EAGXShaderStages::Pixel, FRHIPixelShader>(Context, static_cast<FRHIPixelShader*>(Shader), BufferIndex, Buffer);
+			break;
+
+		default:
+			checkf(0, TEXT("FRHIShader Type %d is invalid or unsupported!"), (int32)Shader->GetFrequency());
+			NOT_SUPPORTED("RHIShaderStage");
+			break;
+	}
+}
+
+void FAGXRHICommandContext::RHISetShaderUniformBuffer(FRHIComputeShader* ComputeShaderRHI, uint32 BufferIndex, FRHIUniformBuffer* BufferRHI)
+{
+	SetUniformBufferInternal<EAGXShaderStages::Compute, FRHIComputeShader>(Context, ComputeShaderRHI, BufferIndex, BufferRHI);
+}
diff --git a/Engine/Source/Runtime/Apple/AGXRHI/Private/AGXVertexBuffer.cpp b/Engine/Source/Runtime/Apple/AGXRHI/Private/AGXVertexBuffer.cpp
new file mode 100644
--- /dev/null
+++ b/Engine/Source/Runtime/Apple/AGXRHI/Private/AGXVertexBuffer.cpp
@@ -0,0 +1,920 @@
+// Copyright Epic Games, Inc. All Rights Reserved.
+
+/*=============================================================================
+	AGXVertexBuffer.cpp: AGX vertex buffer RHI implementation.
+=============================================================================*/
+
+#include "AGXRHIPrivate.h"
+#include "AGXProfiler.h"
+#include "AGXCommandBuffer.h"
+#include "AGXCommandQueue.h"
+#include "Containers/ResourceArray.h"
+#include "RenderUtils.h"
+#include "AGXLLM.h"
+#include <objc/runtime.h>
+
+#define METAL_POOL_BUFFER_BACKING 1
+
+#if !METAL_POOL_BUFFER_BACKING
+DECLARE_MEMORY_STAT(TEXT("Used Device Buffer Memory"), STAT_AGXDeviceBufferMemory, STATGROUP_AGXRHI);
+#endif
+
+#if STATS
+#define METAL_INC_DWORD_STAT_BY(Type, Name, Size, Usage) \
+	switch(Type)	{ \
+		case RRT_UniformBuffer: INC_DWORD_STAT_BY(STAT_AGXUniform##Name, Size); break; \
+        case RRT_Buffer: if (Usage & BUF_IndexBuffer){ INC_DWORD_STAT_BY(STAT_AGXIndex##Name, Size); } else { INC_DWORD_STAT_BY(STAT_AGXVertex##Name, Size); } break; \
+		default: break; \
+	}
+#else
+#define METAL_INC_DWORD_STAT_BY(Type, Name, Size, Usage)
+#endif
+
+@implementation FAGXBufferData
+
+APPLE_PLATFORM_OBJECT_ALLOC_OVERRIDES(FAGXBufferData)
+
+-(instancetype)init
+{
+	id Self = [super init];
+	if (Self)
+	{
+		self->Data = nullptr;
+		self->Len = 0;
+	}
+	return Self;
+}
+-(instancetype)initWithSize:(uint32)InSize
+{
+	id Self = [super init];
+	if (Self)
+	{
+		self->Data = (uint8*)FMemory::Malloc(InSize);
+		self->Len = InSize;
+		check(self->Data);
+	}
+	return Self;
+}
+-(instancetype)initWithBytes:(void const*)InData length:(uint32)InSize
+{
+	id Self = [super init];
+	if (Self)
+	{
+		self->Data = (uint8*)FMemory::Malloc(InSize);
+		self->Len = InSize;
+		check(self->Data);
+		FMemory::Memcpy(self->Data, InData, InSize);
+	}
+	return Self;
+}
+-(void)dealloc
+{
+	if (self->Data)
+	{
+		FMemory::Free(self->Data);
+		self->Data = nullptr;
+		self->Len = 0;
+	}
+	[super dealloc];
+}
+@end
+
+static uint32 MetalBufferUsage(uint32 InUsage)
+{
+	uint32 Usage = InUsage;
+
+	if (InUsage & BUF_VertexBuffer)
+	{
+		Usage |= EAGXBufferUsage_LinearTex;
+	}
+
+	if (InUsage & BUF_IndexBuffer)
+	{
+		Usage |= (EAGXBufferUsage_GPUOnly | EAGXBufferUsage_LinearTex);
+	}
+
+	if (InUsage & BUF_StructuredBuffer)
+	{
+		Usage |= EAGXBufferUsage_GPUOnly;
+	}
+
+	return Usage;
+}
+
+void FAGXRHIBuffer::Swap(FAGXRHIBuffer& Other)
+{
+	::Swap(*this, Other);
+}
+
+static bool CanUsePrivateMemory()
+{
+	return (FAGXCommandQueue::SupportsFeature(EAGXFeaturesEfficientBufferBlits) || FAGXCommandQueue::SupportsFeature(EAGXFeaturesIABs));
+}
+
+bool FAGXRHIBuffer::UsePrivateMemory() const
+{
+	return (FAGXCommandQueue::SupportsFeature(EAGXFeaturesEfficientBufferBlits) && (Usage & (BUF_Dynamic|BUF_Static)))
+	|| (FAGXCommandQueue::SupportsFeature(EAGXFeaturesIABs) && (Usage & (BUF_ShaderResource|BUF_UnorderedAccess)));
+}
+
+FAGXRHIBuffer::FAGXRHIBuffer(uint32 InSize, uint32 InUsage, ERHIResourceType InType)
+: Data(nullptr)
+, LastLockFrame(0)
+, CurrentIndex(0)
+, NumberOfBuffers(0)
+, CurrentLockMode(RLM_Num)
+, LockOffset(0)
+, LockSize(0)
+, Size(InSize)
+, Usage(InUsage)
+, Mode(BUFFER_STORAGE_MODE)
+, Type(InType)
+{
+	// No life-time usage information? Enforce Dynamic.
+	if((Usage & (BUF_Static | BUF_Dynamic | BUF_Volatile)) == 0)
+	{
+		Usage |= BUF_Dynamic;
+	}
+	
+	const bool bIsStatic = (Usage & BUF_Static) != 0;
+	const bool bIsDynamic = (Usage & BUF_Dynamic) != 0;
+	const bool bIsVolatile = (Usage & BUF_Volatile) != 0;
+	const bool bWantsView = (Usage & (BUF_ShaderResource | BUF_UnorderedAccess)) != 0;
+	
+	check(bIsStatic ^ bIsDynamic ^ bIsVolatile);
+
+	Mode = UsePrivateMemory() ? mtlpp::StorageMode::Private : BUFFER_STORAGE_MODE;
+	Mode = CanUsePrivateMemory() ? mtlpp::StorageMode::Private : Mode;
+	
+	if (InSize)
+	{
+		checkf(InSize <= 1024 * 1024 * 1024, TEXT("Metal doesn't support buffers > 1GB"));
+		
+		// Temporary buffers less than the buffer page size - currently 4Kb - is better off going through the set*Bytes API if available.
+		// These can't be used for shader resources or UAVs if we want to use the 'Linear Texture' code path
+		if (!(InUsage & (BUF_UnorderedAccess|BUF_ShaderResource|EAGXBufferUsage_GPUOnly)) && (InUsage & BUF_Volatile) && InSize < AGXBufferPageSize && (InSize < AGXBufferBytesSize))
+		{
+			Data = [[FAGXBufferData alloc] initWithSize:InSize];
+			METAL_INC_DWORD_STAT_BY(Type, MemAlloc, InSize, Usage);
+		}
+		else
+		{
+			uint32 AllocSize = Size;
+			
+			if ((InUsage & EAGXBufferUsage_LinearTex) && !FAGXCommandQueue::SupportsFeature(EAGXFeaturesTextureBuffers))
+			{
+				if (InUsage & BUF_UnorderedAccess)
+				{
+					// Padding for write flushing when not using linear texture bindings for buffers
+					AllocSize = Align(AllocSize + 512, 1024);
+				}
+				
+				if (bWantsView)
+				{
+					uint32 NumElements = AllocSize;
+					uint32 SizeX = NumElements;
+					uint32 SizeY = 1;
+					uint32 Dimension = GMaxTextureDimensions;
+					while (SizeX > GMaxTextureDimensions)
+					{
+						while((NumElements % Dimension) != 0)
+						{
+							check(Dimension >= 1);
+							Dimension = (Dimension >> 1);
+						}
+						SizeX = Dimension;
+						SizeY = NumElements / Dimension;
+						if(SizeY > GMaxTextureDimensions)
+						{
+							Dimension <<= 1;
+							checkf(SizeX <= GMaxTextureDimensions, TEXT("Calculated width %u is greater than maximum permitted %d when converting buffer of size %u to a 2D texture."), Dimension, (int32)GMaxTextureDimensions, AllocSize);
+							if(Dimension <= GMaxTextureDimensions)
+							{
+								AllocSize = Align(Size, Dimension);
+								NumElements = AllocSize;
+								SizeX = NumElements;
+							}
+							else
+							{
+								// We don't know the Pixel Format and so the bytes per element for the potential linear texture
+								// Use max texture dimension as the align to be a worst case rather than crashing
+								AllocSize = Align(Size, GMaxTextureDimensions);
+								break;
+							}
+						}
+					}
+					
+					AllocSize = Align(AllocSize, 1024);
+				}
+			}
+			
+			// Static buffers will never be discarded. You can update them directly.
+			if(bIsStatic)
+			{
+				NumberOfBuffers = 1;
+			}
+			else
+			{
+				check(bIsDynamic || bIsVolatile);
+				NumberOfBuffers = 3;
+			}
+			
+			check(NumberOfBuffers > 0);
+			
+			BufferPool.SetNum(NumberOfBuffers);
+			
+			// These allocations will not go into the pool.
+			uint32 RequestedBufferOffsetAlignment = BufferOffsetAlignment;
+			if(bWantsView)
+			{
+				// Buffer backed linear textures have specific align requirements
+				// We don't know upfront the pixel format that may be requested for an SRV so we can't use minimumLinearTextureAlignmentForPixelFormat:
+				RequestedBufferOffsetAlignment = BufferBackedLinearTextureOffsetAlignment;
+			}
+			
+			AllocSize = Align(AllocSize, RequestedBufferOffsetAlignment);
+			for(uint32 i = 0; i < NumberOfBuffers; i++)
+			{
+				FAGXBufferAndViews& Backing = BufferPool[i];
+
+#if METAL_POOL_BUFFER_BACKING
+				FAGXPooledBufferArgs ArgsCPU(GetAGXDeviceContext().GetDevice(), AllocSize, Usage, Mode);
+				Backing.Buffer = GetAGXDeviceContext().CreatePooledBuffer(ArgsCPU);
+				Backing.Buffer.SetOwner(nullptr, false);
+#else
+				
+				NSUInteger Options = (((NSUInteger) Mode) << mtlpp::ResourceStorageModeShift);
+				
+				METAL_GPUPROFILE(FAGXScopedCPUStats CPUStat(FString::Printf(TEXT("AllocBuffer: %llu, %llu"), AllocSize, Options)));
+				// Allocate one.
+				Backing.Buffer = FAGXBuffer(MTLPP_VALIDATE(mtlpp::Device, GetAGXDeviceContext().GetDevice(), AGXSafeGetRuntimeDebuggingLevel() >= EAGXDebugLevelValidation, NewBuffer(AllocSize, (mtlpp::ResourceOptions) Options)), false);
+				
+				#if STATS || ENABLE_LOW_LEVEL_MEM_TRACKER
+					AGXLLM::LogAllocBuffer(GetAGXDeviceContext().GetDevice(), Backing.Buffer);
+				#endif
+				INC_MEMORY_STAT_BY(STAT_AGXDeviceBufferMemory, Backing.Buffer.GetLength());
+				
+				if (GAGXBufferZeroFill && Mode != mtlpp::StorageMode::Private)
+				{
+					FMemory::Memset(((uint8*)Backing.Buffer.GetContents()), 0, Backing.Buffer.GetLength());
+				}
+				
+				METAL_DEBUG_OPTION(GetAGXDeviceContext().ValidateIsInactiveBuffer(Backing.Buffer));
+				METAL_FATAL_ASSERT(Backing.Buffer, TEXT("Failed to create buffer of size %u and resource options %u"), Size, (uint32)Options);
+				
+				if(bIsStatic)
+				{
+					[Backing.Buffer.GetPtr() setLabel:[NSString stringWithFormat:@"Static on frame %u", GetAGXDeviceContext().GetFrameNumberRHIThread()]];
+				}
+				else
+				{
+					[Backing.Buffer.GetPtr() setLabel:[NSString stringWithFormat:@"buffer on frame %u", GetAGXDeviceContext().GetFrameNumberRHIThread()]];
+				}
+				
+#endif
+			}
+			
+			for(FAGXBufferAndViews& Backing : BufferPool)
+			{
+				check(Backing.Buffer);
+				check(AllocSize <= Backing.Buffer.GetLength());
+				check(Backing.Buffer.GetStorageMode() == Mode);
+				check(Backing.Views.Num() == 0);
+			}
+		}
+	}
+}
+
+FAGXRHIBuffer::~FAGXRHIBuffer()
+{
+	if(TransferBuffer)
+	{
+		METAL_INC_DWORD_STAT_BY(Type, MemFreed, TransferBuffer.GetLength(), Usage);
+		AGXSafeReleaseMetalBuffer(TransferBuffer);
+	}
+	
+	for(FAGXBufferAndViews& Backing : BufferPool)
+	{
+		check(Backing.Buffer);
+		
+		METAL_INC_DWORD_STAT_BY(Type, MemFreed, Backing.Buffer.GetLength(), Usage);
+		AGXSafeReleaseMetalBuffer(Backing.Buffer);
+		
+		for (auto& Pair : Backing.Views)
+		{
+			AGXSafeReleaseMetalTexture(Pair.Value);
+			Pair.Value = nil;
+		}
+		Backing.Views.Empty();
+	}
+
+	if (Data)
+	{
+		METAL_INC_DWORD_STAT_BY(Type, MemFreed, Size, Usage);
+		AGXSafeReleaseMetalObject(Data);
+	}
+}
+
+void FAGXRHIBuffer::AllocTransferBuffer(bool bOnRHIThread, uint32 InSize, EResourceLockMode LockMode)
+{
+	check(!TransferBuffer);
+	FAGXPooledBufferArgs ArgsCPU(GetAGXDeviceContext().GetDevice(), InSize, BUF_Dynamic, mtlpp::StorageMode::Shared);
+	TransferBuffer = GetAGXDeviceContext().CreatePooledBuffer(ArgsCPU);
+	TransferBuffer.SetOwner(nullptr, false);
+	check(TransferBuffer && TransferBuffer.GetPtr());
+	METAL_INC_DWORD_STAT_BY(Type, MemAlloc, InSize, Usage);
+	METAL_FATAL_ASSERT(TransferBuffer, TEXT("Failed to create buffer of size %u and storage mode %u"), InSize, (uint32)mtlpp::StorageMode::Shared);
+}
+
+void FAGXRHIBuffer::AllocLinearTextures(const LinearTextureMapKey& InLinearTextureMapKey)
+{
+	check(MetalIsSafeToUseRHIThreadResources());
+	
+	const bool bWantsView = ((Usage & (BUF_ShaderResource | BUF_UnorderedAccess)) != 0);
+	check(bWantsView);
+	{
+		FAGXBufferAndViews& CurrentBacking = GetCurrentBackingInternal();
+		FAGXBuffer& CurrentBuffer = CurrentBacking.Buffer;
+		
+		check(CurrentBuffer);
+		uint32 Length = CurrentBuffer.GetLength();
+		uint8 InFormat = InLinearTextureMapKey.Key;
+		const FAGXLinearTextureDescriptor& LinearTextureDesc = InLinearTextureMapKey.Value;
+		
+		mtlpp::PixelFormat MTLFormat = (mtlpp::PixelFormat)GAGXBufferFormats[InFormat].LinearTextureFormat;
+		
+		mtlpp::TextureDescriptor Desc;
+		NSUInteger TexUsage = mtlpp::TextureUsage::Unknown;
+		if (Usage & BUF_ShaderResource)
+		{
+			TexUsage |= mtlpp::TextureUsage::ShaderRead;
+		}
+		if (Usage & BUF_UnorderedAccess)
+		{
+			TexUsage |= mtlpp::TextureUsage::ShaderWrite;
+		}
+		
+		uint32 BytesPerElement = (0 == LinearTextureDesc.BytesPerElement) ? GPixelFormats[InFormat].BlockBytes : LinearTextureDesc.BytesPerElement;
+		if (MTLFormat == mtlpp::PixelFormat::RG11B10Float && MTLFormat != (mtlpp::PixelFormat)GPixelFormats[InFormat].PlatformFormat)
+		{
+			BytesPerElement = 4;
+		}
+
+		const uint32 MinimumByteAlignment = GetAGXDeviceContext().GetDevice().GetMinimumLinearTextureAlignmentForPixelFormat((mtlpp::PixelFormat)GAGXBufferFormats[InFormat].LinearTextureFormat);
+		const uint32 MinimumElementAlignment = MinimumByteAlignment / BytesPerElement;
+
+		uint32 Offset = LinearTextureDesc.StartOffsetBytes;
+		check(Offset % MinimumByteAlignment == 0);
+
+		uint32 NumElements = (UINT_MAX == LinearTextureDesc.NumElements) ? ((Size - Offset) / BytesPerElement) : LinearTextureDesc.NumElements;
+		NumElements = Align(NumElements, MinimumElementAlignment);
+
+		uint32 RowBytes = NumElements * BytesPerElement;
+
+		if (FAGXCommandQueue::SupportsFeature(EAGXFeaturesTextureBuffers))
+		{
+			NSUInteger Options = ((NSUInteger) Mode) << mtlpp::ResourceStorageModeShift;
+			Desc = mtlpp::TextureDescriptor::TextureBufferDescriptor(MTLFormat, NumElements, mtlpp::ResourceOptions(Options), mtlpp::TextureUsage(TexUsage));
+			Desc.SetAllowGPUOptimisedContents(false);
+		}
+		else
+		{
+			uint32 Width = NumElements;
+			uint32 Height = 1;
+
+			if (NumElements > GMaxTextureDimensions)
+			{
+				uint32 Dimension = GMaxTextureDimensions;
+				while ((NumElements % Dimension) != 0)
+				{
+					check(Dimension >= 1);
+					Dimension = (Dimension >> 1);
+				}
+
+				Width = Dimension;
+				Height = NumElements / Dimension;
+
+				// If we're just trying to fit as many elements as we can into
+				// the available buffer space, we can trim some padding at the
+				// end of the buffer in order to create widest possible linear
+				// texture that will fit.
+				if ((UINT_MAX == LinearTextureDesc.NumElements) && (Height > GMaxTextureDimensions))
+				{
+					Width = GMaxTextureDimensions;
+					Height = 1;
+
+					while ((Width * Height) < NumElements)
+					{
+						Height <<= 1;
+					}
+
+					while ((Width * Height) > NumElements)
+					{
+						Height -= 1;
+					}
+				}
+
+				checkf(Width <= GMaxTextureDimensions, TEXT("Calculated width %u is greater than maximum permitted %d when converting buffer of size %llu with element stride %u to a 2D texture with %u elements."), Width, (int32)GMaxTextureDimensions, Length, BytesPerElement, NumElements);
+				checkf(Height <= GMaxTextureDimensions, TEXT("Calculated height %u is greater than maximum permitted %d when converting buffer of size %llu with element stride %u to a 2D texture with %u elements."), Height, (int32)GMaxTextureDimensions, Length, BytesPerElement, NumElements);
+			}
+
+			RowBytes = Width * BytesPerElement;
+
+			check(RowBytes % MinimumByteAlignment == 0);
+			check((RowBytes * Height) + Offset <= Length);
+
+			Desc = mtlpp::TextureDescriptor::Texture2DDescriptor(MTLFormat, Width, Height, NO);
+			Desc.SetStorageMode(Mode);
+			Desc.SetCpuCacheMode(CurrentBuffer.GetCpuCacheMode());
+			Desc.SetUsage((mtlpp::TextureUsage)TexUsage);
+		}
+
+		for(FAGXBufferAndViews& Backing : BufferPool)
+		{
+			FAGXBuffer& Buffer = Backing.Buffer;
+			FAGXTexture NewTexture = MTLPP_VALIDATE(mtlpp::Buffer, Buffer, AGXSafeGetRuntimeDebuggingLevel() >= EAGXDebugLevelValidation, NewTexture(Desc, Offset, RowBytes));
+			METAL_FATAL_ASSERT(NewTexture, TEXT("Failed to create linear texture, desc %s from buffer %s"), *FString([Desc description]), *FString([Buffer description]));
+			
+			check(GAGXBufferFormats[InFormat].LinearTextureFormat == mtlpp::PixelFormat::RG11B10Float || GAGXBufferFormats[InFormat].LinearTextureFormat == (mtlpp::PixelFormat)NewTexture.GetPixelFormat());
+			Backing.Views.Add(InLinearTextureMapKey, NewTexture);
+		}
+	}
+	
+	for(FAGXBufferAndViews& Backing : BufferPool)
+	{
+		LinearTextureMap& Views = Backing.Views;
+		check(Views.Find(InLinearTextureMapKey) != nullptr);
+	}
+}
+
+struct FAGXRHICommandCreateLinearTexture : public FRHICommand<FAGXRHICommandCreateLinearTexture>
+{
+	FAGXRHIBuffer* Buffer;
+	TRefCountPtr<FRHIResource> Parent;
+	EPixelFormat Format;
+	FAGXLinearTextureDescriptor LinearTextureDesc;
+	
+	FORCEINLINE_DEBUGGABLE FAGXRHICommandCreateLinearTexture(FAGXRHIBuffer* InBuffer, FRHIResource* InParent, EPixelFormat InFormat, const FAGXLinearTextureDescriptor* InLinearTextureDescriptor)
+		: Buffer(InBuffer)
+		, Parent(InParent)
+		, Format(InFormat)
+		, LinearTextureDesc()
+	{
+		if (InLinearTextureDescriptor)
+		{
+			LinearTextureDesc = *InLinearTextureDescriptor;
+		}
+	}
+	
+	virtual ~FAGXRHICommandCreateLinearTexture()
+	{
+	}
+	
+	void Execute(FRHICommandListBase& CmdList)
+	{
+		check(MetalIsSafeToUseRHIThreadResources());
+		Buffer->CreateLinearTexture(Format, Parent.GetReference(), &LinearTextureDesc);
+	}
+};
+
+void FAGXRHIBuffer::CreateLinearTexture(EPixelFormat InFormat, FRHIResource* InParent, const FAGXLinearTextureDescriptor* InLinearTextureDescriptor)
+{
+	if ((Usage & (BUF_UnorderedAccess|BUF_ShaderResource)) && GAGXBufferFormats[InFormat].LinearTextureFormat != mtlpp::PixelFormat::Invalid)
+	{
+		if (IsRunningRHIInSeparateThread() && !IsInRHIThread() && !FRHICommandListExecutor::GetImmediateCommandList().Bypass())
+		{
+			new (FRHICommandListExecutor::GetImmediateCommandList().AllocCommand<FAGXRHICommandCreateLinearTexture>()) FAGXRHICommandCreateLinearTexture(this, InParent, InFormat, InLinearTextureDescriptor);
+		}
+		else
+		{
+			check(MetalIsSafeToUseRHIThreadResources());
+			LinearTextureMapKey MapKey = (InLinearTextureDescriptor != nullptr) ? LinearTextureMapKey(InFormat, *InLinearTextureDescriptor) : LinearTextureMapKey(InFormat, FAGXLinearTextureDescriptor());
+			
+			FAGXBufferAndViews& Backing = GetCurrentBackingInternal();
+
+			FAGXTexture* ExistingTexture = Backing.Views.Find(MapKey);
+			if (!ExistingTexture)
+			{
+				AllocLinearTextures(MapKey);
+			}
+		}
+	}
+	
+}
+
+ns::AutoReleased<FAGXTexture> FAGXRHIBuffer::GetLinearTexture(EPixelFormat InFormat, const FAGXLinearTextureDescriptor* InLinearTextureDescriptor)
+{
+	ns::AutoReleased<FAGXTexture> Texture;
+	if ((Usage & (BUF_UnorderedAccess|BUF_ShaderResource)) && GAGXBufferFormats[InFormat].LinearTextureFormat != mtlpp::PixelFormat::Invalid)
+	{
+		LinearTextureMapKey MapKey = (InLinearTextureDescriptor != nullptr) ? LinearTextureMapKey(InFormat, *InLinearTextureDescriptor) : LinearTextureMapKey(InFormat, FAGXLinearTextureDescriptor());
+
+		FAGXBufferAndViews& Backing = GetCurrentBackingInternal();
+		FAGXTexture* ExistingTexture = Backing.Views.Find(MapKey);
+		if (ExistingTexture)
+		{
+			Texture = *ExistingTexture;
+		}
+	}
+	return Texture;
+}
+
+void* FAGXRHIBuffer::Lock(bool bIsOnRHIThread, EResourceLockMode InLockMode, uint32 Offset, uint32 InSize)
+{
+	check(CurrentLockMode == RLM_Num);
+	check(LockSize == 0 && LockOffset == 0);
+	check(MetalIsSafeToUseRHIThreadResources());
+	check(!TransferBuffer);
+	
+	if (Data)
+	{
+		check(Data->Data);
+		return ((uint8*)Data->Data) + Offset;
+	}
+	
+	// The system is very naughty and does not obey this rule
+//	check(LastLockFrame == 0 || LastLockFrame != GetAGXDeviceContext().GetFrameNumberRHIThread());
+	
+	const bool bWriteLock = InLockMode == RLM_WriteOnly;
+	const bool bIsStatic = (Usage & BUF_Static) != 0;
+	const bool bIsDynamic = (Usage & BUF_Dynamic) != 0;
+	const bool bIsVolatile = (Usage & BUF_Volatile) != 0;
+	
+	void* ReturnPointer = nullptr;
+	
+	uint32 Len = GetCurrentBacking().Buffer.GetLength(); // all buffers should have the same length or we are in trouble.
+	check(Len >= InSize);
+	
+	if(bWriteLock)
+	{
+		if(bIsStatic)
+		{
+			// Static buffers do not discard. They just return the buffer or a transfer buffer.
+			// You are not supposed to lock more than once a frame.
+		}
+		else
+		{
+			check(bIsDynamic || bIsVolatile);
+			// cycle to next allocation
+			AdvanceBackingIndex();
+		}
+		
+		if(Mode == mtlpp::StorageMode::Private)
+		{
+			FAGXFrameAllocator::AllocationEntry TempBacking = GetAGXDeviceContext().GetTransferAllocator()->AcquireSpace(Len);
+			GetAGXDeviceContext().NewLock(this, TempBacking);
+			check(TempBacking.Backing);
+			ReturnPointer = (uint8*) [TempBacking.Backing contents] + TempBacking.Offset;
+		}
+		else
+		{
+			check(GetCurrentBacking().Buffer);
+			ReturnPointer = GetCurrentBackingInternal().Buffer.GetContents();
+		}
+		check(ReturnPointer != nullptr);
+	}
+	else
+	{
+		check(InLockMode == EResourceLockMode::RLM_ReadOnly);
+		// assumes offset is 0 for reads.
+		check(Offset == 0);
+		
+		if(Mode == mtlpp::StorageMode::Private)
+		{
+			SCOPE_CYCLE_COUNTER(STAT_AGXBufferPageOffTime);
+			AllocTransferBuffer(true, Len, RLM_WriteOnly);
+			check(TransferBuffer.GetLength() >= InSize);
+			
+			// Synchronise the buffer with the CPU
+			GetAGXDeviceContext().CopyFromBufferToBuffer(GetCurrentBacking().Buffer, 0, TransferBuffer, 0, GetCurrentBacking().Buffer.GetLength());
+			
+			//kick the current command buffer.
+			GetAGXDeviceContext().SubmitCommandBufferAndWait();
+			
+			ReturnPointer = TransferBuffer.GetContents();
+		}
+		#if PLATFORM_MAC
+		else if(Mode == mtlpp::StorageMode::Managed)
+		{
+			SCOPE_CYCLE_COUNTER(STAT_AGXBufferPageOffTime);
+			
+			// Synchronise the buffer with the CPU
+			GetAGXDeviceContext().SynchroniseResource(GetCurrentBacking().Buffer);
+			
+			//kick the current command buffer.
+			GetAGXDeviceContext().SubmitCommandBufferAndWait();
+			
+			ReturnPointer = GetCurrentBackingInternal().Buffer.GetContents();
+		}
+		#endif
+		else
+		{
+			// Shared
+			ReturnPointer = GetCurrentBackingInternal().Buffer.GetContents();
+		}
+	} // Read Path
+	
+	
+	
+	check(GetCurrentBacking().Buffer);
+	check(!GetCurrentBacking().Buffer.IsAliasable());
+	
+	check(ReturnPointer);
+	LockOffset = Offset;
+	LockSize = InSize;
+	CurrentLockMode = InLockMode;
+	
+	if(InSize == 0)
+	{
+		LockSize = Len;
+	}
+	
+	ReturnPointer = ((uint8*) (ReturnPointer)) + Offset;
+	return ReturnPointer;
+}
+
+void FAGXRHIBuffer::Unlock()
+{
+	check(MetalIsSafeToUseRHIThreadResources());
+
+	if(!Data)
+	{
+		FAGXBufferAndViews& Backing = GetCurrentBackingInternal();
+		FAGXBuffer& CurrentBuffer = Backing.Buffer;
+		
+		check(CurrentBuffer);
+		check(LockSize > 0);
+		const bool bWriteLock = CurrentLockMode == RLM_WriteOnly;
+		
+		if(bWriteLock)
+		{
+			check(!TransferBuffer);
+			check(LockOffset == 0);
+			check(LockSize <= CurrentBuffer.GetLength());
+			if(Mode == mtlpp::StorageMode::Private)
+			{
+				FAGXFrameAllocator::AllocationEntry Entry = GetAGXDeviceContext().FetchAndRemoveLock(this);
+				FAGXBuffer Transfer = Entry.Backing;
+				GetAGXDeviceContext().AsyncCopyFromBufferToBuffer(Transfer, Entry.Offset, CurrentBuffer, 0, LockSize);
+			}
+#if PLATFORM_MAC
+			else if(Mode == mtlpp::StorageMode::Managed)
+			{
+				if (GAGXBufferZeroFill)
+					MTLPP_VALIDATE(mtlpp::Buffer, CurrentBuffer, AGXSafeGetRuntimeDebuggingLevel() >= EAGXDebugLevelValidation, DidModify(ns::Range(0, CurrentBuffer.GetLength())));
+				else
+					MTLPP_VALIDATE(mtlpp::Buffer, CurrentBuffer, AGXSafeGetRuntimeDebuggingLevel() >= EAGXDebugLevelValidation, DidModify(ns::Range(LockOffset, LockSize)));
+			}
+#endif //PLATFORM_MAC
+			else
+			{
+				// shared buffers are always mapped so nothing happens
+				check(Mode == mtlpp::StorageMode::Shared);
+			}
+		}
+		else
+		{
+			check(CurrentLockMode == RLM_ReadOnly);
+			if(TransferBuffer)
+			{
+				check(Mode == mtlpp::StorageMode::Private);
+				AGXSafeReleaseMetalBuffer(TransferBuffer);
+				TransferBuffer = nil;
+			}
+		}
+	}
+	
+	check(!TransferBuffer);
+	CurrentLockMode = RLM_Num;
+	LockSize = 0;
+	LockOffset = 0;
+	LastLockFrame = GetAGXDeviceContext().GetFrameNumberRHIThread();
+}
+
+FAGXResourceMultiBuffer::FAGXResourceMultiBuffer(uint32 InSize, uint32 InUsage, uint32 InStride, FResourceArrayInterface* ResourceArray, ERHIResourceType Type)
+	: FRHIBuffer(InSize, InUsage & ~EAGXBufferUsageFlags, InStride)
+	, FAGXRHIBuffer(InSize, InUsage, Type)
+	, IndexType((InStride == 2) ? mtlpp::IndexType::UInt16 : mtlpp::IndexType::UInt32)
+{
+	if (InUsage & BUF_StructuredBuffer)
+	{
+		check((InSize % InStride) == 0);
+
+		if (ResourceArray)
+		{
+			// copy any resources to the CPU address
+			void* LockedMemory = RHILockBuffer(this, 0, InSize, RLM_WriteOnly);
+			FMemory::Memcpy(LockedMemory, ResourceArray->GetResourceData(), InSize);
+			ResourceArray->Discard();
+			RHIUnlockBuffer(this);
+		}
+	}
+}
+
+FAGXResourceMultiBuffer::~FAGXResourceMultiBuffer()
+{
+}
+
+void FAGXResourceMultiBuffer::Swap(FAGXResourceMultiBuffer& Other)
+{
+	@autoreleasepool {
+		FRHIBuffer::Swap(Other);
+		FAGXRHIBuffer::Swap(Other);
+		::Swap(IndexType, Other.IndexType);
+	}
+}
+
+FBufferRHIRef FAGXDynamicRHI::RHICreateBuffer(uint32 Size, EBufferUsageFlags Usage, uint32 Stride, ERHIAccess InResourceState, FRHIResourceCreateInfo& CreateInfo)
+{
+	check(0);
+	@autoreleasepool {
+	if (CreateInfo.bWithoutNativeResource)
+	{
+		return new FAGXResourceMultiBuffer(0, MetalBufferUsage(0), 0, nullptr, RRT_Buffer);
+	}
+	
+	// make the RHI object, which will allocate memory
+	FAGXResourceMultiBuffer* Buffer = new FAGXResourceMultiBuffer(Size, MetalBufferUsage(Usage), 0, nullptr, RRT_Buffer);
+
+	if (CreateInfo.ResourceArray)
+	{
+		check(Size >= CreateInfo.ResourceArray->GetResourceDataSize());
+
+		// make a buffer usable by CPU
+		void* BufferData = ::RHILockBuffer(Buffer, 0, Size, RLM_WriteOnly);
+		
+		// copy the contents of the given data into the buffer
+		FMemory::Memcpy(BufferData, CreateInfo.ResourceArray->GetResourceData(), Size);
+		
+		::RHIUnlockBuffer(Buffer);
+
+		// Discard the resource array's contents.
+		CreateInfo.ResourceArray->Discard();
+	}
+	else if (Buffer->Mode == mtlpp::StorageMode::Private)
+	{
+		check (!Buffer->TransferBuffer);
+
+		if (GAGXBufferZeroFill && !FAGXCommandQueue::SupportsFeature(EAGXFeaturesFences))
+		{
+			for(FAGXRHIBuffer::FAGXBufferAndViews& Backing : Buffer->BufferPool)
+			{
+				FAGXBuffer& TheBuffer = Backing.Buffer;
+				GetAGXDeviceContext().FillBuffer(TheBuffer, ns::Range(0, TheBuffer.GetLength()), 0);
+			}
+		}
+	}
+#if PLATFORM_MAC
+	else if (GAGXBufferZeroFill && Buffer->Mode == mtlpp::StorageMode::Managed)
+	{
+		for(FAGXRHIBuffer::FAGXBufferAndViews& Backing : Buffer->BufferPool)
+		{
+			FAGXBuffer& TheBuffer = Backing.Buffer;
+			GetAGXDeviceContext().FillBuffer(TheBuffer, ns::Range(0, TheBuffer.GetLength()), 0);
+			MTLPP_VALIDATE(mtlpp::Buffer, TheBuffer, AGXSafeGetRuntimeDebuggingLevel() >= EAGXDebugLevelValidation, DidModify(ns::Range(0, TheBuffer)));
+		}
+	}
+#endif
+
+	return Buffer;
+	}
+}
+
+void* FAGXDynamicRHI::LockBuffer_BottomOfPipe(FRHICommandListImmediate& RHICmdList, FRHIBuffer* BufferRHI, uint32 Offset, uint32 Size, EResourceLockMode LockMode)
+{
+	@autoreleasepool {
+	FAGXResourceMultiBuffer* Buffer = ResourceCast(BufferRHI);
+
+	// default to buffer memory
+	return (uint8*)Buffer->Lock(true, LockMode, Offset, Size);
+	}
+}
+
+void FAGXDynamicRHI::UnlockBuffer_BottomOfPipe(FRHICommandListImmediate& RHICmdList, FRHIBuffer* BufferRHI)
+{
+	@autoreleasepool {
+	FAGXResourceMultiBuffer* Buffer = ResourceCast(BufferRHI);
+
+	Buffer->Unlock();
+	}
+}
+
+void FAGXDynamicRHI::RHICopyBuffer(FRHIBuffer* SourceBufferRHI, FRHIBuffer* DestBufferRHI)
+{
+	@autoreleasepool {
+		FAGXResourceMultiBuffer* SrcBuffer = ResourceCast(SourceBufferRHI);
+		FAGXResourceMultiBuffer* DstBuffer = ResourceCast(DestBufferRHI);
+		
+		const FAGXBuffer& TheSrcBuffer = SrcBuffer->GetCurrentBuffer();
+		const FAGXBuffer& TheDstBuffer = DstBuffer->GetCurrentBuffer();
+	
+		if (TheSrcBuffer && TheDstBuffer)
+		{
+			GetAGXDeviceContext().CopyFromBufferToBuffer(TheSrcBuffer, 0, TheDstBuffer, 0, FMath::Min(SrcBuffer->GetSize(), DstBuffer->GetSize()));
+		}
+		else if (TheDstBuffer)
+		{
+			FAGXPooledBufferArgs ArgsCPU(GetAGXDeviceContext().GetDevice(), SrcBuffer->GetSize(), BUF_Dynamic, mtlpp::StorageMode::Shared);
+			FAGXBuffer TempBuffer = GetAGXDeviceContext().CreatePooledBuffer(ArgsCPU);
+			FMemory::Memcpy(TempBuffer.GetContents(), SrcBuffer->Data->Data, SrcBuffer->GetSize());
+			GetAGXDeviceContext().CopyFromBufferToBuffer(TempBuffer, 0, TheDstBuffer, 0, FMath::Min(SrcBuffer->GetSize(), DstBuffer->GetSize()));
+			AGXSafeReleaseMetalBuffer(TempBuffer);
+		}
+		else
+		{
+			void const* SrcData = SrcBuffer->Lock(true, RLM_ReadOnly, 0);
+			void* DstData = DstBuffer->Lock(true, RLM_WriteOnly, 0);
+			FMemory::Memcpy(DstData, SrcData, FMath::Min(SrcBuffer->GetSize(), DstBuffer->GetSize()));
+			SrcBuffer->Unlock();
+			DstBuffer->Unlock();
+		}
+	}
+}
+
+void FAGXDynamicRHI::RHITransferBufferUnderlyingResource(FRHIBuffer* DestBuffer, FRHIBuffer* SrcBuffer)
+{
+	@autoreleasepool {
+		check(DestBuffer);
+		FAGXResourceMultiBuffer* Dest = ResourceCast(DestBuffer);
+		if (!SrcBuffer)
+		{
+			TRefCountPtr<FAGXResourceMultiBuffer> DeletionProxy = new FAGXResourceMultiBuffer(0, Dest->GetUsage(), Dest->GetStride(), nullptr, Dest->Type);
+			Dest->Swap(*DeletionProxy);
+		}
+		else
+		{
+			FAGXResourceMultiBuffer* Src = ResourceCast(SrcBuffer);
+			Dest->Swap(*Src);
+		}
+	}
+}
+
+void FAGXRHIBuffer::Init_RenderThread(FRHICommandListImmediate& RHICmdList, uint32 InSize, uint32 InUsage, FRHIResourceCreateInfo& CreateInfo, FRHIResource* Resource)
+{
+	if (CreateInfo.ResourceArray)
+	{
+		check(InSize == CreateInfo.ResourceArray->GetResourceDataSize());
+		
+		if(Data)
+		{
+			FMemory::Memcpy(Data->Data, CreateInfo.ResourceArray->GetResourceData(), InSize);
+		}
+		else
+		{
+			if (Mode == mtlpp::StorageMode::Private)
+			{
+				if (RHICmdList.IsBottomOfPipe())
+				{
+					void* Backing = this->Lock(true, RLM_WriteOnly, 0, InSize);
+					FMemory::Memcpy(Backing, CreateInfo.ResourceArray->GetResourceData(), InSize);
+					this->Unlock();
+				}
+				else
+				{
+					void* Result = FMemory::Malloc(InSize, 16);
+					FMemory::Memcpy(Result, CreateInfo.ResourceArray->GetResourceData(), InSize);
+					
+					RHICmdList.EnqueueLambda(
+						[this, Result, InSize](FRHICommandList& RHICmdList)
+						{
+							void* Backing = this->Lock(true, RLM_WriteOnly, 0, InSize);
+							FMemory::Memcpy(Backing, Result, InSize);
+							this->Unlock();
+							FMemory::Free(Result);
+						});
+				}
+			}
+			else
+			{
+				FAGXBuffer& TheBuffer = GetCurrentBufferInternal();
+				FMemory::Memcpy(TheBuffer.GetContents(), CreateInfo.ResourceArray->GetResourceData(), InSize);
+	#if PLATFORM_MAC
+				if(Mode == mtlpp::StorageMode::Managed)
+				{
+					MTLPP_VALIDATE(mtlpp::Buffer, TheBuffer, AGXSafeGetRuntimeDebuggingLevel() >= EAGXDebugLevelValidation, DidModify(ns::Range(0, GAGXBufferZeroFill ? TheBuffer.GetLength() : InSize)));
+				}
+	#endif
+			}
+		}
+		
+		// Discard the resource array's contents.
+		CreateInfo.ResourceArray->Discard();
+	}
+}
+
+FBufferRHIRef FAGXDynamicRHI::CreateBuffer_RenderThread(class FRHICommandListImmediate& RHICmdList, uint32 Size, EBufferUsageFlags Usage, uint32 Stride, ERHIAccess ResourceState, FRHIResourceCreateInfo& CreateInfo)
+{
+	@autoreleasepool {
+		if (CreateInfo.bWithoutNativeResource)
+		{
+			return new FAGXResourceMultiBuffer(0, MetalBufferUsage(0), 0, nullptr, RRT_Buffer);
+		}
+		
+		// make the RHI object, which will allocate memory
+		TRefCountPtr<FAGXResourceMultiBuffer> Buffer = new FAGXResourceMultiBuffer(Size, MetalBufferUsage(Usage), Stride, nullptr, RRT_Buffer);
+		
+		Buffer->Init_RenderThread(RHICmdList, Size, Usage, CreateInfo, Buffer);
+		
+		return Buffer.GetReference();
+	}
+}
diff --git a/Engine/Source/Runtime/Apple/AGXRHI/Private/AGXVertexDeclaration.cpp b/Engine/Source/Runtime/Apple/AGXRHI/Private/AGXVertexDeclaration.cpp
new file mode 100644
--- /dev/null
+++ b/Engine/Source/Runtime/Apple/AGXRHI/Private/AGXVertexDeclaration.cpp
@@ -0,0 +1,180 @@
+// Copyright Epic Games, Inc. All Rights Reserved.
+
+/*=============================================================================
+	AGXVertexDeclaration.cpp: AGX RHI vertex declaration implementation.
+=============================================================================*/
+
+
+#include "AGXRHIPrivate.h"
+#include "AGXHashedVertexDescriptor.h"
+#include "AGXVertexDeclaration.h"
+#include "AGXProfiler.h"
+
+
+//------------------------------------------------------------------------------
+
+#pragma mark - AGX RHI Vertex Declaration Globals
+
+
+mtlpp::VertexFormat GAGXFColorVertexFormat = mtlpp::VertexFormat::UChar4Normalized;
+
+
+//------------------------------------------------------------------------------
+
+#pragma mark - AGX RHI Vertex Declaration Support Routines
+
+
+static mtlpp::VertexFormat TranslateElementTypeToMTLType(EVertexElementType Type)
+{
+	switch (Type)
+	{
+		case VET_Float1:		return mtlpp::VertexFormat::Float;
+		case VET_Float2:		return mtlpp::VertexFormat::Float2;
+		case VET_Float3:		return mtlpp::VertexFormat::Float3;
+		case VET_Float4:		return mtlpp::VertexFormat::Float4;
+		case VET_PackedNormal:	return mtlpp::VertexFormat::Char4Normalized;
+		case VET_UByte4:		return mtlpp::VertexFormat::UChar4;
+		case VET_UByte4N:		return mtlpp::VertexFormat::UChar4Normalized;
+		case VET_Color:			return GAGXFColorVertexFormat;
+		case VET_Short2:		return mtlpp::VertexFormat::Short2;
+		case VET_Short4:		return mtlpp::VertexFormat::Short4;
+		case VET_Short2N:		return mtlpp::VertexFormat::Short2Normalized;
+		case VET_Half2:			return mtlpp::VertexFormat::Half2;
+		case VET_Half4:			return mtlpp::VertexFormat::Half4;
+		case VET_Short4N:		return mtlpp::VertexFormat::Short4Normalized;
+		case VET_UShort2:		return mtlpp::VertexFormat::UShort2;
+		case VET_UShort4:		return mtlpp::VertexFormat::UShort4;
+		case VET_UShort2N:		return mtlpp::VertexFormat::UShort2Normalized;
+		case VET_UShort4N:		return mtlpp::VertexFormat::UShort4Normalized;
+		case VET_URGB10A2N:		return mtlpp::VertexFormat::UInt1010102Normalized;
+		case VET_UInt:			return mtlpp::VertexFormat::UInt;
+        default:				METAL_FATAL_ERROR(TEXT("Unknown vertex element type %d!"), (uint32)Type); return mtlpp::VertexFormat::Float;
+    };
+}
+
+static uint32 TranslateElementTypeToSize(EVertexElementType Type)
+{
+	switch (Type)
+	{
+		case VET_Float1:		return 4;
+		case VET_Float2:		return 8;
+		case VET_Float3:		return 12;
+		case VET_Float4:		return 16;
+		case VET_PackedNormal:	return 4;
+		case VET_UByte4:		return 4;
+		case VET_UByte4N:		return 4;
+		case VET_Color:			return 4;
+		case VET_Short2:		return 4;
+		case VET_Short4:		return 8;
+		case VET_UShort2:		return 4;
+		case VET_UShort4:		return 8;
+		case VET_Short2N:		return 4;
+		case VET_UShort2N:		return 4;
+		case VET_Half2:			return 4;
+		case VET_Half4:			return 8;
+		case VET_Short4N:		return 8;
+		case VET_UShort4N:		return 8;
+		case VET_URGB10A2N:		return 4;
+		case VET_UInt:			return 4;
+		default:				METAL_FATAL_ERROR(TEXT("Unknown vertex element type %d!"), (uint32)Type); return 0;
+	};
+}
+
+
+//------------------------------------------------------------------------------
+
+#pragma mark - AGX RHI Vertex Declaration Class Implementation
+
+
+FAGXVertexDeclaration::FAGXVertexDeclaration(const FVertexDeclarationElementList& InElements)
+	: Elements(InElements)
+	, BaseHash(0)
+{
+	GenerateLayout(InElements);
+}
+
+FAGXVertexDeclaration::~FAGXVertexDeclaration()
+{
+	// void
+}
+
+bool FAGXVertexDeclaration::GetInitializer(FVertexDeclarationElementList& Init)
+{
+	Init = Elements;
+	return true;
+}
+
+void FAGXVertexDeclaration::GenerateLayout(const FVertexDeclarationElementList& InElements)
+{
+	mtlpp::VertexDescriptor NewLayout;
+
+	ns::Array<mtlpp::VertexBufferLayoutDescriptor> Layouts = NewLayout.GetLayouts();
+	ns::Array<mtlpp::VertexAttributeDescriptor> Attributes = NewLayout.GetAttributes();
+
+	BaseHash = 0;
+	uint32 StrideHash = BaseHash;
+
+	TMap<uint32, uint32> BufferStrides;
+	for (uint32 ElementIndex = 0; ElementIndex < InElements.Num(); ElementIndex++)
+	{
+		const FVertexElement& Element = InElements[ElementIndex];
+
+		checkf(Element.Stride == 0 || Element.Offset + TranslateElementTypeToSize(Element.Type) <= Element.Stride,
+			TEXT("Stream component is bigger than stride: Offset: %d, Size: %d [Type %d], Stride: %d"), Element.Offset, TranslateElementTypeToSize(Element.Type), (uint32)Element.Type, Element.Stride);
+
+		BaseHash = FCrc::MemCrc32(&Element.StreamIndex, sizeof(Element.StreamIndex), BaseHash);
+		BaseHash = FCrc::MemCrc32(&Element.Offset, sizeof(Element.Offset), BaseHash);
+		BaseHash = FCrc::MemCrc32(&Element.Type, sizeof(Element.Type), BaseHash);
+		BaseHash = FCrc::MemCrc32(&Element.AttributeIndex, sizeof(Element.AttributeIndex), BaseHash);
+
+		uint32 Stride = Element.Stride;
+		StrideHash = FCrc::MemCrc32(&Stride, sizeof(Stride), StrideHash);
+
+		// Vertex & Constant buffers are set up in the same space, so add VB's from the top
+		uint32 ShaderBufferIndex = UNREAL_TO_METAL_BUFFER_INDEX(Element.StreamIndex);
+
+		// track the buffer stride, making sure all elements with the same buffer have the same stride
+		uint32* ExistingStride = BufferStrides.Find(ShaderBufferIndex);
+		if (ExistingStride == NULL)
+		{
+			// handle 0 stride buffers
+			mtlpp::VertexStepFunction Function = (Element.Stride == 0 ? mtlpp::VertexStepFunction::Constant : (Element.bUseInstanceIndex ? mtlpp::VertexStepFunction::PerInstance : mtlpp::VertexStepFunction::PerVertex));
+			uint32 StepRate = (Element.Stride == 0 ? 0 : 1);
+
+			// even with MTLVertexStepFunctionConstant, it needs a non-zero stride (not sure why)
+			if (Element.Stride == 0)
+			{
+				Stride = TranslateElementTypeToSize(Element.Type);
+			}
+
+			// look for any unset strides coming from UE4 (this can be removed when all are fixed)
+			if (Element.Stride == 0xFFFF)
+			{
+				NSLog(@"Setting illegal stride - break here if you want to find out why, but this won't break until we try to render with it");
+				Stride = 200;
+			}
+
+			// set the stride once per buffer
+			mtlpp::VertexBufferLayoutDescriptor VBLayout = Layouts[ShaderBufferIndex];
+			VBLayout.SetStride(Stride);
+			VBLayout.SetStepFunction(Function);
+			VBLayout.SetStepRate(StepRate);
+
+			// track this buffer and stride
+			BufferStrides.Add(ShaderBufferIndex, Element.Stride);
+		}
+		else
+		{
+			// if the strides of elements with same buffer index have different strides, something is VERY wrong
+			check(Element.Stride == *ExistingStride);
+		}
+
+		// set the format for each element
+		mtlpp::VertexAttributeDescriptor Attrib = Attributes[Element.AttributeIndex];
+		Attrib.SetFormat(TranslateElementTypeToMTLType(Element.Type));
+		Attrib.SetOffset(Element.Offset);
+		Attrib.SetBufferIndex(ShaderBufferIndex);
+	}
+
+	Layout = FAGXHashedVertexDescriptor(NewLayout, HashCombine(BaseHash, StrideHash));
+}
diff --git a/Engine/Source/Runtime/Apple/AGXRHI/Private/AGXVertexDeclaration.h b/Engine/Source/Runtime/Apple/AGXRHI/Private/AGXVertexDeclaration.h
new file mode 100644
--- /dev/null
+++ b/Engine/Source/Runtime/Apple/AGXRHI/Private/AGXVertexDeclaration.h
@@ -0,0 +1,42 @@
+// Copyright Epic Games, Inc. All Rights Reserved.
+
+/*=============================================================================
+	AGXVertexDeclaration.h: AGX RHI Vertex Declaration.
+=============================================================================*/
+
+#pragma once
+
+
+#include "AGXHashedVertexDescriptor.h"
+
+
+//------------------------------------------------------------------------------
+
+#pragma mark - AGX Vertex Declaration Class
+
+
+/**
+ * This represents a vertex declaration that hasn't been combined with a
+ * specific shader to create a bound shader.
+ */
+class FAGXVertexDeclaration : public FRHIVertexDeclaration
+{
+public:
+	/** Initialization constructor. */
+	FAGXVertexDeclaration(const FVertexDeclarationElementList& InElements);
+	~FAGXVertexDeclaration();
+
+	/** Cached element info array (offset, stream index, etc) */
+	FVertexDeclarationElementList Elements;
+
+	/** This is the layout for the vertex elements */
+	FAGXHashedVertexDescriptor Layout;
+
+	/** Hash without considering strides which may be overriden */
+	uint32 BaseHash;
+
+	virtual bool GetInitializer(FVertexDeclarationElementList& Init) override final;
+
+protected:
+	void GenerateLayout(const FVertexDeclarationElementList& Elements);
+};
diff --git a/Engine/Source/Runtime/Apple/AGXRHI/Private/AGXViewport.cpp b/Engine/Source/Runtime/Apple/AGXRHI/Private/AGXViewport.cpp
new file mode 100644
--- /dev/null
+++ b/Engine/Source/Runtime/Apple/AGXRHI/Private/AGXViewport.cpp
@@ -0,0 +1,689 @@
+// Copyright Epic Games, Inc. All Rights Reserved.
+
+/*=============================================================================
+	AGXViewport.cpp: AGX RHI viewport implementation.
+=============================================================================*/
+
+#include "AGXRHIPrivate.h"
+#if PLATFORM_MAC
+#include "Mac/CocoaWindow.h"
+#include "Mac/CocoaThread.h"
+#else
+#include "IOS/IOSAppDelegate.h"
+#endif
+#include "RenderCommandFence.h"
+#include "Containers/Set.h"
+#include "AGXProfiler.h"
+#include "RenderUtils.h"
+
+extern int32 GAGXSupportsIntermediateBackBuffer;
+extern int32 GAGXSeparatePresentThread;
+extern int32 GAGXNonBlockingPresent;
+extern float GAGXPresentFramePacing;
+
+#if PLATFORM_IOS
+static int32 GEnablePresentPacing = 0;
+static FAutoConsoleVariableRef CVarAGXEnablePresentPacing(
+	   TEXT("ios.PresentPacing"),
+	   GEnablePresentPacing,
+	   TEXT(""),
+		ECVF_Default);
+#endif
+
+#if PLATFORM_MAC
+
+// Quick way to disable availability warnings is to duplicate the definitions into a new type - gotta love ObjC dynamic-dispatch!
+@interface FCAMetalLayer : CALayer
+@property BOOL displaySyncEnabled;
+@property BOOL allowsNextDrawableTimeout;
+@end
+
+@implementation FAGXView
+
+- (id)initWithFrame:(NSRect)frameRect
+{
+	self = [super initWithFrame:frameRect];
+	if (self)
+	{
+	}
+	return self;
+}
+
+- (BOOL)isOpaque
+{
+	return YES;
+}
+
+- (BOOL)mouseDownCanMoveWindow
+{
+	return YES;
+}
+
+@end
+
+#endif
+
+static FCriticalSection ViewportsMutex;
+static TSet<FAGXViewport*> Viewports;
+
+FAGXViewport::FAGXViewport(void* WindowHandle, uint32 InSizeX, uint32 InSizeY, bool bInIsFullscreen, EPixelFormat Format)
+	: Drawable{nil}
+	, BackBuffer{nullptr, nullptr}
+	, Mutex{}
+	, DrawableTextures{}
+	, DisplayID{0}
+	, Block{nullptr}
+	, FrameAvailable{0}
+	, LastCompleteFrame{nullptr}
+	, bIsFullScreen{bInIsFullscreen}
+#if PLATFORM_MAC
+	, View{nullptr}
+	, CustomPresent{nullptr}
+#endif
+{
+#if PLATFORM_MAC
+	MainThreadCall(^{
+		FCocoaWindow* Window = (FCocoaWindow*)WindowHandle;
+		const NSRect ContentRect = NSMakeRect(0, 0, InSizeX, InSizeY);
+		View = [[FAGXView alloc] initWithFrame:ContentRect];
+		[View setAutoresizingMask:NSViewWidthSizable | NSViewHeightSizable];
+		[View setWantsLayer:YES];
+
+		CAMetalLayer* Layer = [CAMetalLayer new];
+
+		CGFloat bgColor[] = { 0.0, 0.0, 0.0, 0.0 };
+		Layer.edgeAntialiasingMask = 0;
+		Layer.masksToBounds = YES;
+		Layer.backgroundColor = CGColorCreate(CGColorSpaceCreateDeviceRGB(), bgColor);
+		Layer.presentsWithTransaction = NO;
+		Layer.anchorPoint = CGPointMake(0.5, 0.5);
+		Layer.frame = ContentRect;
+		Layer.magnificationFilter = kCAFilterNearest;
+		Layer.minificationFilter = kCAFilterNearest;
+
+		[Layer setDevice:GetAGXDeviceContext().GetDevice()];
+		
+		[Layer setFramebufferOnly:NO];
+		[Layer removeAllAnimations];
+
+		[View setLayer:Layer];
+
+		[Window setContentView:View];
+		[[Window standardWindowButton:NSWindowCloseButton] setAction:@selector(performClose:)];
+	}, NSDefaultRunLoopMode, true);
+#endif
+	Resize(InSizeX, InSizeY, bInIsFullscreen, Format);
+	
+	{
+		FScopeLock Lock(&ViewportsMutex);
+		Viewports.Add(this);
+	}
+}
+
+FAGXViewport::~FAGXViewport()
+{
+	if (Block)
+	{
+		FScopeLock BlockLock(&Mutex);
+		if (GAGXSeparatePresentThread)
+		{
+			FPlatformRHIFramePacer::RemoveHandler(Block);
+		}
+		Block_release(Block);
+		Block = nil;
+	}
+	{
+		FScopeLock Lock(&ViewportsMutex);
+		Viewports.Remove(this);
+	}
+	
+	BackBuffer[0].SafeRelease();	// when the rest of the engine releases it, its framebuffers will be released too (those the engine knows about)
+	BackBuffer[1].SafeRelease();
+	check(!IsValidRef(BackBuffer[0]));
+	check(!IsValidRef(BackBuffer[1]));
+}
+
+uint32 FAGXViewport::GetViewportIndex(EAGXViewportAccessFlag Accessor) const
+{
+	switch(Accessor)
+	{
+		case EAGXViewportAccessRHI:
+			check(IsInRHIThread() || IsInRenderingThread());
+			// Deliberate fall-through
+		case EAGXViewportAccessDisplayLink: // Displaylink is not an index, merely an alias that avoids the check...
+			return (GRHISupportsRHIThread && IsRunningRHIInSeparateThread()) ? EAGXViewportAccessRHI : EAGXViewportAccessRenderer;
+		case EAGXViewportAccessRenderer:
+			check(IsInRenderingThread());
+			return Accessor;
+		case EAGXViewportAccessGame:
+			check(IsInGameThread());
+			return EAGXViewportAccessRenderer;
+		default:
+			check(false);
+			return EAGXViewportAccessRenderer;
+	}
+}
+
+void FAGXViewport::Resize(uint32 InSizeX, uint32 InSizeY, bool bInIsFullscreen,EPixelFormat Format)
+{
+	bIsFullScreen = bInIsFullscreen;
+	uint32 Index = GetViewportIndex(EAGXViewportAccessGame);
+	
+	bool bUseHDR = GRHISupportsHDROutput && Format == GRHIHDRDisplayOutputFormat;
+	
+	// Format can come in as PF_Unknown in the LDR case or if this RHI doesn't support HDR.
+	// So we'll fall back to BGRA8 in those cases.
+	if (!bUseHDR)
+	{
+		Format = PF_B8G8R8A8;
+	}
+	
+	mtlpp::PixelFormat MetalFormat = (mtlpp::PixelFormat)GPixelFormats[Format].PlatformFormat;
+	
+	if (IsValidRef(BackBuffer[Index]) && Format != BackBuffer[Index]->GetFormat())
+	{
+		// Really need to flush the RHI thread & GPU here...
+		AddRef();
+		ENQUEUE_RENDER_COMMAND(FlushPendingRHICommands)(
+			[Viewport = this](FRHICommandListImmediate& RHICmdList)
+			{
+				GRHICommandList.GetImmediateCommandList().BlockUntilGPUIdle();
+				Viewport->ReleaseDrawable();
+				Viewport->Release();
+			});
+		
+		// Issue a fence command to the rendering thread and wait for it to complete.
+		FRenderCommandFence Fence;
+		Fence.BeginFence();	
+		Fence.Wait();
+	}
+    
+#if PLATFORM_MAC
+	MainThreadCall(^
+	{
+		CAMetalLayer* MetalLayer = (CAMetalLayer*)View.layer;
+		
+		MetalLayer.drawableSize = CGSizeMake(InSizeX, InSizeY);
+		
+		if (MetalFormat != (mtlpp::PixelFormat)MetalLayer.pixelFormat)
+		{
+			MetalLayer.pixelFormat = (MTLPixelFormat)MetalFormat;
+		}
+		
+		if (bUseHDR != MetalLayer.wantsExtendedDynamicRangeContent)
+		{
+			MetalLayer.wantsExtendedDynamicRangeContent = bUseHDR;
+		}
+		
+	}, NSDefaultRunLoopMode, true);
+#else
+	// A note on HDR in iOS
+	// Setting the pixel format to one of the Apple XR formats is all you need.
+	// iOS expects the app to output in sRGB regadless of the display
+	// (even though Apple's HDR displays are P3)
+	// and its compositor will do the conversion.
+	{
+		IOSAppDelegate* AppDelegate = [IOSAppDelegate GetDelegate];
+		FIOSView* IOSView = AppDelegate.IOSView;
+		CAMetalLayer* MetalLayer = (CAMetalLayer*) IOSView.layer;
+		
+		if (MetalFormat != (mtlpp::PixelFormat) MetalLayer.pixelFormat)
+		{
+			MetalLayer.pixelFormat = (MTLPixelFormat) MetalFormat;
+		}
+		
+		[IOSView UpdateRenderWidth:InSizeX andHeight:InSizeY];
+	}
+#endif
+
+    {
+        FScopeLock Lock(&Mutex);
+        FRHIResourceCreateInfo CreateInfo(TEXT("BackBuffer"));
+        FTexture2DRHIRef NewBackBuffer;
+        FTexture2DRHIRef DoubleBuffer;
+        if (GAGXSupportsIntermediateBackBuffer)
+        {
+            NewBackBuffer = (FAGXTexture2D*)(FRHITexture2D*)GDynamicRHI->RHICreateTexture2D(InSizeX, InSizeY, Format, 1, 1, TexCreate_RenderTargetable, ERHIAccess::Unknown, CreateInfo);
+            
+            if (GAGXSeparatePresentThread)
+            {
+                DoubleBuffer = GDynamicRHI->RHICreateTexture2D(InSizeX, InSizeY, Format, 1, 1, TexCreate_RenderTargetable, ERHIAccess::Unknown, CreateInfo);
+                ((FAGXTexture2D*)DoubleBuffer.GetReference())->Surface.Viewport = this;
+            }
+        }
+        else
+        {
+            NewBackBuffer = (FAGXTexture2D*)(FRHITexture2D*)GDynamicRHI->RHICreateTexture2D(InSizeX, InSizeY, Format, 1, 1, TexCreate_RenderTargetable | TexCreate_Presentable, ERHIAccess::Unknown, CreateInfo);
+        }
+        ((FAGXTexture2D*)NewBackBuffer.GetReference())->Surface.Viewport = this;
+        
+        BackBuffer[Index] = (FAGXTexture2D*)NewBackBuffer.GetReference();
+        if (GAGXSeparatePresentThread)
+        {
+            BackBuffer[EAGXViewportAccessRHI] = (FAGXTexture2D*)DoubleBuffer.GetReference();
+        }
+        else
+        {
+            BackBuffer[EAGXViewportAccessRHI] = BackBuffer[Index];
+        }
+    }
+}
+
+TRefCountPtr<FAGXTexture2D> FAGXViewport::GetBackBuffer(EAGXViewportAccessFlag Accessor) const
+{
+	FScopeLock Lock(&Mutex);
+	
+	uint32 Index = GetViewportIndex(Accessor);
+	check(IsValidRef(BackBuffer[Index]));
+	return BackBuffer[Index];
+}
+
+#if PLATFORM_MAC
+@protocol CAMetalLayerSPI <NSObject>
+- (BOOL)isDrawableAvailable;
+@end
+#endif
+
+id<CAMetalDrawable> FAGXViewport::GetDrawable(EAGXViewportAccessFlag Accessor)
+{
+	SCOPE_CYCLE_COUNTER(STAT_AGXMakeDrawableTime);
+    if (!Drawable
+#if !PLATFORM_MAC
+        || (Drawable.texture.width != BackBuffer[GetViewportIndex(Accessor)]->GetSizeX() || Drawable.texture.height != BackBuffer[GetViewportIndex(Accessor)]->GetSizeY())
+#endif
+        )
+	{
+		// Drawable changed, release the previously retained object.
+		if (Drawable != nil)
+		{
+			[Drawable release];
+			Drawable = nil;
+		}
+
+		@autoreleasepool
+		{
+			uint32 IdleStart = FPlatformTime::Cycles();
+
+#if PLATFORM_MAC
+			CAMetalLayer* CurrentLayer = (CAMetalLayer*)[View layer];
+			if (GAGXNonBlockingPresent == 0 || [((id<CAMetalLayerSPI>)CurrentLayer) isDrawableAvailable])
+			{
+				Drawable = CurrentLayer ? [CurrentLayer nextDrawable] : nil;
+			}
+
+#if METAL_DEBUG_OPTIONS
+			if (Drawable)
+			{
+				CGSize Size = Drawable.layer.drawableSize;
+				if ((Size.width != BackBuffer[GetViewportIndex(Accessor)]->GetSizeX() || Size.height != BackBuffer[GetViewportIndex(Accessor)]->GetSizeY()))
+				{
+					UE_LOG(LogAGX, Display, TEXT("Viewport Size Mismatch: Drawable W:%f H:%f, Viewport W:%u H:%u"), Size.width, Size.height, BackBuffer[GetViewportIndex(Accessor)]->GetSizeX(), BackBuffer[GetViewportIndex(Accessor)]->GetSizeY());
+				}
+			}
+#endif // METAL_DEBUG_OPTIONS
+
+#else // PLATFORM_MAC
+			CGSize Size;
+			IOSAppDelegate* AppDelegate = [IOSAppDelegate GetDelegate];
+			do
+			{
+				Drawable = [AppDelegate.IOSView MakeDrawable];
+				if (Drawable != nil)
+				{
+					Size.width = Drawable.texture.width;
+					Size.height = Drawable.texture.height;
+				}
+				else
+				{
+					FPlatformProcess::SleepNoStats(0.001f);
+				}
+			}
+			while (Drawable == nil || Size.width != BackBuffer[GetViewportIndex(Accessor)]->GetSizeX() || Size.height != BackBuffer[GetViewportIndex(Accessor)]->GetSizeY());
+
+#endif // PLATFORM_MAC
+
+			if (IsInRHIThread())
+			{
+				GWorkingRHIThreadStallTime += FPlatformTime::Cycles() - IdleStart;
+			}
+			else
+			{
+				GRenderThreadIdle[ERenderThreadIdleTypes::WaitingForGPUPresent] += FPlatformTime::Cycles() - IdleStart;
+				GRenderThreadNumIdle[ERenderThreadIdleTypes::WaitingForGPUPresent]++;
+			}
+
+			// Retain the drawable here or it will be released when the
+			// autorelease pool goes out of scope.
+			if (Drawable != nil)
+			{
+				[Drawable retain];
+			}
+		} // autoreleasepool
+	}
+
+	return Drawable;
+}
+
+FAGXTexture FAGXViewport::GetDrawableTexture(EAGXViewportAccessFlag Accessor)
+{
+	id<CAMetalDrawable> CurrentDrawable = GetDrawable(Accessor);
+#if METAL_DEBUG_OPTIONS
+	@autoreleasepool
+	{
+#if PLATFORM_MAC
+		CAMetalLayer* CurrentLayer = (CAMetalLayer*)[View layer];
+#else
+		CAMetalLayer* CurrentLayer = (CAMetalLayer*)[[IOSAppDelegate GetDelegate].IOSView layer];
+#endif
+		
+		uint32 Index = GetViewportIndex(Accessor);
+		CGSize Size = CurrentLayer.drawableSize;
+		if (CurrentDrawable.texture.width != BackBuffer[Index]->GetSizeX() || CurrentDrawable.texture.height != BackBuffer[Index]->GetSizeY())
+		{
+			UE_LOG(LogAGX, Display, TEXT("Viewport Size Mismatch: Drawable W:%f H:%f, Texture W:%llu H:%llu, Viewport W:%u H:%u"), Size.width, Size.height, CurrentDrawable.texture.width, CurrentDrawable.texture.height, BackBuffer[Index]->GetSizeX(), BackBuffer[Index]->GetSizeY());
+		}
+	}
+#endif
+	DrawableTextures[Accessor] = CurrentDrawable.texture;
+	return CurrentDrawable.texture;
+}
+
+ns::AutoReleased<FAGXTexture> FAGXViewport::GetCurrentTexture(EAGXViewportAccessFlag Accessor)
+{
+	return DrawableTextures[Accessor];
+}
+
+void FAGXViewport::ReleaseDrawable()
+{
+	if (!GAGXSeparatePresentThread)
+	{
+		if (Drawable != nil)
+		{
+			[Drawable release];
+			Drawable = nil;
+		}
+
+		if (!GAGXSupportsIntermediateBackBuffer && IsValidRef(BackBuffer[GetViewportIndex(EAGXViewportAccessRHI)]))
+		{
+			BackBuffer[GetViewportIndex(EAGXViewportAccessRHI)]->Surface.Texture = nil;
+		}
+	}
+}
+
+#if PLATFORM_MAC
+NSWindow* FAGXViewport::GetWindow() const
+{
+	return [View window];
+}
+#endif
+
+void FAGXViewport::Present(FAGXCommandQueue& CommandQueue, bool bLockToVsync)
+{
+	FScopeLock Lock(&Mutex);
+	
+	bool bIsLiveResize = false;
+#if PLATFORM_MAC
+	NSNumber* ScreenId = [View.window.screen.deviceDescription objectForKey:@"NSScreenNumber"];
+	DisplayID = ScreenId.unsignedIntValue;
+	bIsLiveResize = View.inLiveResize;
+	{
+		FCAMetalLayer* CurrentLayer = (FCAMetalLayer*)[View layer];
+		static bool sVSyncSafe = FPlatformMisc::MacOSXVersionCompare(10,13,4) >= 0;
+		CurrentLayer.displaySyncEnabled = bLockToVsync || (!sVSyncSafe && !(IsRunningGame() && bIsFullScreen));
+	}
+#endif
+	
+	LastCompleteFrame = GetBackBuffer(EAGXViewportAccessRHI);
+	FPlatformAtomics::InterlockedExchange(&FrameAvailable, 1);
+	
+	if (!Block)
+	{
+		Block = Block_copy(^(uint32 InDisplayID, double OutputSeconds, double OutputDuration)
+		{
+#if !PLATFORM_MAC
+			uint32 FramePace = FPlatformRHIFramePacer::GetFramePace();
+			float MinPresentDuration = FramePace ? (1.0f / (float)FramePace) : 0.0f;
+#endif
+			bool bIsInLiveResize = false;
+#if PLATFORM_MAC
+			bIsInLiveResize = View.inLiveResize;
+#endif
+			if (FrameAvailable > 0 && (InDisplayID == 0 || (DisplayID == InDisplayID && !bIsInLiveResize)))
+			{
+				FPlatformAtomics::InterlockedDecrement(&FrameAvailable);
+				id<CAMetalDrawable> LocalDrawable = [GetDrawable(EAGXViewportAccessDisplayLink) retain];
+				{
+					FScopeLock BlockLock(&Mutex);
+#if PLATFORM_MAC
+					bIsInLiveResize = View.inLiveResize;
+#endif
+					
+					if (LocalDrawable && LocalDrawable.texture && (InDisplayID == 0 || !bIsInLiveResize))
+					{
+						mtlpp::CommandBuffer CurrentCommandBuffer = CommandQueue.CreateCommandBuffer();
+						check(CurrentCommandBuffer);
+						
+#if ENABLE_METAL_GPUPROFILE
+						FAGXProfiler* Profiler = FAGXProfiler::GetProfiler();
+						FAGXCommandBufferStats* Stats = Profiler->AllocateCommandBuffer(CurrentCommandBuffer, 0);
+#endif
+						
+						if (GAGXSupportsIntermediateBackBuffer)
+						{
+							TRefCountPtr<FAGXTexture2D> Texture = LastCompleteFrame;
+							check(IsValidRef(Texture));
+							
+							FAGXTexture Src = Texture->Surface.Texture;
+							FAGXTexture Dst = LocalDrawable.texture;
+							
+							NSUInteger Width = FMath::Min(Src.GetWidth(), Dst.GetWidth());
+							NSUInteger Height = FMath::Min(Src.GetHeight(), Dst.GetHeight());
+							
+							mtlpp::BlitCommandEncoder Encoder = CurrentCommandBuffer.BlitCommandEncoder();
+							check(Encoder.GetPtr());
+#if MTLPP_CONFIG_VALIDATE && METAL_DEBUG_OPTIONS
+							FAGXBlitCommandEncoderDebugging Debugging;
+							if (AGXSafeGetRuntimeDebuggingLevel() >= EAGXDebugLevelFastValidation)
+							{
+								FAGXCommandBufferDebugging CmdDebug = FAGXCommandBufferDebugging::Get(CurrentCommandBuffer);
+								Debugging = FAGXBlitCommandEncoderDebugging(Encoder, CmdDebug);
+							}
+#endif
+							METAL_GPUPROFILE(Profiler->EncodeBlit(Stats, __FUNCTION__));
+
+							Encoder.Copy(Src, 0, 0, mtlpp::Origin(0, 0, 0), mtlpp::Size(Width, Height, 1), Dst, 0, 0, mtlpp::Origin(0, 0, 0));
+							METAL_DEBUG_LAYER(EAGXDebugLevelFastValidation, Debugging.Copy(Src, 0, 0, mtlpp::Origin(0, 0, 0), mtlpp::Size(Width, Height, 1), Dst, 0, 0, mtlpp::Origin(0, 0, 0)));
+
+							Encoder.EndEncoding();
+							METAL_DEBUG_LAYER(EAGXDebugLevelFastValidation, Debugging.EndEncoder());
+
+							mtlpp::CommandBufferHandler H = [Src, Dst](const mtlpp::CommandBuffer &) {
+								// void
+							};
+
+							CurrentCommandBuffer.AddCompletedHandler(H);
+
+							[Drawable release];
+							Drawable = nil;
+						}
+						
+						// This is a bit different than the usual pattern.
+						// This command buffer here is committed directly, instead of going through
+						// FAGXCommandList::Commit. So long as Present() is called within
+						// high level RHI BeginFrame/EndFrame this will not fine.
+						// Otherwise the recording of the Present time will be offset by one in the
+						// FAGXGPUProfiler frame indices.
+      
+#if PLATFORM_MAC
+						FAGXView* theView = View;
+						mtlpp::CommandBufferHandler C = [LocalDrawable, theView](const mtlpp::CommandBuffer & cmd_buf) {
+#else
+						mtlpp::CommandBufferHandler C = [LocalDrawable](const mtlpp::CommandBuffer & cmd_buf) {
+#endif
+							FAGXGPUProfiler::RecordPresent(cmd_buf);
+							[LocalDrawable release];
+#if PLATFORM_MAC
+							MainThreadCall(^{
+								FCocoaWindow* Window = (FCocoaWindow*)[theView window];
+								[Window startRendering];
+							}, NSDefaultRunLoopMode, false);
+#endif
+						};
+						
+#if PLATFORM_MAC		// Mac needs the older way to present otherwise we end up with bad behaviour of the completion handlers that causes GPU timeouts.
+						mtlpp::CommandBufferHandler H = [LocalDrawable](mtlpp::CommandBuffer const&)
+						{
+							[LocalDrawable present];
+						};
+								
+						CurrentCommandBuffer.AddCompletedHandler(C);
+						CurrentCommandBuffer.AddScheduledHandler(H);
+
+#else // PLATFORM_MAC
+						CurrentCommandBuffer.AddCompletedHandler(C);
+
+						if (MinPresentDuration && GEnablePresentPacing)
+						{
+							CurrentCommandBuffer.PresentAfterMinimumDuration(LocalDrawable, 1.0f/(float)FramePace);
+						}
+						else
+						{
+							CurrentCommandBuffer.Present(LocalDrawable);
+						}
+#endif // PLATFORM_MAC
+
+						METAL_GPUPROFILE(Stats->End(CurrentCommandBuffer));
+						CommandQueue.CommitCommandBuffer(CurrentCommandBuffer);
+					}
+				}
+			}
+		});
+		
+		if (GAGXSeparatePresentThread)
+		{
+			FPlatformRHIFramePacer::AddHandler(Block);
+		}
+	}
+	
+	if (bIsLiveResize || !GAGXSeparatePresentThread)
+	{
+		Block(0, 0.0, 0.0);
+	}
+	
+	if (!(GRHISupportsRHIThread && IsRunningRHIInSeparateThread()))
+	{
+		Swap();
+	}
+}
+
+void FAGXViewport::Swap()
+{
+	if (GAGXSeparatePresentThread)
+	{
+		FScopeLock Lock(&Mutex);
+		
+		check(IsValidRef(BackBuffer[0]));
+		check(IsValidRef(BackBuffer[1]));
+		
+		TRefCountPtr<FAGXTexture2D> BB0 = BackBuffer[0];
+		TRefCountPtr<FAGXTexture2D> BB1 = BackBuffer[1];
+		
+		BackBuffer[0] = BB1;
+		BackBuffer[1] = BB0;
+	}
+}
+
+/*=============================================================================
+ *	The following RHI functions must be called from the main thread.
+ *=============================================================================*/
+FViewportRHIRef FAGXDynamicRHI::RHICreateViewport(void* WindowHandle,uint32 SizeX,uint32 SizeY,bool bIsFullscreen,EPixelFormat PreferredPixelFormat)
+{
+	check( IsInGameThread() );
+	@autoreleasepool {
+	return new FAGXViewport(WindowHandle, SizeX, SizeY, bIsFullscreen, PreferredPixelFormat);
+	}
+}
+
+void FAGXDynamicRHI::RHIResizeViewport(FRHIViewport* Viewport, uint32 SizeX, uint32 SizeY, bool bIsFullscreen)
+{
+	RHIResizeViewport(Viewport, SizeX, SizeY, bIsFullscreen, PF_Unknown);
+}
+
+void FAGXDynamicRHI::RHIResizeViewport(FRHIViewport* ViewportRHI,uint32 SizeX,uint32 SizeY,bool bIsFullscreen,EPixelFormat Format)
+{
+	@autoreleasepool {
+	check( IsInGameThread() );
+
+	FAGXViewport* Viewport = ResourceCast(ViewportRHI);
+	Viewport->Resize(SizeX, SizeY, bIsFullscreen, Format);
+	}
+}
+
+void FAGXDynamicRHI::RHITick( float DeltaTime )
+{
+	check( IsInGameThread() );
+}
+
+/*=============================================================================
+ *	Viewport functions.
+ *=============================================================================*/
+
+void FAGXRHICommandContext::RHIBeginDrawingViewport(FRHIViewport* ViewportRHI, FRHITexture* RenderTargetRHI)
+{
+	check(false);
+}
+
+void FAGXRHIImmediateCommandContext::RHIBeginDrawingViewport(FRHIViewport* ViewportRHI, FRHITexture* RenderTargetRHI)
+{
+	@autoreleasepool {
+	FAGXViewport* Viewport = ResourceCast(ViewportRHI);
+	check(Viewport);
+
+	((FAGXDeviceContext*)Context)->BeginDrawingViewport(Viewport);
+
+	// Set the render target and viewport.
+	if (RenderTargetRHI)
+	{
+		FRHIRenderTargetView RTV(RenderTargetRHI, GIsEditor ? ERenderTargetLoadAction::EClear : ERenderTargetLoadAction::ELoad);
+		SetRenderTargets(1, &RTV, nullptr);
+	}
+	else
+	{
+		FRHIRenderTargetView RTV(Viewport->GetBackBuffer(EAGXViewportAccessRHI), GIsEditor ? ERenderTargetLoadAction::EClear : ERenderTargetLoadAction::ELoad);
+		SetRenderTargets(1, &RTV, nullptr);
+	}
+	}
+}
+
+void FAGXRHICommandContext::RHIEndDrawingViewport(FRHIViewport* ViewportRHI,bool bPresent,bool bLockToVsync)
+{
+	check(false);
+}
+
+void FAGXRHIImmediateCommandContext::RHIEndDrawingViewport(FRHIViewport* ViewportRHI,bool bPresent,bool bLockToVsync)
+{
+	@autoreleasepool {
+	FAGXViewport* Viewport = ResourceCast(ViewportRHI);
+	((FAGXDeviceContext*)Context)->EndDrawingViewport(Viewport, bPresent, bLockToVsync);
+	}
+}
+
+FTexture2DRHIRef FAGXDynamicRHI::RHIGetViewportBackBuffer(FRHIViewport* ViewportRHI)
+{
+	@autoreleasepool {
+	FAGXViewport* Viewport = ResourceCast(ViewportRHI);
+	return FTexture2DRHIRef(Viewport->GetBackBuffer(EAGXViewportAccessRenderer).GetReference());
+	}
+}
+
+void FAGXDynamicRHI::RHIAdvanceFrameForGetViewportBackBuffer(FRHIViewport* ViewportRHI)
+{
+	if (GAGXSeparatePresentThread && (GRHISupportsRHIThread && IsRunningRHIInSeparateThread()))
+	{
+		FScopeLock Lock(&ViewportsMutex);
+		for (FAGXViewport* Viewport : Viewports)
+		{
+			Viewport->Swap();
+		}
+	}
+}
