Initial import of AGXRHI

diff --git a/Engine/Source/Runtime/Apple/AGXRHI/Private/Shaders/AGXCompiledShaderCache.cpp b/Engine/Source/Runtime/Apple/AGXRHI/Private/Shaders/AGXCompiledShaderCache.cpp
new file mode 100644
--- /dev/null
+++ b/Engine/Source/Runtime/Apple/AGXRHI/Private/Shaders/AGXCompiledShaderCache.cpp
@@ -0,0 +1,20 @@
+// Copyright Epic Games, Inc. All Rights Reserved.
+
+/*=============================================================================
+	AGXCompiledShaderCache.cpp: AGX RHI Compiled Shader Cache.
+=============================================================================*/
+
+#include "CoreMinimal.h"
+
+THIRD_PARTY_INCLUDES_START
+#include "mtlpp.hpp"
+THIRD_PARTY_INCLUDES_END
+
+#include "AGXCompiledShaderKey.h"
+#include "AGXCompiledShaderCache.h"
+
+FAGXCompiledShaderCache& GetAGXCompiledShaderCache()
+{
+	static FAGXCompiledShaderCache CompiledShaderCache;
+	return CompiledShaderCache;
+}
diff --git a/Engine/Source/Runtime/Apple/AGXRHI/Private/Shaders/AGXCompiledShaderCache.h b/Engine/Source/Runtime/Apple/AGXRHI/Private/Shaders/AGXCompiledShaderCache.h
new file mode 100644
--- /dev/null
+++ b/Engine/Source/Runtime/Apple/AGXRHI/Private/Shaders/AGXCompiledShaderCache.h
@@ -0,0 +1,54 @@
+// Copyright Epic Games, Inc. All Rights Reserved.
+
+/*=============================================================================
+	AGXCompiledShaderCache.h: AGX RHI Compiled Shader Cache.
+=============================================================================*/
+
+#pragma once
+
+#include "Misc/ScopeRWLock.h"
+
+struct FAGXCompiledShaderCache
+{
+public:
+	FAGXCompiledShaderCache()
+	{
+		// VOID
+	}
+
+	~FAGXCompiledShaderCache()
+	{
+		// VOID
+	}
+
+	mtlpp::Function FindRef(FAGXCompiledShaderKey const& Key)
+	{
+		FRWScopeLock ScopedLock(Lock, SLT_ReadOnly);
+		mtlpp::Function Func = Cache.FindRef(Key);
+		return Func;
+	}
+
+	mtlpp::Library FindLibrary(mtlpp::Function const& Function)
+	{
+		FRWScopeLock ScopedLock(Lock, SLT_ReadOnly);
+		mtlpp::Library Lib = LibCache.FindRef(Function.GetPtr());
+		return Lib;
+	}
+
+	void Add(FAGXCompiledShaderKey Key, mtlpp::Library const& Lib, mtlpp::Function const& Function)
+	{
+		FRWScopeLock ScopedLock(Lock, SLT_Write);
+		if (Cache.FindRef(Key) == nil)
+		{
+			Cache.Add(Key, Function);
+			LibCache.Add(Function.GetPtr(), Lib);
+		}
+	}
+
+private:
+	FRWLock Lock;
+	TMap<FAGXCompiledShaderKey, mtlpp::Function> Cache;
+	TMap<mtlpp::Function::Type, mtlpp::Library> LibCache;
+};
+
+extern FAGXCompiledShaderCache& GetAGXCompiledShaderCache();
diff --git a/Engine/Source/Runtime/Apple/AGXRHI/Private/Shaders/AGXCompiledShaderKey.h b/Engine/Source/Runtime/Apple/AGXRHI/Private/Shaders/AGXCompiledShaderKey.h
new file mode 100644
--- /dev/null
+++ b/Engine/Source/Runtime/Apple/AGXRHI/Private/Shaders/AGXCompiledShaderKey.h
@@ -0,0 +1,32 @@
+// Copyright Epic Games, Inc. All Rights Reserved.
+
+/*=============================================================================
+	AGXCompiledShaderKey.h: AGX RHI Compiled Shader Key.
+=============================================================================*/
+
+#pragma once
+
+struct FAGXCompiledShaderKey
+{
+	FAGXCompiledShaderKey(uint32 InCodeSize, uint32 InCodeCRC, uint32 InConstants)
+		: CodeSize(InCodeSize)
+		, CodeCRC(InCodeCRC)
+		, Constants(InConstants)
+	{
+		// VOID
+	}
+
+	friend bool operator ==(const FAGXCompiledShaderKey& A, const FAGXCompiledShaderKey& B)
+	{
+		return A.CodeSize == B.CodeSize && A.CodeCRC == B.CodeCRC && A.Constants == B.Constants;
+	}
+
+	friend uint32 GetTypeHash(const FAGXCompiledShaderKey &Key)
+	{
+		return HashCombine(HashCombine(GetTypeHash(Key.CodeSize), GetTypeHash(Key.CodeCRC)), GetTypeHash(Key.Constants));
+	}
+
+	uint32 CodeSize;
+	uint32 CodeCRC;
+	uint32 Constants;
+};
diff --git a/Engine/Source/Runtime/Apple/AGXRHI/Private/Shaders/AGXShaderLibrary.cpp b/Engine/Source/Runtime/Apple/AGXRHI/Private/Shaders/AGXShaderLibrary.cpp
new file mode 100644
--- /dev/null
+++ b/Engine/Source/Runtime/Apple/AGXRHI/Private/Shaders/AGXShaderLibrary.cpp
@@ -0,0 +1,161 @@
+// Copyright Epic Games, Inc. All Rights Reserved.
+
+/*=============================================================================
+	AGXShaderLibrary.cpp: AGX RHI Shader Library Class Implementation.
+=============================================================================*/
+
+
+#include "AGXRHIPrivate.h"
+#if !UE_BUILD_SHIPPING
+#include "Debugging/AGXShaderDebugCache.h"
+#include "Debugging/AGXShaderDebugZipFile.h"
+#endif // !UE_BUILD_SHIPPING
+#include "AGXShaderLibrary.h"
+#include "AGXShaderTypes.h"
+
+
+//------------------------------------------------------------------------------
+
+#pragma mark - AGX RHI Shader Library Class Support Routines
+
+
+template<typename ShaderType>
+static TRefCountPtr<FRHIShader> AGXCreateMetalShader(TArrayView<const uint8> InCode, mtlpp::Library InLibrary)
+{
+	ShaderType* Shader = new ShaderType(InCode, InLibrary);
+	if (!Shader->GetFunction())
+	{
+		delete Shader;
+		Shader = nullptr;
+	}
+
+	return TRefCountPtr<FRHIShader>(Shader);
+}
+
+
+//------------------------------------------------------------------------------
+
+#pragma mark - AGX RHI Shader Library Class Public Static Members
+
+
+FCriticalSection FAGXShaderLibrary::LoadedShaderLibraryMutex;
+TMap<FString, FRHIShaderLibrary*> FAGXShaderLibrary::LoadedShaderLibraryMap;
+
+
+//------------------------------------------------------------------------------
+
+#pragma mark - AGX RHI Shader Library Class
+
+
+FAGXShaderLibrary::FAGXShaderLibrary(EShaderPlatform Platform,
+										 FString const& Name,
+										 const FString& InShaderLibraryFilename,
+										 const FMetalShaderLibraryHeader& InHeader,
+										 const FSerializedShaderArchive& InSerializedShaders,
+										 const TArray<uint8>& InShaderCode,
+										 const TArray<mtlpp::Library>& InLibrary)
+	: FRHIShaderLibrary(Platform, Name)
+	, ShaderLibraryFilename(InShaderLibraryFilename)
+	, Library(InLibrary)
+	, Header(InHeader)
+	, SerializedShaders(InSerializedShaders)
+	, ShaderCode(InShaderCode)
+{
+#if !UE_BUILD_SHIPPING
+	DebugFile = nullptr;
+
+	FName PlatformName = LegacyShaderPlatformToShaderFormat(Platform);
+	FString LibName = FString::Printf(TEXT("%s_%s"), *Name, *PlatformName.GetPlainNameString());
+	LibName.ToLowerInline();
+	FString Path = FPaths::ProjectContentDir() / LibName + TEXT(".zip");
+
+	if (IFileManager::Get().FileExists(*Path))
+	{
+		DebugFile = FAGXShaderDebugCache::Get().GetDebugFile(Path);
+	}
+#endif // !UE_BUILD_SHIPPING
+}
+
+FAGXShaderLibrary::~FAGXShaderLibrary()
+{
+	FScopeLock Lock(&LoadedShaderLibraryMutex);
+	LoadedShaderLibraryMap.Remove(ShaderLibraryFilename);
+}
+
+bool FAGXShaderLibrary::IsNativeLibrary() const
+{
+	return true;
+}
+
+int32 FAGXShaderLibrary::GetNumShaders() const
+{
+	return SerializedShaders.ShaderEntries.Num();
+}
+
+int32 FAGXShaderLibrary::GetNumShaderMaps() const
+{
+	return SerializedShaders.ShaderMapEntries.Num();
+}
+
+int32 FAGXShaderLibrary::GetNumShadersForShaderMap(int32 ShaderMapIndex) const
+{
+	return SerializedShaders.ShaderMapEntries[ShaderMapIndex].NumShaders;
+}
+
+int32 FAGXShaderLibrary::GetShaderIndex(int32 ShaderMapIndex, int32 i) const
+{
+	const FShaderMapEntry& ShaderMapEntry = SerializedShaders.ShaderMapEntries[ShaderMapIndex];
+	return SerializedShaders.ShaderIndices[ShaderMapEntry.ShaderIndicesOffset + i];
+}
+
+int32 FAGXShaderLibrary::FindShaderMapIndex(const FSHAHash& Hash)
+{
+	return SerializedShaders.FindShaderMap(Hash);
+}
+
+int32 FAGXShaderLibrary::FindShaderIndex(const FSHAHash& Hash)
+{
+	return SerializedShaders.FindShader(Hash);
+}
+
+TRefCountPtr<FRHIShader> FAGXShaderLibrary::CreateShader(int32 Index)
+{
+	const FShaderCodeEntry& ShaderEntry = SerializedShaders.ShaderEntries[Index];
+
+	// We don't handle compressed shaders here, since typically these are just tiny headers.
+	check(ShaderEntry.Size == ShaderEntry.UncompressedSize);
+
+	const TArrayView<uint8> Code = MakeArrayView(ShaderCode.GetData() + ShaderEntry.Offset, ShaderEntry.Size);
+	const int32 LibraryIndex = Index / Header.NumShadersPerLibrary;
+
+	TRefCountPtr<FRHIShader> Shader;
+	switch (ShaderEntry.Frequency)
+	{
+		case SF_Vertex:
+			Shader = AGXCreateMetalShader<FAGXVertexShader>(Code, Library[LibraryIndex]);
+			break;
+
+		case SF_Pixel:
+			Shader = AGXCreateMetalShader<FAGXPixelShader>(Code, Library[LibraryIndex]);
+			break;
+
+		case SF_Geometry:
+			checkf(false, TEXT("Geometry shaders not supported"));
+			break;
+
+		case SF_Compute:
+			Shader = AGXCreateMetalShader<FAGXComputeShader>(Code, Library[LibraryIndex]);
+			break;
+
+		default:
+			checkNoEntry();
+			break;
+	}
+
+	if (Shader)
+	{
+		Shader->SetHash(SerializedShaders.ShaderHashes[Index]);
+	}
+
+	return Shader;
+}
diff --git a/Engine/Source/Runtime/Apple/AGXRHI/Private/Shaders/AGXShaderLibrary.h b/Engine/Source/Runtime/Apple/AGXRHI/Private/Shaders/AGXShaderLibrary.h
new file mode 100644
--- /dev/null
+++ b/Engine/Source/Runtime/Apple/AGXRHI/Private/Shaders/AGXShaderLibrary.h
@@ -0,0 +1,49 @@
+// Copyright Epic Games, Inc. All Rights Reserved.
+
+/*=============================================================================
+	AGXShaderLibrary.h: AGX RHI Shader Library Class.
+=============================================================================*/
+
+#pragma once
+
+class FAGXShaderLibrary final : public FRHIShaderLibrary
+{
+public:
+	static FCriticalSection LoadedShaderLibraryMutex;
+	static TMap<FString, FRHIShaderLibrary*> LoadedShaderLibraryMap;
+
+	FAGXShaderLibrary(EShaderPlatform Platform,
+						FString const& Name,
+						const FString& InShaderLibraryFilename,
+						const FMetalShaderLibraryHeader& InHeader,
+						const FSerializedShaderArchive& InSerializedShaders,
+						const TArray<uint8>& InShaderCode,
+						const TArray<mtlpp::Library>& InLibrary);
+
+	virtual ~FAGXShaderLibrary();
+
+	virtual bool IsNativeLibrary() const override final;
+
+	virtual int32 GetNumShaders() const override;
+	virtual int32 GetNumShaderMaps() const override;
+	virtual int32 GetNumShadersForShaderMap(int32 ShaderMapIndex) const override;
+	virtual int32 GetShaderIndex(int32 ShaderMapIndex, int32 i) const override;
+
+	virtual int32 FindShaderMapIndex(const FSHAHash& Hash) override;
+	virtual int32 FindShaderIndex(const FSHAHash& Hash) override;
+
+	virtual bool PreloadShader(int32 ShaderIndex, FGraphEventArray& OutCompletionEvents) override { return false; }
+	virtual bool PreloadShaderMap(int32 ShaderMapIndex, FGraphEventArray& OutCompletionEvents) override { return false; }
+
+	virtual TRefCountPtr<FRHIShader> CreateShader(int32 Index) override;
+
+private:
+	FString ShaderLibraryFilename;
+	TArray<mtlpp::Library> Library;
+	FMetalShaderLibraryHeader Header;
+	FSerializedShaderArchive SerializedShaders;
+	TArray<uint8> ShaderCode;
+#if !UE_BUILD_SHIPPING
+	class FAGXShaderDebugZipFile* DebugFile;
+#endif
+};
diff --git a/Engine/Source/Runtime/Apple/AGXRHI/Private/Shaders/AGXShaderParameterCache.cpp b/Engine/Source/Runtime/Apple/AGXRHI/Private/Shaders/AGXShaderParameterCache.cpp
new file mode 100644
--- /dev/null
+++ b/Engine/Source/Runtime/Apple/AGXRHI/Private/Shaders/AGXShaderParameterCache.cpp
@@ -0,0 +1,120 @@
+// Copyright Epic Games, Inc. All Rights Reserved.
+
+/*=============================================================================
+	AGXShaderParameterCache.cpp: AGX RHI Shader Parameter Cache Class.
+=============================================================================*/
+
+
+#include "AGXRHIPrivate.h"
+#include "AGXShaderParameterCache.h"
+
+
+//------------------------------------------------------------------------------
+
+#pragma mark - AGX Shader Parameter Cache Class
+
+
+FAGXShaderParameterCache::FAGXShaderParameterCache()
+{
+	for (int32 ArrayIndex = 0; ArrayIndex < CrossCompiler::PACKED_TYPEINDEX_MAX; ++ArrayIndex)
+	{
+		PackedGlobalUniforms[ArrayIndex] = nullptr;
+		PackedGlobalUniformsSizes[ArrayIndex] = 0;
+		PackedGlobalUniformDirty[ArrayIndex].LowVector = 0;
+		PackedGlobalUniformDirty[ArrayIndex].HighVector = 0;
+	}
+}
+
+void FAGXShaderParameterCache::ResizeGlobalUniforms(uint32 TypeIndex, uint32 UniformArraySize)
+{
+	if (!PackedGlobalUniforms[TypeIndex])
+	{
+		PackedGlobalUniforms[TypeIndex] = [[FAGXBufferData alloc] initWithSize:UniformArraySize];
+	}
+	else
+	{
+		PackedGlobalUniforms[TypeIndex]->Data = (uint8*)FMemory::Realloc(PackedGlobalUniforms[TypeIndex]->Data, UniformArraySize);
+		PackedGlobalUniforms[TypeIndex]->Len = UniformArraySize;
+	}
+	PackedGlobalUniformsSizes[TypeIndex] = UniformArraySize;
+	PackedGlobalUniformDirty[TypeIndex].LowVector = 0;
+	PackedGlobalUniformDirty[TypeIndex].HighVector = 0;
+}
+
+FAGXShaderParameterCache::~FAGXShaderParameterCache()
+{
+	for (int32 ArrayIndex = 0; ArrayIndex < CrossCompiler::PACKED_TYPEINDEX_MAX; ++ArrayIndex)
+	{
+		[PackedGlobalUniforms[ArrayIndex] release];
+	}
+}
+
+void FAGXShaderParameterCache::Reset()
+{
+	for (int32 ArrayIndex = 0; ArrayIndex < CrossCompiler::PACKED_TYPEINDEX_MAX; ++ArrayIndex)
+	{
+		PackedGlobalUniformDirty[ArrayIndex].LowVector = 0;
+		PackedGlobalUniformDirty[ArrayIndex].HighVector = 0;
+	}
+}
+
+void FAGXShaderParameterCache::MarkAllDirty()
+{
+	for (int32 ArrayIndex = 0; ArrayIndex < CrossCompiler::PACKED_TYPEINDEX_MAX; ++ArrayIndex)
+	{
+		PackedGlobalUniformDirty[ArrayIndex].LowVector = 0;
+		PackedGlobalUniformDirty[ArrayIndex].HighVector = PackedGlobalUniformsSizes[ArrayIndex] / SizeOfFloat;
+	}
+}
+
+void FAGXShaderParameterCache::Set(uint32 BufferIndexName, uint32 ByteOffset, uint32 NumBytes, const void* NewValues)
+{
+	if (NumBytes)
+	{
+		uint32 BufferIndex = CrossCompiler::PackedTypeNameToTypeIndex(BufferIndexName);
+		check(BufferIndex < CrossCompiler::PACKED_TYPEINDEX_MAX);
+		check(PackedGlobalUniforms[BufferIndex]);
+		check(ByteOffset + NumBytes <= PackedGlobalUniformsSizes[BufferIndex]);
+		PackedGlobalUniformDirty[BufferIndex].LowVector = FMath::Min(PackedGlobalUniformDirty[BufferIndex].LowVector, ByteOffset / SizeOfFloat);
+		PackedGlobalUniformDirty[BufferIndex].HighVector = FMath::Max(PackedGlobalUniformDirty[BufferIndex].HighVector, (ByteOffset + NumBytes + SizeOfFloat - 1) / SizeOfFloat);
+		FMemory::Memcpy(PackedGlobalUniforms[BufferIndex]->Data + ByteOffset, NewValues, NumBytes);
+	}
+}
+
+void FAGXShaderParameterCache::CommitPackedGlobals(FAGXStateCache* Cache, FAGXCommandEncoder* Encoder, uint32 Frequency, const FMetalShaderBindings& Bindings)
+{
+	// copy the current uniform buffer into the ring buffer to submit
+	for (int32 Index = 0; Index < Bindings.PackedGlobalArrays.Num(); ++Index)
+	{
+		int32 UniformBufferIndex = Bindings.PackedGlobalArrays[Index].TypeIndex;
+
+		// is there any data that needs to be copied?
+		if (PackedGlobalUniformDirty[Index].HighVector > 0)
+		{
+			uint32 TotalSize = Bindings.PackedGlobalArrays[Index].Size;
+			uint32 SizeToUpload = PackedGlobalUniformDirty[Index].HighVector * SizeOfFloat;
+
+			//@todo-rco: Temp workaround
+			SizeToUpload = TotalSize;
+
+			//@todo-rco: Temp workaround
+			uint32 Size = FMath::Min(TotalSize, SizeToUpload);
+			if (Size > AGXBufferPageSize)
+			{
+				uint8 const* Bytes = PackedGlobalUniforms[Index]->Data;
+				ns::AutoReleased<FAGXBuffer> Buffer(Encoder->GetRingBuffer().NewBuffer(Size, 0));
+				FMemory::Memcpy((uint8*)Buffer.GetContents(), Bytes, Size);
+				Cache->SetShaderBuffer((EAGXShaderStages)Frequency, Buffer, nil, 0, Size, UniformBufferIndex, mtlpp::ResourceUsage::Read);
+			}
+			else
+			{
+				PackedGlobalUniforms[Index]->Len = Size;
+				Cache->SetShaderBuffer((EAGXShaderStages)Frequency, nil, nil, 0, 0, UniformBufferIndex, mtlpp::ResourceUsage(0));
+				Cache->SetShaderBuffer((EAGXShaderStages)Frequency, nil, PackedGlobalUniforms[Index], 0, Size, UniformBufferIndex, mtlpp::ResourceUsage::Read);
+			}
+
+			// mark as clean
+			PackedGlobalUniformDirty[Index].HighVector = 0;
+		}
+	}
+}
diff --git a/Engine/Source/Runtime/Apple/AGXRHI/Private/Shaders/AGXShaderParameterCache.h b/Engine/Source/Runtime/Apple/AGXRHI/Private/Shaders/AGXShaderParameterCache.h
new file mode 100644
--- /dev/null
+++ b/Engine/Source/Runtime/Apple/AGXRHI/Private/Shaders/AGXShaderParameterCache.h
@@ -0,0 +1,64 @@
+// Copyright Epic Games, Inc. All Rights Reserved.
+
+/*=============================================================================
+	AGXShaderParameterCache.h: AGX RHI Shader Parameter Cache Class.
+=============================================================================*/
+
+#pragma once
+
+class FAGXShaderParameterCache
+{
+public:
+	/** Constructor. */
+	FAGXShaderParameterCache();
+
+	/** Destructor. */
+	~FAGXShaderParameterCache();
+
+	inline void PrepareGlobalUniforms(uint32 TypeIndex, uint32 UniformArraySize)
+	{
+		if (PackedGlobalUniformsSizes[TypeIndex] < UniformArraySize)
+		{
+			ResizeGlobalUniforms(TypeIndex, UniformArraySize);
+		}
+	}
+
+	/**
+	 * Invalidates all existing data.
+	 */
+	void Reset();
+
+	/**
+	 * Marks all uniform arrays as dirty.
+	 */
+	void MarkAllDirty();
+
+	/**
+	 * Sets values directly into the packed uniform array
+	 */
+	void Set(uint32 BufferIndex, uint32 ByteOffset, uint32 NumBytes, const void* NewValues);
+
+	/**
+	 * Commit shader parameters to the currently bound program.
+	 */
+	void CommitPackedGlobals(class FAGXStateCache* Cache, class FAGXCommandEncoder* Encoder, uint32 Frequency, const FMetalShaderBindings& Bindings);
+
+private:
+	static constexpr uint32 SizeOfFloat = sizeof(float);
+
+	/** CPU memory block for storing uniform values. */
+	FAGXBufferData* PackedGlobalUniforms[CrossCompiler::PACKED_TYPEINDEX_MAX];
+
+	struct FRange
+	{
+		uint32 LowVector;
+		uint32 HighVector;
+	};
+
+	/** Dirty ranges for each uniform array. */
+	FRange PackedGlobalUniformDirty[CrossCompiler::PACKED_TYPEINDEX_MAX];
+
+	uint32 PackedGlobalUniformsSizes[CrossCompiler::PACKED_TYPEINDEX_MAX];
+
+	void ResizeGlobalUniforms(uint32 TypeIndex, uint32 UniformArraySize);
+};
diff --git a/Engine/Source/Runtime/Apple/AGXRHI/Private/Shaders/Debugging/AGXShaderDebugCache.cpp b/Engine/Source/Runtime/Apple/AGXRHI/Private/Shaders/Debugging/AGXShaderDebugCache.cpp
new file mode 100644
--- /dev/null
+++ b/Engine/Source/Runtime/Apple/AGXRHI/Private/Shaders/Debugging/AGXShaderDebugCache.cpp
@@ -0,0 +1,50 @@
+// Copyright Epic Games, Inc. All Rights Reserved.
+
+/*=============================================================================
+	AGXShaderDebugCache.cpp: AGX RHI Shader Debug Cache.
+=============================================================================*/
+
+#include "CoreMinimal.h"
+
+THIRD_PARTY_INCLUDES_START
+#include "mtlpp.hpp"
+THIRD_PARTY_INCLUDES_END
+
+#include "HAL/FileManager.h"
+#include "HAL/PlatformFileManager.h"
+#include "Misc/Paths.h"
+#include "Misc/ScopeLock.h"
+
+#include "AGXShaderDebugCache.h"
+#include "AGXShaderDebugZipFile.h"
+
+#if !UE_BUILD_SHIPPING
+
+FAGXShaderDebugZipFile* FAGXShaderDebugCache::GetDebugFile(FString Path)
+{
+	FScopeLock Lock(&Mutex);
+	FAGXShaderDebugZipFile* Ref = DebugFiles.FindRef(Path);
+	if (!Ref)
+	{
+		Ref = new FAGXShaderDebugZipFile(Path);
+		DebugFiles.Add(Path, Ref);
+	}
+	return Ref;
+}
+
+ns::String FAGXShaderDebugCache::GetShaderCode(uint32 ShaderSrcLen, uint32 ShaderSrcCRC)
+{
+	ns::String Code;
+	FScopeLock Lock(&Mutex);
+	for (auto const& Ref : DebugFiles)
+	{
+		Code = Ref.Value->GetShaderCode(ShaderSrcLen, ShaderSrcCRC);
+		if (Code)
+		{
+			break;
+		}
+	}
+	return Code;
+}
+
+#endif // !UE_BUILD_SHIPPING
diff --git a/Engine/Source/Runtime/Apple/AGXRHI/Private/Shaders/Debugging/AGXShaderDebugCache.h b/Engine/Source/Runtime/Apple/AGXRHI/Private/Shaders/Debugging/AGXShaderDebugCache.h
new file mode 100644
--- /dev/null
+++ b/Engine/Source/Runtime/Apple/AGXRHI/Private/Shaders/Debugging/AGXShaderDebugCache.h
@@ -0,0 +1,26 @@
+// Copyright Epic Games, Inc. All Rights Reserved.
+
+/*=============================================================================
+	AGXShaderDebugCache.h: AGX RHI Shader Debug Cache.
+=============================================================================*/
+
+#pragma once
+
+#if !UE_BUILD_SHIPPING
+
+struct FAGXShaderDebugCache
+{
+	static FAGXShaderDebugCache& Get()
+	{
+		static FAGXShaderDebugCache sSelf;
+		return sSelf;
+	}
+	
+	class FAGXShaderDebugZipFile* GetDebugFile(FString Path);
+	ns::String GetShaderCode(uint32 ShaderSrcLen, uint32 ShaderSrcCRC);
+	
+	FCriticalSection Mutex;
+	TMap<FString, class FAGXShaderDebugZipFile*> DebugFiles;
+};
+
+#endif // !UE_BUILD_SHIPPING
diff --git a/Engine/Source/Runtime/Apple/AGXRHI/Private/Shaders/Debugging/AGXShaderDebugZipFile.cpp b/Engine/Source/Runtime/Apple/AGXRHI/Private/Shaders/Debugging/AGXShaderDebugZipFile.cpp
new file mode 100644
--- /dev/null
+++ b/Engine/Source/Runtime/Apple/AGXRHI/Private/Shaders/Debugging/AGXShaderDebugZipFile.cpp
@@ -0,0 +1,250 @@
+// Copyright Epic Games, Inc. All Rights Reserved.
+
+/*=============================================================================
+	AGXShaderDebugZipFile.cpp: AGX shader RHI implementation.
+=============================================================================*/
+
+#include "CoreMinimal.h"
+
+THIRD_PARTY_INCLUDES_START
+#include "mtlpp.hpp"
+THIRD_PARTY_INCLUDES_END
+
+
+#include "HAL/FileManager.h"
+#include "HAL/PlatformFileManager.h"
+#include "Misc/Paths.h"
+#include "Misc/ScopeLock.h"
+
+#include "AGXShaderDebugZipFile.h"
+
+#if !UE_BUILD_SHIPPING
+
+FAGXShaderDebugZipFile::FAGXShaderDebugZipFile(FString LibPath)
+	: File(nullptr)
+{
+	IPlatformFile& PlatformFile = FPlatformFileManager::Get().GetPlatformFile();
+	File = PlatformFile.OpenRead(*LibPath);
+	if (File)
+	{
+		int64 SeekEndOffset = -1;
+
+		// Write normal end of central directory record
+		const static uint8 EndRecord[] =
+		{
+			0x50, 0x4b, 0x05, 0x06, 0xff, 0xff, 0xff, 0xff,
+			0xff, 0xff, 0xff, 0xff, 0xff, 0xff, 0xff, 0xff,
+			0xff, 0xff, 0xff, 0xff, 0x00, 0x00
+		};
+
+		SeekEndOffset += sizeof(EndRecord);
+		bool bOK = File->SeekFromEnd(-SeekEndOffset);
+		if (bOK)
+		{
+			TArray<uint8> Data;
+			Data.AddZeroed(sizeof(EndRecord));
+			bOK = File->Read(Data.GetData(), sizeof(EndRecord));
+			if (bOK)
+			{
+				bOK = (FMemory::Memcmp(Data.GetData(), EndRecord, sizeof(EndRecord)) == 0);
+			}
+		}
+
+		// Write ZIP64 end of central directory locator
+		const static uint8 Locator[] =
+		{
+			0x50, 0x4b, 0x06, 0x07, 0x00, 0x00, 0x00, 0x00,
+		};
+		uint64 DirEndOffset = 0;
+		if (bOK)
+		{
+			SeekEndOffset += sizeof(Locator) + sizeof(uint64) + sizeof(uint32);
+			bOK = File->SeekFromEnd(-SeekEndOffset);
+			if (bOK)
+			{
+				TArray<uint8> Data;
+				Data.AddZeroed(sizeof(Locator));
+				bOK = File->Read(Data.GetData(), sizeof(Locator));
+				if (bOK)
+				{
+					bOK = (FMemory::Memcmp(Data.GetData(), Locator, sizeof(Locator)) == 0);
+				}
+				if (bOK)
+				{
+					bOK = File->Read((uint8*)&DirEndOffset, sizeof(uint64));
+				}
+			}
+		}
+
+		// Write ZIP64 end of central directory record
+		const static uint8 Record[] =
+		{
+			0x50, 0x4b, 0x06, 0x06, 0x2c, 0x00, 0x00, 0x00,
+			0x00, 0x00, 0x00, 0x00, 0x2d, 0x00, 0x2d, 0x00,
+			0x00, 0x00, 0x00, 0x00, 0x00, 0x00, 0x00, 0x00,
+		};
+		struct FAGXZipRecordData
+		{
+			uint64 FilesNum;
+			uint64 FilesNum2;
+			uint64 DirectorySizeInBytes;
+			uint64 DirStartOffset;
+		} RecordData;
+		if (bOK)
+		{
+			SeekEndOffset += sizeof(Record) + (sizeof(uint64) * 4);
+			bOK = File->SeekFromEnd(-SeekEndOffset);
+			if (bOK)
+			{
+				TArray<uint8> Data;
+				Data.AddZeroed(sizeof(Record));
+				bOK = File->Read(Data.GetData(), sizeof(Record));
+				if (bOK)
+				{
+					bOK = (FMemory::Memcmp(Data.GetData(), Record, sizeof(Record)) == 0);
+				}
+				if (bOK)
+				{
+					bOK = File->Read((uint8*)&RecordData, sizeof(RecordData));
+				}
+			}
+		}
+
+		if (bOK)
+		{
+			bOK = File->Seek(RecordData.DirStartOffset);
+			if (bOK)
+			{
+				const static uint8 Footer[] =
+				{
+					0x50, 0x4b, 0x01, 0x02, 0x3f, 0x00, 0x2d, 0x00,
+					0x00, 0x00, 0x00, 0x00
+				};
+				const static uint8 Fields[] =
+				{
+					0x20, 0x00, 0x00, 0x00, 0x00, 0x00, 0x00, 0x00,
+					0x20, 0x00, 0x00, 0x00, 0xff, 0xff, 0xff, 0xff
+				};
+
+				TArray<uint8> FooterData;
+				FooterData.AddZeroed(sizeof(Footer));
+				TArray<uint8> FieldsData;
+				FieldsData.AddZeroed(sizeof(Fields));
+
+				struct FAGXZipFileHeader {
+					uint32 Time;
+					uint32 CRC;
+					uint64 SizeMarker;
+					uint16 FilenameLen;
+				} __attribute__((packed)) Header;
+
+				struct FAGXZipFileTrailer {
+					uint16 Flags;
+					uint16 Attribs;
+					uint64 CompressedLen;
+					uint64 UncompressedLen;
+					uint64 Offset;
+					uint32 DiskNum;
+				} __attribute__((packed)) Trailer;
+
+				static const uint8 FileHeader[] =
+				{
+					0x50, 0x4b, 0x03, 0x04, 0x2d, 0x00, 0x00, 0x00,
+					0x00, 0x00
+				};
+
+				uint32 FileHeaderFixedSize = sizeof(FileHeader) + sizeof(FAGXZipFileHeader) + sizeof(uint16) + sizeof(FAGXZipFileTrailer);
+
+				FString Filename;
+
+				while (bOK && Files.Num() < RecordData.FilesNum && File->Tell() < RecordData.DirStartOffset + RecordData.DirectorySizeInBytes)
+				{
+					bOK = File->Read(FooterData.GetData(), sizeof(Footer));
+					if (bOK)
+					{
+						bOK = (FMemory::Memcmp(FooterData.GetData(), Footer, sizeof(Footer)) == 0);
+					}
+
+					if (bOK)
+					{
+						bOK = File->Read((uint8*)&Header, sizeof(Header));
+						if (bOK)
+						{
+							bOK = (Header.SizeMarker == (uint64)0xffffffffffffffff);
+						}
+					}
+
+					if (bOK)
+					{
+						bOK = File->Read(FieldsData.GetData(), sizeof(Fields));
+						if (bOK)
+						{
+							bOK = (FMemory::Memcmp(FieldsData.GetData(), Fields, sizeof(Fields)) == 0);
+						}
+					}
+
+					if (bOK)
+					{
+						TArray<uint8> FilenameData;
+						FilenameData.AddZeroed(Header.FilenameLen+1);
+						bOK = File->Read(FilenameData.GetData(), Header.FilenameLen);
+						if (bOK)
+						{
+							Filename = UTF8_TO_TCHAR((char const*)FilenameData.GetData());
+						}
+					}
+
+					if (bOK)
+					{
+						bOK = File->Read((uint8*)&Trailer, sizeof(Trailer));
+						if (bOK)
+						{
+							bOK = (Trailer.Flags == (uint16)0x01 && Trailer.Attribs == (uint16)0x1c && Trailer.DiskNum == 0);
+						}
+					}
+
+					if (bOK)
+					{
+						FFileEntry NewEntry(Filename, Header.CRC, Trailer.UncompressedLen, Trailer.Offset + FileHeaderFixedSize + Header.FilenameLen, Header.Time);
+						Files.Add(NewEntry);
+					}
+				}
+			}
+		}
+	}
+}
+
+FAGXShaderDebugZipFile::~FAGXShaderDebugZipFile()
+{
+	if (File)
+	{
+		delete File;
+	}
+}
+
+ns::String FAGXShaderDebugZipFile::GetShaderCode(uint32 ShaderSrcLen, uint32 ShaderSrcCRC)
+{
+	ns::String Source;
+	FScopeLock Lock(&Mutex);
+	FString Name = FString::Printf(TEXT("%u_%u.metal"), ShaderSrcLen, ShaderSrcCRC);
+	for (auto const& Entry : Files)
+	{
+		if (FPaths::GetCleanFilename(Entry.Filename) == Name)
+		{
+			if (File->Seek(Entry.Offset))
+			{
+				TArray<uint8> Data;
+				Data.AddZeroed(Entry.Length+1);
+				if (File->Read(Data.GetData(), Entry.Length))
+				{
+					Source = [NSString stringWithUTF8String:(char const*)Data.GetData()];
+				}
+			}
+
+			break;
+		}
+	}
+	return Source;
+}
+
+#endif // !UE_BUILD_SHIPPING
diff --git a/Engine/Source/Runtime/Apple/AGXRHI/Private/Shaders/Debugging/AGXShaderDebugZipFile.h b/Engine/Source/Runtime/Apple/AGXRHI/Private/Shaders/Debugging/AGXShaderDebugZipFile.h
new file mode 100644
--- /dev/null
+++ b/Engine/Source/Runtime/Apple/AGXRHI/Private/Shaders/Debugging/AGXShaderDebugZipFile.h
@@ -0,0 +1,44 @@
+// Copyright Epic Games, Inc. All Rights Reserved.
+
+/*=============================================================================
+	AGXShaderDebugZipFile.h: AGX RHI Shader Debug Zip File.
+=============================================================================*/
+
+#pragma once
+
+#if !UE_BUILD_SHIPPING
+
+class FAGXShaderDebugZipFile
+{
+	struct FFileEntry
+	{
+		FString Filename;
+		uint32  Crc32;
+		uint64  Length;
+		uint64  Offset;
+		uint32  Time;
+
+		FFileEntry(const FString& InFilename, uint32 InCrc32, uint64 InLength, uint64 InOffset, uint32 InTime)
+			: Filename(InFilename)
+			, Crc32(InCrc32)
+			, Length(InLength)
+			, Offset(InOffset)
+			, Time(InTime)
+		{
+			// VOID
+		}
+	};
+
+public:
+	FAGXShaderDebugZipFile(FString LibPath);
+	~FAGXShaderDebugZipFile();
+
+	ns::String GetShaderCode(uint32 ShaderSrcLen, uint32 ShaderSrcCRC);
+
+private:
+	FCriticalSection Mutex;
+	IFileHandle* File;
+	TArray<FFileEntry> Files;
+};
+
+#endif // !UE_BUILD_SHIPPING
diff --git a/Engine/Source/Runtime/Apple/AGXRHI/Private/Shaders/Types/AGXComputeShader.cpp b/Engine/Source/Runtime/Apple/AGXRHI/Private/Shaders/Types/AGXComputeShader.cpp
new file mode 100644
--- /dev/null
+++ b/Engine/Source/Runtime/Apple/AGXRHI/Private/Shaders/Types/AGXComputeShader.cpp
@@ -0,0 +1,127 @@
+// Copyright Epic Games, Inc. All Rights Reserved.
+
+/*=============================================================================
+	AGXComputeShader.cpp: AGX RHI Compute Shader Class Implementation.
+=============================================================================*/
+
+
+#include "AGXRHIPrivate.h"
+#include "Templates/AGXBaseShader.h"
+#include "AGXComputeShader.h"
+
+
+//------------------------------------------------------------------------------
+
+#pragma mark - AGX RHI Compute Shader Class
+
+
+FAGXComputeShader::FAGXComputeShader(TArrayView<const uint8> InCode, mtlpp::Library InLibrary)
+	: NumThreadsX(0)
+	, NumThreadsY(0)
+	, NumThreadsZ(0)
+{
+	Pipeline = nil;
+	FMetalCodeHeader Header;
+	Init(InCode, Header, InLibrary);
+
+	NumThreadsX = FMath::Max((int32)Header.NumThreadsX, 1);
+	NumThreadsY = FMath::Max((int32)Header.NumThreadsY, 1);
+	NumThreadsZ = FMath::Max((int32)Header.NumThreadsZ, 1);
+}
+
+FAGXComputeShader::~FAGXComputeShader()
+{
+	[Pipeline release];
+	Pipeline = nil;
+}
+
+FAGXShaderPipeline* FAGXComputeShader::GetPipeline()
+{
+	if (!Pipeline)
+	{
+		mtlpp::Function Func = GetCompiledFunction();
+		check(Func);
+
+		ns::Error Error;
+		mtlpp::ComputePipelineDescriptor Descriptor;
+		Descriptor.SetLabel(Func.GetName());
+		Descriptor.SetComputeFunction(Func);
+		if (FAGXCommandQueue::SupportsFeature(EAGXFeaturesTextureBuffers))
+		{
+			Descriptor.SetMaxTotalThreadsPerThreadgroup(NumThreadsX*NumThreadsY*NumThreadsZ);
+		}
+
+		if (FAGXCommandQueue::SupportsFeature(EAGXFeaturesPipelineBufferMutability))
+		{
+			ns::AutoReleased<ns::Array<mtlpp::PipelineBufferDescriptor>> PipelineBuffers = Descriptor.GetBuffers();
+
+			uint32 ImmutableBuffers = Bindings.ConstantBuffers | Bindings.ArgumentBuffers;
+			while(ImmutableBuffers)
+			{
+				uint32 Index = __builtin_ctz(ImmutableBuffers);
+				ImmutableBuffers &= ~(1 << Index);
+
+				if (Index < ML_MaxBuffers)
+				{
+					ns::AutoReleased<mtlpp::PipelineBufferDescriptor> PipelineBuffer = PipelineBuffers[Index];
+					PipelineBuffer.SetMutability(mtlpp::Mutability::Immutable);
+				}
+			}
+			if (SideTableBinding > 0)
+			{
+				ns::AutoReleased<mtlpp::PipelineBufferDescriptor> PipelineBuffer = PipelineBuffers[SideTableBinding];
+				PipelineBuffer.SetMutability(mtlpp::Mutability::Immutable);
+			}
+		}
+
+		mtlpp::ComputePipelineState Kernel;
+		mtlpp::ComputePipelineReflection Reflection;
+
+		METAL_GPUPROFILE(FAGXScopedCPUStats CPUStat(FString::Printf(TEXT("NewComputePipeline: %d_%d"), SourceLen, SourceCRC)));
+#if METAL_DEBUG_OPTIONS
+		if (GetAGXDeviceContext().GetCommandQueue().GetRuntimeDebuggingLevel() >= EAGXDebugLevelFastValidation)
+		{
+			ns::AutoReleasedError ComputeError;
+			mtlpp::AutoReleasedComputePipelineReflection ComputeReflection;
+
+			NSUInteger ComputeOption = mtlpp::PipelineOption::ArgumentInfo | mtlpp::PipelineOption::BufferTypeInfo;
+			Kernel = GetAGXDeviceContext().GetDevice().NewComputePipelineState(Descriptor, mtlpp::PipelineOption(ComputeOption), &ComputeReflection, &ComputeError);
+			Error = ComputeError;
+			Reflection = ComputeReflection;
+		}
+		else
+#endif // METAL_DEBUG_OPTIONS
+		{
+			ns::AutoReleasedError ComputeError;
+			Kernel = GetAGXDeviceContext().GetDevice().NewComputePipelineState(Descriptor, mtlpp::PipelineOption(0), nil, &ComputeError);
+			Error = ComputeError;
+		}
+
+		if (Kernel == nil)
+		{
+			UE_LOG(LogRHI, Error, TEXT("*********** Error\n%s"), *FString(GetSourceCode()));
+			UE_LOG(LogRHI, Fatal, TEXT("Failed to create compute kernel: %s"), *FString([Error description]));
+		}
+
+		Pipeline = [FAGXShaderPipeline new];
+		Pipeline->ComputePipelineState = Kernel;
+#if METAL_DEBUG_OPTIONS
+		Pipeline->ComputePipelineReflection = Reflection;
+		Pipeline->ComputeSource = GetSourceCode();
+		if (Reflection)
+		{
+			Pipeline->ComputeDesc = Descriptor;
+		}
+#endif // METAL_DEBUG_OPTIONS
+		METAL_DEBUG_OPTION(FMemory::Memzero(Pipeline->ResourceMask, sizeof(Pipeline->ResourceMask)));
+	}
+	check(Pipeline);
+
+	return Pipeline;
+}
+
+mtlpp::Function FAGXComputeShader::GetFunction()
+{
+	return GetCompiledFunction();
+}
+
diff --git a/Engine/Source/Runtime/Apple/AGXRHI/Private/Shaders/Types/AGXComputeShader.h b/Engine/Source/Runtime/Apple/AGXRHI/Private/Shaders/Types/AGXComputeShader.h
new file mode 100644
--- /dev/null
+++ b/Engine/Source/Runtime/Apple/AGXRHI/Private/Shaders/Types/AGXComputeShader.h
@@ -0,0 +1,32 @@
+// Copyright Epic Games, Inc. All Rights Reserved.
+
+/*=============================================================================
+	AGXComputeShader.h: AGX RHI Compute Shader Class Definition.
+=============================================================================*/
+
+#pragma once
+
+
+//------------------------------------------------------------------------------
+
+#pragma mark - AGX RHI Compute Shader Class
+
+
+class FAGXComputeShader : public TAGXBaseShader<FRHIComputeShader, SF_Compute>
+{
+public:
+	FAGXComputeShader(TArrayView<const uint8> InCode, mtlpp::Library InLibrary = nil);
+	virtual ~FAGXComputeShader();
+
+	FAGXShaderPipeline* GetPipeline();
+	mtlpp::Function GetFunction();
+
+	// thread group counts
+	int32 NumThreadsX;
+	int32 NumThreadsY;
+	int32 NumThreadsZ;
+
+private:
+	// the state object for a compute shader
+	FAGXShaderPipeline* Pipeline;
+};
diff --git a/Engine/Source/Runtime/Apple/AGXRHI/Private/Shaders/Types/AGXGeometryShader.h b/Engine/Source/Runtime/Apple/AGXRHI/Private/Shaders/Types/AGXGeometryShader.h
new file mode 100644
--- /dev/null
+++ b/Engine/Source/Runtime/Apple/AGXRHI/Private/Shaders/Types/AGXGeometryShader.h
@@ -0,0 +1,17 @@
+// Copyright Epic Games, Inc. All Rights Reserved.
+
+/*=============================================================================
+	AGXGeometryShader.h: AGX RHI Geometry Shader Class Definition.
+=============================================================================*/
+
+#pragma once
+
+
+//------------------------------------------------------------------------------
+
+#pragma mark - AGX RHI Geometry Shader Class
+
+
+class FAGXGeometryShader : public TAGXBaseShader<FRHIGeometryShader, SF_Geometry>
+{
+};
diff --git a/Engine/Source/Runtime/Apple/AGXRHI/Private/Shaders/Types/AGXPixelShader.cpp b/Engine/Source/Runtime/Apple/AGXRHI/Private/Shaders/Types/AGXPixelShader.cpp
new file mode 100644
--- /dev/null
+++ b/Engine/Source/Runtime/Apple/AGXRHI/Private/Shaders/Types/AGXPixelShader.cpp
@@ -0,0 +1,33 @@
+// Copyright Epic Games, Inc. All Rights Reserved.
+
+/*=============================================================================
+	AGXPixelShader.cpp: AGX RHI Pixel Shader Class Implementation.
+=============================================================================*/
+
+
+#include "AGXRHIPrivate.h"
+#include "Templates/AGXBaseShader.h"
+#include "AGXPixelShader.h"
+
+
+//------------------------------------------------------------------------------
+
+#pragma mark - AGX RHI Pixel Shader Class
+
+
+FAGXPixelShader::FAGXPixelShader(TArrayView<const uint8> InCode)
+{
+	FMetalCodeHeader Header;
+	Init(InCode, Header);
+}
+
+FAGXPixelShader::FAGXPixelShader(TArrayView<const uint8> InCode, mtlpp::Library InLibrary)
+{
+	FMetalCodeHeader Header;
+	Init(InCode, Header, InLibrary);
+}
+
+mtlpp::Function FAGXPixelShader::GetFunction()
+{
+	return GetCompiledFunction();
+}
diff --git a/Engine/Source/Runtime/Apple/AGXRHI/Private/Shaders/Types/AGXPixelShader.h b/Engine/Source/Runtime/Apple/AGXRHI/Private/Shaders/Types/AGXPixelShader.h
new file mode 100644
--- /dev/null
+++ b/Engine/Source/Runtime/Apple/AGXRHI/Private/Shaders/Types/AGXPixelShader.h
@@ -0,0 +1,22 @@
+// Copyright Epic Games, Inc. All Rights Reserved.
+
+/*=============================================================================
+	AGXPixelShader.h: AGX RHI Pixel Shader Class Definition.
+=============================================================================*/
+
+#pragma once
+
+
+//------------------------------------------------------------------------------
+
+#pragma mark - AGX RHI Pixel Shader Class
+
+
+class FAGXPixelShader : public TAGXBaseShader<FRHIPixelShader, SF_Pixel>
+{
+public:
+	FAGXPixelShader(TArrayView<const uint8> InCode);
+	FAGXPixelShader(TArrayView<const uint8> InCode, mtlpp::Library InLibrary);
+
+	mtlpp::Function GetFunction();
+};
diff --git a/Engine/Source/Runtime/Apple/AGXRHI/Private/Shaders/Types/AGXVertexShader.cpp b/Engine/Source/Runtime/Apple/AGXRHI/Private/Shaders/Types/AGXVertexShader.cpp
new file mode 100644
--- /dev/null
+++ b/Engine/Source/Runtime/Apple/AGXRHI/Private/Shaders/Types/AGXVertexShader.cpp
@@ -0,0 +1,33 @@
+// Copyright Epic Games, Inc. All Rights Reserved.
+
+/*=============================================================================
+	AGXVertexShader.cpp: AGX RHI Vertex Shader Class Implementation.
+=============================================================================*/
+
+
+#include "AGXRHIPrivate.h"
+#include "Templates/AGXBaseShader.h"
+#include "AGXVertexShader.h"
+
+
+//------------------------------------------------------------------------------
+
+#pragma mark - AGX RHI Vertex Shader Class
+
+
+FAGXVertexShader::FAGXVertexShader(TArrayView<const uint8> InCode)
+{
+	FMetalCodeHeader Header;
+	Init(InCode, Header);
+}
+
+FAGXVertexShader::FAGXVertexShader(TArrayView<const uint8> InCode, mtlpp::Library InLibrary)
+{
+	FMetalCodeHeader Header;
+	Init(InCode, Header, InLibrary);
+}
+
+mtlpp::Function FAGXVertexShader::GetFunction()
+{
+	return GetCompiledFunction();
+}
diff --git a/Engine/Source/Runtime/Apple/AGXRHI/Private/Shaders/Types/AGXVertexShader.h b/Engine/Source/Runtime/Apple/AGXRHI/Private/Shaders/Types/AGXVertexShader.h
new file mode 100644
--- /dev/null
+++ b/Engine/Source/Runtime/Apple/AGXRHI/Private/Shaders/Types/AGXVertexShader.h
@@ -0,0 +1,22 @@
+// Copyright Epic Games, Inc. All Rights Reserved.
+
+/*=============================================================================
+	AGXVertexShader.h: AGX RHI Vertex Shader Class Definition.
+=============================================================================*/
+
+#pragma once
+
+
+//------------------------------------------------------------------------------
+
+#pragma mark - AGX RHI Vertex Shader Class
+
+
+class FAGXVertexShader : public TAGXBaseShader<FRHIVertexShader, SF_Vertex>
+{
+public:
+	FAGXVertexShader(TArrayView<const uint8> InCode);
+	FAGXVertexShader(TArrayView<const uint8> InCode, mtlpp::Library InLibrary);
+
+	mtlpp::Function GetFunction();
+};
diff --git a/Engine/Source/Runtime/Apple/AGXRHI/Private/Shaders/Types/Templates/AGXBaseShader.h b/Engine/Source/Runtime/Apple/AGXRHI/Private/Shaders/Types/Templates/AGXBaseShader.h
new file mode 100644
--- /dev/null
+++ b/Engine/Source/Runtime/Apple/AGXRHI/Private/Shaders/Types/Templates/AGXBaseShader.h
@@ -0,0 +1,516 @@
+// Copyright Epic Games, Inc. All Rights Reserved.
+
+/*=============================================================================
+	AGXBaseShader.h: AGX RHI Base Shader Class Template.
+=============================================================================*/
+
+#pragma once
+
+#include "Misc/Paths.h"
+#include "HAL/FileManager.h"
+#include "HAL/PlatformFileManager.h"
+#include "Misc/FileHelper.h"
+#include "Shaders/Debugging/AGXShaderDebugCache.h"
+#include "Shaders/AGXCompiledShaderKey.h"
+#include "Shaders/AGXCompiledShaderCache.h"
+
+
+//------------------------------------------------------------------------------
+
+#pragma mark - AGX RHI Base Shader Class Support Routines
+
+
+extern mtlpp::LanguageVersion AGXValidateVersion(uint8 Version);
+
+
+//------------------------------------------------------------------------------
+
+#pragma mark - AGX RHI Base Shader Class Template Defines
+
+
+/** Set to 1 to enable shader debugging (makes the driver save the shader source) */
+#define DEBUG_METAL_SHADERS (UE_BUILD_DEBUG || UE_BUILD_DEVELOPMENT)
+
+
+//------------------------------------------------------------------------------
+
+#pragma mark - AGX RHI Base Shader Class Template
+
+template<typename BaseResourceType, int32 ShaderType>
+class TAGXBaseShader : public BaseResourceType, public IRefCountedObject
+{
+public:
+	enum
+	{
+		StaticFrequency = ShaderType
+	};
+
+	TAGXBaseShader()
+	{
+		// void
+	}
+
+	virtual ~TAGXBaseShader()
+	{
+		this->Destroy();
+	}
+
+	void Init(TArrayView<const uint8> InCode, FMetalCodeHeader& Header, mtlpp::Library InLibrary = nil);
+	void Destroy();
+
+	/**
+	 * Gets the metal source code as an NSString if available or nil if not.  Note
+	 * that this will dynamically decompress from compressed data on first
+	 * invocation.
+	 */
+	NSString* GetSourceCode();
+
+	// IRefCountedObject interface:
+	virtual uint32 AddRef() const override final;
+	virtual uint32 Release() const override final;
+	virtual uint32 GetRefCount() const override final;
+
+	/** External bindings for this shader. */
+	FMetalShaderBindings Bindings;
+
+	// List of memory copies from RHIUniformBuffer to packed uniforms
+	TArray<CrossCompiler::FUniformBufferCopyInfo> UniformBuffersCopyInfo;
+
+	/* Argument encoders for shader IABs */
+	TMap<uint32, mtlpp::ArgumentEncoder> ArgumentEncoders;
+
+	/* Tier1 Argument buffer bitmasks */
+	TMap<uint32, TBitArray<>> ArgumentBitmasks;
+
+	/* Uniform buffer static slots */
+	TArray<FUniformBufferStaticSlot> StaticSlots;
+
+	/** The binding for the buffer side-table if present */
+	int32 SideTableBinding = -1;
+
+	/** CRC & Len for name disambiguation */
+	uint32 SourceLen = 0;
+	uint32 SourceCRC = 0;
+
+	/** Hash for the shader/material permutation constants */
+	uint32 ConstantValueHash = 0;
+
+protected:
+	mtlpp::Function GetCompiledFunction(bool const bAsync = false);
+
+	// this is the compiler shader
+	mtlpp::Function Function = nil;
+
+private:
+	// This is the MTLLibrary for the shader so we can dynamically refine the MTLFunction
+	mtlpp::Library Library = nil;
+
+	/** The debuggable text source */
+	NSString* GlslCodeNSString = nil;
+
+	/** The compressed text source */
+	TArray<uint8> CompressedSource;
+
+	/** The uncompressed text source size */
+	uint32 CodeSize = 0;
+
+	// Function constant states
+	bool bHasFunctionConstants = false;
+	bool bDeviceFunctionConstants = false;
+};
+
+
+//------------------------------------------------------------------------------
+
+#pragma mark - AGX RHI Base Shader Class Template Member Functions
+
+
+template<typename BaseResourceType, int32 ShaderType>
+void TAGXBaseShader<BaseResourceType, ShaderType>::Init(TArrayView<const uint8> InShaderCode, FMetalCodeHeader& Header, mtlpp::Library InLibrary)
+{
+	FShaderCodeReader ShaderCode(InShaderCode);
+
+	FMemoryReaderView Ar(InShaderCode, true);
+
+	Ar.SetLimitSize(ShaderCode.GetActualShaderCodeSize());
+
+	// was the shader already compiled offline?
+	uint8 OfflineCompiledFlag;
+	Ar << OfflineCompiledFlag;
+	check(OfflineCompiledFlag == 0 || OfflineCompiledFlag == 1);
+
+	// get the header
+	Ar << Header;
+
+	AGXValidateVersion(Header.Version);
+
+	SourceLen = Header.SourceLen;
+	SourceCRC = Header.SourceCRC;
+
+	// If this triggers than a level above us has failed to provide valid shader data and the cook is probably bogus
+	UE_CLOG(Header.SourceLen == 0 || Header.SourceCRC == 0, LogAGX, Fatal, TEXT("Invalid Shader Bytecode provided."));
+
+	bDeviceFunctionConstants = Header.bDeviceFunctionConstants;
+
+	// remember where the header ended and code (precompiled or source) begins
+	int32 CodeOffset = Ar.Tell();
+	uint32 BufferSize = ShaderCode.GetActualShaderCodeSize() - CodeOffset;
+	const ANSICHAR* SourceCode = (ANSICHAR*)InShaderCode.GetData() + CodeOffset;
+
+	// Only archived shaders should be in here.
+	UE_CLOG(InLibrary && !(Header.CompileFlags & (1 << CFLAG_Archive)), LogAGX, Warning, TEXT("Shader being loaded wasn't marked for archiving but a MTLLibrary was provided - this is unsupported."));
+
+	if (!OfflineCompiledFlag)
+	{
+		UE_LOG(LogAGX, Display, TEXT("Loaded a text shader (will be slower to load)"));
+	}
+
+	bool bOfflineCompile = (OfflineCompiledFlag > 0);
+
+	const ANSICHAR* ShaderSource = ShaderCode.FindOptionalData('c');
+	bool bHasShaderSource = (ShaderSource && FCStringAnsi::Strlen(ShaderSource) > 0);
+
+	static bool bForceTextShaders = FAGXCommandQueue::SupportsFeature(EAGXFeaturesGPUTrace);
+	if (!bHasShaderSource)
+	{
+		int32 LZMASourceSize = 0;
+		int32 SourceSize = 0;
+		const uint8* LZMASource = ShaderCode.FindOptionalDataAndSize('z', LZMASourceSize);
+		const uint8* UnSourceLen = ShaderCode.FindOptionalDataAndSize('u', SourceSize);
+		if (LZMASource && LZMASourceSize > 0 && UnSourceLen && SourceSize == sizeof(uint32))
+		{
+			CompressedSource.Append(LZMASource, LZMASourceSize);
+			memcpy(&CodeSize, UnSourceLen, sizeof(uint32));
+			bHasShaderSource = false;
+		}
+#if !UE_BUILD_SHIPPING
+		else if(bForceTextShaders)
+		{
+			GlslCodeNSString = [FAGXShaderDebugCache::Get().GetShaderCode(SourceLen, SourceCRC).GetPtr() retain];
+		}
+#endif
+		if (bForceTextShaders && CodeSize && CompressedSource.Num())
+		{
+			bHasShaderSource = (GetSourceCode() != nil);
+		}
+	}
+	else if (bOfflineCompile && bHasShaderSource)
+	{
+		GlslCodeNSString = [NSString stringWithUTF8String:ShaderSource];
+		check(GlslCodeNSString);
+		[GlslCodeNSString retain];
+	}
+
+	bHasFunctionConstants = (Header.bDeviceFunctionConstants);
+
+	ConstantValueHash = 0;
+
+	Library = InLibrary;
+
+	bool bNeedsCompiling = false;
+
+	// Find the existing compiled shader in the cache.
+	uint32 FunctionConstantHash = ConstantValueHash;
+	FAGXCompiledShaderKey Key(Header.SourceLen, Header.SourceCRC, FunctionConstantHash);
+
+	Function = GetAGXCompiledShaderCache().FindRef(Key);
+	if (!Library && Function)
+	{
+		Library = GetAGXCompiledShaderCache().FindLibrary(Function);
+	}
+	else
+	{
+		bNeedsCompiling = true;
+	}
+
+	Bindings = Header.Bindings;
+	if (bNeedsCompiling || !Library)
+	{
+		if (bOfflineCompile METAL_DEBUG_OPTION(&& !(bHasShaderSource && bForceTextShaders)))
+		{
+			if (InLibrary)
+			{
+				Library = InLibrary;
+			}
+			else
+			{
+				METAL_GPUPROFILE(FAGXScopedCPUStats CPUStat(FString::Printf(TEXT("NewLibraryBinary: %d_%d"), SourceLen, SourceCRC)));
+
+				// Archived shaders should never get in here.
+				check(!(Header.CompileFlags & (1 << CFLAG_Archive)) || BufferSize > 0);
+
+				// allow GCD to copy the data into its own buffer
+				//		dispatch_data_t GCDBuffer = dispatch_data_create(InShaderCode.GetTypedData() + CodeOffset, ShaderCode.GetActualShaderCodeSize() - CodeOffset, nil, DISPATCH_DATA_DESTRUCTOR_DEFAULT);
+				ns::AutoReleasedError AError;
+				void* Buffer = FMemory::Malloc( BufferSize );
+				FMemory::Memcpy( Buffer, InShaderCode.GetData() + CodeOffset, BufferSize );
+				dispatch_data_t GCDBuffer = dispatch_data_create(Buffer, BufferSize, dispatch_get_main_queue(), ^(void) { FMemory::Free(Buffer); } );
+
+				// load up the already compiled shader
+				Library = GetAGXDeviceContext().GetDevice().NewLibrary(GCDBuffer, &AError);
+				dispatch_release(GCDBuffer);
+
+				if (Library == nil)
+				{
+					NSLog(@"Failed to create library: %@", ns::Error(AError).GetPtr());
+				}
+			}
+		}
+		else
+		{
+			METAL_GPUPROFILE(FAGXScopedCPUStats CPUStat(FString::Printf(TEXT("NewLibrarySource: %d_%d"), SourceLen, SourceCRC)));
+			NSString* ShaderString = ((OfflineCompiledFlag == 0) ? [NSString stringWithUTF8String:SourceCode] : GlslCodeNSString);
+
+			if(Header.ShaderName.Len())
+			{
+				ShaderString = [NSString stringWithFormat:@"// %@\n%@", Header.ShaderName.GetNSString(), ShaderString];
+			}
+
+			NSString* NewShaderString = [ShaderString stringByReplacingOccurrencesOfString:@"#pragma once" withString:@""];
+
+			mtlpp::CompileOptions CompileOptions;
+
+#if DEBUG_METAL_SHADERS
+			static bool bForceFastMath = FParse::Param(FCommandLine::Get(), TEXT("metalfastmath"));
+			static bool bForceNoFastMath = FParse::Param(FCommandLine::Get(), TEXT("metalnofastmath"));
+			if (bForceNoFastMath)
+			{
+				CompileOptions.SetFastMathEnabled(NO);
+			}
+			else if (bForceFastMath)
+			{
+				CompileOptions.SetFastMathEnabled(YES);
+			}
+			else
+#endif
+			{
+				CompileOptions.SetFastMathEnabled((BOOL)(!(Header.CompileFlags & (1 << CFLAG_NoFastMath))));
+			}
+
+#if !PLATFORM_MAC || DEBUG_METAL_SHADERS
+			NSMutableDictionary *PreprocessorMacros = [NSMutableDictionary new];
+#if !PLATFORM_MAC // Pretty sure that as_type-casts work on macOS, but they don't for half2<->uint on older versions of the iOS runtime compiler.
+			[PreprocessorMacros addEntriesFromDictionary: @{ @"METAL_RUNTIME_COMPILER" : @(1)}];
+#endif
+#if DEBUG_METAL_SHADERS
+			[PreprocessorMacros addEntriesFromDictionary: @{ @"MTLSL_ENABLE_DEBUG_INFO" : @(1)}];
+#endif
+			CompileOptions.SetPreprocessorMacros(PreprocessorMacros);
+#endif
+
+			mtlpp::LanguageVersion MetalVersion;
+			switch(Header.Version)
+			{
+				case 6:
+				case 5:
+				case 4:
+					MetalVersion = mtlpp::LanguageVersion::Version2_1;
+					break;
+				case 3:
+					MetalVersion = mtlpp::LanguageVersion::Version2_0;
+					break;
+				case 2:
+					MetalVersion = mtlpp::LanguageVersion::Version1_2;
+					break;
+				case 1:
+					MetalVersion = mtlpp::LanguageVersion::Version1_1;
+					break;
+				case 0:
+#if PLATFORM_MAC
+					MetalVersion = mtlpp::LanguageVersion::Version1_1;
+#else
+					MetalVersion = mtlpp::LanguageVersion::Version1_0;
+#endif
+					break;
+				default:
+					UE_LOG(LogRHI, Fatal, TEXT("Failed to create shader with unknown version %d: %s"), Header.Version, *FString(NewShaderString));
+#if PLATFORM_MAC
+					MetalVersion = mtlpp::LanguageVersion::Version1_1;
+#else
+					MetalVersion = mtlpp::LanguageVersion::Version1_0;
+#endif
+					break;
+				}
+			CompileOptions.SetLanguageVersion(MetalVersion);
+
+			ns::AutoReleasedError Error;
+			Library = GetAGXDeviceContext().GetDevice().NewLibrary(NewShaderString, CompileOptions, &Error);
+			if (Library == nil)
+			{
+				UE_LOG(LogRHI, Error, TEXT("*********** Error\n%s"), *FString(NewShaderString));
+				UE_LOG(LogRHI, Fatal, TEXT("Failed to create shader: %s"), *FString([Error.GetPtr() description]));
+			}
+			else if (Error != nil)
+			{
+				// Warning...
+				UE_LOG(LogRHI, Warning, TEXT("*********** Warning\n%s"), *FString(NewShaderString));
+				UE_LOG(LogRHI, Warning, TEXT("Created shader with warnings: %s"), *FString([Error.GetPtr() description]));
+			}
+
+			GlslCodeNSString = NewShaderString;
+			[GlslCodeNSString retain];
+		}
+
+		GetCompiledFunction(true);
+	}
+	UniformBuffersCopyInfo = Header.UniformBuffersCopyInfo;
+	SideTableBinding = Header.SideTable;
+
+	StaticSlots.Reserve(Bindings.ShaderResourceTable.ResourceTableLayoutHashes.Num());
+
+	for (uint32 LayoutHash : Bindings.ShaderResourceTable.ResourceTableLayoutHashes)
+	{
+		if (const FShaderParametersMetadata* Metadata = FindUniformBufferStructByLayoutHash(LayoutHash))
+		{
+			StaticSlots.Add(Metadata->GetLayout().StaticSlot);
+		}
+		else
+		{
+			StaticSlots.Add(MAX_UNIFORM_BUFFER_STATIC_SLOTS);
+		}
+	}
+}
+
+template<typename BaseResourceType, int32 ShaderType>
+void TAGXBaseShader<BaseResourceType, ShaderType>::Destroy()
+{
+	[GlslCodeNSString release];
+}
+
+template<typename BaseResourceType, int32 ShaderType>
+inline NSString* TAGXBaseShader<BaseResourceType, ShaderType>::GetSourceCode()
+{
+	if (!GlslCodeNSString && CodeSize && CompressedSource.Num())
+	{
+		GlslCodeNSString = AGXDecodeMetalSourceCode(CodeSize, CompressedSource);
+	}
+	if (!GlslCodeNSString)
+	{
+		GlslCodeNSString = [FString::Printf(TEXT("Hash: %s, Name: Main_%0.8x_%0.8x"), *BaseResourceType::GetHash().ToString(), SourceLen, SourceCRC).GetNSString() retain];
+	}
+	return GlslCodeNSString;
+}
+
+template<typename BaseResourceType, int32 ShaderType>
+uint32 TAGXBaseShader<BaseResourceType, ShaderType>::AddRef() const
+{
+	return FRHIResource::AddRef();
+}
+
+template<typename BaseResourceType, int32 ShaderType>
+uint32 TAGXBaseShader<BaseResourceType, ShaderType>::Release() const
+{
+	return FRHIResource::Release();
+}
+
+template<typename BaseResourceType, int32 ShaderType>
+uint32 TAGXBaseShader<BaseResourceType, ShaderType>::GetRefCount() const
+{
+	return FRHIResource::GetRefCount();
+}
+
+template<typename BaseResourceType, int32 ShaderType>
+mtlpp::Function TAGXBaseShader<BaseResourceType, ShaderType>::GetCompiledFunction(bool const bAsync)
+{
+	mtlpp::Function Func = Function;
+
+	if (!Func)
+	{
+		// Find the existing compiled shader in the cache.
+		uint32 FunctionConstantHash = ConstantValueHash;
+		FAGXCompiledShaderKey Key(SourceLen, SourceCRC, FunctionConstantHash);
+		Func = Function = GetAGXCompiledShaderCache().FindRef(Key);
+
+		if (!Func)
+		{
+			// Get the function from the library - the function name is "Main" followed by the CRC32 of the source MTLSL as 0-padded hex.
+			// This ensures that even if we move to a unified library that the function names will be unique - duplicates will only have one entry in the library.
+			NSString* Name = [NSString stringWithFormat:@"Main_%0.8x_%0.8x", SourceLen, SourceCRC];
+			mtlpp::FunctionConstantValues ConstantValues(nil);
+			if (bHasFunctionConstants)
+			{
+				ConstantValues = mtlpp::FunctionConstantValues();
+
+				if (bDeviceFunctionConstants)
+				{
+					// Index 33 is the device vendor id constant
+					ConstantValues.SetConstantValue(&GRHIVendorId, mtlpp::DataType::UInt, @"GAGXDeviceManufacturer");
+				}
+			}
+
+			if (!bHasFunctionConstants || !bAsync)
+			{
+				METAL_GPUPROFILE(FAGXScopedCPUStats CPUStat(FString::Printf(TEXT("NewFunction: %s"), *FString(Name))));
+				if (!bHasFunctionConstants)
+				{
+					Function = Library.NewFunction(Name);
+				}
+				else
+				{
+					ns::AutoReleasedError AError;
+					Function = Library.NewFunction(Name, ConstantValues, &AError);
+					ns::Error Error = AError;
+					UE_CLOG(Function == nil, LogAGX, Error, TEXT("Failed to create function: %s"), *FString(Error.GetPtr().description));
+					UE_CLOG(Function == nil, LogAGX, Fatal, TEXT("*********** Error\n%s"), *FString(GetSourceCode()));
+				}
+
+				check(Function);
+				GetAGXCompiledShaderCache().Add(Key, Library, Function);
+
+				Func = Function;
+			}
+			else
+			{
+				METAL_GPUPROFILE(FAGXScopedCPUStats CPUStat(FString::Printf(TEXT("NewFunctionAsync: %s"), *FString(Name))));
+				METAL_GPUPROFILE(uint64 CPUStart = CPUStat.Stats ? CPUStat.Stats->CPUStartTime : 0);
+#if ENABLE_METAL_GPUPROFILE
+				ns::String nsName(Name);
+				Library.NewFunction(Name, ConstantValues, [Key, this, CPUStart, nsName](mtlpp::Function const& NewFunction, ns::Error const& Error){
+#else
+				Library.NewFunction(Name, ConstantValues, [Key, this](mtlpp::Function const& NewFunction, ns::Error const& Error){
+#endif
+					METAL_GPUPROFILE(FAGXScopedCPUStats CompletionStat(FString::Printf(TEXT("NewFunctionCompletion: %s"), *FString(nsName.GetPtr()))));
+					UE_CLOG(NewFunction == nil, LogAGX, Error, TEXT("Failed to create function: %s"), *FString(Error.GetPtr().description));
+					UE_CLOG(NewFunction == nil, LogAGX, Fatal, TEXT("*********** Error\n%s"), *FString(GetSourceCode()));
+
+					GetAGXCompiledShaderCache().Add(Key, Library, NewFunction);
+#if ENABLE_METAL_GPUPROFILE
+					if (CompletionStat.Stats)
+					{
+						CompletionStat.Stats->CPUStartTime = CPUStart;
+					}
+#endif
+				});
+
+				return nil;
+			}
+		}
+	}
+
+	if (FAGXCommandQueue::SupportsFeature(EAGXFeaturesIABs) && Bindings.ArgumentBuffers && ArgumentEncoders.Num() == 0)
+	{
+		uint32 ArgumentBuffers = Bindings.ArgumentBuffers;
+		while(ArgumentBuffers)
+		{
+			uint32 Index = __builtin_ctz(ArgumentBuffers);
+			ArgumentBuffers &= ~(1 << Index);
+
+			mtlpp::ArgumentEncoder ArgumentEncoder = Function.NewArgumentEncoderWithBufferIndex(Index);
+			ArgumentEncoders.Add(Index, ArgumentEncoder);
+
+			TBitArray<> Resources;
+			for (uint8 Id : Bindings.ArgumentBufferMasks[Index])
+			{
+				if (Id >= Resources.Num())
+				{
+					Resources.Add(false, (Id + 1) - Resources.Num());
+				}
+				Resources[Id] = true;
+			}
+			ArgumentBitmasks.Add(Index, Resources);
+		}
+	}
+
+	check(Func);
+	return Func;
+}
diff --git a/Engine/Source/Runtime/Apple/AGXRHI/Public/AGXFrameAllocator.h b/Engine/Source/Runtime/Apple/AGXRHI/Public/AGXFrameAllocator.h
new file mode 100644
--- /dev/null
+++ b/Engine/Source/Runtime/Apple/AGXRHI/Public/AGXFrameAllocator.h
@@ -0,0 +1,88 @@
+// Copyright Epic Games, Inc. All Rights Reserved.
+
+#pragma once
+
+#import <Metal/Metal.h>
+#import <os/lock.h>
+
+#include "AGXRHI.h"
+
+#define METAL_FRAME_ALLOCATOR_VALIDATION !(UE_BUILD_SHIPPING || UE_BUILD_TEST)
+
+class FAGXFrameAllocator
+{
+public:
+
+    struct AllocationEntry
+    {
+        id <MTLBuffer> Backing;
+        uint32 Offset;
+    };
+    
+private:
+    // The device. This object will take a reference.
+    // You must make sure your rendering is on this same device or things will explode horribly.
+    id <MTLDevice> Device;
+#if METAL_FRAME_ALLOCATOR_VALIDATION
+    // Every buffer that is currently allocated.
+    NSMutableArray<id <MTLBuffer>>* AllAllocatedBuffers;
+#endif
+    // Buffers that are available. These buffers are not currently being read from nor written to.
+    NSMutableArray<id <MTLBuffer>>* AvailableBuffers;
+    // The current set of buffers in use for this frame.
+    // Must ONLY be mutated on RHI thread if present.
+    NSMutableArray<id <MTLBuffer>>* BuffersInFlight;
+    // The current buffer we are suballocating from and the current offset.
+    id <MTLBuffer> CurrentBuffer;
+    uint32 CurrentCursor;
+    uint32 CurrentFrame;
+    // Lock around AllAllocatedBuffers and AvailableBuffers since they are mutated from a completion handler.
+    os_unfair_lock UnfairLock;
+    // Allocation info and target limits.
+    uint32 TargetAllocationLimit;
+    uint32 DefaultAllocationSize;
+    // Stats!
+    uint64 TotalBytesAllocated;
+    uint64 BytesAcquiredInCurrentFrame;
+    uint64 BytesInFlight;
+    
+    TStatId TotalAllocationStat;
+    TStatId MemoryInFlightStat;
+    TStatId BytesPerFrameStat;
+public:
+    // API
+    
+    // This object will hold a reference to TargetDevice for its lifetime.
+    FAGXFrameAllocator(id <MTLDevice> TargetDevice);
+    ~FAGXFrameAllocator();
+    
+    // Sets a target allocation limit.
+    // This is not a hard high-water mark.
+    // If a frame needs to use more memory it will do so.
+    // When the total bytes allocated goes over this amount we will start removing buffers from the pool.
+    void SetTargetAllocationLimitInBytes(uint32 LimitInBytes);
+    
+    // The size of each buffer.
+    // The actual size is computed by max(SizeNeeded, DefaultSize)
+    void SetDefaultAllocationSizeInBytes(uint32 DefaultInBytes);
+    
+    void SetStatIds(const TStatId& TotalAllocationStat, const TStatId& MemoryInFlightStat, const TStatId& BytesPerFrameStat);
+    
+    // Acquires SizeInBytes bytes for uniform data.
+    // Returns the currently active buffer and offset.
+    // WARNING: This is NOT completely thread safe.
+    // WARNING: You MUST only call from the RHI thread if present. You MAY call it from the parallel translate threads.
+    AllocationEntry AcquireSpace(uint32 SizeInBytes);
+    // Marks the end of a frame.
+    // When LastCommandBufferInFrame is retired, the BuffersInFlight
+    // will be returned to the AvailableBuffers pool.
+    // When this call returns CurrentBuffer and CurrentCursor will be reset to 0.
+    // This is NOT thread-safe. It must be called from the RHI thread if present.
+    void MarkEndOfFrame(uint32 FrameNumberThatEnded, id <MTLCommandBuffer> LastCommandBufferInFrame);
+    
+private:
+    // This is NOT thread safe. You must lock around it.
+    id <MTLBuffer> FindOrAllocateBufferUnsafe(uint32 SizeInBytes);
+    void RecycleBuffers(NSArray <id <MTLBuffer>>* BuffersToRecycle);
+};
+
diff --git a/Engine/Source/Runtime/Apple/AGXRHI/Public/AGXRHI.h b/Engine/Source/Runtime/Apple/AGXRHI/Public/AGXRHI.h
new file mode 100644
--- /dev/null
+++ b/Engine/Source/Runtime/Apple/AGXRHI/Public/AGXRHI.h
@@ -0,0 +1,9 @@
+// Copyright Epic Games, Inc. All Rights Reserved.
+
+/*=============================================================================
+	AGXRHI.h: Public AGX RHI definitions.
+=============================================================================*/
+
+#pragma once 
+
+// TODO: Put public AGX RHI definitions here.
diff --git a/Engine/Source/Runtime/Apple/AGXRHI/Public/AGXRHIContext.h b/Engine/Source/Runtime/Apple/AGXRHI/Public/AGXRHIContext.h
new file mode 100644
--- /dev/null
+++ b/Engine/Source/Runtime/Apple/AGXRHI/Public/AGXRHIContext.h
@@ -0,0 +1,283 @@
+// Copyright Epic Games, Inc. All Rights Reserved.
+
+#pragma once
+
+#include <Metal/Metal.h>
+#include "AGXState.h"
+#include "AGXResources.h"
+#include "AGXViewport.h"
+#include "RHICoreShader.h"
+
+#define UE_METAL_RHI_SUPPORT_CLEAR_UAV_WITH_BLIT_ENCODER	1
+
+class FAGXContext;
+struct FAGXCommandBufferFence;
+
+#if UE_METAL_RHI_SUPPORT_CLEAR_UAV_WITH_BLIT_ENCODER
+enum class EAGXRHIClearUAVType : uint8
+{
+	VertexBuffer,
+	StructuredBuffer
+};
+#endif // UE_METAL_RHI_SUPPORT_CLEAR_UAV_WITH_BLIT_ENCODER
+
+/** The interface RHI command context. */
+class FAGXRHICommandContext : public IRHICommandContext
+{
+public:
+	FAGXRHICommandContext(class FAGXProfiler* InProfiler, FAGXContext* WrapContext);
+	virtual ~FAGXRHICommandContext();
+
+	/** Get the internal context */
+	FORCEINLINE FAGXContext& GetInternalContext() const { return *Context; }
+	
+	/** Get the profiler pointer */
+	FORCEINLINE class FAGXProfiler* GetProfiler() const { return Profiler; }
+
+	/**
+	 *Sets the current compute shader.  Mostly for compliance with platforms
+	 *that require shader setting before resource binding.
+	 */
+	virtual void RHISetComputeShader(FRHIComputeShader* ComputeShader) override;
+	
+	virtual void RHISetComputePipelineState(FRHIComputePipelineState* ComputePipelineState) override;
+	
+	virtual void RHIDispatchComputeShader(uint32 ThreadGroupCountX, uint32 ThreadGroupCountY, uint32 ThreadGroupCountZ) final override;
+	
+	virtual void RHIDispatchIndirectComputeShader(FRHIBuffer* ArgumentBuffer, uint32 ArgumentOffset) final override;
+	
+	// Useful when used with geometry shader (emit polygons to different viewports), otherwise SetViewPort() is simpler
+	// @param Count >0
+	// @param Data must not be 0
+	virtual void RHISetMultipleViewports(uint32 Count, const FViewportBounds* Data) final override;
+	
+	/** Clears a UAV to the multi-component value provided. */
+	virtual void RHIClearUAVFloat(FRHIUnorderedAccessView* UnorderedAccessViewRHI, const FVector4& Values) final override;
+	virtual void RHIClearUAVUint(FRHIUnorderedAccessView* UnorderedAccessViewRHI, const FUintVector4& Values) final override;
+	
+	virtual void RHICopyTexture(FRHITexture* SourceTextureRHI, FRHITexture* DestTextureRHI, const FRHICopyTextureInfo& CopyInfo) final override;
+	
+	virtual void RHICopyBufferRegion(FRHIBuffer* DstBufferRHI, uint64 DstOffset, FRHIBuffer* SrcBufferRHI, uint64 SrcOffset, uint64 NumBytes) final override;
+
+	/**
+	 * Resolves from one texture to another.
+	 * @param SourceTexture - texture to resolve from, 0 is silenty ignored
+	 * @param DestTexture - texture to resolve to, 0 is silenty ignored
+	 * @param ResolveParams - optional resolve params
+	 */
+	virtual void RHICopyToResolveTarget(FRHITexture* SourceTexture, FRHITexture* DestTexture, const FResolveParams& ResolveParams) final override;
+	
+	virtual void RHIBeginRenderQuery(FRHIRenderQuery* RenderQuery) final override;
+	
+	virtual void RHIEndRenderQuery(FRHIRenderQuery* RenderQuery) final override;
+	
+	void RHIBeginOcclusionQueryBatch(uint32 NumQueriesInBatch);
+	
+	void RHIEndOcclusionQueryBatch();
+	
+	virtual void RHISubmitCommandsHint() override;
+
+	virtual void RHIDiscardRenderTargets(bool Depth, bool Stencil, uint32 ColorBitMask) final override;
+	
+	// This method is queued with an RHIThread, otherwise it will flush after it is queued; without an RHI thread there is no benefit to queuing this frame advance commands
+	virtual void RHIBeginDrawingViewport(FRHIViewport* Viewport, FRHITexture* RenderTargetRHI) override;
+	
+	// This method is queued with an RHIThread, otherwise it will flush after it is queued; without an RHI thread there is no benefit to queuing this frame advance commands
+	virtual void RHIEndDrawingViewport(FRHIViewport* Viewport, bool bPresent, bool bLockToVsync) override;
+	
+	// This method is queued with an RHIThread, otherwise it will flush after it is queued; without an RHI thread there is no benefit to queuing this frame advance commands
+	virtual void RHIBeginFrame() override;
+	
+	// This method is queued with an RHIThread, otherwise it will flush after it is queued; without an RHI thread there is no benefit to queuing this frame advance commands
+	virtual void RHIEndFrame() override;
+	
+	/**
+		* Signals the beginning of scene rendering. The RHI makes certain caching assumptions between
+		* calls to BeginScene/EndScene. Currently the only restriction is that you can't update texture
+		* references.
+		*/
+	// This method is queued with an RHIThread, otherwise it will flush after it is queued; without an RHI thread there is no benefit to queuing this frame advance commands
+	virtual void RHIBeginScene() override;
+	
+	/**
+		* Signals the end of scene rendering. See RHIBeginScene.
+		*/
+	// This method is queued with an RHIThread, otherwise it will flush after it is queued; without an RHI thread there is no benefit to queuing this frame advance commands
+	virtual void RHIEndScene() override;
+	
+	virtual void RHISetStreamSource(uint32 StreamIndex, FRHIBuffer* VertexBuffer, uint32 Offset) final override;
+
+	// @param MinX including like Win32 RECT
+	// @param MinY including like Win32 RECT
+	// @param MaxX excluding like Win32 RECT
+	// @param MaxY excluding like Win32 RECT
+	virtual void RHISetViewport(float MinX, float MinY, float MinZ, float MaxX, float MaxY, float MaxZ) final override;
+
+	virtual void RHISetStereoViewport(float LeftMinX, float RightMinX, float LeftMinY, float RightMinY, float MinZ, float LeftMaxX, float RightMaxX, float LeftMaxY, float RightMaxY, float MaxZ) final override;
+	
+	// @param MinX including like Win32 RECT
+	// @param MinY including like Win32 RECT
+	// @param MaxX excluding like Win32 RECT
+	// @param MaxY excluding like Win32 RECT
+	virtual void RHISetScissorRect(bool bEnable, uint32 MinX, uint32 MinY, uint32 MaxX, uint32 MaxY) final override;
+	
+	virtual void RHISetGraphicsPipelineState(FRHIGraphicsPipelineState* GraphicsState, bool bApplyAdditionalState) final override;
+	
+	virtual void RHISetStaticUniformBuffers(const FUniformBufferStaticBindings& InUniformBuffers) final override;
+
+	/** Set the shader resource view of a surface. */
+	virtual void RHISetShaderTexture(FRHIGraphicsShader* Shader, uint32 TextureIndex, FRHITexture* NewTexture) final override;
+	
+	/** Set the shader resource view of a surface.  This is used for binding TextureMS parameter types that need a multi sampled view. */
+	virtual void RHISetShaderTexture(FRHIComputeShader* PixelShader, uint32 TextureIndex, FRHITexture* NewTexture) final override;
+	
+	/**
+	 * Sets sampler state.
+	 * @param ComputeShader		The compute shader to set the sampler for.
+	 * @param SamplerIndex		The index of the sampler.
+	 * @param NewState			The new sampler state.
+	 */
+	virtual void RHISetShaderSampler(FRHIComputeShader* ComputeShader, uint32 SamplerIndex, FRHISamplerState* NewState) final override;
+	
+	/**
+	 * Sets sampler state.
+	 * @param Shader			The shader to set the sampler for.
+	 * @param SamplerIndex		The index of the sampler.
+	 * @param NewState			The new sampler state.
+	 */
+	virtual void RHISetShaderSampler(FRHIGraphicsShader* Shader, uint32 SamplerIndex, FRHISamplerState* NewState) final override;
+	
+	/** Sets a pixel shader UAV parameter. */
+	virtual void RHISetUAVParameter(FRHIPixelShader* PixelShaderRHI, uint32 UAVIndex, FRHIUnorderedAccessView* UAVRHI) final override;
+	
+	/** Sets a compute shader UAV parameter. */
+	virtual void RHISetUAVParameter(FRHIComputeShader* ComputeShader, uint32 UAVIndex, FRHIUnorderedAccessView* UAV) final override;
+	
+	/** Sets a compute shader UAV parameter and initial count */
+	virtual void RHISetUAVParameter(FRHIComputeShader* ComputeShader, uint32 UAVIndex, FRHIUnorderedAccessView* UAV, uint32 InitialCount) final override;
+	
+	virtual void RHISetShaderResourceViewParameter(FRHIGraphicsShader* Shader, uint32 SamplerIndex, FRHIShaderResourceView* SRV) final override;
+	
+	virtual void RHISetShaderResourceViewParameter(FRHIComputeShader* ComputeShader, uint32 SamplerIndex, FRHIShaderResourceView* SRV) final override;
+	
+	virtual void RHISetShaderUniformBuffer(FRHIGraphicsShader* Shader, uint32 BufferIndex, FRHIUniformBuffer* Buffer) final override;
+	
+	virtual void RHISetShaderUniformBuffer(FRHIComputeShader* ComputeShader, uint32 BufferIndex, FRHIUniformBuffer* Buffer) final override;
+	
+	virtual void RHISetShaderParameter(FRHIGraphicsShader* Shader, uint32 BufferIndex, uint32 BaseIndex, uint32 NumBytes, const void* NewValue) final override;
+	
+	virtual void RHISetShaderParameter(FRHIComputeShader* ComputeShader, uint32 BufferIndex, uint32 BaseIndex, uint32 NumBytes, const void* NewValue) final override;
+	
+	virtual void RHISetStencilRef(uint32 StencilRef) final override;
+
+	virtual void RHISetBlendFactor(const FLinearColor& BlendFactor) final override;
+
+	void SetRenderTargets(uint32 NumSimultaneousRenderTargets, const FRHIRenderTargetView* NewRenderTargets, const FRHIDepthRenderTargetView* NewDepthStencilTarget);
+	
+	void SetRenderTargetsAndClear(const FRHISetRenderTargetsInfo& RenderTargetsInfo);
+	
+	virtual void RHIDrawPrimitive(uint32 BaseVertexIndex, uint32 NumPrimitives, uint32 NumInstances) final override;
+	
+	virtual void RHIDrawPrimitiveIndirect(FRHIBuffer* ArgumentBuffer, uint32 ArgumentOffset) final override;
+	
+	virtual void RHIDrawIndexedIndirect(FRHIBuffer* IndexBufferRHI, FRHIBuffer* ArgumentsBufferRHI, int32 DrawArgumentsIndex, uint32 NumInstances) final override;
+	
+	// @param NumPrimitives need to be >0
+	virtual void RHIDrawIndexedPrimitive(FRHIBuffer* IndexBuffer, int32 BaseVertexIndex, uint32 FirstInstance, uint32 NumVertices, uint32 StartIndex, uint32 NumPrimitives, uint32 NumInstances) final override;
+	
+	virtual void RHIDrawIndexedPrimitiveIndirect(FRHIBuffer* IndexBuffer, FRHIBuffer* ArgumentBuffer, uint32 ArgumentOffset) final override;
+
+	/**
+	* Sets Depth Bounds Testing with the given min/max depth.
+	* @param MinDepth	The minimum depth for depth bounds test
+	* @param MaxDepth	The maximum depth for depth bounds test.
+	*					The valid values for fMinDepth and fMaxDepth are such that 0 <= fMinDepth <= fMaxDepth <= 1
+	*/
+	virtual void RHISetDepthBounds(float MinDepth, float MaxDepth) final override;
+
+	virtual void RHIPushEvent(const TCHAR* Name, FColor Color) final override;
+	
+	virtual void RHIPopEvent() final override;
+	
+	virtual void RHICopyToStagingBuffer(FRHIBuffer* SourceBufferRHI, FRHIStagingBuffer* DestinationStagingBufferRHI, uint32 Offset, uint32 NumBytes) final override;
+	virtual void RHIWriteGPUFence(FRHIGPUFence* FenceRHI) final override;
+
+	virtual void RHIBeginTransitions(TArrayView<const FRHITransition*> Transitions);
+	virtual void RHIEndTransitions(TArrayView<const FRHITransition*> Transitions);
+
+	virtual void RHIBeginRenderPass(const FRHIRenderPassInfo& InInfo, const TCHAR* InName) final override;
+
+	virtual void RHIEndRenderPass() final override;
+	
+	virtual void RHINextSubpass() final override;
+
+protected:
+	static TGlobalResource<TBoundShaderStateHistory<10000>> BoundShaderStateHistory;
+	
+	/** Context implementation details. */
+	FAGXContext* Context;
+	
+	/** Occlusion query batch fence */
+	TSharedPtr<FAGXCommandBufferFence, ESPMode::ThreadSafe> CommandBufferFence;
+	
+	/** Profiling implementation details. */
+	class FAGXProfiler* Profiler;
+	
+	/** Some local variables to track the pending primitive information used in RHIEnd*UP functions */
+	FAGXBuffer PendingVertexBuffer;
+	uint32 PendingVertexDataStride;
+	
+	FAGXBuffer PendingIndexBuffer;
+	uint32 PendingIndexDataStride;
+	
+	uint32 PendingPrimitiveType;
+	uint32 PendingNumPrimitives;
+
+	template <typename TRHIShader>
+	void ApplyStaticUniformBuffers(TRHIShader* Shader)
+	{
+		if (Shader)
+		{
+			UE::RHICore::ApplyStaticUniformBuffers(this, Shader, Shader->StaticSlots, Shader->Bindings.ShaderResourceTable.ResourceTableLayoutHashes, GlobalUniformBuffers);
+		}
+	}
+
+	TArray<FRHIUniformBuffer*> GlobalUniformBuffers;
+
+private:
+#if UE_METAL_RHI_SUPPORT_CLEAR_UAV_WITH_BLIT_ENCODER
+	void ClearUAVWithBlitEncoder(FRHIUnorderedAccessView* UnorderedAccessViewRHI, EAGXRHIClearUAVType Type, uint32 Pattern);
+#endif // UE_METAL_RHI_SUPPORT_CLEAR_UAV_WITH_BLIT_ENCODER
+	void ClearUAV(TRHICommandList_RecursiveHazardous<FAGXRHICommandContext>& RHICmdList, FAGXUnorderedAccessView* UnorderedAccessView, const void* ClearValue, bool bFloat);
+
+	void RHIClearMRT(bool bClearColor, int32 NumClearColors, const FLinearColor* ColorArray, bool bClearDepth, float Depth, bool bClearStencil, uint32 Stencil);
+};
+
+class FAGXRHIComputeContext : public FAGXRHICommandContext
+{
+public:
+	FAGXRHIComputeContext(class FAGXProfiler* InProfiler, FAGXContext* WrapContext);
+	virtual ~FAGXRHIComputeContext();
+	
+	virtual void RHISetAsyncComputeBudget(EAsyncComputeBudget Budget) final override;
+	virtual void RHISetComputeShader(FRHIComputeShader* ComputeShader) final override;
+	virtual void RHISetComputePipelineState(FRHIComputePipelineState* ComputePipelineState) final override;
+	virtual void RHISubmitCommandsHint() final override;
+};
+
+class FAGXRHIImmediateCommandContext : public FAGXRHICommandContext
+{
+public:
+	FAGXRHIImmediateCommandContext(class FAGXProfiler* InProfiler, FAGXContext* WrapContext);
+
+	// FRHICommandContext API accessible only on the immediate device context
+	virtual void RHIBeginDrawingViewport(FRHIViewport* Viewport, FRHITexture* RenderTargetRHI) final override;
+	virtual void RHIEndDrawingViewport(FRHIViewport* Viewport, bool bPresent, bool bLockToVsync) final override;
+	virtual void RHIBeginFrame() final override;
+	virtual void RHIEndFrame() final override;
+	virtual void RHIBeginScene() final override;
+	virtual void RHIEndScene() final override;
+	
+protected:
+	friend class FAGXDynamicRHI;
+};
diff --git a/Engine/Source/Runtime/Apple/AGXRHI/Public/AGXResources.h b/Engine/Source/Runtime/Apple/AGXRHI/Public/AGXResources.h
new file mode 100644
--- /dev/null
+++ b/Engine/Source/Runtime/Apple/AGXRHI/Public/AGXResources.h
@@ -0,0 +1,795 @@
+// Copyright Epic Games, Inc. All Rights Reserved.
+
+/*=============================================================================
+	AGXResources.h: AGX resource RHI definitions..
+=============================================================================*/
+
+#pragma once
+
+#include "BoundShaderStateCache.h"
+#include "MetalShaderResources.h"
+#include "ShaderCodeArchive.h"
+
+THIRD_PARTY_INCLUDES_START
+#include "mtlpp.hpp"
+THIRD_PARTY_INCLUDES_END
+
+/** Parallel execution is available on Mac but not iOS for the moment - it needs to be tested because it isn't cost-free */
+#define METAL_SUPPORTS_PARALLEL_RHI_EXECUTE 1
+
+class FAGXContext;
+@class FAGXShaderPipeline;
+
+extern NSString* AGXDecodeMetalSourceCode(uint32 CodeSize, TArray<uint8> const& CompressedSource);
+
+struct FAGXRenderPipelineHash
+{
+	friend uint32 GetTypeHash(FAGXRenderPipelineHash const& Hash)
+	{
+		return HashCombine(GetTypeHash(Hash.RasterBits), GetTypeHash(Hash.TargetBits));
+	}
+	
+	friend bool operator==(FAGXRenderPipelineHash const& Left, FAGXRenderPipelineHash const& Right)
+	{
+		return Left.RasterBits == Right.RasterBits && Left.TargetBits == Right.TargetBits;
+	}
+	
+	uint64 RasterBits;
+	uint64 TargetBits;
+};
+
+class FAGXSubBufferHeap;
+class FAGXSubBufferLinear;
+class FAGXSubBufferMagazine;
+
+class FAGXBuffer : public mtlpp::Buffer
+{
+public:
+	FAGXBuffer(ns::Ownership retain = ns::Ownership::Retain) : mtlpp::Buffer(retain), Heap(nullptr), Linear(nullptr), Magazine(nullptr), bPooled(false) { }
+	FAGXBuffer(ns::Protocol<id<MTLBuffer>>::type handle, ns::Ownership retain = ns::Ownership::Retain);
+	
+	FAGXBuffer(mtlpp::Buffer&& rhs, FAGXSubBufferHeap* heap);
+	FAGXBuffer(mtlpp::Buffer&& rhs, FAGXSubBufferLinear* heap);
+	FAGXBuffer(mtlpp::Buffer&& rhs, FAGXSubBufferMagazine* magazine);
+	FAGXBuffer(mtlpp::Buffer&& rhs, bool bInPooled);
+	
+	FAGXBuffer(const FAGXBuffer& rhs);
+	FAGXBuffer(FAGXBuffer&& rhs);
+	virtual ~FAGXBuffer();
+	
+	FAGXBuffer& operator=(const FAGXBuffer& rhs);
+	FAGXBuffer& operator=(FAGXBuffer&& rhs);
+	
+	inline bool operator==(FAGXBuffer const& rhs) const
+	{
+		return mtlpp::Buffer::operator==(rhs);
+	}
+	
+	inline bool IsPooled() const { return bPooled; }
+	inline bool IsSingleUse() const { return bSingleUse; }
+	inline void MarkSingleUse() { bSingleUse = true; }
+    void SetOwner(class FAGXRHIBuffer* Owner, bool bIsSwap);
+	void Release();
+	
+	friend uint32 GetTypeHash(FAGXBuffer const& Hash)
+	{
+		return HashCombine(GetTypeHash(Hash.GetPtr()), GetTypeHash((uint64)Hash.GetOffset()));
+	}
+	
+private:
+	FAGXSubBufferHeap* Heap;
+	FAGXSubBufferLinear* Linear;
+	FAGXSubBufferMagazine* Magazine;
+	bool bPooled;
+	bool bSingleUse;
+};
+
+class FAGXTexture : public mtlpp::Texture
+{
+public:
+	FAGXTexture(ns::Ownership retain = ns::Ownership::Retain) : mtlpp::Texture(retain) { }
+	FAGXTexture(ns::Protocol<id<MTLTexture>>::type handle, ns::Ownership retain = ns::Ownership::Retain)
+	: mtlpp::Texture(handle, nullptr, retain) {}
+	
+	FAGXTexture(mtlpp::Texture&& rhs)
+	: mtlpp::Texture((mtlpp::Texture&&)rhs)
+	{
+		
+	}
+	
+	FAGXTexture(const FAGXTexture& rhs)
+	: mtlpp::Texture(rhs)
+	{
+		
+	}
+	
+	FAGXTexture(FAGXTexture&& rhs)
+	: mtlpp::Texture((mtlpp::Texture&&)rhs)
+	{
+		
+	}
+	
+	FAGXTexture& operator=(const FAGXTexture& rhs)
+	{
+		if (this != &rhs)
+		{
+			mtlpp::Texture::operator=(rhs);
+		}
+		return *this;
+	}
+	
+	FAGXTexture& operator=(FAGXTexture&& rhs)
+	{
+		mtlpp::Texture::operator=((mtlpp::Texture&&)rhs);
+		return *this;
+	}
+	
+	inline bool operator==(FAGXTexture const& rhs) const
+	{
+		return mtlpp::Texture::operator==(rhs);
+	}
+	
+	friend uint32 GetTypeHash(FAGXTexture const& Hash)
+	{
+		return GetTypeHash(Hash.GetPtr());
+	}
+};
+
+/** Texture/RT wrapper. */
+class AGXRHI_API FAGXSurface
+{
+public:
+
+	/** 
+	 * Constructor that will create Texture and Color/DepthBuffers as needed
+	 */
+	FAGXSurface(ERHIResourceType ResourceType, EPixelFormat Format, uint32 SizeX, uint32 SizeY, uint32 SizeZ, uint32 NumSamples, bool bArray, uint32 ArraySize, uint32 NumMips, ETextureCreateFlags Flags, FResourceBulkDataInterface* BulkData);
+
+	FAGXSurface(FAGXSurface& Source, NSRange MipRange);
+	
+	FAGXSurface(FAGXSurface& Source, NSRange MipRange, EPixelFormat Format, bool bSRGBForceDisable);
+	
+	/**
+	 * Destructor
+	 */
+	~FAGXSurface();
+
+	/** Prepare for texture-view support - need only call this once on the source texture which is to be viewed. */
+	void PrepareTextureView();
+	
+	/** @returns A newly allocated buffer object large enough for the surface within the texture specified. */
+	id <MTLBuffer> AllocSurface(uint32 MipIndex, uint32 ArrayIndex, EResourceLockMode LockMode, uint32& DestStride, bool SingleLayer = false);
+
+	/** Apply the data in Buffer to the surface specified.
+	 * Will also handle destroying SourceBuffer appropriately.
+	 */
+	void UpdateSurfaceAndDestroySourceBuffer(id <MTLBuffer> SourceBuffer, uint32 MipIndex, uint32 ArrayIndex);
+	
+	/**
+	 * Locks one of the texture's mip-maps.
+	 * @param ArrayIndex Index of the texture array/face in the form Index*6+Face
+	 * @return A pointer to the specified texture data.
+	 */
+	void* Lock(uint32 MipIndex, uint32 ArrayIndex, EResourceLockMode LockMode, uint32& DestStride, bool SingleLayer = false);
+	
+	/** Unlocks a previously locked mip-map.
+	 * @param ArrayIndex Index of the texture array/face in the form Index*6+Face
+	 */
+	void Unlock(uint32 MipIndex, uint32 ArrayIndex, bool bTryAsync);
+	
+	/**
+	 * Locks one of the texture's mip-maps.
+	 * @param ArrayIndex Index of the texture array/face in the form Index*6+Face
+	 * @return A pointer to the specified texture data.
+	 */
+	void* AsyncLock(class FRHICommandListImmediate& RHICmdList, uint32 MipIndex, uint32 ArrayIndex, EResourceLockMode LockMode, uint32& DestStride, bool bNeedsDefaultRHIFlush);
+	
+	/** Unlocks a previously locked mip-map.
+	 * @param ArrayIndex Index of the texture array/face in the form Index*6+Face
+	 */
+	void AsyncUnlock(id <MTLBuffer> SourceData, uint32 MipIndex, uint32 ArrayIndex);
+
+	/**
+	 * Returns how much memory a single mip uses, and optionally returns the stride
+	 */
+	uint32 GetMipSize(uint32 MipIndex, uint32* Stride, bool bSingleLayer);
+
+	/**
+	 * Returns how much memory is used by the surface
+	 */
+	uint32 GetMemorySize();
+
+	/** Returns the number of faces for the texture */
+	uint32 GetNumFaces();
+	
+	/** Gets the drawable texture if this is a back-buffer surface. */
+	FAGXTexture GetDrawableTexture();
+	ns::AutoReleased<FAGXTexture> GetCurrentTexture();
+
+	FAGXTexture Reallocate(FAGXTexture Texture, mtlpp::TextureUsage UsageModifier);
+	void ReplaceTexture(FAGXContext& Context, FAGXTexture OldTexture, FAGXTexture NewTexture);
+	void MakeAliasable(void);
+	void MakeUnAliasable(void);
+	
+	ERHIResourceType Type;
+	EPixelFormat PixelFormat;
+	uint8 FormatKey;
+	//texture used for store actions and binding to shader params
+	FAGXTexture Texture;
+	//if surface is MSAA, texture used to bind for RT
+	FAGXTexture MSAATexture;
+
+	//texture used for a resolve target.  Same as texture on iOS.  
+	//Dummy target on Mac where RHISupportsSeparateMSAAAndResolveTextures is true.	In this case we don't always want a resolve texture but we
+	//have to have one until renderpasses are implemented at a high level.
+	// Mac / RHISupportsSeparateMSAAAndResolveTextures == true
+	// iOS A9+ where depth resolve is available
+	// iOS < A9 where depth resolve is unavailable.
+	FAGXTexture MSAAResolveTexture;
+	uint32 SizeX, SizeY, SizeZ;
+	bool bIsCubemap;
+	int16 volatile Written;
+	int16 GPUReadback = 0;
+	enum EAGXGPUReadbackFlags : int16
+	{
+		ReadbackRequestedShift 			= 0,
+		ReadbackFenceCompleteShift  	= 1,
+		ReadbackRequested 				= 1 << ReadbackRequestedShift,
+		ReadbackFenceComplete 			= 1 << ReadbackFenceCompleteShift,
+		ReadbackRequestedAndComplete 	= ReadbackRequested | ReadbackFenceComplete
+	};
+	
+	ETextureCreateFlags Flags;
+
+	uint32 BufferLocks;
+
+	// how much memory is allocated for this texture
+	uint64 TotalTextureSize;
+	
+	// For back-buffers, the owning viewport.
+	class FAGXViewport* Viewport;
+	
+	TSet<class FAGXShaderResourceView*> SRVs;
+
+private:
+	void Init(FAGXSurface& Source, NSRange MipRange);
+	
+	void Init(FAGXSurface& Source, NSRange MipRange, EPixelFormat Format, bool bSRGBForceDisable);
+	
+private:
+	// The movie playback IOSurface/CVTexture wrapper to avoid page-off
+	CFTypeRef ImageSurfaceRef;
+	
+	// Texture view surfaces don't own their resources, only reference
+	bool bTextureView;
+	
+	// Count of outstanding async. texture uploads
+	static volatile int64 ActiveUploads;
+};
+
+class FAGXTexture2D : public FRHITexture2D
+{
+public:
+	/** The surface info */
+	FAGXSurface Surface;
+
+	// Constructor, just calls base and Surface constructor
+	FAGXTexture2D(EPixelFormat Format, uint32 SizeX, uint32 SizeY, uint32 NumMips, uint32 NumSamples, ETextureCreateFlags Flags, FResourceBulkDataInterface* BulkData, const FClearValueBinding& InClearValue)
+		: FRHITexture2D(SizeX, SizeY, NumMips, NumSamples, Format, Flags, InClearValue)
+		, Surface(RRT_Texture2D, Format, SizeX, SizeY, 1, NumSamples, /*bArray=*/ false, 1, NumMips, Flags, BulkData)
+	{
+	}
+	
+	virtual ~FAGXTexture2D()
+	{
+	}
+	
+	virtual void* GetTextureBaseRHI() override final
+	{
+		return &Surface;
+	}
+	
+	virtual void* GetNativeResource() const override final
+	{
+		return Surface.Texture;
+	}
+};
+
+class FAGXTexture2DArray : public FRHITexture2DArray
+{
+public:
+	/** The surface info */
+	FAGXSurface Surface;
+
+	// Constructor, just calls base and Surface constructor
+	FAGXTexture2DArray(EPixelFormat Format, uint32 SizeX, uint32 SizeY, uint32 ArraySize, uint32 NumMips, ETextureCreateFlags Flags, FResourceBulkDataInterface* BulkData, const FClearValueBinding& InClearValue)
+		: FRHITexture2DArray(SizeX, SizeY, ArraySize, NumMips, 1, Format, Flags, InClearValue)
+		, Surface(RRT_Texture2DArray, Format, SizeX, SizeY, 1, /*NumSamples=*/1, /*bArray=*/ true, ArraySize, NumMips, Flags, BulkData)
+	{
+	}
+	
+	virtual ~FAGXTexture2DArray()
+	{
+	}
+	
+	virtual void* GetTextureBaseRHI() override final
+	{
+		return &Surface;
+	}
+};
+
+class FAGXTexture3D : public FRHITexture3D
+{
+public:
+	/** The surface info */
+	FAGXSurface Surface;
+
+	// Constructor, just calls base and Surface constructor
+	FAGXTexture3D(EPixelFormat Format, uint32 SizeX, uint32 SizeY, uint32 SizeZ, uint32 NumMips, ETextureCreateFlags Flags, FResourceBulkDataInterface* BulkData, const FClearValueBinding& InClearValue)
+		: FRHITexture3D(SizeX, SizeY, SizeZ, NumMips, Format, Flags, InClearValue)
+		, Surface(RRT_Texture3D, Format, SizeX, SizeY, SizeZ, /*NumSamples=*/1, /*bArray=*/ false, 1, NumMips, Flags, BulkData)
+	{
+	}
+	
+	virtual ~FAGXTexture3D()
+	{
+	}
+	
+	virtual void* GetTextureBaseRHI() override final
+	{
+		return &Surface;
+	}
+};
+
+class FAGXTextureCube : public FRHITextureCube
+{
+public:
+	/** The surface info */
+	FAGXSurface Surface;
+
+	// Constructor, just calls base and Surface constructor
+	FAGXTextureCube(EPixelFormat Format, uint32 Size, bool bArray, uint32 ArraySize, uint32 NumMips, ETextureCreateFlags Flags, FResourceBulkDataInterface* BulkData, const FClearValueBinding& InClearValue)
+		: FRHITextureCube(Size, NumMips, Format, Flags, InClearValue)
+		, Surface(RRT_TextureCube, Format, Size, Size, 6, /*NumSamples=*/1, bArray, ArraySize, NumMips, Flags, BulkData)
+	{
+	}
+	
+	virtual ~FAGXTextureCube()
+	{
+	}
+	
+	virtual void* GetTextureBaseRHI() override final
+	{
+		return &Surface;
+	}
+
+	virtual void* GetNativeResource() const override final
+	{
+		return Surface.Texture;
+	}
+};
+
+@interface FAGXBufferData : FApplePlatformObject<NSObject>
+{
+@public
+	uint8* Data;
+	uint32 Len;	
+}
+-(instancetype)initWithSize:(uint32)Size;
+-(instancetype)initWithBytes:(void const*)Data length:(uint32)Size;
+@end
+
+enum EAGXBufferUsage
+{
+	EAGXBufferUsage_GPUOnly = 0x80000000,
+	EAGXBufferUsage_LinearTex = 0x40000000,
+
+	EAGXBufferUsageMask	= 0xF,
+	EAGXBufferUsageShift	= 0x1C,
+	EAGXBufferUsageFlags	= (EAGXBufferUsageMask << EAGXBufferUsageShift)
+};
+
+class FAGXLinearTextureDescriptor
+{
+public:
+	FAGXLinearTextureDescriptor()
+		: StartOffsetBytes(0)
+		, NumElements(UINT_MAX)
+		, BytesPerElement(0)
+	{
+		// void
+	}
+
+	FAGXLinearTextureDescriptor(uint32 InStartOffsetBytes, uint32 InNumElements, uint32 InBytesPerElement)
+		: StartOffsetBytes(InStartOffsetBytes)
+		, NumElements(InNumElements)
+		, BytesPerElement(InBytesPerElement)
+	{
+		// void
+	}
+	
+	FAGXLinearTextureDescriptor(const FAGXLinearTextureDescriptor& Other)
+		: StartOffsetBytes(Other.StartOffsetBytes)
+		, NumElements(Other.NumElements)
+		, BytesPerElement(Other.BytesPerElement)
+	{
+		// void
+	}
+	
+	~FAGXLinearTextureDescriptor()
+	{
+		// void
+	}
+
+	friend uint32 GetTypeHash(FAGXLinearTextureDescriptor const& Key)
+	{
+		uint32 Hash = GetTypeHash((uint64)Key.StartOffsetBytes);
+		Hash = HashCombine(Hash, GetTypeHash((uint64)Key.NumElements));
+		Hash = HashCombine(Hash, GetTypeHash((uint64)Key.BytesPerElement));
+		return Hash;
+	}
+
+	bool operator==(FAGXLinearTextureDescriptor const& Other) const
+	{
+		return    StartOffsetBytes == Other.StartOffsetBytes
+		       && NumElements      == Other.NumElements
+		       && BytesPerElement  == Other.BytesPerElement;
+	}
+
+	uint32 StartOffsetBytes;
+	uint32 NumElements;
+	uint32 BytesPerElement;
+};
+
+class FAGXRHIBuffer
+{
+public:
+	// Matches other RHIs
+	static constexpr const uint32 MaxNumBufferedFrames = 4;
+	
+	using LinearTextureMapKey = TTuple<EPixelFormat, FAGXLinearTextureDescriptor>;
+	using LinearTextureMap = TMap<LinearTextureMapKey, FAGXTexture>;
+	
+	struct FAGXBufferAndViews
+	{
+		FAGXBuffer Buffer;
+		LinearTextureMap Views;
+	};
+	
+	FAGXRHIBuffer(uint32 InSize, uint32 InUsage, ERHIResourceType InType);
+	virtual ~FAGXRHIBuffer();
+	
+	/**
+	 * Initialize the buffer contents from the render-thread.
+	 */
+	void Init_RenderThread(class FRHICommandListImmediate& RHICmdList, uint32 Size, uint32 InUsage, FRHIResourceCreateInfo& CreateInfo, FRHIResource* Resource);
+	
+	/**
+	 * Get a linear texture for given format.
+	 */
+	void CreateLinearTexture(EPixelFormat InFormat, FRHIResource* InParent, const FAGXLinearTextureDescriptor* InLinearTextureDescriptor = nullptr);
+	
+	/**
+	 * Get a linear texture for given format.
+	 */
+	ns::AutoReleased<FAGXTexture> GetLinearTexture(EPixelFormat InFormat, const FAGXLinearTextureDescriptor* InLinearTextureDescriptor = nullptr);
+	
+	/**
+	 * Prepare a CPU accessible buffer for uploading to GPU memory
+	 */
+	void* Lock(bool bIsOnRHIThread, EResourceLockMode LockMode, uint32 Offset, uint32 Size=0);
+	
+	/**
+	 * Prepare a CPU accessible buffer for uploading to GPU memory
+	 */
+	void Unlock();
+	
+	void Swap(FAGXRHIBuffer& Other);
+	
+	const FAGXBufferAndViews& GetCurrentBacking()
+	{
+		check(NumberOfBuffers > 0);
+		return BufferPool[CurrentIndex];
+	}
+	
+	const FAGXBuffer& GetCurrentBuffer()
+	{
+		return BufferPool[CurrentIndex].Buffer;
+	}
+	
+	FAGXBuffer GetCurrentBufferOrNil()
+	{
+		if(NumberOfBuffers > 0)
+		{
+			return GetCurrentBuffer();
+		}
+		return nil;
+	}
+	
+	void AdvanceBackingIndex()
+	{
+		CurrentIndex = (CurrentIndex + 1) % NumberOfBuffers;
+	}
+	
+	/**
+	 * Whether to allocate the resource from private memory.
+	 */
+	bool UsePrivateMemory() const;
+	
+	// A temporary shared/CPU accessible buffer for upload/download
+	FAGXBuffer TransferBuffer;
+	
+	TArray<FAGXBufferAndViews> BufferPool;
+	
+	/** Buffer for small buffers < 4Kb to avoid heap fragmentation. */
+	FAGXBufferData* Data;
+	
+	// Frame we last locked (for debugging, mainly)
+	uint32 LastLockFrame;
+	
+	// The active buffer.
+	uint32 CurrentIndex		: 8;
+	// How many buffers are actually allocated
+	uint32 NumberOfBuffers	: 8;
+	// Current lock mode. RLM_Num indicates this buffer is not locked.
+	uint32 CurrentLockMode	: 16;
+	
+	// offset into the buffer (for lock usage)
+	uint32 LockOffset;
+	
+	// Sizeof outstanding lock.
+	uint32 LockSize;
+	
+	// Initial buffer size.
+	uint32 Size;
+	
+	// Buffer usage.
+	uint32 Usage;
+	
+	// Storage mode
+	mtlpp::StorageMode Mode;
+	
+	// Resource type
+	ERHIResourceType Type;
+	
+	static_assert((1 << 16) > RLM_Num, "Lock mode does not fit in bitfield");
+	static_assert((1 << 8) > MaxNumBufferedFrames, "Buffer count does not fit in bitfield");
+	
+private:
+	FAGXBufferAndViews& GetCurrentBackingInternal()
+	{
+		return BufferPool[CurrentIndex];
+	}
+	
+	FAGXBuffer& GetCurrentBufferInternal()
+	{
+		return BufferPool[CurrentIndex].Buffer;
+	}
+	
+	/**
+	 * Allocate the CPU accessible buffer for data transfer.
+	 */
+	void AllocTransferBuffer(bool bOnRHIThread, uint32 InSize, EResourceLockMode LockMode);
+	
+	/**
+	 * Allocate a linear texture for given format.
+	 */
+	void AllocLinearTextures(const LinearTextureMapKey& InLinearTextureMapKey);
+};
+
+class FAGXResourceMultiBuffer : public FRHIBuffer, public FAGXRHIBuffer
+{
+public:
+	FAGXResourceMultiBuffer(uint32 InSize, uint32 InUsage, uint32 InStride, FResourceArrayInterface* ResourceArray, ERHIResourceType ResourceType);
+	virtual ~FAGXResourceMultiBuffer();
+
+	void Swap(FAGXResourceMultiBuffer& Other);
+
+	// 16- or 32-bit; used for index buffers only.
+	mtlpp::IndexType IndexType;
+};
+
+typedef FAGXResourceMultiBuffer FAGXIndexBuffer;
+typedef FAGXResourceMultiBuffer FAGXVertexBuffer;
+typedef FAGXResourceMultiBuffer FAGXStructuredBuffer;
+
+class FAGXShaderResourceView : public FRHIShaderResourceView
+{
+public:
+
+	// The vertex buffer this SRV comes from (can be null)
+	TRefCountPtr<FAGXVertexBuffer> SourceVertexBuffer;
+	
+	// The index buffer this SRV comes from (can be null)
+	TRefCountPtr<FAGXIndexBuffer> SourceIndexBuffer;
+
+	// The texture that this SRV come from
+	TRefCountPtr<FRHITexture> SourceTexture;
+	
+	// The source structured buffer (can be null)
+	TRefCountPtr<FAGXStructuredBuffer> SourceStructuredBuffer;
+	
+	FAGXSurface* TextureView;
+	uint32 Offset;
+	uint8 MipLevel          : 4;
+	uint8 bSRGBForceDisable : 1;
+	uint8 Reserved          : 3;
+	uint8 NumMips;
+	uint8 Format;
+	uint8 Stride;
+
+	void InitLinearTextureDescriptor(const FAGXLinearTextureDescriptor& InLinearTextureDescriptor);
+
+	FAGXShaderResourceView();
+	~FAGXShaderResourceView();
+	
+	ns::AutoReleased<FAGXTexture> GetLinearTexture(bool const bUAV);
+
+private:
+	FAGXLinearTextureDescriptor* LinearTextureDesc;
+};
+
+
+
+class FAGXUnorderedAccessView : public FRHIUnorderedAccessView
+{
+public:
+	
+	// the potential resources to refer to with the UAV object
+	TRefCountPtr<FAGXShaderResourceView> SourceView;
+};
+
+class FAGXGPUFence final : public FRHIGPUFence
+{
+public:
+	FAGXGPUFence(FName InName)
+		: FRHIGPUFence(InName)
+	{
+	}
+
+	~FAGXGPUFence()
+	{
+	}
+
+	virtual void Clear() override final;
+
+	void WriteInternal(mtlpp::CommandBuffer& CmdBuffer);
+
+	virtual bool Poll() const override final;
+
+private:
+	mtlpp::CommandBufferFence Fence;
+};
+
+class FAGXShaderLibrary;
+class FAGXGraphicsPipelineState;
+class FAGXComputePipelineState;
+class FAGXVertexDeclaration;
+class FAGXVertexShader;
+class FAGXGeometryShader;
+class FAGXPixelShader;
+class FAGXComputeShader;
+class FAGXRHIStagingBuffer;
+class FAGXRHIRenderQuery;
+class FAGXSuballocatedUniformBuffer;
+
+template<class T>
+struct TAGXResourceTraits
+{
+};
+template<>
+struct TAGXResourceTraits<FRHIShaderLibrary>
+{
+	typedef FAGXShaderLibrary TConcreteType;
+};
+template<>
+struct TAGXResourceTraits<FRHIVertexDeclaration>
+{
+	typedef FAGXVertexDeclaration TConcreteType;
+};
+template<>
+struct TAGXResourceTraits<FRHIVertexShader>
+{
+	typedef FAGXVertexShader TConcreteType;
+};
+template<>
+struct TAGXResourceTraits<FRHIGeometryShader>
+{
+	typedef FAGXGeometryShader TConcreteType;
+};
+template<>
+struct TAGXResourceTraits<FRHIPixelShader>
+{
+	typedef FAGXPixelShader TConcreteType;
+};
+template<>
+struct TAGXResourceTraits<FRHIComputeShader>
+{
+	typedef FAGXComputeShader TConcreteType;
+};
+template<>
+struct TAGXResourceTraits<FRHITexture3D>
+{
+	typedef FAGXTexture3D TConcreteType;
+};
+template<>
+struct TAGXResourceTraits<FRHITexture2D>
+{
+	typedef FAGXTexture2D TConcreteType;
+};
+template<>
+struct TAGXResourceTraits<FRHITexture2DArray>
+{
+	typedef FAGXTexture2DArray TConcreteType;
+};
+template<>
+struct TAGXResourceTraits<FRHITextureCube>
+{
+	typedef FAGXTextureCube TConcreteType;
+};
+template<>
+struct TAGXResourceTraits<FRHIRenderQuery>
+{
+	typedef FAGXRHIRenderQuery TConcreteType;
+};
+template<>
+struct TAGXResourceTraits<FRHIUniformBuffer>
+{
+	typedef FAGXSuballocatedUniformBuffer TConcreteType;
+};
+template<>
+struct TAGXResourceTraits<FRHIBuffer>
+{
+	typedef FAGXResourceMultiBuffer TConcreteType;
+};
+template<>
+struct TAGXResourceTraits<FRHIShaderResourceView>
+{
+	typedef FAGXShaderResourceView TConcreteType;
+};
+template<>
+struct TAGXResourceTraits<FRHIUnorderedAccessView>
+{
+	typedef FAGXUnorderedAccessView TConcreteType;
+};
+
+template<>
+struct TAGXResourceTraits<FRHISamplerState>
+{
+	typedef FAGXSamplerState TConcreteType;
+};
+template<>
+struct TAGXResourceTraits<FRHIRasterizerState>
+{
+	typedef FAGXRasterizerState TConcreteType;
+};
+template<>
+struct TAGXResourceTraits<FRHIDepthStencilState>
+{
+	typedef FAGXDepthStencilState TConcreteType;
+};
+template<>
+struct TAGXResourceTraits<FRHIBlendState>
+{
+	typedef FAGXBlendState TConcreteType;
+};
+template<>
+struct TAGXResourceTraits<FRHIGraphicsPipelineState>
+{
+	typedef FAGXGraphicsPipelineState TConcreteType;
+};
+template<>
+struct TAGXResourceTraits<FRHIComputePipelineState>
+{
+	typedef FAGXComputePipelineState TConcreteType;
+};
+template<>
+struct TAGXResourceTraits<FRHIGPUFence>
+{
+	typedef FAGXGPUFence TConcreteType;
+};
+template<>
+struct TAGXResourceTraits<FRHIStagingBuffer>
+{
+	typedef FAGXRHIStagingBuffer TConcreteType;
+};
