U5 - file missed from latest CL

diff --git a/Engine/Source/Runtime/Apple/AGXRHI/Private/AGXDynamicRHI.cpp b/Engine/Source/Runtime/Apple/AGXRHI/Private/AGXDynamicRHI.cpp
--- a/Engine/Source/Runtime/Apple/AGXRHI/Private/AGXDynamicRHI.cpp
+++ b/Engine/Source/Runtime/Apple/AGXRHI/Private/AGXDynamicRHI.cpp
@@ -51,8 +51,7 @@ FGraphicsPipelineStateRHIRef FAGXDynamicRHI::RHICreateGraphicsPipelineState(cons
 		if(!State->Compile())
 		{
 			// Compilation failures are propagated up to the caller.
-			State->DoNoDeferDelete();
-			delete State;
+			State->Delete();
 			return nullptr;
 		}
 
