Fix IOS compile

diff --git a/Engine/Source/Runtime/Apple/AGXRHI/Private/AGXTexture.cpp b/Engine/Source/Runtime/Apple/AGXRHI/Private/AGXTexture.cpp
--- a/Engine/Source/Runtime/Apple/AGXRHI/Private/AGXTexture.cpp
+++ b/Engine/Source/Runtime/Apple/AGXRHI/Private/AGXTexture.cpp
@@ -834,11 +834,11 @@ FAGXSurface::FAGXSurface(ERHIResourceType ResourceType, EPixelFormat Format, uin
 		}
 		
 #if PLATFORM_IOS
-		if (Flags & TexCreate_Memoryless)
+		if (EnumHasAnyFlags(Flags, TexCreate_Memoryless))
 		{
-			ensure(Flags & (TexCreate_RenderTargetable | TexCreate_DepthStencilTargetable));
-			ensure(!(Flags & (TexCreate_CPUReadback | TexCreate_CPUWritable)));
-			ensure(!(Flags & TexCreate_UAV));
+			ensure(EnumHasAnyFlags(Flags, (TexCreate_RenderTargetable | TexCreate_DepthStencilTargetable)));
+			ensure(!EnumHasAnyFlags(Flags, (TexCreate_CPUReadback | TexCreate_CPUWritable)));
+			ensure(!EnumHasAnyFlags(Flags, TexCreate_UAV));
 			Desc.SetStorageMode(mtlpp::StorageMode::Memoryless);
 			Desc.SetResourceOptions((mtlpp::ResourceOptions)(mtlpp::ResourceOptions::CpuCacheModeDefaultCache | mtlpp::ResourceOptions::StorageModeMemoryless));
 		}
diff --git a/Engine/Source/Runtime/Apple/MetalRHI/Private/MetalTexture.cpp b/Engine/Source/Runtime/Apple/MetalRHI/Private/MetalTexture.cpp
--- a/Engine/Source/Runtime/Apple/MetalRHI/Private/MetalTexture.cpp
+++ b/Engine/Source/Runtime/Apple/MetalRHI/Private/MetalTexture.cpp
@@ -836,11 +836,11 @@ FMetalSurface::FMetalSurface(ERHIResourceType ResourceType, EPixelFormat Format,
 		}
 		
 #if PLATFORM_IOS
-		if (Flags & TexCreate_Memoryless)
+		if (EnumHasAnyFlags(Flags, TexCreate_Memoryless))
 		{
-			ensure(Flags & (TexCreate_RenderTargetable | TexCreate_DepthStencilTargetable));
-			ensure(!(Flags & (TexCreate_CPUReadback | TexCreate_CPUWritable)));
-			ensure(!(Flags & TexCreate_UAV));
+			ensure(EnumHasAnyFlags(Flags, (TexCreate_RenderTargetable | TexCreate_DepthStencilTargetable)));
+			ensure(!EnumHasAnyFlags(Flags, (TexCreate_CPUReadback | TexCreate_CPUWritable)));
+			ensure(!EnumHasAnyFlags(Flags, TexCreate_UAV));
 			Desc.SetStorageMode(mtlpp::StorageMode::Memoryless);
 			Desc.SetResourceOptions((mtlpp::ResourceOptions)(mtlpp::ResourceOptions::CpuCacheModeDefaultCache | mtlpp::ResourceOptions::StorageModeMemoryless));
 		}
