Fixup Mac UnrealGame build.

#rnx

diff --git a/Engine/Source/Runtime/Apple/AGXRHI/Private/AGXContext.cpp b/Engine/Source/Runtime/Apple/AGXRHI/Private/AGXContext.cpp
--- a/Engine/Source/Runtime/Apple/AGXRHI/Private/AGXContext.cpp
+++ b/Engine/Source/Runtime/Apple/AGXRHI/Private/AGXContext.cpp
@@ -299,7 +299,7 @@ static mtlpp::Device GetMTLDevice(uint32& DeviceIndex)
 #endif // PLATFORM_MAC_ARM64
 }
 
-mtlpp::PrimitiveTopologyClass TranslatePrimitiveTopology(uint32 PrimitiveType)
+mtlpp::PrimitiveTopologyClass AGXTranslatePrimitiveTopology(uint32 PrimitiveType)
 {
 	switch (PrimitiveType)
 	{
diff --git a/Engine/Source/Runtime/Apple/AGXRHI/Private/AGXPipeline.cpp b/Engine/Source/Runtime/Apple/AGXRHI/Private/AGXPipeline.cpp
--- a/Engine/Source/Runtime/Apple/AGXRHI/Private/AGXPipeline.cpp
+++ b/Engine/Source/Runtime/Apple/AGXRHI/Private/AGXPipeline.cpp
@@ -273,7 +273,7 @@ struct FAGXGraphicsPipelineKey
 		Key.SetHashValue(Offset_AlphaToCoverage, NumBits_AlphaToCoverage, Init.NumSamples > 1 && BlendState->bUseAlphaToCoverage ? 1 : 0);
 		
 #if PLATFORM_MAC
-		Key.SetHashValue(Offset_PrimitiveTopology, NumBits_PrimitiveTopology, TranslatePrimitiveTopology(Init.PrimitiveType));
+		Key.SetHashValue(Offset_PrimitiveTopology, NumBits_PrimitiveTopology, AGXTranslatePrimitiveTopology(Init.PrimitiveType));
 #endif
 
 		FAGXVertexDeclaration* VertexDecl = (FAGXVertexDeclaration*)Init.BoundShaderState.VertexDeclarationRHI;
@@ -782,7 +782,7 @@ static bool ConfigureRenderPipelineDescriptor(mtlpp::RenderPipelineDescriptor& R
 	RenderPipelineDesc.SetSampleCount(NumSamples);
 	RenderPipelineDesc.SetAlphaToCoverageEnabled(NumSamples > 1 && BlendState->bUseAlphaToCoverage);
 #if PLATFORM_MAC
-	RenderPipelineDesc.SetInputPrimitiveTopology(TranslatePrimitiveTopology(Init.PrimitiveType));
+	RenderPipelineDesc.SetInputPrimitiveTopology(AGXTranslatePrimitiveTopology(Init.PrimitiveType));
 	DebugPipelineDesc.SetSampleCount(!bNoMSAA ? FMath::Max(Init.NumSamples, (uint16)1u) : (uint16)1u);
 	DebugPipelineDesc.SetInputPrimitiveTopology(mtlpp::PrimitiveTopologyClass::Point);
 #endif
diff --git a/Engine/Source/Runtime/Apple/AGXRHI/Private/AGXProfiler.cpp b/Engine/Source/Runtime/Apple/AGXRHI/Private/AGXProfiler.cpp
--- a/Engine/Source/Runtime/Apple/AGXRHI/Private/AGXProfiler.cpp
+++ b/Engine/Source/Runtime/Apple/AGXRHI/Private/AGXProfiler.cpp
@@ -62,7 +62,7 @@ int64 volatile GAGXGPUWorkTime = 0;
 int64 volatile GAGXGPUIdleTime = 0;
 int64 volatile GAGXPresentTime = 0;
 
-void WriteString(FArchive* OutputFile, const char* String)
+static void AGXWriteString(FArchive* OutputFile, const char* String)
 {
 	OutputFile->Serialize((void*)String, sizeof(ANSICHAR)*FCStringAnsi::Strlen(String));
 }
@@ -719,7 +719,7 @@ void FAGXProfiler::SaveTrace()
 		
 		FArchive* OutputFile = IFileManager::Get().CreateFileWriter(*OutputFilename);
 		
-		WriteString(OutputFile, R"({"traceEvents":[)" "\n");
+		AGXWriteString(OutputFile, R"({"traceEvents":[)" "\n");
 		
 		int32 SortIndex = 0; // Lower numbers result in higher position in the visualizer.
 		const uint32 Pid = FPlatformProcess::GetCurrentProcessId();
@@ -730,28 +730,28 @@ void FAGXProfiler::SaveTrace()
 											 Pid, GPUIndex, GPUIndex, Pid, GPUIndex, SortIndex
 											 );
 			
-			WriteString(OutputFile, TCHAR_TO_UTF8(*Output));
+			AGXWriteString(OutputFile, TCHAR_TO_UTF8(*Output));
 			SortIndex++;
 			
 			Output = FString::Printf(TEXT("{\"pid\":%d, \"tid\":%d, \"ph\": \"M\", \"name\": \"thread_name\", \"args\":{\"name\":\"GPU %d Operations\"}},{\"pid\":%d, \"tid\":%d, \"ph\": \"M\", \"name\": \"thread_sort_index\", \"args\":{\"sort_index\": %d}},\n"),
 											 Pid, GPUIndex+SortIndex, GPUIndex, Pid, GPUIndex+SortIndex, SortIndex
 											 );
 			
-			WriteString(OutputFile, TCHAR_TO_UTF8(*Output));
+			AGXWriteString(OutputFile, TCHAR_TO_UTF8(*Output));
 			SortIndex++;
 			
 			Output = FString::Printf(TEXT("{\"pid\":%d, \"tid\":%d, \"ph\": \"M\", \"name\": \"thread_name\", \"args\":{\"name\":\"Render Events %d\"}},{\"pid\":%d, \"tid\":%d, \"ph\": \"M\", \"name\": \"thread_sort_index\", \"args\":{\"sort_index\": %d}},\n"),
 									 Pid, GPUIndex+SortIndex, GPUIndex, Pid, GPUIndex+SortIndex, SortIndex
 									 );
 			
-			WriteString(OutputFile, TCHAR_TO_UTF8(*Output));
+			AGXWriteString(OutputFile, TCHAR_TO_UTF8(*Output));
 			SortIndex++;
 			
 			Output = FString::Printf(TEXT("{\"pid\":%d, \"tid\":%d, \"ph\": \"M\", \"name\": \"thread_name\", \"args\":{\"name\":\"Driver Stats %d\"}},{\"pid\":%d, \"tid\":%d, \"ph\": \"M\", \"name\": \"thread_sort_index\", \"args\":{\"sort_index\": %d}},\n"),
 									 Pid, GPUIndex+SortIndex, GPUIndex, Pid, GPUIndex+SortIndex, SortIndex
 									 );
 			
-			WriteString(OutputFile, TCHAR_TO_UTF8(*Output));
+			AGXWriteString(OutputFile, TCHAR_TO_UTF8(*Output));
 			SortIndex++;
 			
 			for (uint32 Display : Displays)
@@ -760,7 +760,7 @@ void FAGXProfiler::SaveTrace()
 										 Pid, Display + SortIndex, SortIndex - 3, Pid, Display + SortIndex, SortIndex
 										 );
 				
-				WriteString(OutputFile, TCHAR_TO_UTF8(*Output));
+				AGXWriteString(OutputFile, TCHAR_TO_UTF8(*Output));
 				SortIndex++;
 			}
 		}
@@ -788,27 +788,27 @@ void FAGXProfiler::SaveTrace()
 											 );
 			
 			
-			WriteString(OutputFile, TCHAR_TO_UTF8(*Output));
+			AGXWriteString(OutputFile, TCHAR_TO_UTF8(*Output));
 			SortIndex++;
 		}
 		
 		for (FAGXCommandBufferStats* CmdBufStats : TracedBuffers)
 		{
-			WriteString(OutputFile, TCHAR_TO_UTF8(*CmdBufStats->GetJSONRepresentation(Pid)));
+			AGXWriteString(OutputFile, TCHAR_TO_UTF8(*CmdBufStats->GetJSONRepresentation(Pid)));
 			
 			for (IAGXStatsScope* ES : CmdBufStats->Children)
 			{
-				WriteString(OutputFile, TCHAR_TO_UTF8(*ES->GetJSONRepresentation(Pid)));
+				AGXWriteString(OutputFile, TCHAR_TO_UTF8(*ES->GetJSONRepresentation(Pid)));
 				
 				uint64 PrevTime = ES->GPUStartTime;
 				for (IAGXStatsScope* DS : ES->Children)
 				{
-					WriteString(OutputFile, TCHAR_TO_UTF8(*DS->GetJSONRepresentation(Pid)));
+					AGXWriteString(OutputFile, TCHAR_TO_UTF8(*DS->GetJSONRepresentation(Pid)));
 					if (!DS->GPUStartTime)
 					{
 						DS->GPUStartTime = FMath::Max(PrevTime, DS->GPUStartTime);
 						DS->GPUEndTime = DS->GPUStartTime + 1llu;
-						WriteString(OutputFile, TCHAR_TO_UTF8(*DS->GetJSONRepresentation(Pid)));
+						AGXWriteString(OutputFile, TCHAR_TO_UTF8(*DS->GetJSONRepresentation(Pid)));
 					}
 					PrevTime = DS->GPUEndTime;
 				}
@@ -821,21 +821,21 @@ void FAGXProfiler::SaveTrace()
 		for (FAGXDisplayStats* DisplayStat : DisplayStats)
 		{
 			DisplayStat->GPUThreadIndex += 3;
-			WriteString(OutputFile, TCHAR_TO_UTF8(*DisplayStat->GetJSONRepresentation(Pid)));
+			AGXWriteString(OutputFile, TCHAR_TO_UTF8(*DisplayStat->GetJSONRepresentation(Pid)));
 			delete DisplayStat;
 		}
 		DisplayStats.Empty();
 		
 		for (FAGXCPUStats* CPUStat : CPUStats)
 		{
-			WriteString(OutputFile, TCHAR_TO_UTF8(*CPUStat->GetJSONRepresentation(Pid)));
+			AGXWriteString(OutputFile, TCHAR_TO_UTF8(*CPUStat->GetJSONRepresentation(Pid)));
 			delete CPUStat;
 		}
 		CPUStats.Empty();
 		
 		// All done
 		
-		WriteString(OutputFile, "{}]}");
+		AGXWriteString(OutputFile, "{}]}");
 		
 		OutputFile->Close();
 	}
diff --git a/Engine/Source/Runtime/Apple/AGXRHI/Private/AGXRHIPrivate.h b/Engine/Source/Runtime/Apple/AGXRHI/Private/AGXRHIPrivate.h
--- a/Engine/Source/Runtime/Apple/AGXRHI/Private/AGXRHIPrivate.h
+++ b/Engine/Source/Runtime/Apple/AGXRHI/Private/AGXRHIPrivate.h
@@ -212,7 +212,7 @@ FORCEINLINE mtlpp::LoadAction GetMetalRTLoadAction(ERenderTargetLoadAction LoadA
 mtlpp::PrimitiveType AGXTranslatePrimitiveType(uint32 PrimitiveType);
 
 #if PLATFORM_MAC
-mtlpp::PrimitiveTopologyClass TranslatePrimitiveTopology(uint32 PrimitiveType);
+mtlpp::PrimitiveTopologyClass AGXTranslatePrimitiveTopology(uint32 PrimitiveType);
 #endif
 
 mtlpp::PixelFormat AGXToSRGBFormat(mtlpp::PixelFormat LinMTLFormat);
diff --git a/Engine/Source/Runtime/Apple/AGXRHI/Private/AGXTexture.cpp b/Engine/Source/Runtime/Apple/AGXRHI/Private/AGXTexture.cpp
--- a/Engine/Source/Runtime/Apple/AGXRHI/Private/AGXTexture.cpp
+++ b/Engine/Source/Runtime/Apple/AGXRHI/Private/AGXTexture.cpp
@@ -158,7 +158,7 @@ void AGXSafeReleaseMetalTexture(FAGXSurface* Surface, FAGXTexture& Texture)
 }
 
 #if !PLATFORM_IOS
-mtlpp::PixelFormat ToSRGBFormat_NonAppleMacGPU(mtlpp::PixelFormat MTLFormat)
+static mtlpp::PixelFormat AGX_ToSRGBFormat_NonAppleMacGPU(mtlpp::PixelFormat MTLFormat)
 {
 	switch (MTLFormat)
 	{
@@ -187,7 +187,7 @@ mtlpp::PixelFormat ToSRGBFormat_NonAppleMacGPU(mtlpp::PixelFormat MTLFormat)
 }
 #endif
 
-mtlpp::PixelFormat ToSRGBFormat_AppleGPU(mtlpp::PixelFormat MTLFormat)
+static mtlpp::PixelFormat AGX_ToSRGBFormat_AppleGPU(mtlpp::PixelFormat MTLFormat)
 {
 	switch (MTLFormat)
 	{
@@ -231,12 +231,12 @@ mtlpp::PixelFormat AGXToSRGBFormat(mtlpp::PixelFormat MTLFormat)
 {
 	if([GetAGXDeviceContext().GetDevice().GetPtr() supportsFamily:MTLGPUFamilyApple1])
 	{
-		MTLFormat = ToSRGBFormat_AppleGPU(MTLFormat);
+		MTLFormat = AGX_ToSRGBFormat_AppleGPU(MTLFormat);
 	}
 #if !PLATFORM_IOS
 	else if([GetAGXDeviceContext().GetDevice().GetPtr() supportsFamily:MTLGPUFamilyMac1])
 	{
-		MTLFormat = ToSRGBFormat_NonAppleMacGPU(MTLFormat);
+		MTLFormat = AGX_ToSRGBFormat_NonAppleMacGPU(MTLFormat);
 	}
 #endif
 	
