Reflect MetalRHI changes from CL 16544610 to AGXRHI

#rnx

diff --git a/Engine/Source/Runtime/Apple/AGXRHI/Private/AGXLLM.cpp b/Engine/Source/Runtime/Apple/AGXRHI/Private/AGXLLM.cpp
--- a/Engine/Source/Runtime/Apple/AGXRHI/Private/AGXLLM.cpp
+++ b/Engine/Source/Runtime/Apple/AGXRHI/Private/AGXLLM.cpp
@@ -166,12 +166,10 @@ static mtlpp::SizeAndAlign TextureSizeAndAlign(mtlpp::TextureType TextureType, u
 			SizeAlign.Size = RHICalcTextureCubePlatformSize(Width, MetalToRHIPixelFormat(Format), MipCount, TexCreate_None, CreateInfo, Align);
 			SizeAlign.Align = Align;
 			break;
-#if PLATFORM_MAC || __IPHONE_OS_VERSION_MAX_ALLOWED >= 110000
 		case mtlpp::TextureType::TextureCubeArray:
 			SizeAlign.Size = RHICalcTextureCubePlatformSize(Width, MetalToRHIPixelFormat(Format), MipCount, TexCreate_None, CreateInfo, Align) * ArrayCount;
 			SizeAlign.Align = Align;
 			break;
-#endif
 		case mtlpp::TextureType::Texture3D:
 			SizeAlign.Size = RHICalcTexture3DPlatformSize(Width, Height, Depth, MetalToRHIPixelFormat(Format), MipCount, TexCreate_None, CreateInfo, Align);
 			SizeAlign.Align = Align;
@@ -189,16 +187,10 @@ static mtlpp::SizeAndAlign TextureSizeAndAlign(mtlpp::TextureType TextureType, u
 void AGXLLM::LogAllocTexture(mtlpp::Device& Device, mtlpp::TextureDescriptor const& Desc, mtlpp::Texture const& Texture)
 {
 	mtlpp::SizeAndAlign SizeAlign;
-#if PLATFORM_MAC || __IPHONE_OS_VERSION_MAX_ALLOWED >= 110000
 	if (FAGXCommandQueue::SupportsFeature(EAGXFeaturesGPUCaptureManager))
 	{
 		SizeAlign = Device.HeapTextureSizeAndAlign(Desc);
 	}
-	else
-#endif
-	{
-		SizeAlign = TextureSizeAndAlign(Desc.GetTextureType(), Desc.GetWidth(), Desc.GetHeight(), Desc.GetDepth(), Desc.GetPixelFormat(), Desc.GetMipmapLevelCount(), Desc.GetSampleCount(), Desc.GetArrayLength());
-	}
 	
 	void* Ptr = (void*)Texture.GetPtr();
 	uint64 Size = SizeAlign.Size;
diff --git a/Engine/Source/Runtime/Apple/AGXRHI/Private/AGXRHIPrivate.h b/Engine/Source/Runtime/Apple/AGXRHI/Private/AGXRHIPrivate.h
--- a/Engine/Source/Runtime/Apple/AGXRHI/Private/AGXRHIPrivate.h
+++ b/Engine/Source/Runtime/Apple/AGXRHI/Private/AGXRHIPrivate.h
@@ -71,9 +71,9 @@ enum EMTLTextureType
 };
 
 // This is the right VERSION check, see Availability.h in the SDK
-#define METAL_SUPPORTS_INDIRECT_ARGUMENT_BUFFERS ((PLATFORM_MAC && __MAC_OS_X_VERSION_MAX_ALLOWED >= 101300) || (!PLATFORM_MAC && __IPHONE_OS_VERSION_MAX_ALLOWED >= 110000)) && (__clang_major__ >= 9)
-#define METAL_SUPPORTS_CAPTURE_MANAGER (PLATFORM_MAC && __MAC_OS_X_VERSION_MAX_ALLOWED >= 101300) || (!PLATFORM_MAC && __IPHONE_OS_VERSION_MAX_ALLOWED >= 110000) && (__clang_major__ >= 9)
-#define METAL_SUPPORTS_TILE_SHADERS (!PLATFORM_MAC && !PLATFORM_TVOS && __IPHONE_OS_VERSION_MAX_ALLOWED >= 110000) && (__clang_major__ >= 9)
+#define METAL_SUPPORTS_INDIRECT_ARGUMENT_BUFFERS 1
+#define METAL_SUPPORTS_CAPTURE_MANAGER 1
+#define METAL_SUPPORTS_TILE_SHADERS 1
 // In addition to compile-time SDK checks we also need a way to check if these are available on runtime
 extern bool GAGXSupportsCaptureManager;
 
