Enable multi-pipe resource transitions for Metal.

#rnx

diff --git a/Engine/Source/Runtime/Apple/AGXRHI/Private/AGXDynamicRHI.cpp b/Engine/Source/Runtime/Apple/AGXRHI/Private/AGXDynamicRHI.cpp
--- a/Engine/Source/Runtime/Apple/AGXRHI/Private/AGXDynamicRHI.cpp
+++ b/Engine/Source/Runtime/Apple/AGXRHI/Private/AGXDynamicRHI.cpp
@@ -107,8 +107,6 @@ void FAGXDynamicRHI::RHIUnlockStagingBuffer(FRHIStagingBuffer* StagingBuffer)
 
 void FAGXDynamicRHI::RHICreateTransition(FRHITransition* Transition, const FRHITransitionCreateInfo& CreateInfo)
 {
-	checkf(FMath::IsPowerOfTwo(uint32(CreateInfo.SrcPipelines)) && FMath::IsPowerOfTwo(uint32(CreateInfo.DstPipelines)), TEXT("Support for multi-pipe resources is not yet implemented."));
-
 	// Construct the data in-place on the transition instance
 	new (Transition->GetPrivateData<FAGXTransitionData>()) FAGXTransitionData(CreateInfo.SrcPipelines, CreateInfo.DstPipelines, CreateInfo.Flags, CreateInfo.TransitionInfos);
 }
diff --git a/Engine/Source/Runtime/Apple/MetalRHI/Private/MetalDynamicRHI.cpp b/Engine/Source/Runtime/Apple/MetalRHI/Private/MetalDynamicRHI.cpp
--- a/Engine/Source/Runtime/Apple/MetalRHI/Private/MetalDynamicRHI.cpp
+++ b/Engine/Source/Runtime/Apple/MetalRHI/Private/MetalDynamicRHI.cpp
@@ -107,8 +107,6 @@ void FMetalDynamicRHI::RHIUnlockStagingBuffer(FRHIStagingBuffer* StagingBuffer)
 
 void FMetalDynamicRHI::RHICreateTransition(FRHITransition* Transition, const FRHITransitionCreateInfo& CreateInfo)
 {
-	checkf(FMath::IsPowerOfTwo(uint32(CreateInfo.SrcPipelines)) && FMath::IsPowerOfTwo(uint32(CreateInfo.DstPipelines)), TEXT("Support for multi-pipe resources is not yet implemented."));
-
 	// Construct the data in-place on the transition instance
 	new (Transition->GetPrivateData<FMetalTransitionData>()) FMetalTransitionData(CreateInfo.SrcPipelines, CreateInfo.DstPipelines, CreateInfo.Flags, CreateInfo.TransitionInfos);
 }
