Non Apple Silicon Macs:  When non sRGB SRV/UAV has been requested on a sRGB texture we need to still apply the RGB expansion or we have non compatible texture format views.

diff --git a/Engine/Source/Runtime/Apple/MetalRHI/Private/MetalTexture.cpp b/Engine/Source/Runtime/Apple/MetalRHI/Private/MetalTexture.cpp
--- a/Engine/Source/Runtime/Apple/MetalRHI/Private/MetalTexture.cpp
+++ b/Engine/Source/Runtime/Apple/MetalRHI/Private/MetalTexture.cpp
@@ -544,6 +544,13 @@ void FMetalSurface::Init(FMetalSurface& Source, NSRange MipRange, EPixelFormat F
 		{
 			// Force no sRGB, should create new view without sRGB format
 			bUseSourceTex = false;
+#if PLATFORM_MAC
+			// R8Unorm has been expanded in the source surface for sRGBA support - we need to expand to RGBA to enable compatible texture format view for non apple silicon macs
+			if (Format == PF_G8 && Source.Texture.GetPixelFormat() == mtlpp::PixelFormat::RGBA8Unorm_sRGB)
+			{
+				MetalFormat = mtlpp::PixelFormat::RGBA8Unorm;
+			}
+#endif
 		}
 		else
 		{
