Remove Metal texture resource write combined usage.   Causes problems with our buffer handing code when using buffer backed textures.

diff --git a/Engine/Source/Runtime/Apple/MetalRHI/Private/MetalBuffer.cpp b/Engine/Source/Runtime/Apple/MetalRHI/Private/MetalBuffer.cpp
--- a/Engine/Source/Runtime/Apple/MetalRHI/Private/MetalBuffer.cpp
+++ b/Engine/Source/Runtime/Apple/MetalRHI/Private/MetalBuffer.cpp
@@ -1216,7 +1216,7 @@ FMetalBuffer FMetalBufferPoolPolicyData::CreateResource(CreationArguments Args)
 
 FMetalBufferPoolPolicyData::CreationArguments FMetalBufferPoolPolicyData::GetCreationArguments(FMetalBuffer const& Resource)
 {
-	return FMetalBufferPoolPolicyData::CreationArguments(Resource.GetDevice(), Resource.GetLength(), BUF_None, Resource.GetStorageMode());
+	return FMetalBufferPoolPolicyData::CreationArguments(Resource.GetDevice(), Resource.GetLength(), BUF_None, Resource.GetStorageMode(), Resource.GetCpuCacheMode());
 }
 
 void FMetalBufferPoolPolicyData::FreeResource(FMetalBuffer& Resource)
@@ -1484,6 +1484,11 @@ FMetalBuffer FMetalResourceHeap::CreateBuffer(uint32 Size, uint32 Alignment, EBu
 	FMetalBuffer Buffer;
 	uint32 BlockSize = Align(Size, Alignment);
 	mtlpp::StorageMode StorageMode = (mtlpp::StorageMode)(((NSUInteger)Options & mtlpp::ResourceStorageModeMask) >> mtlpp::ResourceStorageModeShift);
+	mtlpp::CpuCacheMode CpuMode = (mtlpp::CpuCacheMode)(((NSUInteger)Options & mtlpp::ResourceCpuCacheModeMask) >> mtlpp::ResourceCpuCacheModeShift);
+	
+	// Write combined should be on a case by case basis
+	check(CpuMode == mtlpp::CpuCacheMode::DefaultCache);
+	
 	if (BlockSize <= 33554432)
 	{
 		switch (StorageMode)
@@ -1518,11 +1523,11 @@ FMetalBuffer FMetalResourceHeap::CreateBuffer(uint32 Size, uint32 Alignment, EBu
 				 }
 				 else
 				 {
-                    Buffer = ManagedBuffers.CreatePooledResource(FMetalPooledBufferArgs(Queue->GetDevice(), BlockSize, Flags, StorageMode));
-                     if (GMetalResourcePurgeInPool)
-                     {
-                         Buffer.SetPurgeableState(mtlpp::PurgeableState::NonVolatile);
-                     }
+                    Buffer = ManagedBuffers.CreatePooledResource(FMetalPooledBufferArgs(Queue->GetDevice(), BlockSize, Flags, StorageMode, CpuMode));
+					if (GMetalResourcePurgeInPool)
+					{
+						Buffer.SetPurgeableState(mtlpp::PurgeableState::NonVolatile);
+					}
 					DEC_MEMORY_STAT_BY(STAT_MetalBufferUnusedMemory, Buffer.GetLength());
 					DEC_MEMORY_STAT_BY(STAT_MetalPooledBufferUnusedMemory, Buffer.GetLength());
 					INC_MEMORY_STAT_BY(STAT_MetalPooledBufferMemory, Buffer.GetLength());
@@ -1595,8 +1600,7 @@ FMetalBuffer FMetalResourceHeap::CreateBuffer(uint32 Size, uint32 Alignment, EBu
 				else
 				{
 					FScopeLock Lock(&Mutex);
-					mtlpp::CpuCacheMode CpuMode = (mtlpp::CpuCacheMode)(((NSUInteger)Options & mtlpp::ResourceCpuCacheModeMask) >> mtlpp::ResourceCpuCacheModeShift);
-                    Buffer = Buffers[Storage][(NSUInteger)CpuMode].CreatePooledResource(FMetalPooledBufferArgs(Queue->GetDevice(), BlockSize, Flags, StorageMode, CpuMode));
+                    Buffer = Buffers[Storage].CreatePooledResource(FMetalPooledBufferArgs(Queue->GetDevice(), BlockSize, Flags, StorageMode, CpuMode));
 					if (GMetalResourcePurgeInPool)
 					{
                    		Buffer.SetPurgeableState(mtlpp::PurgeableState::NonVolatile);
@@ -1662,7 +1666,7 @@ void FMetalResourceHeap::ReleaseBuffer(FMetalBuffer& Buffer)
 			case mtlpp::StorageMode::Private:
 			case mtlpp::StorageMode::Shared:
 			{
-				Buffers[(NSUInteger)StorageMode][(NSUInteger)Buffer.GetCpuCacheMode()].ReleasePooledResource(Buffer);
+				Buffers[(NSUInteger)StorageMode].ReleasePooledResource(Buffer);
 				break;
 			}
 			default:
@@ -1761,10 +1765,7 @@ void FMetalResourceHeap::Compact(FMetalRenderPass* Pass, bool const bForce)
 
 	for(uint32 AllocTypeIndex = 0;AllocTypeIndex < NumAllocTypes;++AllocTypeIndex)
 	{
-		for(uint32 CacheModeIndex = 0;CacheModeIndex < NumCacheModes;++CacheModeIndex)
-		{
-			Buffers[AllocTypeIndex][CacheModeIndex].DrainPool(bForce);
-		}
+		Buffers[AllocTypeIndex].DrainPool(bForce);
 	}
 	
 #if PLATFORM_MAC
diff --git a/Engine/Source/Runtime/Apple/MetalRHI/Private/MetalBuffer.h b/Engine/Source/Runtime/Apple/MetalRHI/Private/MetalBuffer.h
--- a/Engine/Source/Runtime/Apple/MetalRHI/Private/MetalBuffer.h
+++ b/Engine/Source/Runtime/Apple/MetalRHI/Private/MetalBuffer.h
@@ -415,13 +415,6 @@ class FMetalResourceHeap
         NumUsageTypes = 2
     };
     
-    enum CPUCacheMode
-    {
-		Default,
-		WriteCombined,
-		NumCacheModes
-    };
-    
 public:
 	FMetalResourceHeap(void);
 	~FMetalResourceHeap();
@@ -461,7 +454,7 @@ private:
     TArray<FMetalSubBufferHeap*> BufferHeaps[NumUsageTypes][NumAllocTypes][NumHeapSizes];
 	
 	/** Larger buffers (up-to 32MB) that are subject to bucketing & pooling rather than sub-allocation */
-	FMetalBufferPool Buffers[NumAllocTypes][NumCacheModes];
+	FMetalBufferPool Buffers[NumAllocTypes];
 #if PLATFORM_MAC // All managed buffers are bucketed & pooled rather than sub-allocated to avoid memory consistency complexities
 	FMetalBufferPool ManagedBuffers;
 	TArray<FMetalSubBufferLinear*> ManagedSubHeaps;
diff --git a/Engine/Source/Runtime/Apple/MetalRHI/Private/MetalTexture.cpp b/Engine/Source/Runtime/Apple/MetalRHI/Private/MetalTexture.cpp
--- a/Engine/Source/Runtime/Apple/MetalRHI/Private/MetalTexture.cpp
+++ b/Engine/Source/Runtime/Apple/MetalRHI/Private/MetalTexture.cpp
@@ -788,12 +788,11 @@ FMetalSurface::FMetalSurface(ERHIResourceType ResourceType, EPixelFormat Format,
 		}
 		else if (EnumHasAnyFlags(Flags, TexCreate_NoTiling) && !EnumHasAnyFlags(Flags, TexCreate_FastVRAM|TexCreate_DepthStencilTargetable|TexCreate_RenderTargetable|TexCreate_UAV))
 		{
+			Desc.SetCpuCacheMode(mtlpp::CpuCacheMode::DefaultCache);
 #if PLATFORM_MAC
-			Desc.SetCpuCacheMode(mtlpp::CpuCacheMode::WriteCombined);
 			Desc.SetStorageMode(mtlpp::StorageMode::Managed);
-			Desc.SetResourceOptions((mtlpp::ResourceOptions)(mtlpp::ResourceOptions::CpuCacheModeWriteCombined|mtlpp::ResourceOptions::StorageModeManaged));
+			Desc.SetResourceOptions((mtlpp::ResourceOptions)(mtlpp::ResourceOptions::CpuCacheModeDefaultCache|mtlpp::ResourceOptions::StorageModeManaged));
 #else
-			Desc.SetCpuCacheMode(mtlpp::CpuCacheMode::DefaultCache);
 			Desc.SetStorageMode(mtlpp::StorageMode::Shared);
 			Desc.SetResourceOptions((mtlpp::ResourceOptions)(mtlpp::ResourceOptions::CpuCacheModeDefaultCache|mtlpp::ResourceOptions::StorageModeShared));
 #endif
@@ -821,12 +820,11 @@ FMetalSurface::FMetalSurface(ERHIResourceType ResourceType, EPixelFormat Format,
 		else
 		{
 			check(!(Flags & TexCreate_CPUReadback));
+			Desc.SetCpuCacheMode(mtlpp::CpuCacheMode::DefaultCache);
 #if PLATFORM_MAC
-			Desc.SetCpuCacheMode(mtlpp::CpuCacheMode::WriteCombined);
 			Desc.SetStorageMode(mtlpp::StorageMode::Private);
-			Desc.SetResourceOptions((mtlpp::ResourceOptions)(mtlpp::ResourceOptions::CpuCacheModeWriteCombined|mtlpp::ResourceOptions::StorageModePrivate));
+			Desc.SetResourceOptions((mtlpp::ResourceOptions)(mtlpp::ResourceOptions::CpuCacheModeDefaultCache|mtlpp::ResourceOptions::StorageModePrivate));
 #else
-			Desc.SetCpuCacheMode(mtlpp::CpuCacheMode::DefaultCache);
 			if(GMetalForceIOSTexturesShared)
 			{
 				Desc.SetStorageMode(mtlpp::StorageMode::Shared);
