Fix bytes-per-row calculation for buffer-backed textures.

- heapTextureSizeAndAlignmentWithDescriptor: does not return a meaningful value for bytes-per-row in the size field of the MTLAlignAndSize; it seems to return the buffer size (at page granularity) required to accommodate a texture based on the descriptor. For example, the API reports 4096 on Intel-based Macs for a 128x1 R32 texture.

#rnx

diff --git a/Engine/Source/Runtime/Apple/MetalRHI/Private/MetalTexture.cpp b/Engine/Source/Runtime/Apple/MetalRHI/Private/MetalTexture.cpp
--- a/Engine/Source/Runtime/Apple/MetalRHI/Private/MetalTexture.cpp
+++ b/Engine/Source/Runtime/Apple/MetalRHI/Private/MetalTexture.cpp
@@ -916,7 +916,7 @@ FMetalSurface::FMetalSurface(ERHIResourceType ResourceType, EPixelFormat Format,
 			FMetalBuffer Buffer = GetMetalDeviceContext().CreatePooledBuffer(Args);
 
 			const uint32 MinimumByteAlignment = GetMetalDeviceContext().GetDevice().GetMinimumLinearTextureAlignmentForPixelFormat(MTLFormat);
-			const NSUInteger BytesPerRow = Align(SizeAlign.Size / SizeY, MinimumByteAlignment);
+			const NSUInteger BytesPerRow = Align(Desc.GetWidth() * GPixelFormats[Format].BlockBytes, MinimumByteAlignment);
 			
 			Texture = Buffer.NewTexture(Desc, 0, BytesPerRow);
 		}
