Move the Metal device object into the global namespace.

- Reduce the footprint of the mtlpp::Device type.

#rnx

diff --git a/Engine/Source/Runtime/Apple/AGXRHI/Private/AGXBuffer.cpp b/Engine/Source/Runtime/Apple/AGXRHI/Private/AGXBuffer.cpp
--- a/Engine/Source/Runtime/Apple/AGXRHI/Private/AGXBuffer.cpp
+++ b/Engine/Source/Runtime/Apple/AGXRHI/Private/AGXBuffer.cpp
@@ -186,19 +186,19 @@ FAGXSubBufferHeap::FAGXSubBufferHeap(NSUInteger Size, NSUInteger Alignment, mtlp
 		mtlpp::HeapDescriptor Desc;
 		Desc.SetSize(FullSize);
 		Desc.SetStorageMode(Storage);
-		ParentHeap = GetAGXDeviceContext().GetDevice().NewHeap(Desc);
+		ParentHeap = GMtlppDevice.NewHeap(Desc);
 		check(ParentHeap.GetPtr());
 #if STATS || ENABLE_LOW_LEVEL_MEM_TRACKER
-		AGXLLM::LogAllocHeap(GetAGXDeviceContext().GetDevice(), ParentHeap);
+		AGXLLM::LogAllocHeap(ParentHeap);
 #endif
 	}
 	else
 	{
-		ParentBuffer = MTLPP_VALIDATE(mtlpp::Device, GetAGXDeviceContext().GetDevice(), AGXSafeGetRuntimeDebuggingLevel() >= EAGXDebugLevelValidation, NewBuffer(FullSize, Options));
+		ParentBuffer = MTLPP_VALIDATE(mtlpp::Device, GMtlppDevice, AGXSafeGetRuntimeDebuggingLevel() >= EAGXDebugLevelValidation, NewBuffer(FullSize, Options));
 		check(ParentBuffer.GetPtr());
 		check(ParentBuffer.GetLength() >= FullSize);
 #if STATS || ENABLE_LOW_LEVEL_MEM_TRACKER
-		AGXLLM::LogAllocBuffer(GetAGXDeviceContext().GetDevice(), ParentBuffer);
+		AGXLLM::LogAllocBuffer(ParentBuffer);
 #endif
 		FreeRanges.Add(ns::Range(0, FullSize));
 	}
@@ -334,18 +334,6 @@ ns::String FAGXSubBufferHeap::GetLabel() const
 	}
 }
 
-mtlpp::Device FAGXSubBufferHeap::GetDevice() const
-{
-	if (ParentHeap)
-	{
-		return ParentHeap.GetDevice();
-	}
-	else
-	{
-		return ParentBuffer.GetDevice();
-	}
-}
-
 mtlpp::StorageMode FAGXSubBufferHeap::GetStorageMode() const
 {
 	if (ParentHeap)
@@ -438,7 +426,7 @@ bool FAGXSubBufferHeap::CanAllocateSize(NSUInteger Size) const
 		NSUInteger Cache = (NSUInteger(GetCpuCacheMode()) << mtlpp::ResourceCpuCacheModeShift);
 		mtlpp::ResourceOptions Opt = mtlpp::ResourceOptions(Storage | Cache);
 		
-		NSUInteger Align = ParentHeap.GetDevice().HeapBufferSizeAndAlign(Size, Opt).Align;
+		NSUInteger Align = GMtlppDevice.HeapBufferSizeAndAlign(Size, Opt).Align;
 		return Size <= ParentHeap.MaxAvailableSizeWithAlignment(Align);
 	}
 	else
@@ -460,7 +448,7 @@ FAGXBuffer FAGXSubBufferHeap::NewBuffer(NSUInteger length)
 		
 		Result = FAGXBuffer(ParentHeap.NewBuffer(Size, Opt), this);
 #if STATS || ENABLE_LOW_LEVEL_MEM_TRACKER
-		AGXLLM::LogAllocBuffer(GetAGXDeviceContext().GetDevice(), Result);
+		AGXLLM::LogAllocBuffer(Result);
 #endif
 		DEC_MEMORY_STAT_BY(STAT_AGXBufferUnusedMemory, Result.GetLength());
 		DEC_MEMORY_STAT_BY(STAT_AGXHeapBufferUnusedMemory, Result.GetLength());
@@ -548,11 +536,11 @@ FAGXSubBufferLinear::FAGXSubBufferLinear(NSUInteger Size, NSUInteger Alignment,
 	METAL_GPUPROFILE(FAGXScopedCPUStats CPUStat(FString::Printf(TEXT("AllocBuffer: %llu, %llu"), FullSize, Options)));
 	
 	mtlpp::StorageMode Storage = (mtlpp::StorageMode)((Options & mtlpp::ResourceStorageModeMask) >> mtlpp::ResourceStorageModeShift);
-	ParentBuffer = MTLPP_VALIDATE(mtlpp::Device, GetAGXDeviceContext().GetDevice(), AGXSafeGetRuntimeDebuggingLevel() >= EAGXDebugLevelValidation, NewBuffer(FullSize, Options));
+	ParentBuffer = MTLPP_VALIDATE(mtlpp::Device, GMtlppDevice, AGXSafeGetRuntimeDebuggingLevel() >= EAGXDebugLevelValidation, NewBuffer(FullSize, Options));
 	check(ParentBuffer.GetPtr());
 	check(ParentBuffer.GetLength() >= FullSize);
 #if STATS || ENABLE_LOW_LEVEL_MEM_TRACKER
-	AGXLLM::LogAllocBuffer(GetAGXDeviceContext().GetDevice(), ParentBuffer);
+	AGXLLM::LogAllocBuffer(ParentBuffer);
 #endif
 	INC_MEMORY_STAT_BY(STAT_AGXBufferUnusedMemory, FullSize);
 	INC_MEMORY_STAT_BY(STAT_AGXLinearBufferUnusedMemory, FullSize);
@@ -595,11 +583,6 @@ ns::String FAGXSubBufferLinear::GetLabel() const
 	return ParentBuffer.GetLabel();
 }
 
-mtlpp::Device FAGXSubBufferLinear::GetDevice() const
-{
-	return ParentBuffer.GetDevice();
-}
-
 mtlpp::StorageMode FAGXSubBufferLinear::GetStorageMode() const
 {
 	return ParentBuffer.GetStorageMode();
@@ -679,7 +662,7 @@ FAGXSubBufferMagazine::FAGXSubBufferMagazine(NSUInteger Size, NSUInteger ChunkSi
     mtlpp::StorageMode Storage = (mtlpp::StorageMode)((Options & mtlpp::ResourceStorageModeMask) >> mtlpp::ResourceStorageModeShift);
     if (PLATFORM_IOS && bSupportsHeaps && Storage == mtlpp::StorageMode::Private)
     {
-        MinAlign = GetAGXDeviceContext().GetDevice().HeapBufferSizeAndAlign(BlockSize, Options).Align;
+        MinAlign = GMtlppDevice.HeapBufferSizeAndAlign(BlockSize, Options).Align;
     }
     
     NSUInteger FullSize = Align(Size, MinAlign);
@@ -694,21 +677,21 @@ FAGXSubBufferMagazine::FAGXSubBufferMagazine(NSUInteger Size, NSUInteger ChunkSi
 		mtlpp::HeapDescriptor Desc;
 		Desc.SetSize(FullSize);
 		Desc.SetStorageMode(Storage);
-		ParentHeap = GetAGXDeviceContext().GetDevice().NewHeap(Desc);
+		ParentHeap = GMtlppDevice.NewHeap(Desc);
 		check(ParentHeap.GetPtr());
 		METAL_FATAL_ASSERT(ParentHeap, TEXT("Failed to create heap of size %u and resource options %u"), Size, (uint32)Options);
 #if STATS || ENABLE_LOW_LEVEL_MEM_TRACKER
-		AGXLLM::LogAllocHeap(GetAGXDeviceContext().GetDevice(), ParentHeap);
+		AGXLLM::LogAllocHeap(ParentHeap);
 #endif
 	}
 	else
 	{
-		ParentBuffer = MTLPP_VALIDATE(mtlpp::Device, GetAGXDeviceContext().GetDevice(), AGXSafeGetRuntimeDebuggingLevel() >= EAGXDebugLevelValidation, NewBuffer(FullSize, Options));
+		ParentBuffer = MTLPP_VALIDATE(mtlpp::Device, GMtlppDevice, AGXSafeGetRuntimeDebuggingLevel() >= EAGXDebugLevelValidation, NewBuffer(FullSize, Options));
 		check(ParentBuffer.GetPtr());
 		check(ParentBuffer.GetLength() >= FullSize);
 		METAL_FATAL_ASSERT(ParentBuffer, TEXT("Failed to create heap of size %u and resource options %u"), Size, (uint32)Options);
 #if STATS || ENABLE_LOW_LEVEL_MEM_TRACKER
-		AGXLLM::LogAllocBuffer(GetAGXDeviceContext().GetDevice(), ParentBuffer);
+		AGXLLM::LogAllocBuffer(ParentBuffer);
 #endif
 		
 		INC_MEMORY_STAT_BY(STAT_AGXBufferUnusedMemory, FullSize);
@@ -774,18 +757,6 @@ ns::String FAGXSubBufferMagazine::GetLabel() const
 	}
 }
 
-mtlpp::Device FAGXSubBufferMagazine::GetDevice() const
-{
-	if (ParentHeap)
-	{
-		return ParentHeap.GetDevice();
-	}
-	else
-	{
-		return ParentBuffer.GetDevice();
-	}
-}
-
 mtlpp::StorageMode FAGXSubBufferMagazine::GetStorageMode() const
 {
 	if (ParentHeap)
@@ -881,7 +852,7 @@ FAGXBuffer FAGXSubBufferMagazine::NewBuffer()
 		
 		Result = FAGXBuffer(ParentHeap.NewBuffer(Size, Opt), this);
 #if STATS || ENABLE_LOW_LEVEL_MEM_TRACKER
-		AGXLLM::LogAllocBuffer(GetAGXDeviceContext().GetDevice(), Result);
+		AGXLLM::LogAllocBuffer(Result);
 #endif
 		DEC_MEMORY_STAT_BY(STAT_AGXBufferUnusedMemory, Result.GetLength());
 		DEC_MEMORY_STAT_BY(STAT_AGXMagazineBufferUnusedMemory, Result.GetLength());
@@ -956,10 +927,6 @@ FAGXSubBufferRing::~FAGXSubBufferRing()
 {
 }
 
-mtlpp::Device FAGXSubBufferRing::GetDevice() const
-{
-	return Buffer.IsValid() ? Buffer->Buffer.GetDevice() : nil;
-}
 mtlpp::StorageMode FAGXSubBufferRing::GetStorageMode() const
 {
 	return Buffer.IsValid() ? Buffer->Buffer.GetStorageMode() : Storage;
@@ -1196,12 +1163,11 @@ uint32 FAGXBufferPoolPolicyData::GetPoolBucketSize(uint32 Bucket)
 
 FAGXBuffer FAGXBufferPoolPolicyData::CreateResource(CreationArguments Args)
 {
-	check(Args.Device);	
 	uint32 BufferSize = GetPoolBucketSize(GetPoolBucketIndex(Args));
 	METAL_GPUPROFILE(FAGXScopedCPUStats CPUStat(FString::Printf(TEXT("AllocBuffer: %llu, %llu"), BufferSize, mtlpp::ResourceOptions(BUFFER_CACHE_MODE | ((NSUInteger)Args.Storage << mtlpp::ResourceStorageModeShift)))));
-	FAGXBuffer NewBuf(MTLPP_VALIDATE(mtlpp::Device, Args.Device, AGXSafeGetRuntimeDebuggingLevel() >= EAGXDebugLevelValidation, NewBuffer(BufferSize, FAGXCommandQueue::GetCompatibleResourceOptions(mtlpp::ResourceOptions(BUFFER_CACHE_MODE | mtlpp::ResourceOptions::HazardTrackingModeUntracked | ((NSUInteger)Args.Storage << mtlpp::ResourceStorageModeShift))))), true);
+	FAGXBuffer NewBuf(MTLPP_VALIDATE(mtlpp::Device, GMtlppDevice, AGXSafeGetRuntimeDebuggingLevel() >= EAGXDebugLevelValidation, NewBuffer(BufferSize, FAGXCommandQueue::GetCompatibleResourceOptions(mtlpp::ResourceOptions(BUFFER_CACHE_MODE | mtlpp::ResourceOptions::HazardTrackingModeUntracked | ((NSUInteger)Args.Storage << mtlpp::ResourceStorageModeShift))))), true);
 #if STATS || ENABLE_LOW_LEVEL_MEM_TRACKER
-	AGXLLM::LogAllocBuffer(Args.Device, NewBuf);
+	AGXLLM::LogAllocBuffer(NewBuf);
 #endif
 	INC_MEMORY_STAT_BY(STAT_AGXBufferUnusedMemory, NewBuf.GetLength());
 	INC_MEMORY_STAT_BY(STAT_AGXPooledBufferUnusedMemory, NewBuf.GetLength());
@@ -1210,7 +1176,7 @@ FAGXBuffer FAGXBufferPoolPolicyData::CreateResource(CreationArguments Args)
 
 FAGXBufferPoolPolicyData::CreationArguments FAGXBufferPoolPolicyData::GetCreationArguments(FAGXBuffer const& Resource)
 {
-	return FAGXBufferPoolPolicyData::CreationArguments(Resource.GetDevice(), Resource.GetLength(), BUF_None, Resource.GetStorageMode());
+	return FAGXBufferPoolPolicyData::CreationArguments(Resource.GetLength(), BUF_None, Resource.GetStorageMode());
 }
 
 void FAGXBufferPoolPolicyData::FreeResource(FAGXBuffer& Resource)
@@ -1229,7 +1195,7 @@ FAGXTexturePool::~FAGXTexturePool()
 {
 }
 
-FAGXTexture FAGXTexturePool::CreateTexture(mtlpp::Device Device, mtlpp::TextureDescriptor Desc)
+FAGXTexture FAGXTexturePool::CreateTexture(mtlpp::TextureDescriptor Desc)
 {
 	FAGXTexturePool::Descriptor Descriptor;
 	Descriptor.textureType = (NSUInteger)Desc.GetTextureType();
@@ -1263,9 +1229,9 @@ FAGXTexture FAGXTexturePool::CreateTexture(mtlpp::Device Device, mtlpp::TextureD
 	else
 	{
 		METAL_GPUPROFILE(FAGXScopedCPUStats CPUStat(FString::Printf(TEXT("AllocTexture: %s"), TEXT("")/**FString([Desc.GetPtr() description])*/)));
-		Texture = MTLPP_VALIDATE(mtlpp::Device, Device, AGXSafeGetRuntimeDebuggingLevel() >= EAGXDebugLevelValidation, NewTexture(Desc));
+		Texture = MTLPP_VALIDATE(mtlpp::Device, GMtlppDevice, AGXSafeGetRuntimeDebuggingLevel() >= EAGXDebugLevelValidation, NewTexture(Desc));
 #if STATS || ENABLE_LOW_LEVEL_MEM_TRACKER
-		AGXLLM::LogAllocTexture(Device, Desc, Texture);
+		AGXLLM::LogAllocTexture(Desc, Texture);
 #endif
 	}
 	return Texture;
@@ -1448,13 +1414,14 @@ mtlpp::Heap FAGXResourceHeap::GetTextureHeap(mtlpp::TextureDescriptor Desc, mtlp
 			}
 			if (!Result)
 			{
-				mtlpp::HeapDescriptor HeapDesc;
-				HeapDesc.SetSize(HeapTextureHeapSizes[HeapIndex]);
-				HeapDesc.SetStorageMode(Desc.GetStorageMode());
-				HeapDesc.SetCpuCacheMode(Desc.GetCpuCacheMode());
-				Result = Queue->GetDevice().NewHeap(HeapDesc);
+				MTLHeapDescriptor* HeapDesc = [[MTLHeapDescriptor alloc] init];
+				[HeapDesc setSize:HeapTextureHeapSizes[HeapIndex]];
+				[HeapDesc setStorageMode:(MTLStorageMode)Desc.GetStorageMode()];
+				[HeapDesc setCpuCacheMode:(MTLCPUCacheMode)Desc.GetCpuCacheMode()];
+				Result = mtlpp::Heap([GMtlDevice newHeapWithDescriptor:HeapDesc], nullptr, ns::Ownership::Assign);
+				[HeapDesc release];
 #if STATS || ENABLE_LOW_LEVEL_MEM_TRACKER
-				AGXLLM::LogAllocHeap(Queue->GetDevice(), Result);
+				AGXLLM::LogAllocHeap(Result);
 #endif
 				TextureHeaps[UsageIndex][HeapIndex].Add(Result);
 			}
@@ -1512,7 +1479,7 @@ FAGXBuffer FAGXResourceHeap::CreateBuffer(uint32 Size, uint32 Alignment, EBuffer
 				 }
 				 else
 				 {
-                    Buffer = ManagedBuffers.CreatePooledResource(FAGXPooledBufferArgs(Queue->GetDevice(), BlockSize, Flags, StorageMode));
+                    Buffer = ManagedBuffers.CreatePooledResource(FAGXPooledBufferArgs(BlockSize, Flags, StorageMode));
                      if (GAGXResourcePurgeInPool)
                      {
                          Buffer.SetPurgeableState(mtlpp::PurgeableState::NonVolatile);
@@ -1589,7 +1556,7 @@ FAGXBuffer FAGXResourceHeap::CreateBuffer(uint32 Size, uint32 Alignment, EBuffer
 				else
 				{
 					FScopeLock Lock(&Mutex);
-                    Buffer = Buffers[Storage].CreatePooledResource(FAGXPooledBufferArgs(Queue->GetDevice(), BlockSize, Flags, StorageMode));
+                    Buffer = Buffers[Storage].CreatePooledResource(FAGXPooledBufferArgs(BlockSize, Flags, StorageMode));
 					if (GAGXResourcePurgeInPool)
 					{
                    		Buffer.SetPurgeableState(mtlpp::PurgeableState::NonVolatile);
@@ -1610,9 +1577,10 @@ FAGXBuffer FAGXResourceHeap::CreateBuffer(uint32 Size, uint32 Alignment, EBuffer
 	else
 	{
 		METAL_GPUPROFILE(FAGXScopedCPUStats CPUStat(FString::Printf(TEXT("AllocBuffer: %llu, %llu"), BlockSize, Options)));
-		Buffer = FAGXBuffer(MTLPP_VALIDATE(mtlpp::Device, Queue->GetDevice(), AGXSafeGetRuntimeDebuggingLevel() >= EAGXDebugLevelValidation, NewBuffer(BlockSize, Options)), false);
+		mtlpp::Buffer MtlppBuffer([GMtlDevice newBufferWithLength:BlockSize options:(MTLResourceOptions)Options], nullptr, ns::Ownership::Assign);
+		Buffer = FAGXBuffer(MoveTemp(MtlppBuffer), false);
 #if STATS || ENABLE_LOW_LEVEL_MEM_TRACKER
-		AGXLLM::LogAllocBuffer(Queue->GetDevice(), Buffer);
+		AGXLLM::LogAllocBuffer(Buffer);
 #endif
 		INC_MEMORY_STAT_BY(STAT_AGXDeviceBufferMemory, Buffer.GetLength());
 	}
@@ -1681,25 +1649,25 @@ FAGXTexture FAGXResourceHeap::CreateTexture(mtlpp::TextureDescriptor Desc, FAGXS
 	LLM_SCOPE_METAL(ELLMTagAGX::Textures);
 	LLM_PLATFORM_SCOPE_METAL(ELLMTagAGX::Textures);
 	
-	mtlpp::SizeAndAlign Res = Queue->GetDevice().HeapTextureSizeAndAlign(Desc);
+	mtlpp::SizeAndAlign Res = GMtlppDevice.HeapTextureSizeAndAlign(Desc);
 	mtlpp::Heap Heap = GetTextureHeap(Desc, Res);
 	if (Heap)
 	{
 		METAL_GPUPROFILE(FAGXScopedCPUStats CPUStat(FString::Printf(TEXT("AllocTexture: %s"), TEXT("")/**FString([Desc.GetPtr() description])*/)));
 		FAGXTexture Texture = Heap.NewTexture(Desc);
 #if STATS || ENABLE_LOW_LEVEL_MEM_TRACKER
-		AGXLLM::LogAllocTexture(Queue->GetDevice(), Desc, Texture);
+		AGXLLM::LogAllocTexture(Desc, Texture);
 #endif
 		return Texture;
 	}
 	else if (Desc.GetUsage() & mtlpp::TextureUsage::RenderTarget)
 	{
 		LLM_PLATFORM_SCOPE_METAL(ELLMTagAGX::RenderTargets);
-		return TargetPool.CreateTexture(Queue->GetDevice(), Desc);
+		return TargetPool.CreateTexture(Desc);
 	}
 	else
 	{
-		return TexturePool.CreateTexture(Queue->GetDevice(), Desc);
+		return TexturePool.CreateTexture(Desc);
 	}
 }
 
diff --git a/Engine/Source/Runtime/Apple/AGXRHI/Private/AGXBuffer.h b/Engine/Source/Runtime/Apple/AGXRHI/Private/AGXBuffer.h
--- a/Engine/Source/Runtime/Apple/AGXRHI/Private/AGXBuffer.h
+++ b/Engine/Source/Runtime/Apple/AGXRHI/Private/AGXBuffer.h
@@ -11,18 +11,16 @@
 
 struct FAGXPooledBufferArgs
 {
-    FAGXPooledBufferArgs() : Device(nil), Size(0), Flags(BUF_None), Storage(mtlpp::StorageMode::Shared) {}
+    FAGXPooledBufferArgs() : Size(0), Flags(BUF_None), Storage(mtlpp::StorageMode::Shared) {}
 	
-    FAGXPooledBufferArgs(mtlpp::Device InDevice, uint32 InSize, EBufferUsageFlags InFlags, mtlpp::StorageMode InStorage, mtlpp::CpuCacheMode InCpuCacheMode = mtlpp::CpuCacheMode::DefaultCache)
-	: Device(InDevice)
-	, Size(InSize)
+    FAGXPooledBufferArgs(uint32 InSize, EBufferUsageFlags InFlags, mtlpp::StorageMode InStorage, mtlpp::CpuCacheMode InCpuCacheMode = mtlpp::CpuCacheMode::DefaultCache)
+	: Size(InSize)
     , Flags(InFlags)
 	, Storage(InStorage)
 	, CpuCacheMode(InCpuCacheMode)
 	{
 	}
 	
-	mtlpp::Device Device;
 	uint32 Size;
 	EBufferUsageFlags Flags;
 	mtlpp::StorageMode Storage;
@@ -38,7 +36,6 @@ public:
 	~FAGXSubBufferHeap();
 	
 	ns::String   GetLabel() const;
-    mtlpp::Device       GetDevice() const;
     mtlpp::StorageMode  GetStorageMode() const;
     mtlpp::CpuCacheMode GetCpuCacheMode() const;
     NSUInteger     GetSize() const;
@@ -80,7 +77,6 @@ public:
 	~FAGXSubBufferLinear();
 	
 	ns::String   GetLabel() const;
-	mtlpp::Device       GetDevice() const;
 	mtlpp::StorageMode  GetStorageMode() const;
 	mtlpp::CpuCacheMode GetCpuCacheMode() const;
 	NSUInteger     GetSize() const;
@@ -109,7 +105,6 @@ public:
 	~FAGXSubBufferMagazine();
 	
 	ns::String   GetLabel() const;
-    mtlpp::Device       GetDevice() const;
     mtlpp::StorageMode  GetStorageMode() const;
     mtlpp::CpuCacheMode GetCpuCacheMode() const;
     NSUInteger     GetSize() const;
@@ -153,7 +148,6 @@ public:
 	FAGXSubBufferRing(NSUInteger Size, NSUInteger Alignment, mtlpp::ResourceOptions Options);
 	~FAGXSubBufferRing();
 	
-	mtlpp::Device       GetDevice() const;
 	mtlpp::StorageMode  GetStorageMode() const;
 	mtlpp::CpuCacheMode GetCpuCacheMode() const;
 	NSUInteger     GetSize() const;
@@ -343,7 +337,7 @@ public:
 	FAGXTexturePool(FCriticalSection& PoolMutex);
 	~FAGXTexturePool();
 	
-	FAGXTexture CreateTexture(mtlpp::Device Device, mtlpp::TextureDescriptor Desc);
+	FAGXTexture CreateTexture(mtlpp::TextureDescriptor Desc);
 	void ReleaseTexture(FAGXTexture& Texture);
 	
 	void Drain(bool const bForce);
diff --git a/Engine/Source/Runtime/Apple/AGXRHI/Private/AGXCaptureManager.cpp b/Engine/Source/Runtime/Apple/AGXRHI/Private/AGXCaptureManager.cpp
--- a/Engine/Source/Runtime/Apple/AGXRHI/Private/AGXCaptureManager.cpp
+++ b/Engine/Source/Runtime/Apple/AGXRHI/Private/AGXCaptureManager.cpp
@@ -7,16 +7,15 @@
 
 bool GAGXSupportsCaptureManager = false;
 
-FAGXCaptureManager::FAGXCaptureManager(mtlpp::Device InDevice, FAGXCommandQueue& InQueue)
-: Device(InDevice)
-, Queue(InQueue)
+FAGXCaptureManager::FAGXCaptureManager(FAGXCommandQueue& InQueue)
+: Queue(InQueue)
 {
 	MTLPP_IF_AVAILABLE(10.13, 11.0, 11.0)
 	{
 		GAGXSupportsCaptureManager = true;
 		
 		mtlpp::CaptureManager Manager = mtlpp::CaptureManager::SharedCaptureManager();
-		Manager.SetDefaultCaptureScope(Manager.NewCaptureScopeWithDevice(Device));
+		Manager.SetDefaultCaptureScope(Manager.NewCaptureScopeWithDevice(GMtlppDevice));
 		Manager.GetDefaultCaptureScope().SetLabel(@"1 Frame");
 		
 		FAGXCaptureScope DefaultScope;
@@ -31,7 +30,7 @@ FAGXCaptureManager::FAGXCaptureManager(mtlpp::Device InDevice, FAGXCommandQueue&
 		for (uint32 i = 0; i < (sizeof(PresentStepCounts) / sizeof(uint32)); i++)
 		{
 			FAGXCaptureScope Scope;
-			Scope.MTLScope = Manager.NewCaptureScopeWithDevice(Device);
+			Scope.MTLScope = Manager.NewCaptureScopeWithDevice(GMtlppDevice);
 			Scope.MTLScope.SetLabel(FString::Printf(TEXT("%u Frames"), PresentStepCounts[i]).GetNSString());
 			Scope.Type = EAGXCaptureTypePresent;
 			Scope.StepCount = PresentStepCounts[i];
@@ -80,7 +79,7 @@ void FAGXCaptureManager::BeginCapture(void)
 {
 	if (GAGXSupportsCaptureManager)
 	{
-		mtlpp::CaptureManager::SharedCaptureManager().StartCaptureWithDevice(Device);
+		mtlpp::CaptureManager::SharedCaptureManager().StartCaptureWithDevice(GMtlppDevice);
 	}
 }
 
diff --git a/Engine/Source/Runtime/Apple/AGXRHI/Private/AGXCaptureManager.h b/Engine/Source/Runtime/Apple/AGXRHI/Private/AGXCaptureManager.h
--- a/Engine/Source/Runtime/Apple/AGXRHI/Private/AGXCaptureManager.h
+++ b/Engine/Source/Runtime/Apple/AGXRHI/Private/AGXCaptureManager.h
@@ -11,7 +11,7 @@ class FAGXCommandQueue;
 class FAGXCaptureManager
 {
 public:
-	FAGXCaptureManager(mtlpp::Device Device, FAGXCommandQueue& Queue);
+	FAGXCaptureManager(FAGXCommandQueue& Queue);
 	~FAGXCaptureManager();
 	
 	// Called by the AGXRHI code to trigger the provided capture scopes visible in Xcode.
@@ -23,7 +23,6 @@ public:
 	void EndCapture(void);
 	
 private:
-	mtlpp::Device Device;
 	FAGXCommandQueue& Queue;
 	bool bSupportsCaptureManager;
 	
diff --git a/Engine/Source/Runtime/Apple/AGXRHI/Private/AGXCommandQueue.cpp b/Engine/Source/Runtime/Apple/AGXRHI/Private/AGXCommandQueue.cpp
--- a/Engine/Source/Runtime/Apple/AGXRHI/Private/AGXCommandQueue.cpp
+++ b/Engine/Source/Runtime/Apple/AGXRHI/Private/AGXCommandQueue.cpp
@@ -21,9 +21,8 @@ bool GAGXCommandBufferDebuggingEnabled = 0;
 
 #pragma mark - Public C++ Boilerplate -
 
-FAGXCommandQueue::FAGXCommandQueue(mtlpp::Device InDevice, uint32 const MaxNumCommandBuffers /* = 0 */)
-: Device(InDevice)
-, ParallelCommandLists(0)
+FAGXCommandQueue::FAGXCommandQueue(uint32 const MaxNumCommandBuffers /* = 0 */)
+: ParallelCommandLists(0)
 , RuntimeDebuggingLevel(EAGXDebugLevelOff)
 {
 	int32 MaxShaderVersion = 0;
@@ -50,11 +49,11 @@ FAGXCommandQueue::FAGXCommandQueue(mtlpp::Device InDevice, uint32 const MaxNumCo
 
 	if(MaxNumCommandBuffers == 0)
 	{
-		CommandQueue = Device.NewCommandQueue();
+		CommandQueue = [GMtlDevice newCommandQueue];
 	}
 	else
 	{
-		CommandQueue = Device.NewCommandQueue(MaxNumCommandBuffers);
+		CommandQueue = [GMtlDevice newCommandQueueWithMaxCommandBufferCount:MaxNumCommandBuffers];
 	}
 	check(CommandQueue);
 #if PLATFORM_IOS
@@ -66,7 +65,7 @@ FAGXCommandQueue::FAGXCommandQueue(mtlpp::Device InDevice, uint32 const MaxNumCo
 #if PLATFORM_TVOS
         Features &= ~(EAGXFeaturesSetBytes);
 		
-		if(Device.SupportsFeatureSet(mtlpp::FeatureSet::tvOS_GPUFamily2_v1))
+		if ([GMtlDevice supportsFeatureSet:MTLFeatureSet_tvOS_GPUFamily2_v1])
 		{
 			Features |= EAGXFeaturesCountingQueries | EAGXFeaturesBaseVertexInstance | EAGXFeaturesIndirectBuffer | EAGXFeaturesMSAADepthResolve | EAGXFeaturesMSAAStoreAndResolve;
 		}
@@ -106,12 +105,12 @@ FAGXCommandQueue::FAGXCommandQueue(mtlpp::Device InDevice, uint32 const MaxNumCo
 			}
 		}
 #else
-		if (Device.SupportsFeatureSet(mtlpp::FeatureSet::iOS_GPUFamily3_v1))
+		if ([GMtlDevice supportsFeatureSet:MTLFeatureSet_iOS_GPUFamily3_v1])
 		{
 			Features |= EAGXFeaturesCountingQueries | EAGXFeaturesBaseVertexInstance | EAGXFeaturesIndirectBuffer | EAGXFeaturesMSAADepthResolve;
 		}
 		
-		if(Device.SupportsFeatureSet(mtlpp::FeatureSet::iOS_GPUFamily3_v2) || Device.SupportsFeatureSet(mtlpp::FeatureSet::iOS_GPUFamily2_v3) || Device.SupportsFeatureSet(mtlpp::FeatureSet::iOS_GPUFamily1_v3))
+		if ([GMtlDevice supportsFeatureSet:MTLFeatureSet_iOS_GPUFamily3_v2] || [GMtlDevice supportsFeatureSet:MTLFeatureSet_iOS_GPUFamily2_v3] || [GMtlDevice supportsFeatureSet:MTLFeatureSet_iOS_GPUFamily1_v3])
 		{
 			if (FParse::Param(FCommandLine::Get(),TEXT("metalfence")))
 			{
@@ -124,7 +123,7 @@ FAGXCommandQueue::FAGXCommandQueue(mtlpp::Device InDevice, uint32 const MaxNumCo
 			}
 		}
 		
-		if(Device.SupportsFeatureSet(mtlpp::FeatureSet::iOS_GPUFamily3_v2))
+		if ([GMtlDevice supportsFeatureSet:MTLFeatureSet_iOS_GPUFamily3_v2])
 		{
 			Features |= EAGXFeaturesMSAAStoreAndResolve;
 		}
@@ -166,7 +165,7 @@ FAGXCommandQueue::FAGXCommandQueue(mtlpp::Device InDevice, uint32 const MaxNumCo
 						Features |= EAGXFeaturesTextureBuffers;
 					}
                     
-                    if(Device.SupportsFeatureSet(mtlpp::FeatureSet::iOS_GPUFamily4_v1))
+                    if ([GMtlDevice supportsFeatureSet:MTLFeatureSet_iOS_GPUFamily4_v1])
                     {
                         Features |= EAGXFeaturesTileShaders;
                         
@@ -175,7 +174,7 @@ FAGXCommandQueue::FAGXCommandQueue(mtlpp::Device InDevice, uint32 const MaxNumCo
                         GAGXCommandBufferDebuggingEnabled = (GPUCrashDebuggingCVar && GPUCrashDebuggingCVar->GetInt() != 0) || FParse::Param(FCommandLine::Get(),TEXT("metalgpudebug"));
                     }
                     
-					if (Device.SupportsFeatureSet(mtlpp::FeatureSet::iOS_GPUFamily5_v1))
+					if ([GMtlDevice supportsFeatureSet:MTLFeatureSet_iOS_GPUFamily5_v1])
 					{
 						Features |= EAGXFeaturesLayeredRendering;
 					}
@@ -189,51 +188,30 @@ FAGXCommandQueue::FAGXCommandQueue(mtlpp::Device InDevice, uint32 const MaxNumCo
 		Features = EAGXFeaturesSetBufferOffset;
 	}
 #else // Assume that Mac & other platforms all support these from the start. They can diverge later.
-	const bool bIsNVIDIA = [Device.GetName().GetPtr() rangeOfString:@"Nvidia" options:NSCaseInsensitiveSearch].location != NSNotFound;
+	const bool bIsNVIDIA = [[GMtlDevice name] rangeOfString:@"Nvidia" options:NSCaseInsensitiveSearch].location != NSNotFound;
 	Features = EAGXFeaturesCountingQueries | EAGXFeaturesBaseVertexInstance | EAGXFeaturesIndirectBuffer | EAGXFeaturesLayeredRendering | EAGXFeaturesCubemapArrays;
 	if (!bIsNVIDIA)
 	{
 		Features |= EAGXFeaturesSetBufferOffset;
 	}
-	if (Device.SupportsFeatureSet(mtlpp::FeatureSet::macOS_GPUFamily1_v2))
+	if ([GMtlDevice supportsFeatureSet:MTLFeatureSet_macOS_GPUFamily1_v2])
     {
-        Features |= EAGXFeaturesMSAADepthResolve | EAGXFeaturesMSAAStoreAndResolve;
-        
-        // Assume that set*Bytes only works on macOS Sierra and above as no-one has tested it anywhere else.
-		Features |= EAGXFeaturesSetBytes;
-		
-		FString DeviceName(Device.GetName());
-		// On earlier OS versions Intel Broadwell couldn't suballocate properly
-		if (!(DeviceName.Contains(TEXT("Intel")) && (DeviceName.Contains(TEXT("5300")) || DeviceName.Contains(TEXT("6000")) || DeviceName.Contains(TEXT("6100")))) || FPlatformMisc::MacOSXVersionCompare(10,14,0) >= 0)
-		{
-			// Using Private Memory & BlitEncoders for Vertex & Index data should be *much* faster.
-        	Features |= EAGXFeaturesEfficientBufferBlits;
-        	
-			Features |= EAGXFeaturesBufferSubAllocation;
-					
-	        // On earlier OS versions Vega didn't like non-zero blit offsets
-	        if (!DeviceName.Contains(TEXT("Vega")) || FPlatformMisc::MacOSXVersionCompare(10,13,5) >= 0)
-	        {
-				Features |= EAGXFeaturesPrivateBufferSubAllocation;
-			}
-		}
+        Features |= (   EAGXFeaturesSetBytes
+					  | EAGXFeaturesMSAADepthResolve
+					  | EAGXFeaturesMSAAStoreAndResolve
+					  | EAGXFeaturesEfficientBufferBlits
+					  | EAGXFeaturesBufferSubAllocation
+					  | EAGXFeaturesPrivateBufferSubAllocation
+					  | EAGXFeaturesMaxThreadsPerThreadgroup );
 		
 		GAGXFColorVertexFormat = mtlpp::VertexFormat::UChar4Normalized_BGRA;
 		
-		// On 10.13.5+ we can use MTLParallelRenderEncoder
-		if (FPlatformMisc::MacOSXVersionCompare(10,13,5) >= 0)
+		if ([[GMtlDevice name] rangeOfString:@"Nvidia" options:NSCaseInsensitiveSearch].location == NSNotFound && !FParse::Param(FCommandLine::Get(),TEXT("nometalparallelencoder")))
 		{
-			// Except on Nvidia for the moment
-			if ([Device.GetName().GetPtr() rangeOfString:@"Nvidia" options:NSCaseInsensitiveSearch].location == NSNotFound && !FParse::Param(FCommandLine::Get(),TEXT("nometalparallelencoder")))
-			{
-				Features |= EAGXFeaturesParallelRenderEncoders;
-			}
+			Features |= EAGXFeaturesParallelRenderEncoders;
 		}
 
-		// Turn on Texture Buffers! These are faster on the GPU as we don't need to do out-of-bounds tests but require Metal 2.1 and macOS 10.14
-		if (FPlatformMisc::MacOSXVersionCompare(10,14,0) >= 0)
 		{
-			Features |= EAGXFeaturesMaxThreadsPerThreadgroup;
 			if (MaxShaderVersion >= 4)
 			{
 				Features |= EAGXFeaturesTextureBuffers;
@@ -270,13 +248,13 @@ FAGXCommandQueue::FAGXCommandQueue(mtlpp::Device InDevice, uint32 const MaxNumCo
 			}
 		}
     }
-    else if ([Device.GetName().GetPtr() rangeOfString:@"Nvidia" options:NSCaseInsensitiveSearch].location != NSNotFound)
+    else if ([[GMtlDevice name] rangeOfString:@"Nvidia" options:NSCaseInsensitiveSearch].location != NSNotFound)
     {
 		// Using set*Bytes fixes bugs on Nvidia for 10.11 so we should use it...
     	Features |= EAGXFeaturesSetBytes;
     }
     
-    if(Device.SupportsFeatureSet(mtlpp::FeatureSet::macOS_GPUFamily1_v3) && FPlatformMisc::MacOSXVersionCompare(10,13,0) >= 0)
+    if ([GMtlDevice supportsFeatureSet:MTLFeatureSet_macOS_GPUFamily1_v3])
     {
         Features |= EAGXFeaturesMultipleViewports | EAGXFeaturesPipelineBufferMutability | EAGXFeaturesGPUCaptureManager;
 		
@@ -299,7 +277,7 @@ FAGXCommandQueue::FAGXCommandQueue(mtlpp::Device InDevice, uint32 const MaxNumCo
 	
 #if !UE_BUILD_SHIPPING
 	Class MTLDebugDevice = NSClassFromString(@"MTLDebugDevice");
-	if ([Device isKindOfClass:MTLDebugDevice])
+	if ([GMtlDevice isKindOfClass:MTLDebugDevice])
 	{
 		Features |= EAGXFeaturesValidation;
 	}
@@ -332,7 +310,11 @@ FAGXCommandQueue::FAGXCommandQueue(mtlpp::Device InDevice, uint32 const MaxNumCo
 
 FAGXCommandQueue::~FAGXCommandQueue(void)
 {
-	// void
+	if (CommandQueue != nil)
+	{
+		[CommandQueue release];
+		CommandQueue = nil;
+	}
 }
 	
 #pragma mark - Public Command Buffer Mutators -
@@ -342,8 +324,8 @@ mtlpp::CommandBuffer FAGXCommandQueue::CreateCommandBuffer(void)
 #if PLATFORM_MAC
 	static bool bUnretainedRefs = FParse::Param(FCommandLine::Get(),TEXT("metalunretained"))
 	|| (!FParse::Param(FCommandLine::Get(),TEXT("metalretainrefs"))
-		&& ([Device.GetName() rangeOfString:@"Nvidia" options:NSCaseInsensitiveSearch].location == NSNotFound)
-		&& ([Device.GetName() rangeOfString:@"Intel" options:NSCaseInsensitiveSearch].location == NSNotFound));
+		&& ([[GMtlDevice name] rangeOfString:@"Nvidia" options:NSCaseInsensitiveSearch].location == NSNotFound)
+		&& ([[GMtlDevice name] rangeOfString:@"Intel"  options:NSCaseInsensitiveSearch].location == NSNotFound));
 #else
 	static bool bUnretainedRefs = !FParse::Param(FCommandLine::Get(),TEXT("metalretainrefs"));
 #endif
@@ -351,7 +333,10 @@ mtlpp::CommandBuffer FAGXCommandQueue::CreateCommandBuffer(void)
 	mtlpp::CommandBuffer CmdBuffer;
 	@autoreleasepool
 	{
-		CmdBuffer = bUnretainedRefs ? MTLPP_VALIDATE(mtlpp::CommandQueue, CommandQueue, AGXSafeGetRuntimeDebuggingLevel() >= EAGXDebugLevelValidation, CommandBufferWithUnretainedReferences()) : MTLPP_VALIDATE(mtlpp::CommandQueue, CommandQueue, AGXSafeGetRuntimeDebuggingLevel() >= EAGXDebugLevelValidation, CommandBuffer());
+		id<MTLCommandBuffer> MtlCmdBuf = bUnretainedRefs ? [CommandQueue commandBufferWithUnretainedReferences] : [CommandQueue commandBuffer];
+		check(MtlCmdBuf);
+		
+		CmdBuffer = mtlpp::CommandBuffer(MtlCmdBuf);
 		
 		if (RuntimeDebuggingLevel > EAGXDebugLevelOff)
 		{			
@@ -446,11 +431,6 @@ void FAGXCommandQueue::GetCommittedCommandBufferFences(TArray<mtlpp::CommandBuff
 }
 
 #pragma mark - Public Command Queue Accessors -
-	
-mtlpp::Device& FAGXCommandQueue::GetDevice(void)
-{
-	return Device;
-}
 
 mtlpp::ResourceOptions FAGXCommandQueue::GetCompatibleResourceOptions(mtlpp::ResourceOptions Options)
 {
diff --git a/Engine/Source/Runtime/Apple/AGXRHI/Private/AGXCommandQueue.h b/Engine/Source/Runtime/Apple/AGXRHI/Private/AGXCommandQueue.h
--- a/Engine/Source/Runtime/Apple/AGXRHI/Private/AGXCommandQueue.h
+++ b/Engine/Source/Runtime/Apple/AGXRHI/Private/AGXCommandQueue.h
@@ -90,7 +90,7 @@ public:
 	 * @param Device The Metal device to create on.
 	 * @param MaxNumCommandBuffers The maximum number of incomplete command-buffers, defaults to 0 which implies the system default.
 	 */
-	FAGXCommandQueue(mtlpp::Device Device, uint32 const MaxNumCommandBuffers = 0);
+	FAGXCommandQueue(uint32 const MaxNumCommandBuffers = 0);
 	
 	/** Destructor */
 	~FAGXCommandQueue(void);
@@ -126,9 +126,6 @@ public:
 	
 #pragma mark - Public Command Queue Accessors -
 	
-	/** @returns The command queue's native device. */
-	mtlpp::Device& GetDevice(void);
-	
 	/** Converts a Metal v1.1+ resource option to something valid on the current version. */
 	static mtlpp::ResourceOptions GetCompatibleResourceOptions(mtlpp::ResourceOptions Options);
 	
@@ -158,8 +155,7 @@ public:
 
 private:
 #pragma mark - Private Member Variables -
-	mtlpp::Device Device;
-	mtlpp::CommandQueue CommandQueue;
+	id<MTLCommandQueue> CommandQueue;
 	TArray<TArray<mtlpp::CommandBuffer>> CommandBuffers;
 	TLockFreePointerListLIFO<mtlpp::CommandBufferFence> CommandBufferFences;
 	uint64 ParallelCommandLists;
diff --git a/Engine/Source/Runtime/Apple/AGXRHI/Private/AGXContext.cpp b/Engine/Source/Runtime/Apple/AGXRHI/Private/AGXContext.cpp
--- a/Engine/Source/Runtime/Apple/AGXRHI/Private/AGXContext.cpp
+++ b/Engine/Source/Runtime/Apple/AGXRHI/Private/AGXContext.cpp
@@ -147,44 +147,49 @@ static FAutoConsoleVariableRef CVarAGXDefaultTransferAllocation(
 	TEXT("Default size of a single entry in the upload pool."));
 
 
+//------------------------------------------------------------------------------
+
+#pragma mark - AGXRHI Private Globals
+
+
+id<MTLDevice> GMtlDevice = nil;
+
+// Placeholder: TODO: remove
+mtlpp::Device GMtlppDevice;
+
+
+//------------------------------------------------------------------------------
+
+#pragma mark - AGX Device Context Support Routines
+
 
 #if PLATFORM_MAC
-static ns::AutoReleased<ns::Object<id <NSObject>>> GAGXDeviceObserver;
-static mtlpp::Device GetMTLDevice(uint32& DeviceIndex)
+static id<NSObject> GAGXDeviceObserver = nil;
+
+static id<MTLDevice> GetMTLDevice(uint32& DeviceIndex)
 {
 #if PLATFORM_MAC_ARM64
-    return mtlpp::Device::CreateSystemDefaultDevice();
+    return MTLCreateSystemDefaultDevice();
 #else
 	SCOPED_AUTORELEASE_POOL;
 	
 	DeviceIndex = 0;
 	
-	ns::Array<mtlpp::Device> DeviceList;
-	
-	if (FPlatformMisc::MacOSXVersionCompare(10, 13, 4) >= 0)
-	{
-			DeviceList = mtlpp::Device::CopyAllDevicesWithObserver(GAGXDeviceObserver, ^(const mtlpp::Device & Device, const ns::String & Notification)
-			{
-				if ([Notification.GetPtr() isEqualToString:MTLDeviceWasAddedNotification])
-				{
-					FPlatformMisc::GPUChangeNotification(Device.GetRegistryID(), FPlatformMisc::EMacGPUNotification::Added);
-				}
-				else if ([Notification.GetPtr() isEqualToString:MTLDeviceRemovalRequestedNotification])
-				{
-					FPlatformMisc::GPUChangeNotification(Device.GetRegistryID(), FPlatformMisc::EMacGPUNotification::RemovalRequested);
-				}
-				else if ([Notification.GetPtr() isEqualToString:MTLDeviceWasRemovedNotification])
-				{
-					FPlatformMisc::GPUChangeNotification(Device.GetRegistryID(), FPlatformMisc::EMacGPUNotification::Removed);
-				}
-			});
-	}
-	else
+	NSArray<id<MTLDevice>>* DeviceList = MTLCopyAllDevicesWithObserver(&GAGXDeviceObserver, ^(id<MTLDevice> Device, MTLDeviceNotificationName Notification)
 	{
-		DeviceList = mtlpp::Device::CopyAllDevices();
-	}
-	
-	const int32 NumDevices = DeviceList.GetSize();
+		if ([Notification isEqualToString:MTLDeviceWasAddedNotification])
+		{
+			FPlatformMisc::GPUChangeNotification([Device registryID], FPlatformMisc::EMacGPUNotification::Added);
+		}
+		else if ([Notification isEqualToString:MTLDeviceRemovalRequestedNotification])
+		{
+			FPlatformMisc::GPUChangeNotification([Device registryID], FPlatformMisc::EMacGPUNotification::RemovalRequested);
+		}
+		else if ([Notification isEqualToString:MTLDeviceWasRemovedNotification])
+		{
+			FPlatformMisc::GPUChangeNotification([Device registryID], FPlatformMisc::EMacGPUNotification::Removed);
+		}
+	});
 	
 	TArray<FMacPlatformMisc::FGPUDescriptor> const& GPUs = FPlatformMisc::GetGPUDescriptors();
 	check(GPUs.Num() > 0);
@@ -194,62 +199,45 @@ static mtlpp::Device GetMTLDevice(uint32& DeviceIndex)
 	//        we drop support for 10.12
 	//  NOTE: this means any implementation of GetGraphicsAdapterLuid() for Mac should return an index, and use -1 as a 
 	//        sentinel value representing "no device" (instead of 0, which is used in the LUID case)
-	int32 HmdGraphicsAdapter  = IHeadMountedDisplayModule::IsAvailable() ? (int32)IHeadMountedDisplayModule::Get().GetGraphicsAdapterLuid() : -1;
+	int32 HmdGraphicsAdapter = IHeadMountedDisplayModule::IsAvailable() ? (int32)IHeadMountedDisplayModule::Get().GetGraphicsAdapterLuid() : -1;
  	int32 OverrideRendererId = FPlatformMisc::GetExplicitRendererIndex();
-	
 	int32 ExplicitRendererId = OverrideRendererId >= 0 ? OverrideRendererId : HmdGraphicsAdapter;
-	if(ExplicitRendererId < 0 && GPUs.Num() > 1 && FMacPlatformMisc::MacOSXVersionCompare(10, 11, 5) == 0)
-	{
-		OverrideRendererId = -1;
-		bool bForceExplicitRendererId = false;
-		for(uint32 i = 0; i < GPUs.Num(); i++)
-		{
-			FMacPlatformMisc::FGPUDescriptor const& GPU = GPUs[i];
-			if((GPU.GPUVendorId == 0x10DE))
-			{
-				OverrideRendererId = i;
-				bForceExplicitRendererId = (GPU.GPUMetalBundle && ![GPU.GPUMetalBundle isEqualToString:@"GeForceMTLDriverWeb"]);
-			}
-			else if(!GPU.GPUHeadless && GPU.GPUVendorId != 0x8086)
-			{
-				OverrideRendererId = i;
-			}
-		}
-		if (bForceExplicitRendererId)
-		{
-			ExplicitRendererId = OverrideRendererId;
-		}
-	}
-	
-	mtlpp::Device SelectedDevice;
+
+	id<MTLDevice> SelectedDevice = nil;
 	if (ExplicitRendererId >= 0 && ExplicitRendererId < GPUs.Num())
 	{
 		FMacPlatformMisc::FGPUDescriptor const& GPU = GPUs[ExplicitRendererId];
 		TArray<FString> NameComponents;
-		FString(GPU.GPUName).TrimStart().ParseIntoArray(NameComponents, TEXT(" "));	
-		for (uint32 index = 0; index < NumDevices; index++)
+		FString(GPU.GPUName).TrimStart().ParseIntoArray(NameComponents, TEXT(" "));
+		NSUInteger NumDevices = [DeviceList count];
+		for (NSUInteger index = 0; index < NumDevices; ++index)
 		{
-			mtlpp::Device Device = DeviceList[index];
-			
-			if(MTLPP_CHECK_AVAILABLE(10.13, 11.0, 11.0) && (Device.GetRegistryID() == GPU.RegistryID))
+			id<MTLDevice> Device = [DeviceList objectAtIndex:index];
+
+			if ([Device registryID] == GPU.RegistryID)
 			{
 				DeviceIndex = ExplicitRendererId;
 				SelectedDevice = Device;
 			}
-			else if(([Device.GetName().GetPtr() rangeOfString:@"Nvidia" options:NSCaseInsensitiveSearch].location != NSNotFound && GPU.GPUVendorId == 0x10DE)
-			   || ([Device.GetName().GetPtr() rangeOfString:@"AMD" options:NSCaseInsensitiveSearch].location != NSNotFound && GPU.GPUVendorId == 0x1002)
-			   || ([Device.GetName().GetPtr() rangeOfString:@"Intel" options:NSCaseInsensitiveSearch].location != NSNotFound && GPU.GPUVendorId == 0x8086))
+			else
 			{
-				bool bMatchesName = (NameComponents.Num() > 0);
-				for (FString& Component : NameComponents)
+				NSString *DeviceName = [Device name];
+
+				if (    (([DeviceName rangeOfString:@"Nvidia" options:NSCaseInsensitiveSearch].location != NSNotFound) && (GPU.GPUVendorId == 0x10DE))
+					 || (([DeviceName rangeOfString:@"AMD"    options:NSCaseInsensitiveSearch].location != NSNotFound) && (GPU.GPUVendorId == 0x1002))
+					 || (([DeviceName rangeOfString:@"Intel"  options:NSCaseInsensitiveSearch].location != NSNotFound) && (GPU.GPUVendorId == 0x8086)) )
 				{
-					bMatchesName &= FString(Device.GetName().GetPtr()).Contains(Component);
-				}
-				if((Device.IsHeadless() == GPU.GPUHeadless || GPU.GPUVendorId != 0x1002) && bMatchesName)
-                {
-					DeviceIndex = ExplicitRendererId;
-					SelectedDevice = Device;
-					break;
+					bool bMatchesName = (NameComponents.Num() > 0);
+					for (FString& Component : NameComponents)
+					{
+						bMatchesName &= FString(DeviceName).Contains(Component);
+					}
+					if (([Device isHeadless] == GPU.GPUHeadless || GPU.GPUVendorId != 0x1002) && bMatchesName)
+					{
+						DeviceIndex = ExplicitRendererId;
+						SelectedDevice = Device;
+						break;
+					}
 				}
 			}
 		}
@@ -261,38 +249,43 @@ static mtlpp::Device GetMTLDevice(uint32& DeviceIndex)
 	if (SelectedDevice == nil)
 	{
 		TArray<FString> NameComponents;
-		SelectedDevice = mtlpp::Device::CreateSystemDefaultDevice();
+		SelectedDevice = MTLCreateSystemDefaultDevice();
 		bool bFoundDefault = false;
 		for (uint32 i = 0; i < GPUs.Num(); i++)
 		{
 			FMacPlatformMisc::FGPUDescriptor const& GPU = GPUs[i];
-			if(MTLPP_CHECK_AVAILABLE(10.13, 11.0, 11.0) && (SelectedDevice.GetRegistryID() == GPU.RegistryID))
+			if ([SelectedDevice registryID] == GPU.RegistryID)
 			{
 				DeviceIndex = i;
 				bFoundDefault = true;
 				break;
 			}
-			else if(([SelectedDevice.GetName().GetPtr() rangeOfString:@"Nvidia" options:NSCaseInsensitiveSearch].location != NSNotFound && GPU.GPUVendorId == 0x10DE)
-					|| ([SelectedDevice.GetName().GetPtr() rangeOfString:@"AMD" options:NSCaseInsensitiveSearch].location != NSNotFound && GPU.GPUVendorId == 0x1002)
-					|| ([SelectedDevice.GetName().GetPtr() rangeOfString:@"Intel" options:NSCaseInsensitiveSearch].location != NSNotFound && GPU.GPUVendorId == 0x8086))
+			else
 			{
-				NameComponents.Empty();
-				bool bMatchesName = FString(GPU.GPUName).TrimStart().ParseIntoArray(NameComponents, TEXT(" ")) > 0;
-				for (FString& Component : NameComponents)
+				NSString* SelectedDeviceName = [SelectedDevice name];
+
+				if (    (([SelectedDeviceName rangeOfString:@"Nvidia" options:NSCaseInsensitiveSearch].location != NSNotFound) && (GPU.GPUVendorId == 0x10DE))
+					 || (([SelectedDeviceName rangeOfString:@"AMD"    options:NSCaseInsensitiveSearch].location != NSNotFound) && (GPU.GPUVendorId == 0x1002))
+					 || (([SelectedDeviceName rangeOfString:@"Intel"  options:NSCaseInsensitiveSearch].location != NSNotFound) && (GPU.GPUVendorId == 0x8086)) )
 				{
-					bMatchesName &= FString(SelectedDevice.GetName().GetPtr()).Contains(Component);
-				}
-				if((SelectedDevice.IsHeadless() == GPU.GPUHeadless || GPU.GPUVendorId != 0x1002) && bMatchesName)
-                {
-					DeviceIndex = i;
-					bFoundDefault = true;
-					break;
+					NameComponents.Empty();
+					bool bMatchesName = FString(GPU.GPUName).TrimStart().ParseIntoArray(NameComponents, TEXT(" ")) > 0;
+					for (FString& Component : NameComponents)
+					{
+						bMatchesName &= FString(SelectedDeviceName).Contains(Component);
+					}
+					if (([SelectedDevice isHeadless] == GPU.GPUHeadless || GPU.GPUVendorId != 0x1002) && bMatchesName)
+					{
+						DeviceIndex = i;
+						bFoundDefault = true;
+						break;
+					}
 				}
 			}
 		}
 		if(!bFoundDefault)
 		{
-			UE_LOG(LogAGX, Warning,  TEXT("Couldn't find Metal device %s in GPU descriptors from IORegistry - capability reporting may be wrong."), *FString(SelectedDevice.GetName().GetPtr()));
+			UE_LOG(LogAGX, Warning,  TEXT("Couldn't find Metal device %s in GPU descriptors from IORegistry - capability reporting may be wrong."), *FString([SelectedDevice name]));
 		}
 	}
 	return SelectedDevice;
@@ -352,16 +345,22 @@ mtlpp::PrimitiveTopologyClass AGXTranslatePrimitiveTopology(uint32 PrimitiveType
 }
 #endif
 
+
+//------------------------------------------------------------------------------
+
+#pragma mark - AGX Device Context Class
+
+
 FAGXDeviceContext* FAGXDeviceContext::CreateDeviceContext()
 {
 	uint32 DeviceIndex = 0;
 #if PLATFORM_IOS
-	mtlpp::Device Device = mtlpp::Device([IOSAppDelegate GetDelegate].IOSView->MetalDevice);
+	GMtlDevice = [IOSAppDelegate GetDelegate].IOSView->MetalDevice;
 #else
-	mtlpp::Device Device = GetMTLDevice(DeviceIndex);
-	if (!Device)
+	GMtlDevice = GetMTLDevice(DeviceIndex);
+	if (!GMtlDevice)
 	{
-		FPlatformMisc::MessageBoxExt(EAppMsgType::Ok, TEXT("The graphics card in this Mac appears to erroneously report support for Metal graphics technology, which is required to run this application, but failed to create a Metal device. The application will now exit."), TEXT("Failed to initialize Metal"));
+		FPlatformMisc::MessageBoxExt(EAppMsgType::Ok, TEXT("Metal device creation failed. The application will now exit."), TEXT("Failed to initialize Metal"));
 		exit(0);
 	}
 #endif
@@ -372,24 +371,26 @@ FAGXDeviceContext* FAGXDeviceContext::CreateDeviceContext()
 	{
 		GAGXRuntimeDebugLevel = MetalDebug;
 	}
+
+	GMtlppDevice = mtlpp::Device(GMtlDevice, ns::Ownership::AutoRelease);
+
+	MTLPP_VALIDATION(mtlpp::ValidatedDevice::Register(GMtlppDevice));
 	
-	MTLPP_VALIDATION(mtlpp::ValidatedDevice::Register(Device));
-	
-	FAGXCommandQueue* Queue = new FAGXCommandQueue(Device, GAGXCommandQueueSize);
+	FAGXCommandQueue* Queue = new FAGXCommandQueue(GAGXCommandQueueSize);
 	check(Queue);
 	
 	if (FAGXCommandQueue::SupportsFeature(EAGXFeaturesFences))
 	{
-		FAGXFencePool::Get().Initialise(Device);
+		FAGXFencePool::Get().Init();
 	}
 	
-	return new FAGXDeviceContext(Device, DeviceIndex, Queue);
+	return new FAGXDeviceContext(DeviceIndex, Queue);
 }
 
-FAGXDeviceContext::FAGXDeviceContext(mtlpp::Device MetalDevice, uint32 InDeviceIndex, FAGXCommandQueue* Queue)
-: FAGXContext(MetalDevice, *Queue, true)
+FAGXDeviceContext::FAGXDeviceContext(uint32 InDeviceIndex, FAGXCommandQueue* Queue)
+: FAGXContext(*Queue, true)
 , DeviceIndex(InDeviceIndex)
-, CaptureManager(MetalDevice.GetPtr(), *Queue)
+, CaptureManager(*Queue)
 , SceneFrameCounter(0)
 , FrameCounter(0)
 , ActiveContexts(1)
@@ -431,12 +432,12 @@ FAGXDeviceContext::FAGXDeviceContext(mtlpp::Device MetalDevice, uint32 InDeviceI
 	}
     
     // initialize uniform allocator
-    UniformBufferAllocator = new FAGXFrameAllocator(MetalDevice.GetPtr());
+    UniformBufferAllocator = new FAGXFrameAllocator();
     UniformBufferAllocator->SetTargetAllocationLimitInBytes(GAGXTargetUniformAllocationLimit);
     UniformBufferAllocator->SetDefaultAllocationSizeInBytes(GAGXDefaultUniformBufferAllocation);
     UniformBufferAllocator->SetStatIds(GET_STATID(STAT_AGXUniformAllocatedMemory), GET_STATID(STAT_AGXUniformMemoryInFlight), GET_STATID(STAT_AGXUniformBytesPerFrame));
 	
-	TransferBufferAllocator = new FAGXFrameAllocator(MetalDevice.GetPtr());
+	TransferBufferAllocator = new FAGXFrameAllocator();
 	TransferBufferAllocator->SetTargetAllocationLimitInBytes(GAGXTargetTransferAllocatorLimit);
 	TransferBufferAllocator->SetDefaultAllocationSizeInBytes(GAGXDefaultTransferAllocation);
 	// We won't set StatIds here so it goes to the default frame allocator stats
@@ -458,10 +459,7 @@ FAGXDeviceContext::~FAGXDeviceContext()
     delete UniformBufferAllocator;
 	
 #if PLATFORM_MAC
-	if (FPlatformMisc::MacOSXVersionCompare(10, 13, 4) >= 0)
-	{
-		mtlpp::Device::RemoveDeviceObserver(GAGXDeviceObserver);
-	}
+	MTLRemoveDeviceObserver(GAGXDeviceObserver);
 #endif
 }
 
@@ -878,7 +876,7 @@ FAGXRHICommandContext* FAGXDeviceContext::AcquireContext(int32 NewIndex, int32 N
 	FAGXRHICommandContext* Context = ParallelContexts.Pop();
 	if (!Context)
 	{
-		FAGXContext* AGXContext = new FAGXContext(GetDevice(), GetCommandQueue(), false);
+		FAGXContext* AGXContext = new FAGXContext(GetCommandQueue(), false);
 		check(AGXContext);
 		
 		FAGXRHICommandContext* CmdContext = static_cast<FAGXRHICommandContext*>(RHIGetDefaultContext());
@@ -1020,9 +1018,8 @@ bool FAGXDeviceContext::ValidateIsInactiveBuffer(FAGXBuffer const& Buffer)
 uint32 FAGXContext::CurrentContextTLSSlot = FPlatformTLS::AllocTlsSlot();
 #endif
 
-FAGXContext::FAGXContext(mtlpp::Device InDevice, FAGXCommandQueue& Queue, bool const bIsImmediate)
-: Device(InDevice)
-, CommandQueue(Queue)
+FAGXContext::FAGXContext(FAGXCommandQueue& Queue, bool const bIsImmediate)
+: CommandQueue(Queue)
 , CommandList(Queue, bIsImmediate)
 , StateCache(bIsImmediate)
 , RenderPass(CommandList, StateCache)
@@ -1040,11 +1037,6 @@ FAGXContext::~FAGXContext()
 	SubmitCommandsHint(EAGXSubmitFlagsWaitOnCommandBuffer);
 }
 
-mtlpp::Device& FAGXContext::GetDevice()
-{
-	return Device;
-}
-
 FAGXCommandQueue& FAGXContext::GetCommandQueue()
 {
 	return CommandQueue;
diff --git a/Engine/Source/Runtime/Apple/AGXRHI/Private/AGXContext.h b/Engine/Source/Runtime/Apple/AGXRHI/Private/AGXContext.h
--- a/Engine/Source/Runtime/Apple/AGXRHI/Private/AGXContext.h
+++ b/Engine/Source/Runtime/Apple/AGXRHI/Private/AGXContext.h
@@ -17,7 +17,6 @@
 #include "IOS/IOSView.h"
 #endif
 #include "Containers/LockFreeList.h"
-#include "device.hpp"
 
 #define NUM_SAFE_FRAMES 4
 
@@ -30,10 +29,9 @@ class FAGXContext
 {
 	friend class FAGXCommandContextContainer;
 public:
-	FAGXContext(mtlpp::Device InDevice, FAGXCommandQueue& Queue, bool const bIsImmediate);
+	FAGXContext(FAGXCommandQueue& Queue, bool const bIsImmediate);
 	virtual ~FAGXContext();
 	
-	mtlpp::Device& GetDevice();
 	FAGXCommandQueue& GetCommandQueue();
 	FAGXCommandList& GetCommandList();
 	mtlpp::CommandBuffer const& GetCurrentCommandBuffer() const;
@@ -135,9 +133,6 @@ public:
 	void TransitionRHIResource(T* InResource);
 
 protected:
-	/** The underlying Metal device */
-	mtlpp::Device Device;
-	
 	/** The wrapper around the device command-queue for creating & committing command buffers to */
 	FAGXCommandQueue& CommandQueue;
 	
@@ -261,7 +256,7 @@ public:
 #endif
 	
 private:
-	FAGXDeviceContext(mtlpp::Device MetalDevice, uint32 DeviceIndex, FAGXCommandQueue* Queue);
+	FAGXDeviceContext(uint32 DeviceIndex, FAGXCommandQueue* Queue);
 	
 private:
 	/** The index into the GPU device list for the selected Metal device */
diff --git a/Engine/Source/Runtime/Apple/AGXRHI/Private/AGXDynamicRHI_Shaders.cpp b/Engine/Source/Runtime/Apple/AGXRHI/Private/AGXDynamicRHI_Shaders.cpp
--- a/Engine/Source/Runtime/Apple/AGXRHI/Private/AGXDynamicRHI_Shaders.cpp
+++ b/Engine/Source/Runtime/Apple/AGXRHI/Private/AGXDynamicRHI_Shaders.cpp
@@ -110,7 +110,7 @@ FRHIShaderLibraryRef FAGXDynamicRHI::RHICreateShaderLibrary(EShaderPlatform Plat
 
 					METAL_GPUPROFILE(FAGXScopedCPUStats CPUStat(FString::Printf(TEXT("NewLibraryFile: %s"), *MetalLibraryFilePath)));
 					NSError* Error;
-					mtlpp::Library Library = [GetAGXDeviceContext().GetDevice() newLibraryWithFile:MetalLibraryFilePath.GetNSString() error:&Error];
+					mtlpp::Library Library([GMtlDevice newLibraryWithFile:MetalLibraryFilePath.GetNSString() error:&Error], nullptr, ns::Ownership::Assign);
 					if (Library != nil)
 					{
 						Libraries.Add(Library);
diff --git a/Engine/Source/Runtime/Apple/AGXRHI/Private/AGXFence.cpp b/Engine/Source/Runtime/Apple/AGXRHI/Private/AGXFence.cpp
--- a/Engine/Source/Runtime/Apple/AGXRHI/Private/AGXFence.cpp
+++ b/Engine/Source/Runtime/Apple/AGXRHI/Private/AGXFence.cpp
@@ -133,18 +133,17 @@ void FAGXFence::Validate(void) const
 }
 #endif
 
-void FAGXFencePool::Initialise(mtlpp::Device const& InDevice)
+void FAGXFencePool::Init()
 {
-	Device = InDevice;
 	for (int32 i = 0; i < FAGXFencePool::NumFences; i++)
 	{
 #if METAL_DEBUG_OPTIONS
 		if (GAGXRuntimeDebugLevel >= EAGXDebugLevelValidation)
 		{
 			FAGXDebugFence* VertexFence = [[FAGXDebugFence new] autorelease];
-			VertexFence.Inner = Device.NewFence();
+			VertexFence.Inner = GMtlppDevice.NewFence();
 			FAGXDebugFence* FragmentFence = [[FAGXDebugFence new] autorelease];
-			FragmentFence.Inner = Device.NewFence();
+			FragmentFence.Inner = GMtlppDevice.NewFence();
 			FAGXFence* F = new FAGXFence;
 			F->Set(mtlpp::RenderStages::Vertex, VertexFence);
 			F->Set(mtlpp::RenderStages::Fragment, FragmentFence);
@@ -155,8 +154,8 @@ void FAGXFencePool::Initialise(mtlpp::Device const& InDevice)
 #endif
 		{
 			FAGXFence* F = new FAGXFence;
-			F->Set(mtlpp::RenderStages::Vertex, Device.NewFence());
-			F->Set(mtlpp::RenderStages::Fragment, Device.NewFence());
+			F->Set(mtlpp::RenderStages::Vertex, GMtlppDevice.NewFence());
+			F->Set(mtlpp::RenderStages::Fragment, GMtlppDevice.NewFence());
 #if METAL_DEBUG_OPTIONS
 			if (GAGXRuntimeDebugLevel >= EAGXDebugLevelValidation)
 			{
diff --git a/Engine/Source/Runtime/Apple/AGXRHI/Private/AGXFence.h b/Engine/Source/Runtime/Apple/AGXRHI/Private/AGXFence.h
--- a/Engine/Source/Runtime/Apple/AGXRHI/Private/AGXFence.h
+++ b/Engine/Source/Runtime/Apple/AGXRHI/Private/AGXFence.h
@@ -188,7 +188,7 @@ public:
 		return sSelf;
 	}
 	
-	void Initialise(mtlpp::Device const& InDevice);
+	void Init();
 	
 	FAGXFence* AllocateFence();
 	void ReleaseFence(FAGXFence* const InFence);
@@ -199,7 +199,6 @@ public:
 private:
 	int32 Count;
 	int32 Allocated;
-	mtlpp::Device Device;
 #if METAL_DEBUG_OPTIONS
 	TSet<FAGXFence*> Fences;
 	FCriticalSection Mutex;
diff --git a/Engine/Source/Runtime/Apple/AGXRHI/Private/AGXFrameAllocator.cpp b/Engine/Source/Runtime/Apple/AGXRHI/Private/AGXFrameAllocator.cpp
--- a/Engine/Source/Runtime/Apple/AGXRHI/Private/AGXFrameAllocator.cpp
+++ b/Engine/Source/Runtime/Apple/AGXRHI/Private/AGXFrameAllocator.cpp
@@ -7,14 +7,13 @@
 
 #pragma mark Constructor/Destructor
 
-FAGXFrameAllocator::FAGXFrameAllocator(id <MTLDevice> TargetDevice)
-    : Device(TargetDevice)
+FAGXFrameAllocator::FAGXFrameAllocator()
+    : CurrentBuffer(nil)
 #if METAL_FRAME_ALLOCATOR_VALIDATION
     , AllAllocatedBuffers(nil)
 #endif
     , AvailableBuffers(nil)
     , BuffersInFlight(nil)
-    , CurrentBuffer(nil)
     , CurrentCursor(0)
     , CurrentFrame(0)
     , TargetAllocationLimit(0)
@@ -23,7 +22,6 @@ FAGXFrameAllocator::FAGXFrameAllocator(id <MTLDevice> TargetDevice)
     , BytesAcquiredInCurrentFrame(0)
     , BytesInFlight(0)
 {
-    check(Device);
 #if METAL_FRAME_ALLOCATOR_VALIDATION
     AllAllocatedBuffers = [NSMutableArray new];
 #endif
@@ -34,8 +32,6 @@ FAGXFrameAllocator::FAGXFrameAllocator(id <MTLDevice> TargetDevice)
     TotalAllocationStat = GET_STATID(STAT_AGXFrameAllocatorAllocatedMemory);
     MemoryInFlightStat = GET_STATID(STAT_AGXFrameAllocatorMemoryInFlight);
     BytesPerFrameStat = GET_STATID(STAT_AGXFrameAllocatorBytesPerFrame);
-    
-    [Device retain];
 }
 
 FAGXFrameAllocator::~FAGXFrameAllocator()
@@ -50,8 +46,6 @@ FAGXFrameAllocator::~FAGXFrameAllocator()
     [AllAllocatedBuffers release];
 #endif
     
-    [Device release];
-    
     os_unfair_lock_unlock(&UnfairLock);
 }
 
@@ -153,7 +147,7 @@ id <MTLBuffer> FAGXFrameAllocator::FindOrAllocateBufferUnsafe(uint32 SizeInBytes
     else
     {
         uint32 AllocationSize = FMath::Max(DefaultAllocationSize, SizeInBytes);
-        Found = [Device newBufferWithLength:AllocationSize options:MTLResourceStorageModeShared | MTLResourceCPUCacheModeWriteCombined];
+        Found = [GMtlDevice newBufferWithLength:AllocationSize options:MTLResourceStorageModeShared | MTLResourceCPUCacheModeWriteCombined];
         [Found setLabel: [NSString stringWithFormat:@"UniformBacking %u", AllocationSize]];
         
         check(Found);
@@ -161,7 +155,7 @@ id <MTLBuffer> FAGXFrameAllocator::FindOrAllocateBufferUnsafe(uint32 SizeInBytes
         this->TotalBytesAllocated += AllocationSize;
         INC_MEMORY_STAT_BY_FName(TotalAllocationStat.GetName(), [Found length]);
 #if STATS || ENABLE_LOW_LEVEL_MEM_TRACKER
-        AGXLLM::LogAllocBuffer(GetAGXDeviceContext().GetDevice(), Found);
+        AGXLLM::LogAllocBuffer(Found);
 #endif
 
 #if METAL_FRAME_ALLOCATOR_VALIDATION
diff --git a/Engine/Source/Runtime/Apple/AGXRHI/Private/AGXLLM.cpp b/Engine/Source/Runtime/Apple/AGXRHI/Private/AGXLLM.cpp
--- a/Engine/Source/Runtime/Apple/AGXRHI/Private/AGXLLM.cpp
+++ b/Engine/Source/Runtime/Apple/AGXRHI/Private/AGXLLM.cpp
@@ -184,16 +184,16 @@ static mtlpp::SizeAndAlign TextureSizeAndAlign(mtlpp::TextureType TextureType, u
 	return SizeAlign;
 }
 
-void AGXLLM::LogAllocTexture(mtlpp::Device& Device, mtlpp::TextureDescriptor const& Desc, mtlpp::Texture const& Texture)
+void AGXLLM::LogAllocTexture(mtlpp::TextureDescriptor const& Desc, mtlpp::Texture const& Texture)
 {
-	mtlpp::SizeAndAlign SizeAlign;
+	MTLSizeAndAlign SizeAlign;
 	if (FAGXCommandQueue::SupportsFeature(EAGXFeaturesGPUCaptureManager))
 	{
-		SizeAlign = Device.HeapTextureSizeAndAlign(Desc);
+		SizeAlign = [GMtlDevice heapTextureSizeAndAlignWithDescriptor:Desc.GetPtr()];
 	}
 	
 	void* Ptr = (void*)Texture.GetPtr();
-	uint64 Size = SizeAlign.Size;
+	uint64 Size = static_cast<uint64>(SizeAlign.size);
 	
 #if PLATFORM_IOS
 	bool bMemoryless = (Texture.GetStorageMode() == mtlpp::StorageMode::Memoryless);
@@ -248,7 +248,7 @@ void AGXLLM::LogAllocTexture(mtlpp::Device& Device, mtlpp::TextureDescriptor con
 	}
 }
 
-void AGXLLM::LogAllocBuffer(mtlpp::Device& Device, mtlpp::Buffer const& Buffer)
+void AGXLLM::LogAllocBuffer(mtlpp::Buffer const& Buffer)
 {
 	void* Ptr = (void*)Buffer.GetPtr();
 	uint64 Size = Buffer.GetLength();
@@ -274,7 +274,7 @@ void AGXLLM::LogAllocBuffer(mtlpp::Device& Device, mtlpp::Buffer const& Buffer)
 	}
 }
 
-void AGXLLM::LogAllocHeap(mtlpp::Device& Device, mtlpp::Heap const& Heap)
+void AGXLLM::LogAllocHeap(mtlpp::Heap const& Heap)
 {
 	void* Ptr = (void*)Heap.GetPtr();
 	uint64 Size = Heap.GetSize();
diff --git a/Engine/Source/Runtime/Apple/AGXRHI/Private/AGXLLM.h b/Engine/Source/Runtime/Apple/AGXRHI/Private/AGXLLM.h
--- a/Engine/Source/Runtime/Apple/AGXRHI/Private/AGXLLM.h
+++ b/Engine/Source/Runtime/Apple/AGXRHI/Private/AGXLLM.h
@@ -46,9 +46,9 @@ namespace AGXLLM
 // These work without the LLM module
 namespace AGXLLM
 {
-	void LogAllocTexture(mtlpp::Device& Device, mtlpp::TextureDescriptor const& Desc, mtlpp::Texture const& Texture);
-	void LogAllocBuffer(mtlpp::Device& Device, mtlpp::Buffer const& Buffer);
-	void LogAllocHeap(mtlpp::Device& Device, mtlpp::Heap const& Heap);
+	void LogAllocTexture(mtlpp::TextureDescriptor const& Desc, mtlpp::Texture const& Texture);
+	void LogAllocBuffer(mtlpp::Buffer const& Buffer);
+	void LogAllocHeap(mtlpp::Heap const& Heap);
 	
 	void LogAliasTexture(mtlpp::Texture const& Texture);
 	void LogAliasBuffer(mtlpp::Buffer const& Buffer);
diff --git a/Engine/Source/Runtime/Apple/AGXRHI/Private/AGXPipeline.cpp b/Engine/Source/Runtime/Apple/AGXRHI/Private/AGXPipeline.cpp
--- a/Engine/Source/Runtime/Apple/AGXRHI/Private/AGXPipeline.cpp
+++ b/Engine/Source/Runtime/Apple/AGXRHI/Private/AGXPipeline.cpp
@@ -119,13 +119,13 @@ struct FAGXHelperFunctions
             mtlpp::CompileOptions CompileOptions;
             ns::AutoReleasedError Error;
 
-			DebugShadersLib = GetAGXDeviceContext().GetDevice().NewLibrary(GAGXDebugShader, CompileOptions, &Error);
+			DebugShadersLib = GMtlppDevice.NewLibrary(GAGXDebugShader, CompileOptions, &Error);
             DebugFunc = DebugShadersLib.NewFunction(@"Main_Debug");
 			
-			DebugComputeShadersLib = GetAGXDeviceContext().GetDevice().NewLibrary(GAGXDebugMarkerComputeShader, CompileOptions, &Error);
+			DebugComputeShadersLib = GMtlppDevice.NewLibrary(GAGXDebugMarkerComputeShader, CompileOptions, &Error);
 			DebugComputeFunc = DebugComputeShadersLib.NewFunction(@"Main_Debug");
 			
-			DebugComputeState = GetAGXDeviceContext().GetDevice().NewComputePipelineState(DebugComputeFunc, &Error);
+			DebugComputeState = GMtlppDevice.NewComputePipelineState(DebugComputeFunc, &Error);
         }
 #endif
     }
@@ -852,7 +852,6 @@ static FAGXShaderPipeline* CreateMTLRenderPipeline(bool const bSync, FAGXGraphic
     if (vertexFunction && ((PixelShader != nullptr) == (fragmentFunction != nil)))
     {
 		ns::Error Error;
-		mtlpp::Device Device = GetAGXDeviceContext().GetDevice();
 
 		uint32 const NumActiveTargets = Init.ComputeNumValidRenderTargets();
         check(NumActiveTargets <= MaxSimultaneousRenderTargets);
@@ -900,7 +899,7 @@ static FAGXShaderPipeline* CreateMTLRenderPipeline(bool const bSync, FAGXGraphic
 		{
 			ns::AutoReleasedError RenderError;
 			METAL_GPUPROFILE(FAGXScopedCPUStats CPUStat(FString::Printf(TEXT("NewRenderPipeline: %s"), TEXT("")/**FString([RenderPipelineDesc.GetPtr() description])*/)));
-			Pipeline->RenderPipelineState = Device.NewRenderPipelineState(RenderPipelineDesc, (mtlpp::PipelineOption)RenderOption, Reflection, &RenderError);
+			Pipeline->RenderPipelineState = GMtlppDevice.NewRenderPipelineState(RenderPipelineDesc, (mtlpp::PipelineOption)RenderOption, Reflection, &RenderError);
 			if (Reflection)
 			{
 				Pipeline->RenderPipelineReflection = *Reflection;
@@ -946,7 +945,7 @@ static FAGXShaderPipeline* CreateMTLRenderPipeline(bool const bSync, FAGXGraphic
 
 			ns::AutoReleasedError RenderError;
 			METAL_GPUPROFILE(FAGXScopedCPUStats CPUStat(FString::Printf(TEXT("NewDebugPipeline: %s"), TEXT("")/**FString([RenderPipelineDesc.GetPtr() description])*/)));
-			Pipeline->DebugPipelineState = Device.NewRenderPipelineState(DebugPipelineDesc, mtlpp::PipelineOption::NoPipelineOption, Reflection, nullptr);
+			Pipeline->DebugPipelineState = GMtlppDevice.NewRenderPipelineState(DebugPipelineDesc, mtlpp::PipelineOption::NoPipelineOption, Reflection, nullptr);
 		}
 #endif
         
diff --git a/Engine/Source/Runtime/Apple/AGXRHI/Private/AGXRHI.cpp b/Engine/Source/Runtime/Apple/AGXRHI/Private/AGXRHI.cpp
--- a/Engine/Source/Runtime/Apple/AGXRHI/Private/AGXRHI.cpp
+++ b/Engine/Source/Runtime/Apple/AGXRHI/Private/AGXRHI.cpp
@@ -191,24 +191,21 @@ FAGXDynamicRHI::FAGXDynamicRHI(ERHIFeatureLevel::Type RequestedFeatureLevel)
 	bool bSupportsPointLights = false;
 	bool bSupportsRHIThread = false;
 	
-	// get the device to ask about capabilities?
-	mtlpp::Device Device = ImmediateContext.Context->GetDevice();
-		
 #if PLATFORM_IOS
 	// A8 can use 256 bits of MRTs
 #if PLATFORM_TVOS
-	GRHISupportsDrawIndirect = Device.SupportsFeatureSet(mtlpp::FeatureSet::tvOS_GPUFamily2_v1);
-	GRHISupportsPixelShaderUAVs = Device.SupportsFeatureSet(mtlpp::FeatureSet::tvOS_GPUFamily2_v1);
+	GRHISupportsDrawIndirect = [GMtlDevice supportsFeatureSet:MTLFeatureSet_tvOS_GPUFamily2_v1];
+	GRHISupportsPixelShaderUAVs = [GMtlDevice supportsFeatureSet:MTLFeatureSet_tvOS_GPUFamily2_v1];
 #else
-	GRHISupportsRWTextureBuffers = Device.SupportsFeatureSet(mtlpp::FeatureSet::iOS_GPUFamily4_v1);
-	GRHISupportsDrawIndirect = Device.SupportsFeatureSet(mtlpp::FeatureSet::iOS_GPUFamily3_v1);
-	GRHISupportsPixelShaderUAVs = Device.SupportsFeatureSet(mtlpp::FeatureSet::iOS_GPUFamily3_v1);
+	GRHISupportsRWTextureBuffers = [GMtlDevice supportsFeatureSet:MTLFeatureSet_iOS_GPUFamily4_v1];
+	GRHISupportsDrawIndirect = [GMtlDevice supportsFeatureSet:MTLFeatureSet_iOS_GPUFamily3_v1];
+	GRHISupportsPixelShaderUAVs = [GMtlDevice supportsFeatureSet:MTLFeatureSet_iOS_GPUFamily3_v1];
 
 	const mtlpp::FeatureSet FeatureSets[] = {
-		mtlpp::FeatureSet::iOS_GPUFamily1_v1,
-		mtlpp::FeatureSet::iOS_GPUFamily2_v1,
-		mtlpp::FeatureSet::iOS_GPUFamily3_v1,
-		mtlpp::FeatureSet::iOS_GPUFamily4_v1
+		MTLFeatureSet_iOS_GPUFamily1_v1,
+		MTLFeatureSet_iOS_GPUFamily2_v1,
+		MTLFeatureSet_iOS_GPUFamily3_v1,
+		MTLFeatureSet_iOS_GPUFamily4_v1
 	};
 		
 	const uint8 FeatureSetVersions[][3] = {
@@ -221,7 +218,7 @@ FAGXDynamicRHI::FAGXDynamicRHI(ERHIFeatureLevel::Type RequestedFeatureLevel)
 	GRHIDeviceId = 0;
 	for (uint32 i = 0; i < 4; i++)
 	{
-		if (FPlatformMisc::IOSVersionCompare(FeatureSetVersions[i][0],FeatureSetVersions[i][1],FeatureSetVersions[i][2]) >= 0 && Device.SupportsFeatureSet(FeatureSets[i]))
+		if (FPlatformMisc::IOSVersionCompare(FeatureSetVersions[i][0],FeatureSetVersions[i][1],FeatureSetVersions[i][2]) >= 0 && [GMtlDevice supportsFeatureSet:FeatureSets[i]])
 		{
 			GRHIDeviceId++;
 		}
@@ -294,7 +291,7 @@ FAGXDynamicRHI::FAGXDynamicRHI(ERHIFeatureLevel::Type RequestedFeatureLevel)
 	bool bSupportsD24S8 = false;
 	bool bSupportsD16 = false;
 	
-	GRHIAdapterName = FString(Device.GetName());
+	GRHIAdapterName = FString([GMtlDevice name]);
 	
 	// However they don't all support other features depending on the version of the OS.
 	bool bSupportsTiledReflections = false;
@@ -407,7 +404,7 @@ FAGXDynamicRHI::FAGXDynamicRHI(ERHIFeatureLevel::Type RequestedFeatureLevel)
 	}
 	
 	// Change the support depth format if we can
-	bSupportsD24S8 = Device.IsDepth24Stencil8PixelFormatSupported();
+	bSupportsD24S8 = [GMtlDevice isDepth24Stencil8PixelFormatSupported];
 	
 	// Disable tiled reflections on Mac Metal for some GPU drivers that ignore the lod-level and so render incorrectly.
 	if (!bSupportsTiledReflections && !FParse::Param(FCommandLine::Get(),TEXT("metaltiledreflections")))
@@ -442,7 +439,7 @@ FAGXDynamicRHI::FAGXDynamicRHI(ERHIFeatureLevel::Type RequestedFeatureLevel)
 		
 	if(
 	   #if PLATFORM_MAC
-	   (Device.SupportsFeatureSet(mtlpp::FeatureSet::macOS_GPUFamily1_v3) && FPlatformMisc::MacOSXVersionCompare(10,13,0) >= 0)
+	   ([GMtlDevice supportsFeatureSet:MTLFeatureSet_macOS_GPUFamily1_v3] && FPlatformMisc::MacOSXVersionCompare(10,13,0) >= 0)
 	   #elif PLATFORM_IOS || PLATFORM_TVOS
 	   FPlatformMisc::IOSVersionCompare(10,3,0)
 	   #endif
@@ -570,7 +567,7 @@ FAGXDynamicRHI::FAGXDynamicRHI(ERHIFeatureLevel::Type RequestedFeatureLevel)
 	GRHISupportsPipelineFileCache = true;
 
 #if PLATFORM_MAC
-	check(Device.SupportsFeatureSet(mtlpp::FeatureSet::macOS_GPUFamily1_v1));
+	check([GMtlDevice supportsFeatureSet:MTLFeatureSet_macOS_GPUFamily1_v1]);
 	GRHISupportsBaseVertexIndex = FPlatformMisc::MacOSXVersionCompare(10,11,2) >= 0 || !IsRHIDeviceAMD(); // Supported on macOS & iOS but not tvOS - broken on AMD prior to 10.11.2
 	GRHISupportsFirstInstance = true; // Supported on macOS & iOS but not tvOS.
 	GMaxTextureDimensions = 16384;
@@ -578,8 +575,8 @@ FAGXDynamicRHI::FAGXDynamicRHI(ERHIFeatureLevel::Type RequestedFeatureLevel)
 	GMaxTextureArrayLayers = 2048;
 	GMaxShadowDepthBufferSizeX = GMaxTextureDimensions;
 	GMaxShadowDepthBufferSizeY = GMaxTextureDimensions;
-    bSupportsD16 = !FParse::Param(FCommandLine::Get(),TEXT("nometalv2")) && Device.SupportsFeatureSet(mtlpp::FeatureSet::macOS_GPUFamily1_v2);
-    GRHISupportsHDROutput = FPlatformMisc::MacOSXVersionCompare(10,14,4) >= 0 && Device.SupportsFeatureSet(mtlpp::FeatureSet::macOS_GPUFamily1_v2);
+    bSupportsD16 = !FParse::Param(FCommandLine::Get(),TEXT("nometalv2")) && [GMtlDevice supportsFeatureSet:MTLFeatureSet_macOS_GPUFamily1_v2];
+    GRHISupportsHDROutput = FPlatformMisc::MacOSXVersionCompare(10,14,4) >= 0 && [GMtlDevice supportsFeatureSet:MTLFeatureSet_macOS_GPUFamily1_v2];
 	GRHIHDRDisplayOutputFormat = (GRHISupportsHDROutput) ? PF_PLATFORM_HDR_0 : PF_B8G8R8A8;
 	// Based on the spec below, the maxTotalThreadsPerThreadgroup is not a fixed number but calculated according to the device current ability, so the available threads could less than the maximum number.
 	// For safety and keep the consistency for all platform, reduce the maximum number to half of the device based.
@@ -595,7 +592,7 @@ FAGXDynamicRHI::FAGXDynamicRHI(ERHIFeatureLevel::Type RequestedFeatureLevel)
 	GRHIHDRDisplayOutputFormat = PF_B8G8R8A8; // must have a default value for non-hdr, just like mac or ios
 #else
 	// Only A9+ can support this, so for now we need to limit this to the desktop-forward renderer only.
-	GRHISupportsBaseVertexIndex = Device.SupportsFeatureSet(mtlpp::FeatureSet::iOS_GPUFamily3_v1) && (GMaxRHIFeatureLevel >= ERHIFeatureLevel::SM5);
+	GRHISupportsBaseVertexIndex = [GMtlDevice supportsFeatureSet:MTLFeatureSet_iOS_GPUFamily3_v1] && (GMaxRHIFeatureLevel >= ERHIFeatureLevel::SM5);
 	GRHISupportsFirstInstance = GRHISupportsBaseVertexIndex;
 	
 	// TODO: Move this into IOSPlatform
@@ -609,7 +606,7 @@ FAGXDynamicRHI::FAGXDynamicRHI(ERHIFeatureLevel::Type RequestedFeatureLevel)
 	// Based on the spec below, the maxTotalThreadsPerThreadgroup is not a fixed number but calculated according to the device current ability, so the available threads could less than the maximum number.
 	// For safety and keep the consistency for all platform, reduce the maximum number to half of the device based.
 	// https://developer.apple.com/documentation/metal/mtlcomputepipelinedescriptor/2966560-maxtotalthreadsperthreadgroup?language=objc
-	GMaxWorkGroupInvocations = Device.SupportsFeatureSet(mtlpp::FeatureSet::iOS_GPUFamily4_v1) ? 512 : 256;
+	GMaxWorkGroupInvocations = [GMtlDevice supportsFeatureSet:MTLFeatureSet_iOS_GPUFamily4_v1] ? 512 : 256;
 #endif
 	GMaxTextureDimensions = 8192;
 	GMaxCubeTextureDimensions = 8192;
@@ -760,9 +757,9 @@ FAGXDynamicRHI::FAGXDynamicRHI(ERHIFeatureLevel::Type RequestedFeatureLevel)
 	GPixelFormats[PF_PLATFORM_HDR_0		].Supported			= GRHISupportsHDROutput;
 		
 #if PLATFORM_TVOS
-    if (!Device.SupportsFeatureSet(mtlpp::FeatureSet::tvOS_GPUFamily2_v1))
+    if (![GMtlDevice supportsFeatureSet:MTLFeatureSet_tvOS_GPUFamily2_v1])
 #else
-	if (!Device.SupportsFeatureSet(mtlpp::FeatureSet::iOS_GPUFamily3_v2))
+	if (![GMtlDevice supportsFeatureSet:MTLFeatureSet_iOS_GPUFamily3_v2])
 #endif
 	{
 		GPixelFormats[PF_FloatRGB			].PlatformFormat 	= (uint32)mtlpp::PixelFormat::RGBA16Float;
@@ -903,9 +900,9 @@ FAGXDynamicRHI::FAGXDynamicRHI(ERHIFeatureLevel::Type RequestedFeatureLevel)
 	}
 #endif
 
-	switch (Device.GetReadWriteTextureSupport())
+	switch ([GMtlDevice readWriteTextureSupport])
 	{
-		case mtlpp::ReadWriteTextureTier::Tier2:
+		case MTLReadWriteTextureTier2:
 			GFormatSupportsTypedUAVLoad[PF_A32B32G32R32F]			= true;
 			GFormatSupportsTypedUAVLoad[PF_R32G32B32A32_UINT]		= true;
 			GFormatSupportsTypedUAVLoad[PF_FloatRGBA]				= true;
@@ -920,13 +917,13 @@ FAGXDynamicRHI::FAGXDynamicRHI(ERHIFeatureLevel::Type RequestedFeatureLevel)
 			GFormatSupportsTypedUAVLoad[PF_R8_UINT]					= true;
 			// Fall through
 
-		case mtlpp::ReadWriteTextureTier::Tier1:
+		case MTLReadWriteTextureTier1:
 			GFormatSupportsTypedUAVLoad[PF_R32_FLOAT]				= true;
 			GFormatSupportsTypedUAVLoad[PF_R32_UINT]				= true;
 			GFormatSupportsTypedUAVLoad[PF_R32_SINT]				= true;
 			// Fall through
 
-		case mtlpp::ReadWriteTextureTier::None:
+		case MTLReadWriteTextureTierNone:
 			break;
 	};
 
@@ -976,7 +973,7 @@ FAGXDynamicRHI::FAGXDynamicRHI(ERHIFeatureLevel::Type RequestedFeatureLevel)
 	if (ImmediateContext.Profiler)
 		ImmediateContext.Profiler->BeginFrame();
 #endif
-	AsyncComputeContext = GSupportsEfficientAsyncCompute ? new FAGXRHIComputeContext(ImmediateContext.Profiler, new FAGXContext(ImmediateContext.Context->GetDevice(), ImmediateContext.Context->GetCommandQueue(), true)) : nullptr;
+	AsyncComputeContext = GSupportsEfficientAsyncCompute ? new FAGXRHIComputeContext(ImmediateContext.Profiler, new FAGXContext(ImmediateContext.Context->GetCommandQueue(), true)) : nullptr;
 
 #if ENABLE_METAL_GPUPROFILE
 		if (ImmediateContext.Profiler)
@@ -1026,7 +1023,7 @@ uint64 FAGXDynamicRHI::RHICalcTextureCubePlatformSize(uint32 Size, uint8 Format,
 
 uint64 FAGXDynamicRHI::RHIGetMinimumAlignmentForBufferBackedSRV(EPixelFormat Format)
 {
-	return ImmediateContext.Context->GetDevice().GetMinimumLinearTextureAlignmentForPixelFormat((mtlpp::PixelFormat)GAGXBufferFormats[Format].LinearTextureFormat);
+	return static_cast<uint64>([GMtlDevice minimumLinearTextureAlignmentForPixelFormat:(MTLPixelFormat)GAGXBufferFormats[Format].LinearTextureFormat]);
 }
 
 void FAGXDynamicRHI::Init()
@@ -1230,7 +1227,7 @@ void FAGXDynamicRHI::RHIReleaseThreadOwnership()
 
 void* FAGXDynamicRHI::RHIGetNativeDevice()
 {
-	return (void*)ImmediateContext.Context->GetDevice().GetPtr();
+	return GMtlDevice;
 }
 
 void* FAGXDynamicRHI::RHIGetNativeInstance()
@@ -1253,10 +1250,7 @@ uint16 FAGXDynamicRHI::RHIGetPlatformTextureMaxSampleCount()
 		int sample = *sampleIt;
 
 #if PLATFORM_IOS || PLATFORM_MAC
-		id<MTLDevice> Device = (id<MTLDevice>)RHIGetNativeDevice();
-		check(Device);
-
-		if (![Device supportsTextureSampleCount : sample])
+		if (![GMtlDevice supportsTextureSampleCount:sample])
 		{
 			break;
 		}
diff --git a/Engine/Source/Runtime/Apple/AGXRHI/Private/AGXRHIPrivate.h b/Engine/Source/Runtime/Apple/AGXRHI/Private/AGXRHIPrivate.h
--- a/Engine/Source/Runtime/Apple/AGXRHI/Private/AGXRHIPrivate.h
+++ b/Engine/Source/Runtime/Apple/AGXRHI/Private/AGXRHIPrivate.h
@@ -20,6 +20,12 @@ THIRD_PARTY_INCLUDES_START
 #include "mtlpp.hpp"
 THIRD_PARTY_INCLUDES_END
 
+// The Metal device object.
+extern id<MTLDevice> GMtlDevice;
+
+// Placeholder: TODO: remove
+extern mtlpp::Device GMtlppDevice;
+
 // Whether the AGX RHI is initialized sufficiently to handle resources
 extern bool GIsAGXInitialized;
 
diff --git a/Engine/Source/Runtime/Apple/AGXRHI/Private/AGXRHIRenderQuery.cpp b/Engine/Source/Runtime/Apple/AGXRHI/Private/AGXRHIRenderQuery.cpp
--- a/Engine/Source/Runtime/Apple/AGXRHI/Private/AGXRHIRenderQuery.cpp
+++ b/Engine/Source/Runtime/Apple/AGXRHI/Private/AGXRHIRenderQuery.cpp
@@ -110,7 +110,7 @@ FAGXQueryBuffer* FAGXQueryBufferPool::GetCurrentQueryBuffer()
 			FMemory::Memzero((((uint8*)Buffer.GetContents())), EQueryBufferMaxSize);
 
 #if STATS || ENABLE_LOW_LEVEL_MEM_TRACKER
-			AGXLLM::LogAllocBuffer(Context->GetDevice(), Buffer);
+			AGXLLM::LogAllocBuffer(Buffer);
 #endif // STATS || ENABLE_LOW_LEVEL_MEM_TRACKER
 		}
 
diff --git a/Engine/Source/Runtime/Apple/AGXRHI/Private/AGXRenderTarget.cpp b/Engine/Source/Runtime/Apple/AGXRHI/Private/AGXRenderTarget.cpp
--- a/Engine/Source/Runtime/Apple/AGXRHI/Private/AGXRenderTarget.cpp
+++ b/Engine/Source/Runtime/Apple/AGXRHI/Private/AGXRenderTarget.cpp
@@ -280,7 +280,7 @@ void FAGXDynamicRHI::RHIReadSurfaceData(FRHITexture* TextureRHI, FIntRect Rect,
 			Desc.SetStorageMode(StorageMode);
 			Desc.SetUsage(Texture.GetUsage());
 			
-			TempTexture = GetAGXDeviceContext().GetDevice().NewTexture(Desc);
+			TempTexture = GMtlppDevice.NewTexture(Desc);
 			
 			ImmediateContext.Context->CopyFromTextureToTexture(Texture, 0, InFlags.GetMip(), mtlpp::Origin(Region.origin), mtlpp::Size(Region.size), TempTexture, 0, 0, mtlpp::Origin(0, 0, 0));
 			
@@ -320,7 +320,7 @@ void FAGXDynamicRHI::RHIReadSurfaceData(FRHITexture* TextureRHI, FIntRect Rect,
 		const uint32 Alignment = PLATFORM_MAC ? 1u : 64u; // Mac permits natural row alignment (tightly-packed) but iOS does not.
 		const uint32 AlignedStride = ((Stride - 1) & ~(Alignment - 1)) + Alignment;
 		const uint32 BytesPerImage = AlignedStride * SizeY;
-		FAGXBuffer Buffer = ((FAGXDeviceContext*)ImmediateContext.Context)->CreatePooledBuffer(FAGXPooledBufferArgs(ImmediateContext.Context->GetDevice(), BytesPerImage, BUF_Dynamic, mtlpp::StorageMode::Shared));
+		FAGXBuffer Buffer = ((FAGXDeviceContext*)ImmediateContext.Context)->CreatePooledBuffer(FAGXPooledBufferArgs(BytesPerImage, BUF_Dynamic, mtlpp::StorageMode::Shared));
 		{
 			// Synchronise the texture with the CPU
 			SCOPE_CYCLE_COUNTER(STAT_AGXTexturePageOffTime);
@@ -425,7 +425,7 @@ void FAGXDynamicRHI::RHIReadSurfaceFloatData(FRHITexture* TextureRHI, FIntRect R
 	const uint32 AlignedStride = ((Stride - 1) & ~(Alignment - 1)) + Alignment;
 	const uint32 BytesPerImage = AlignedStride  * SizeY;
 	int32 FloatBGRADataSize = BytesPerImage;
-	FAGXBuffer Buffer = ((FAGXDeviceContext*)ImmediateContext.Context)->CreatePooledBuffer(FAGXPooledBufferArgs(ImmediateContext.Context->GetDevice(), FloatBGRADataSize, BUF_Dynamic, mtlpp::StorageMode::Shared));
+	FAGXBuffer Buffer = ((FAGXDeviceContext*)ImmediateContext.Context)->CreatePooledBuffer(FAGXPooledBufferArgs(FloatBGRADataSize, BUF_Dynamic, mtlpp::StorageMode::Shared));
 	{
 		// Synchronise the texture with the CPU
 		SCOPE_CYCLE_COUNTER(STAT_AGXTexturePageOffTime);
@@ -491,7 +491,7 @@ void FAGXDynamicRHI::RHIRead3DSurfaceFloatData(FRHITexture* TextureRHI,FIntRect
 	const uint32 AlignedStride = ((Stride - 1) & ~(Alignment - 1)) + Alignment;
 	const uint32 BytesPerImage = AlignedStride  * SizeY;
 	int32 FloatBGRADataSize = BytesPerImage * SizeZ;
-	FAGXBuffer Buffer = ((FAGXDeviceContext*)ImmediateContext.Context)->CreatePooledBuffer(FAGXPooledBufferArgs(ImmediateContext.Context->GetDevice(), FloatBGRADataSize, BUF_Dynamic, mtlpp::StorageMode::Shared));
+	FAGXBuffer Buffer = ((FAGXDeviceContext*)ImmediateContext.Context)->CreatePooledBuffer(FAGXPooledBufferArgs(FloatBGRADataSize, BUF_Dynamic, mtlpp::StorageMode::Shared));
 	{
 		// Synchronise the texture with the CPU
 		SCOPE_CYCLE_COUNTER(STAT_AGXTexturePageOffTime);
diff --git a/Engine/Source/Runtime/Apple/AGXRHI/Private/AGXState.cpp b/Engine/Source/Runtime/Apple/AGXRHI/Private/AGXState.cpp
--- a/Engine/Source/Runtime/Apple/AGXRHI/Private/AGXState.cpp
+++ b/Engine/Source/Runtime/Apple/AGXRHI/Private/AGXState.cpp
@@ -27,54 +27,54 @@ static mtlpp::SamplerMinMagFilter TranslateZFilterMode(ESamplerFilter Filter)
 	}
 }
 
-static mtlpp::SamplerAddressMode TranslateWrapMode(ESamplerAddressMode AddressMode)
+static MTLSamplerAddressMode TranslateWrapMode(ESamplerAddressMode AddressMode)
 {
 	switch (AddressMode)
 	{
-		case AM_Clamp:	return mtlpp::SamplerAddressMode::ClampToEdge;
-		case AM_Mirror: return mtlpp::SamplerAddressMode::MirrorRepeat;
-		case AM_Border: return mtlpp::SamplerAddressMode::ClampToEdge;
-		default:		return mtlpp::SamplerAddressMode::Repeat;
+		case AM_Clamp:	return MTLSamplerAddressModeClampToEdge;
+		case AM_Mirror: return MTLSamplerAddressModeMirrorRepeat;
+		case AM_Border: return MTLSamplerAddressModeClampToEdge;
+		default:		return MTLSamplerAddressModeRepeat;
 	}
 }
 
-static mtlpp::CompareFunction TranslateCompareFunction(ECompareFunction CompareFunction)
+static MTLCompareFunction TranslateCompareFunction(ECompareFunction CompareFunction)
 {
 	switch(CompareFunction)
 	{
-		case CF_Less:			return mtlpp::CompareFunction::Less;
-		case CF_LessEqual:		return mtlpp::CompareFunction::LessEqual;
-		case CF_Greater:		return mtlpp::CompareFunction::Greater;
-		case CF_GreaterEqual:	return mtlpp::CompareFunction::GreaterEqual;
-		case CF_Equal:			return mtlpp::CompareFunction::Equal;
-		case CF_NotEqual:		return mtlpp::CompareFunction::NotEqual;
-		case CF_Never:			return mtlpp::CompareFunction::Never;
-		default:				return mtlpp::CompareFunction::Always;
+		case CF_Less:			return MTLCompareFunctionLess;
+		case CF_LessEqual:		return MTLCompareFunctionLessEqual;
+		case CF_Greater:		return MTLCompareFunctionGreater;
+		case CF_GreaterEqual:	return MTLCompareFunctionGreaterEqual;
+		case CF_Equal:			return MTLCompareFunctionEqual;
+		case CF_NotEqual:		return MTLCompareFunctionNotEqual;
+		case CF_Never:			return MTLCompareFunctionNever;
+		default:				return MTLCompareFunctionAlways;
 	};
 }
 
-static mtlpp::CompareFunction TranslateSamplerCompareFunction(ESamplerCompareFunction SamplerComparisonFunction)
+static MTLCompareFunction TranslateSamplerCompareFunction(ESamplerCompareFunction SamplerComparisonFunction)
 {
 	switch(SamplerComparisonFunction)
 	{
-		case SCF_Less:		return mtlpp::CompareFunction::Less;
-		case SCF_Never:		return mtlpp::CompareFunction::Never;
-		default:			return mtlpp::CompareFunction::Never;
+		case SCF_Less:		return MTLCompareFunctionLess;
+		case SCF_Never:		return MTLCompareFunctionNever;
+		default:			return MTLCompareFunctionNever;
 	};
 }
 
-static mtlpp::StencilOperation TranslateStencilOp(EStencilOp StencilOp)
+static MTLStencilOperation TranslateStencilOp(EStencilOp StencilOp)
 {
 	switch(StencilOp)
 	{
-		case SO_Zero:				return mtlpp::StencilOperation::Zero;
-		case SO_Replace:			return mtlpp::StencilOperation::Replace;
-		case SO_SaturatedIncrement:	return mtlpp::StencilOperation::IncrementClamp;
-		case SO_SaturatedDecrement:	return mtlpp::StencilOperation::DecrementClamp;
-		case SO_Invert:				return mtlpp::StencilOperation::Invert;
-		case SO_Increment:			return mtlpp::StencilOperation::IncrementWrap;
-		case SO_Decrement:			return mtlpp::StencilOperation::DecrementWrap;
-		default:					return mtlpp::StencilOperation::Keep;
+		case SO_Zero:				return MTLStencilOperationZero;
+		case SO_Replace:			return MTLStencilOperationReplace;
+		case SO_SaturatedIncrement:	return MTLStencilOperationIncrementClamp;
+		case SO_SaturatedDecrement:	return MTLStencilOperationDecrementClamp;
+		case SO_Invert:				return MTLStencilOperationInvert;
+		case SO_Increment:			return MTLStencilOperationIncrementWrap;
+		case SO_Decrement:			return MTLStencilOperationDecrementWrap;
+		default:					return MTLStencilOperationKeep;
 	};
 }
 
@@ -219,73 +219,76 @@ private:
 
 static FAGXStateObjectCache<FSamplerStateInitializerRHI, FAGXSampler> Samplers;
 
-static FAGXSampler FindOrCreateSamplerState(mtlpp::Device Device, const FSamplerStateInitializerRHI& Initializer)
+static FAGXSampler FindOrCreateSamplerState(const FSamplerStateInitializerRHI& Initializer)
 {
 	FAGXSampler State = Samplers.Find(Initializer);
 	if (!State.GetPtr())
 	{
-		mtlpp::SamplerDescriptor Desc;
+		MTLSamplerDescriptor* Desc = [[MTLSamplerDescriptor alloc] init];
+		
 		switch(Initializer.Filter)
 		{
 			case SF_AnisotropicLinear:
 			case SF_AnisotropicPoint:
-				Desc.SetMinFilter(mtlpp::SamplerMinMagFilter::Linear);
-				Desc.SetMagFilter(mtlpp::SamplerMinMagFilter::Linear);
-				Desc.SetMipFilter(mtlpp::SamplerMipFilter::Linear);
+				[Desc setMinFilter:MTLSamplerMinMagFilterLinear];
+				[Desc setMagFilter:MTLSamplerMinMagFilterLinear];
+				[Desc setMipFilter:MTLSamplerMipFilterLinear];
 				break;
 			case SF_Trilinear:
-				Desc.SetMinFilter(mtlpp::SamplerMinMagFilter::Linear);
-				Desc.SetMagFilter(mtlpp::SamplerMinMagFilter::Linear);
-				Desc.SetMipFilter(mtlpp::SamplerMipFilter::Linear);
+				[Desc setMinFilter:MTLSamplerMinMagFilterLinear];
+				[Desc setMagFilter:MTLSamplerMinMagFilterLinear];
+				[Desc setMipFilter:MTLSamplerMipFilterLinear];
 				break;
 			case SF_Bilinear:
-				Desc.SetMinFilter(mtlpp::SamplerMinMagFilter::Linear);
-				Desc.SetMagFilter(mtlpp::SamplerMinMagFilter::Linear);
-				Desc.SetMipFilter(mtlpp::SamplerMipFilter::Nearest);
+				[Desc setMinFilter:MTLSamplerMinMagFilterLinear];
+				[Desc setMagFilter:MTLSamplerMinMagFilterLinear];
+				[Desc setMipFilter:MTLSamplerMipFilterNearest];
 				break;
 			case SF_Point:
-				Desc.SetMinFilter(mtlpp::SamplerMinMagFilter::Nearest);
-				Desc.SetMagFilter(mtlpp::SamplerMinMagFilter::Nearest);
-				Desc.SetMipFilter(mtlpp::SamplerMipFilter::Nearest);
+				[Desc setMinFilter:MTLSamplerMinMagFilterNearest];
+				[Desc setMagFilter:MTLSamplerMinMagFilterNearest];
+				[Desc setMipFilter:MTLSamplerMipFilterNearest];
 				break;
 		}
-		Desc.SetMaxAnisotropy(GetMetalMaxAnisotropy(Initializer.Filter, Initializer.MaxAnisotropy));
-		Desc.SetSAddressMode(TranslateWrapMode(Initializer.AddressU));
-		Desc.SetTAddressMode(TranslateWrapMode(Initializer.AddressV));
-		Desc.SetRAddressMode(TranslateWrapMode(Initializer.AddressW));
-		Desc.SetLodMinClamp(Initializer.MinMipLevel);
-		Desc.SetLodMaxClamp(Initializer.MaxMipLevel);
+		[Desc setMaxAnisotropy:GetMetalMaxAnisotropy(Initializer.Filter, Initializer.MaxAnisotropy)];
+		[Desc setSAddressMode:TranslateWrapMode(Initializer.AddressU)];
+		[Desc setTAddressMode:TranslateWrapMode(Initializer.AddressV)];
+		[Desc setRAddressMode:TranslateWrapMode(Initializer.AddressW)];
+		[Desc setLodMinClamp:Initializer.MinMipLevel];
+		[Desc setLodMaxClamp:Initializer.MaxMipLevel];
 #if PLATFORM_TVOS
-		Desc.SetCompareFunction(mtlpp::CompareFunction::Never);	
+		[Desc setCompareFunction:MTLCompareFunctionNever];
 #elif PLATFORM_IOS
-		Desc.SetCompareFunction(Device.SupportsFeatureSet(mtlpp::FeatureSet::iOS_GPUFamily3_v1) ? TranslateSamplerCompareFunction(Initializer.SamplerComparisonFunction) : mtlpp::CompareFunction::Never);
+		[Desc setCompareFunction:[GMtlDevice supportsFeatureSet:MTLFeatureSet_iOS_GPUFamily3_v1] ? TranslateSamplerCompareFunction(Initializer.SamplerComparisonFunction) : MTLCompareFunctionNever];
 #else
-		Desc.SetCompareFunction(TranslateSamplerCompareFunction(Initializer.SamplerComparisonFunction));
+		[Desc setCompareFunction:TranslateSamplerCompareFunction(Initializer.SamplerComparisonFunction)];
 #endif
 #if PLATFORM_MAC
-		Desc.SetBorderColor(Initializer.BorderColor == 0 ? mtlpp::SamplerBorderColor::TransparentBlack : mtlpp::SamplerBorderColor::OpaqueWhite);
+		[Desc setBorderColor:Initializer.BorderColor == 0 ? MTLSamplerBorderColorTransparentBlack : MTLSamplerBorderColorOpaqueWhite];
 #endif
 		if (FAGXCommandQueue::SupportsFeature(EAGXFeaturesIABs))
 		{
-			Desc.SetSupportArgumentBuffers(true);
+			[Desc setSupportArgumentBuffers:YES];
 		}
 		
-		State = Device.NewSamplerState(Desc);
+		State = FAGXSampler([GMtlDevice newSamplerStateWithDescriptor:Desc], ns::Ownership::Assign);
+		
+		[Desc release];
 		
 		Samplers.Add(Initializer, State);
 	}
 	return State;
 }
 
-FAGXSamplerState::FAGXSamplerState(mtlpp::Device Device, const FSamplerStateInitializerRHI& Initializer)
+FAGXSamplerState::FAGXSamplerState(const FSamplerStateInitializerRHI& Initializer)
 {
-	State = FindOrCreateSamplerState(Device, Initializer);
+	State = FindOrCreateSamplerState(Initializer);
 #if !PLATFORM_MAC
 	if (GetMetalMaxAnisotropy(Initializer.Filter, Initializer.MaxAnisotropy))
 	{
 		FSamplerStateInitializerRHI Init = Initializer;
 		Init.MaxAnisotropy = 1;
-		NoAnisoState = FindOrCreateSamplerState(Device, Init);
+		NoAnisoState = FindOrCreateSamplerState(Init);
 	}
 #endif
 }
@@ -312,58 +315,63 @@ bool FAGXRasterizerState::GetInitializer(FRasterizerStateInitializerRHI& Init)
 
 static FAGXStateObjectCache<FDepthStencilStateInitializerRHI, mtlpp::DepthStencilState> DepthStencilStates;
 
-FAGXDepthStencilState::FAGXDepthStencilState(mtlpp::Device Device, const FDepthStencilStateInitializerRHI& InInitializer)
+FAGXDepthStencilState::FAGXDepthStencilState(const FDepthStencilStateInitializerRHI& InInitializer)
 {
 	Initializer = InInitializer;
 
 	State = DepthStencilStates.Find(Initializer);
 	if (!State.GetPtr())
 	{
-		mtlpp::DepthStencilDescriptor Desc;
+		MTLDepthStencilDescriptor* Desc = [[MTLDepthStencilDescriptor alloc] init];
 		
-		Desc.SetDepthCompareFunction(TranslateCompareFunction(Initializer.DepthTest));
-		Desc.SetDepthWriteEnabled(Initializer.bEnableDepthWrite);
+		[Desc setDepthCompareFunction:TranslateCompareFunction(Initializer.DepthTest)];
+		[Desc setDepthWriteEnabled:Initializer.bEnableDepthWrite];
 		
 		if (Initializer.bEnableFrontFaceStencil)
 		{
 			// set up front face stencil operations
-			mtlpp::StencilDescriptor Stencil;
-			Stencil.SetStencilCompareFunction(TranslateCompareFunction(Initializer.FrontFaceStencilTest));
-			Stencil.SetStencilFailureOperation(TranslateStencilOp(Initializer.FrontFaceStencilFailStencilOp));
-			Stencil.SetDepthFailureOperation(TranslateStencilOp(Initializer.FrontFaceDepthFailStencilOp));
-			Stencil.SetDepthStencilPassOperation(TranslateStencilOp(Initializer.FrontFacePassStencilOp));
-			Stencil.SetReadMask(Initializer.StencilReadMask);
-			Stencil.SetWriteMask(Initializer.StencilWriteMask);
-			Desc.SetFrontFaceStencil(Stencil);
+			MTLStencilDescriptor* Stencil = [[MTLStencilDescriptor alloc] init];
+			[Stencil setStencilCompareFunction:TranslateCompareFunction(Initializer.FrontFaceStencilTest)];
+			[Stencil setStencilFailureOperation:TranslateStencilOp(Initializer.FrontFaceStencilFailStencilOp)];
+			[Stencil setDepthFailureOperation:TranslateStencilOp(Initializer.FrontFaceDepthFailStencilOp)];
+			[Stencil setDepthStencilPassOperation:TranslateStencilOp(Initializer.FrontFacePassStencilOp)];
+			[Stencil setReadMask:Initializer.StencilReadMask];
+			[Stencil setWriteMask:Initializer.StencilWriteMask];
+			[Desc setFrontFaceStencil:Stencil];
+			[Stencil release];
 		}
 		
 		if (Initializer.bEnableBackFaceStencil)
 		{
 			// set up back face stencil operations
-			mtlpp::StencilDescriptor Stencil;
-			Stencil.SetStencilCompareFunction(TranslateCompareFunction(Initializer.BackFaceStencilTest));
-			Stencil.SetStencilFailureOperation(TranslateStencilOp(Initializer.BackFaceStencilFailStencilOp));
-			Stencil.SetDepthFailureOperation(TranslateStencilOp(Initializer.BackFaceDepthFailStencilOp));
-			Stencil.SetDepthStencilPassOperation(TranslateStencilOp(Initializer.BackFacePassStencilOp));
-			Stencil.SetReadMask(Initializer.StencilReadMask);
-			Stencil.SetWriteMask(Initializer.StencilWriteMask);
-			Desc.SetBackFaceStencil(Stencil);
+			MTLStencilDescriptor* Stencil = [[MTLStencilDescriptor alloc] init];
+			[Stencil setStencilCompareFunction:TranslateCompareFunction(Initializer.BackFaceStencilTest)];
+			[Stencil setStencilFailureOperation:TranslateStencilOp(Initializer.BackFaceStencilFailStencilOp)];
+			[Stencil setDepthFailureOperation:TranslateStencilOp(Initializer.BackFaceDepthFailStencilOp)];
+			[Stencil setDepthStencilPassOperation:TranslateStencilOp(Initializer.BackFacePassStencilOp)];
+			[Stencil setReadMask:Initializer.StencilReadMask];
+			[Stencil setWriteMask:Initializer.StencilWriteMask];
+			[Desc setBackFaceStencil:Stencil];
+			[Stencil release];
 		}
 		else if(Initializer.bEnableFrontFaceStencil)
 		{
 			// set up back face stencil operations to front face in single-face mode
-			mtlpp::StencilDescriptor Stencil;
-			Stencil.SetStencilCompareFunction(TranslateCompareFunction(Initializer.FrontFaceStencilTest));
-			Stencil.SetStencilFailureOperation(TranslateStencilOp(Initializer.FrontFaceStencilFailStencilOp));
-			Stencil.SetDepthFailureOperation(TranslateStencilOp(Initializer.FrontFaceDepthFailStencilOp));
-			Stencil.SetDepthStencilPassOperation(TranslateStencilOp(Initializer.FrontFacePassStencilOp));
-			Stencil.SetReadMask(Initializer.StencilReadMask);
-			Stencil.SetWriteMask(Initializer.StencilWriteMask);
-			Desc.SetBackFaceStencil(Stencil);
+			MTLStencilDescriptor* Stencil = [[MTLStencilDescriptor alloc] init];
+			[Stencil setStencilCompareFunction:TranslateCompareFunction(Initializer.FrontFaceStencilTest)];
+			[Stencil setStencilFailureOperation:TranslateStencilOp(Initializer.FrontFaceStencilFailStencilOp)];
+			[Stencil setDepthFailureOperation:TranslateStencilOp(Initializer.FrontFaceDepthFailStencilOp)];
+			[Stencil setDepthStencilPassOperation:TranslateStencilOp(Initializer.FrontFacePassStencilOp)];
+			[Stencil setReadMask:Initializer.StencilReadMask];
+			[Stencil setWriteMask:Initializer.StencilWriteMask];
+			[Desc setBackFaceStencil:Stencil];
+			[Stencil release];
 		}
 		
 		// bake out the descriptor
-		State = Device.NewDepthStencilState(Desc);
+		State = mtlpp::DepthStencilState([GMtlDevice newDepthStencilStateWithDescriptor:Desc], nullptr, ns::Ownership::Assign);
+		
+		[Desc release];
 		
 		DepthStencilStates.Add(Initializer, State);
 	}
@@ -496,7 +504,7 @@ bool FAGXBlendState::GetInitializer(FBlendStateInitializerRHI& Initializer)
 FSamplerStateRHIRef FAGXDynamicRHI::RHICreateSamplerState(const FSamplerStateInitializerRHI& Initializer)
 {
     @autoreleasepool {
-	return new FAGXSamplerState(ImmediateContext.Context->GetDevice(), Initializer);
+	return new FAGXSamplerState(Initializer);
 	}
 }
 
@@ -510,7 +518,7 @@ FRasterizerStateRHIRef FAGXDynamicRHI::RHICreateRasterizerState(const FRasterize
 FDepthStencilStateRHIRef FAGXDynamicRHI::RHICreateDepthStencilState(const FDepthStencilStateInitializerRHI& Initializer)
 {
 	@autoreleasepool {
-	return new FAGXDepthStencilState(ImmediateContext.Context->GetDevice(), Initializer);
+	return new FAGXDepthStencilState(Initializer);
 	}
 }
 
diff --git a/Engine/Source/Runtime/Apple/AGXRHI/Private/AGXStateCache.cpp b/Engine/Source/Runtime/Apple/AGXRHI/Private/AGXStateCache.cpp
--- a/Engine/Source/Runtime/Apple/AGXRHI/Private/AGXStateCache.cpp
+++ b/Engine/Source/Runtime/Apple/AGXRHI/Private/AGXStateCache.cpp
@@ -2238,7 +2238,7 @@ FAGXBuffer& FAGXStateCache::GetDebugBuffer()
     {
         // Assume worst case tiling (16x16) and render-target size (4096x4096) on iOS for now
         uint32 NumTiles = PLATFORM_MAC ? 1 : 65536;
-        DebugBuffer = GetAGXDeviceContext().CreatePooledBuffer(FAGXPooledBufferArgs(GetAGXDeviceContext().GetDevice(), NumTiles * sizeof(FAGXDebugInfo), BUF_Dynamic, mtlpp::StorageMode::Shared));
+        DebugBuffer = GetAGXDeviceContext().CreatePooledBuffer(FAGXPooledBufferArgs(NumTiles * sizeof(FAGXDebugInfo), BUF_Dynamic, mtlpp::StorageMode::Shared));
     }
     return DebugBuffer;
 }
diff --git a/Engine/Source/Runtime/Apple/AGXRHI/Private/AGXTexture.cpp b/Engine/Source/Runtime/Apple/AGXRHI/Private/AGXTexture.cpp
--- a/Engine/Source/Runtime/Apple/AGXRHI/Private/AGXTexture.cpp
+++ b/Engine/Source/Runtime/Apple/AGXRHI/Private/AGXTexture.cpp
@@ -229,12 +229,12 @@ static mtlpp::PixelFormat AGX_ToSRGBFormat_AppleGPU(mtlpp::PixelFormat MTLFormat
 
 mtlpp::PixelFormat AGXToSRGBFormat(mtlpp::PixelFormat MTLFormat)
 {
-	if([GetAGXDeviceContext().GetDevice().GetPtr() supportsFamily:MTLGPUFamilyApple1])
+	if([GMtlDevice supportsFamily:MTLGPUFamilyApple1])
 	{
 		MTLFormat = AGX_ToSRGBFormat_AppleGPU(MTLFormat);
 	}
 #if !PLATFORM_IOS
-	else if([GetAGXDeviceContext().GetDevice().GetPtr() supportsFamily:MTLGPUFamilyMac1])
+	else if([GMtlDevice supportsFamily:MTLGPUFamilyMac1])
 	{
 		MTLFormat = AGX_ToSRGBFormat_NonAppleMacGPU(MTLFormat);
 	}
@@ -857,7 +857,7 @@ FAGXSurface::FAGXSurface(ERHIResourceType ResourceType, EPixelFormat Format, uin
 			CFRetain(ImageSurfaceRef);
 			{
 #if !COREVIDEO_SUPPORTS_METAL
-				Texture = MTLPP_VALIDATE(mtlpp::Device, GetAGXDeviceContext().GetDevice(), AGXSafeGetRuntimeDebuggingLevel() >= EAGXDebugLevelValidation, NewTexture(Desc, CVPixelBufferGetIOSurface((CVPixelBufferRef)ImageSurfaceRef), 0));
+				Texture = MTLPP_VALIDATE(mtlpp::Device, GMtlppDevice, AGXSafeGetRuntimeDebuggingLevel() >= EAGXDebugLevelValidation, NewTexture(Desc, CVPixelBufferGetIOSurface((CVPixelBufferRef)ImageSurfaceRef), 0));
 #else
 				Texture = CVMetalTextureGetTexture((CVMetalTextureRef)ImageSurfaceRef);
 #endif
@@ -875,7 +875,8 @@ FAGXSurface::FAGXSurface(ERHIResourceType ResourceType, EPixelFormat Format, uin
 			Desc.SetStorageMode(mtlpp::StorageMode::Managed);
 			Desc.SetResourceOptions((mtlpp::ResourceOptions)((Desc.GetResourceOptions() & ~(mtlpp::ResourceStorageModeMask)) | mtlpp::ResourceOptions::StorageModeManaged));
 			
-			Texture = [GetAGXDeviceContext().GetDevice() newTextureWithDescriptor:Desc iosurface : (IOSurfaceRef)ImageSurfaceRef plane : 0];
+			ns::IOSurface MtlppIOSurface((IOSurfaceRef)ImageSurfaceRef, ns::Ownership::AutoRelease);
+			Texture = MTLPP_VALIDATE(mtlpp::Device, GMtlppDevice, AGXSafeGetRuntimeDebuggingLevel() >= EAGXDebugLevelValidation, NewTextureWithDescriptor(Desc, MtlppIOSurface, 0));
 			
 			METAL_FATAL_ASSERT(Texture, TEXT("Failed to create texture, desc %s"), *FString([Desc description]));
 			
@@ -893,11 +894,10 @@ FAGXSurface::FAGXSurface(ERHIResourceType ResourceType, EPixelFormat Format, uin
 		}
 		else
 		{
-			mtlpp::Device Device = GetAGXDeviceContext().GetDevice();
-			mtlpp::SizeAndAlign SizeAlign = Device.HeapTextureSizeAndAlign(Desc);
+			mtlpp::SizeAndAlign SizeAlign = GMtlppDevice.HeapTextureSizeAndAlign(Desc);
 
 			// Backing buffer resource options must match the texture we are going to create from it
-			FAGXPooledBufferArgs Args(Device, SizeAlign.Size, BUF_Dynamic, mtlpp::StorageMode::Private, Desc.GetCpuCacheMode());
+			FAGXPooledBufferArgs Args(SizeAlign.Size, BUF_Dynamic, mtlpp::StorageMode::Private, Desc.GetCpuCacheMode());
 			FAGXBuffer Buffer = GetAGXDeviceContext().CreatePooledBuffer(Args);
 
 			Texture = Buffer.NewTexture(Desc, 0, Align(SizeAlign.Size / SizeY, 256));
@@ -1166,8 +1166,7 @@ id <MTLBuffer> FAGXSurface::AllocSurface(uint32 MipIndex, uint32 ArrayIndex, ERe
 	
 	// allocate some temporary memory
 	// This should really be pooled and texture transfers should be their own pool
-	id <MTLDevice> Device = GetAGXDeviceContext().GetDevice();
-	id <MTLBuffer> Buffer = [Device newBufferWithLength:MipBytes options:MTLResourceStorageModeShared];
+	id <MTLBuffer> Buffer = [GMtlDevice newBufferWithLength:MipBytes options:MTLResourceStorageModeShared];
 	Buffer.label = @"Temporary Surface Backing";
 	
 	// Note: while the lock is active, this map owns the backing store.
@@ -1974,7 +1973,7 @@ static FAGXBuffer InternalCopyTexture2DUpdateRegion(FRHITexture2D* TextureRHI, c
 	if(Texture->GetFormat() == PF_G8 && EnumHasAnyFlags(Texture->GetFlags(), TexCreate_SRGB))
 	{
 		const uint32 BufferSize = UpdateRegion.Height * UpdateRegion.Width * sizeof(uint32);		
-		Buffer = GetAGXDeviceContext().CreatePooledBuffer(FAGXPooledBufferArgs(GetAGXDeviceContext().GetDevice(), BufferSize, BUF_Dynamic, mtlpp::StorageMode::Shared));
+		Buffer = GetAGXDeviceContext().CreatePooledBuffer(FAGXPooledBufferArgs(BufferSize, BUF_Dynamic, mtlpp::StorageMode::Shared));
 		InternalExpandR8ToStandardRGBA((uint32*)Buffer.GetContents(), UpdateRegion, InOutSourcePitch, SourceData);
 	}
 
@@ -1982,7 +1981,7 @@ static FAGXBuffer InternalCopyTexture2DUpdateRegion(FRHITexture2D* TextureRHI, c
 #endif
 	{
 		const uint32 BufferSize = UpdateRegion.Height * InSourcePitch;
-		Buffer = GetAGXDeviceContext().CreatePooledBuffer(FAGXPooledBufferArgs(GetAGXDeviceContext().GetDevice(), BufferSize, BUF_Dynamic, mtlpp::StorageMode::Shared));
+		Buffer = GetAGXDeviceContext().CreatePooledBuffer(FAGXPooledBufferArgs(BufferSize, BUF_Dynamic, mtlpp::StorageMode::Shared));
 
 		// Limit copy to line by line by update region pitch otherwise we can go off the end of source data on the last row
 		uint8* pDestRow = (uint8*)Buffer.GetContents();
@@ -2186,7 +2185,7 @@ struct FAGXDynamicRHIUpdateTexture3DCommand final : public FRHICommand<FAGXDynam
 		FAGXTexture Tex = Texture->Surface.Texture;
 		const uint32 BufferSize = UpdateRegion.Height * UpdateRegion.Depth* SourceRowPitch;
 		
-		Buffer = GetAGXDeviceContext().CreatePooledBuffer(FAGXPooledBufferArgs(GetAGXDeviceContext().GetDevice(), BufferSize, BUF_Dynamic, mtlpp::StorageMode::Shared));
+		Buffer = GetAGXDeviceContext().CreatePooledBuffer(FAGXPooledBufferArgs(BufferSize, BUF_Dynamic, mtlpp::StorageMode::Shared));
 		InternalCopyTexture3DUpdateRegionData(Texture, UpdateRegion, SourceRowPitch, SourceDepthPitch, SourceData, (uint8*)Buffer.GetContents());
 	}
 	
@@ -2257,7 +2256,7 @@ void FAGXDynamicRHI::RHIUpdateTexture3D(FRHITexture3D* TextureRHI,uint32 MipInde
 			SCOPED_AUTORELEASE_POOL;
 
 			const uint32 BufferSize = UpdateRegion.Height * UpdateRegion.Depth * SourceRowPitch;
-			FAGXBuffer IntermediateBuffer = GetAGXDeviceContext().CreatePooledBuffer(FAGXPooledBufferArgs(GetAGXDeviceContext().GetDevice(), BufferSize, BUF_Dynamic, mtlpp::StorageMode::Shared));
+			FAGXBuffer IntermediateBuffer = GetAGXDeviceContext().CreatePooledBuffer(FAGXPooledBufferArgs(BufferSize, BUF_Dynamic, mtlpp::StorageMode::Shared));
 			InternalCopyTexture3DUpdateRegionData(TextureRHI, UpdateRegion, SourceRowPitch, SourceDepthPitch, SourceData, (uint8*)IntermediateBuffer.GetContents());
 			InternalUpdateTexture3D(ImmediateContext.GetInternalContext(), TextureRHI, MipIndex, UpdateRegion, SourceRowPitch, SourceDepthPitch, IntermediateBuffer);
 			GetAGXDeviceContext().ReleaseBuffer(IntermediateBuffer);
@@ -2611,7 +2610,7 @@ void FAGXDynamicRHI::RHICopySubTextureRegion(FRHITexture2D* SourceTexture, FRHIT
 				const uint32 AlignedStride = ((Stride - 1) & ~(Alignment - 1)) + Alignment;
 				const uint32 BytesPerImage = AlignedStride *  SourceSize.height;
 				
-				FAGXBuffer Buffer = GetAGXDeviceContext().CreatePooledBuffer(FAGXPooledBufferArgs(ImmediateContext.Context->GetDevice(), BytesPerImage, BUF_Dynamic, mtlpp::StorageMode::Shared));
+				FAGXBuffer Buffer = GetAGXDeviceContext().CreatePooledBuffer(FAGXPooledBufferArgs(BytesPerImage, BUF_Dynamic, mtlpp::StorageMode::Shared));
 				
 				check(Buffer);
 				
@@ -2740,7 +2739,7 @@ void FAGXRHICommandContext::RHICopyTexture(FRHITexture* SourceTextureRHI, FRHITe
 						const uint32 BytesPerImage = AlignedStride *  SourceSize.height;
 						const uint32 DataSize = BytesPerImage * SourceSize.depth;
 						
-						FAGXBuffer Buffer = GetAGXDeviceContext().CreatePooledBuffer(FAGXPooledBufferArgs(GetInternalContext().GetDevice(), DataSize, BUF_Dynamic, mtlpp::StorageMode::Shared));
+						FAGXBuffer Buffer = GetAGXDeviceContext().CreatePooledBuffer(FAGXPooledBufferArgs(DataSize, BUF_Dynamic, mtlpp::StorageMode::Shared));
 						
 						check(Buffer);
 						
diff --git a/Engine/Source/Runtime/Apple/AGXRHI/Private/AGXUAV.cpp b/Engine/Source/Runtime/Apple/AGXRHI/Private/AGXUAV.cpp
--- a/Engine/Source/Runtime/Apple/AGXRHI/Private/AGXUAV.cpp
+++ b/Engine/Source/Runtime/Apple/AGXRHI/Private/AGXUAV.cpp
@@ -465,7 +465,7 @@ void FAGXRHICommandContext::ClearUAVWithBlitEncoder(FRHIUnorderedAccessView* Uno
 	};
 
 	uint32 AlignedSize = Align(Size, BufferOffsetAlignment);
-	FAGXPooledBufferArgs Args(GetAGXDeviceContext().GetDevice(), AlignedSize, BUF_Dynamic, mtlpp::StorageMode::Shared);
+	FAGXPooledBufferArgs Args(AlignedSize, BUF_Dynamic, mtlpp::StorageMode::Shared);
 	FAGXBuffer Temp = GetAGXDeviceContext().CreatePooledBuffer(Args);
 	uint32* ContentBytes = (uint32*)Temp.GetContents();
 	for (uint32 Element = 0; Element < (AlignedSize >> 2); ++Element)
@@ -616,7 +616,7 @@ void FAGXRHICommandContext::RHICopyToStagingBuffer(FRHIBuffer* SourceBufferRHI,
 			{
 				AGXSafeReleaseMetalBuffer(ReadbackBuffer);
 			}
-			FAGXPooledBufferArgs ArgsCPU(GetAGXDeviceContext().GetDevice(), NumBytes, BUF_Dynamic, mtlpp::StorageMode::Shared);
+			FAGXPooledBufferArgs ArgsCPU(NumBytes, BUF_Dynamic, mtlpp::StorageMode::Shared);
 			ReadbackBuffer = GetAGXDeviceContext().CreatePooledBuffer(ArgsCPU);
 		}
 
diff --git a/Engine/Source/Runtime/Apple/AGXRHI/Private/AGXVertexBuffer.cpp b/Engine/Source/Runtime/Apple/AGXRHI/Private/AGXVertexBuffer.cpp
--- a/Engine/Source/Runtime/Apple/AGXRHI/Private/AGXVertexBuffer.cpp
+++ b/Engine/Source/Runtime/Apple/AGXRHI/Private/AGXVertexBuffer.cpp
@@ -239,7 +239,7 @@ FAGXRHIBuffer::FAGXRHIBuffer(uint32 InSize, EBufferUsageFlags InUsage, EAGXBuffe
 				FAGXBufferAndViews& Backing = BufferPool[i];
 
 #if METAL_POOL_BUFFER_BACKING
-				FAGXPooledBufferArgs ArgsCPU(GetAGXDeviceContext().GetDevice(), AllocSize, Usage, Mode);
+				FAGXPooledBufferArgs ArgsCPU(AllocSize, Usage, Mode);
 				Backing.Buffer = GetAGXDeviceContext().CreatePooledBuffer(ArgsCPU);
 				Backing.Buffer.SetOwner(nullptr, false);
 #else
@@ -248,11 +248,11 @@ FAGXRHIBuffer::FAGXRHIBuffer(uint32 InSize, EBufferUsageFlags InUsage, EAGXBuffe
 				
 				METAL_GPUPROFILE(FAGXScopedCPUStats CPUStat(FString::Printf(TEXT("AllocBuffer: %llu, %llu"), AllocSize, Options)));
 				// Allocate one.
-				Backing.Buffer = FAGXBuffer(MTLPP_VALIDATE(mtlpp::Device, GetAGXDeviceContext().GetDevice(), AGXSafeGetRuntimeDebuggingLevel() >= EAGXDebugLevelValidation, NewBuffer(AllocSize, (mtlpp::ResourceOptions) Options)), false);
+				Backing.Buffer = FAGXBuffer(MTLPP_VALIDATE(mtlpp::Device, GMtlppDevice, AGXSafeGetRuntimeDebuggingLevel() >= EAGXDebugLevelValidation, NewBuffer(AllocSize, (mtlpp::ResourceOptions) Options)), false);
 				
-				#if STATS || ENABLE_LOW_LEVEL_MEM_TRACKER
-					AGXLLM::LogAllocBuffer(GetAGXDeviceContext().GetDevice(), Backing.Buffer);
-				#endif
+#if STATS || ENABLE_LOW_LEVEL_MEM_TRACKER
+				AGXLLM::LogAllocBuffer(Backing.Buffer);
+#endif
 				INC_MEMORY_STAT_BY(STAT_AGXDeviceBufferMemory, Backing.Buffer.GetLength());
 				
 				if (GAGXBufferZeroFill && Mode != mtlpp::StorageMode::Private)
@@ -319,7 +319,7 @@ FAGXRHIBuffer::~FAGXRHIBuffer()
 void FAGXRHIBuffer::AllocTransferBuffer(bool bOnRHIThread, uint32 InSize, EResourceLockMode LockMode)
 {
 	check(!TransferBuffer);
-	FAGXPooledBufferArgs ArgsCPU(GetAGXDeviceContext().GetDevice(), InSize, BUF_Dynamic, mtlpp::StorageMode::Shared);
+	FAGXPooledBufferArgs ArgsCPU(InSize, BUF_Dynamic, mtlpp::StorageMode::Shared);
 	TransferBuffer = GetAGXDeviceContext().CreatePooledBuffer(ArgsCPU);
 	TransferBuffer.SetOwner(nullptr, false);
 	check(TransferBuffer && TransferBuffer.GetPtr());
@@ -361,7 +361,7 @@ void FAGXRHIBuffer::AllocLinearTextures(const LinearTextureMapKey& InLinearTextu
 			BytesPerElement = 4;
 		}
 
-		const uint32 MinimumByteAlignment = GetAGXDeviceContext().GetDevice().GetMinimumLinearTextureAlignmentForPixelFormat((mtlpp::PixelFormat)GAGXBufferFormats[InFormat].LinearTextureFormat);
+		const uint32 MinimumByteAlignment = GMtlppDevice.GetMinimumLinearTextureAlignmentForPixelFormat((mtlpp::PixelFormat)GAGXBufferFormats[InFormat].LinearTextureFormat);
 		const uint32 MinimumElementAlignment = MinimumByteAlignment / BytesPerElement;
 
 		uint32 Offset = LinearTextureDesc.StartOffsetBytes;
@@ -815,7 +815,7 @@ void FAGXDynamicRHI::RHICopyBuffer(FRHIBuffer* SourceBufferRHI, FRHIBuffer* Dest
 		}
 		else if (TheDstBuffer)
 		{
-			FAGXPooledBufferArgs ArgsCPU(GetAGXDeviceContext().GetDevice(), SrcBuffer->GetSize(), BUF_Dynamic, mtlpp::StorageMode::Shared);
+			FAGXPooledBufferArgs ArgsCPU(SrcBuffer->GetSize(), BUF_Dynamic, mtlpp::StorageMode::Shared);
 			FAGXBuffer TempBuffer = GetAGXDeviceContext().CreatePooledBuffer(ArgsCPU);
 			FMemory::Memcpy(TempBuffer.GetContents(), SrcBuffer->Data->Data, SrcBuffer->GetSize());
 			GetAGXDeviceContext().CopyFromBufferToBuffer(TempBuffer, 0, TheDstBuffer, 0, FMath::Min(SrcBuffer->GetSize(), DstBuffer->GetSize()));
diff --git a/Engine/Source/Runtime/Apple/AGXRHI/Private/AGXViewport.cpp b/Engine/Source/Runtime/Apple/AGXRHI/Private/AGXViewport.cpp
--- a/Engine/Source/Runtime/Apple/AGXRHI/Private/AGXViewport.cpp
+++ b/Engine/Source/Runtime/Apple/AGXRHI/Private/AGXViewport.cpp
@@ -101,7 +101,7 @@ FAGXViewport::FAGXViewport(void* WindowHandle, uint32 InSizeX, uint32 InSizeY, b
 		Layer.magnificationFilter = kCAFilterNearest;
 		Layer.minificationFilter = kCAFilterNearest;
 
-		[Layer setDevice:GetAGXDeviceContext().GetDevice()];
+		[Layer setDevice:GMtlDevice];
 		
 		[Layer setFramebufferOnly:NO];
 		[Layer removeAllAnimations];
diff --git a/Engine/Source/Runtime/Apple/AGXRHI/Private/Shaders/Types/AGXComputeShader.cpp b/Engine/Source/Runtime/Apple/AGXRHI/Private/Shaders/Types/AGXComputeShader.cpp
--- a/Engine/Source/Runtime/Apple/AGXRHI/Private/Shaders/Types/AGXComputeShader.cpp
+++ b/Engine/Source/Runtime/Apple/AGXRHI/Private/Shaders/Types/AGXComputeShader.cpp
@@ -85,7 +85,7 @@ FAGXShaderPipeline* FAGXComputeShader::GetPipeline()
 			mtlpp::AutoReleasedComputePipelineReflection ComputeReflection;
 
 			NSUInteger ComputeOption = mtlpp::PipelineOption::ArgumentInfo | mtlpp::PipelineOption::BufferTypeInfo;
-			Kernel = GetAGXDeviceContext().GetDevice().NewComputePipelineState(Descriptor, mtlpp::PipelineOption(ComputeOption), &ComputeReflection, &ComputeError);
+			Kernel = GMtlppDevice.NewComputePipelineState(Descriptor, mtlpp::PipelineOption(ComputeOption), &ComputeReflection, &ComputeError);
 			Error = ComputeError;
 			Reflection = ComputeReflection;
 		}
@@ -93,7 +93,7 @@ FAGXShaderPipeline* FAGXComputeShader::GetPipeline()
 #endif // METAL_DEBUG_OPTIONS
 		{
 			ns::AutoReleasedError ComputeError;
-			Kernel = GetAGXDeviceContext().GetDevice().NewComputePipelineState(Descriptor, mtlpp::PipelineOption(0), nil, &ComputeError);
+			Kernel = GMtlppDevice.NewComputePipelineState(Descriptor, mtlpp::PipelineOption(0), nil, &ComputeError);
 			Error = ComputeError;
 		}
 
diff --git a/Engine/Source/Runtime/Apple/AGXRHI/Private/Shaders/Types/Templates/AGXBaseShader.h b/Engine/Source/Runtime/Apple/AGXRHI/Private/Shaders/Types/Templates/AGXBaseShader.h
--- a/Engine/Source/Runtime/Apple/AGXRHI/Private/Shaders/Types/Templates/AGXBaseShader.h
+++ b/Engine/Source/Runtime/Apple/AGXRHI/Private/Shaders/Types/Templates/AGXBaseShader.h
@@ -247,7 +247,7 @@ void TAGXBaseShader<BaseResourceType, ShaderType>::Init(TArrayView<const uint8>
 				dispatch_data_t GCDBuffer = dispatch_data_create(Buffer, BufferSize, dispatch_get_main_queue(), ^(void) { FMemory::Free(Buffer); } );
 
 				// load up the already compiled shader
-				Library = GetAGXDeviceContext().GetDevice().NewLibrary(GCDBuffer, &AError);
+				Library = GMtlppDevice.NewLibrary(GCDBuffer, &AError);
 				dispatch_release(GCDBuffer);
 
 				if (Library == nil)
@@ -334,7 +334,7 @@ void TAGXBaseShader<BaseResourceType, ShaderType>::Init(TArrayView<const uint8>
 			CompileOptions.SetLanguageVersion(MetalVersion);
 
 			ns::AutoReleasedError Error;
-			Library = GetAGXDeviceContext().GetDevice().NewLibrary(NewShaderString, CompileOptions, &Error);
+			Library = GMtlppDevice.NewLibrary(NewShaderString, CompileOptions, &Error);
 			if (Library == nil)
 			{
 				UE_LOG(LogRHI, Error, TEXT("*********** Error\n%s"), *FString(NewShaderString));
diff --git a/Engine/Source/Runtime/Apple/AGXRHI/Public/AGXFrameAllocator.h b/Engine/Source/Runtime/Apple/AGXRHI/Public/AGXFrameAllocator.h
--- a/Engine/Source/Runtime/Apple/AGXRHI/Public/AGXFrameAllocator.h
+++ b/Engine/Source/Runtime/Apple/AGXRHI/Public/AGXFrameAllocator.h
@@ -20,9 +20,8 @@ public:
     };
     
 private:
-    // The device. This object will take a reference.
-    // You must make sure your rendering is on this same device or things will explode horribly.
-    id <MTLDevice> Device;
+	// The current buffer we are suballocating from and the current offset.
+	id <MTLBuffer> CurrentBuffer;
 #if METAL_FRAME_ALLOCATOR_VALIDATION
     // Every buffer that is currently allocated.
     NSMutableArray<id <MTLBuffer>>* AllAllocatedBuffers;
@@ -32,8 +31,6 @@ private:
     // The current set of buffers in use for this frame.
     // Must ONLY be mutated on RHI thread if present.
     NSMutableArray<id <MTLBuffer>>* BuffersInFlight;
-    // The current buffer we are suballocating from and the current offset.
-    id <MTLBuffer> CurrentBuffer;
     uint32 CurrentCursor;
     uint32 CurrentFrame;
     // Lock around AllAllocatedBuffers and AvailableBuffers since they are mutated from a completion handler.
@@ -52,8 +49,7 @@ private:
 public:
     // API
     
-    // This object will hold a reference to TargetDevice for its lifetime.
-    FAGXFrameAllocator(id <MTLDevice> TargetDevice);
+    FAGXFrameAllocator();
     ~FAGXFrameAllocator();
     
     // Sets a target allocation limit.
diff --git a/Engine/Source/Runtime/Apple/AGXRHI/Public/AGXState.h b/Engine/Source/Runtime/Apple/AGXRHI/Public/AGXState.h
--- a/Engine/Source/Runtime/Apple/AGXRHI/Public/AGXState.h
+++ b/Engine/Source/Runtime/Apple/AGXRHI/Public/AGXState.h
@@ -81,7 +81,7 @@ public:
 	/** 
 	 * Constructor/destructor
 	 */
-	FAGXSamplerState(mtlpp::Device Device, const FSamplerStateInitializerRHI& Initializer);
+	FAGXSamplerState(const FSamplerStateInitializerRHI& Initializer);
 	~FAGXSamplerState();
 
 	FAGXSampler State;
@@ -112,7 +112,7 @@ public:
 	/**
 	 * Constructor/destructor
 	 */
-	FAGXDepthStencilState(mtlpp::Device Device, const FDepthStencilStateInitializerRHI& Initializer);
+	FAGXDepthStencilState(const FDepthStencilStateInitializerRHI& Initializer);
 	~FAGXDepthStencilState();
 	
 	virtual bool GetInitializer(FDepthStencilStateInitializerRHI& Init) override final;
