Bring CL 17082967 into AGXRHI from Metal RHI.

---
Remove Metal texture resource write combined usage.   Causes problems with our buffer handing code when using buffer backed textures.

---

diff --git a/Engine/Source/Runtime/Apple/AGXRHI/Private/AGXBuffer.cpp b/Engine/Source/Runtime/Apple/AGXRHI/Private/AGXBuffer.cpp
--- a/Engine/Source/Runtime/Apple/AGXRHI/Private/AGXBuffer.cpp
+++ b/Engine/Source/Runtime/Apple/AGXRHI/Private/AGXBuffer.cpp
@@ -1176,7 +1176,7 @@ FAGXBuffer FAGXBufferPoolPolicyData::CreateResource(CreationArguments Args)
 
 FAGXBufferPoolPolicyData::CreationArguments FAGXBufferPoolPolicyData::GetCreationArguments(FAGXBuffer const& Resource)
 {
-	return FAGXBufferPoolPolicyData::CreationArguments(Resource.GetLength(), BUF_None, Resource.GetStorageMode());
+	return FAGXBufferPoolPolicyData::CreationArguments(Resource.GetLength(), BUF_None, Resource.GetStorageMode(), Resource.GetCpuCacheMode());
 }
 
 void FAGXBufferPoolPolicyData::FreeResource(FAGXBuffer& Resource)
@@ -1445,6 +1445,11 @@ FAGXBuffer FAGXResourceHeap::CreateBuffer(uint32 Size, uint32 Alignment, EBuffer
 	FAGXBuffer Buffer;
 	uint32 BlockSize = Align(Size, Alignment);
 	mtlpp::StorageMode StorageMode = (mtlpp::StorageMode)(((NSUInteger)Options & mtlpp::ResourceStorageModeMask) >> mtlpp::ResourceStorageModeShift);
+	mtlpp::CpuCacheMode CpuMode = (mtlpp::CpuCacheMode)(((NSUInteger)Options & mtlpp::ResourceCpuCacheModeMask) >> mtlpp::ResourceCpuCacheModeShift);
+	
+	// Write combined should be on a case by case basis
+	check(CpuMode == mtlpp::CpuCacheMode::DefaultCache);
+	
 	if (BlockSize <= 33554432)
 	{
 		switch (StorageMode)
@@ -1479,11 +1484,11 @@ FAGXBuffer FAGXResourceHeap::CreateBuffer(uint32 Size, uint32 Alignment, EBuffer
 				 }
 				 else
 				 {
-                    Buffer = ManagedBuffers.CreatePooledResource(FAGXPooledBufferArgs(BlockSize, Flags, StorageMode));
-                     if (GAGXResourcePurgeInPool)
-                     {
-                         Buffer.SetPurgeableState(mtlpp::PurgeableState::NonVolatile);
-                     }
+					Buffer = ManagedBuffers.CreatePooledResource(FAGXPooledBufferArgs(BlockSize, Flags, StorageMode, CpuMode));
+					if (GAGXResourcePurgeInPool)
+					{
+						Buffer.SetPurgeableState(mtlpp::PurgeableState::NonVolatile);
+					}
 					DEC_MEMORY_STAT_BY(STAT_AGXBufferUnusedMemory, Buffer.GetLength());
 					DEC_MEMORY_STAT_BY(STAT_AGXPooledBufferUnusedMemory, Buffer.GetLength());
 					INC_MEMORY_STAT_BY(STAT_AGXPooledBufferMemory, Buffer.GetLength());
@@ -1556,7 +1561,7 @@ FAGXBuffer FAGXResourceHeap::CreateBuffer(uint32 Size, uint32 Alignment, EBuffer
 				else
 				{
 					FScopeLock Lock(&Mutex);
-                    Buffer = Buffers[Storage].CreatePooledResource(FAGXPooledBufferArgs(BlockSize, Flags, StorageMode));
+					Buffer = Buffers[Storage].CreatePooledResource(FAGXPooledBufferArgs(BlockSize, Flags, StorageMode, CpuMode));
 					if (GAGXResourcePurgeInPool)
 					{
                    		Buffer.SetPurgeableState(mtlpp::PurgeableState::NonVolatile);
diff --git a/Engine/Source/Runtime/Apple/AGXRHI/Private/AGXTexture.cpp b/Engine/Source/Runtime/Apple/AGXRHI/Private/AGXTexture.cpp
--- a/Engine/Source/Runtime/Apple/AGXRHI/Private/AGXTexture.cpp
+++ b/Engine/Source/Runtime/Apple/AGXRHI/Private/AGXTexture.cpp
@@ -771,12 +771,11 @@ FAGXSurface::FAGXSurface(ERHIResourceType ResourceType, EPixelFormat Format, uin
 		}
 		else if(EnumHasAnyFlags(Flags, TexCreate_NoTiling) && !EnumHasAnyFlags(Flags, TexCreate_FastVRAM|TexCreate_DepthStencilTargetable|TexCreate_RenderTargetable|TexCreate_UAV))
 		{
+			Desc.SetCpuCacheMode(mtlpp::CpuCacheMode::DefaultCache);
 #if PLATFORM_MAC
-			Desc.SetCpuCacheMode(mtlpp::CpuCacheMode::WriteCombined);
 			Desc.SetStorageMode(mtlpp::StorageMode::Managed);
-			Desc.SetResourceOptions((mtlpp::ResourceOptions)(mtlpp::ResourceOptions::CpuCacheModeWriteCombined|mtlpp::ResourceOptions::StorageModeManaged));
+			Desc.SetResourceOptions((mtlpp::ResourceOptions)(mtlpp::ResourceOptions::CpuCacheModeDefaultCache|mtlpp::ResourceOptions::StorageModeManaged));
 #else
-			Desc.SetCpuCacheMode(mtlpp::CpuCacheMode::DefaultCache);
 			Desc.SetStorageMode(mtlpp::StorageMode::Shared);
 			Desc.SetResourceOptions((mtlpp::ResourceOptions)(mtlpp::ResourceOptions::CpuCacheModeDefaultCache|mtlpp::ResourceOptions::StorageModeShared));
 #endif
@@ -804,12 +803,11 @@ FAGXSurface::FAGXSurface(ERHIResourceType ResourceType, EPixelFormat Format, uin
 		else
 		{
 			check(!(Flags & TexCreate_CPUReadback));
+			Desc.SetCpuCacheMode(mtlpp::CpuCacheMode::DefaultCache);
 #if PLATFORM_MAC
-			Desc.SetCpuCacheMode(mtlpp::CpuCacheMode::WriteCombined);
 			Desc.SetStorageMode(mtlpp::StorageMode::Private);
-			Desc.SetResourceOptions((mtlpp::ResourceOptions)(mtlpp::ResourceOptions::CpuCacheModeWriteCombined|mtlpp::ResourceOptions::StorageModePrivate));
+			Desc.SetResourceOptions((mtlpp::ResourceOptions)(mtlpp::ResourceOptions::CpuCacheModeDefaultCache|mtlpp::ResourceOptions::StorageModePrivate));
 #else
-			Desc.SetCpuCacheMode(mtlpp::CpuCacheMode::DefaultCache);
 			if(GAGXForceIOSTexturesShared)
 			{
 				Desc.SetStorageMode(mtlpp::StorageMode::Shared);
