Fix-up iOS build.

#rnx

diff --git a/Engine/Source/Runtime/Apple/AGXRHI/Private/AGXRHI.cpp b/Engine/Source/Runtime/Apple/AGXRHI/Private/AGXRHI.cpp
--- a/Engine/Source/Runtime/Apple/AGXRHI/Private/AGXRHI.cpp
+++ b/Engine/Source/Runtime/Apple/AGXRHI/Private/AGXRHI.cpp
@@ -201,7 +201,7 @@ FAGXDynamicRHI::FAGXDynamicRHI(ERHIFeatureLevel::Type RequestedFeatureLevel)
 	GRHISupportsDrawIndirect = [GMtlDevice supportsFeatureSet:MTLFeatureSet_iOS_GPUFamily3_v1];
 	GRHISupportsPixelShaderUAVs = [GMtlDevice supportsFeatureSet:MTLFeatureSet_iOS_GPUFamily3_v1];
 
-	const mtlpp::FeatureSet FeatureSets[] = {
+	const MTLFeatureSet FeatureSets[] = {
 		MTLFeatureSet_iOS_GPUFamily1_v1,
 		MTLFeatureSet_iOS_GPUFamily2_v1,
 		MTLFeatureSet_iOS_GPUFamily3_v1,
