Port required elements of CL 16782183 (mac texture atomic emulation) from MetalRHI to AGXRHI.  Also pull in MetalRHI BytesPerRow fix from CL 17086762.

diff --git a/Engine/Source/Runtime/Apple/AGXRHI/Private/AGXBuffer.cpp b/Engine/Source/Runtime/Apple/AGXRHI/Private/AGXBuffer.cpp
--- a/Engine/Source/Runtime/Apple/AGXRHI/Private/AGXBuffer.cpp
+++ b/Engine/Source/Runtime/Apple/AGXRHI/Private/AGXBuffer.cpp
@@ -1164,8 +1164,14 @@ uint32 FAGXBufferPoolPolicyData::GetPoolBucketSize(uint32 Bucket)
 FAGXBuffer FAGXBufferPoolPolicyData::CreateResource(CreationArguments Args)
 {
 	uint32 BufferSize = GetPoolBucketSize(GetPoolBucketIndex(Args));
-	METAL_GPUPROFILE(FAGXScopedCPUStats CPUStat(FString::Printf(TEXT("AllocBuffer: %llu, %llu"), BufferSize, mtlpp::ResourceOptions(BUFFER_CACHE_MODE | ((NSUInteger)Args.Storage << mtlpp::ResourceStorageModeShift)))));
-	FAGXBuffer NewBuf(MTLPP_VALIDATE(mtlpp::Device, GMtlppDevice, AGXSafeGetRuntimeDebuggingLevel() >= EAGXDebugLevelValidation, NewBuffer(BufferSize, FAGXCommandQueue::GetCompatibleResourceOptions(mtlpp::ResourceOptions(BUFFER_CACHE_MODE | mtlpp::ResourceOptions::HazardTrackingModeUntracked | ((NSUInteger)Args.Storage << mtlpp::ResourceStorageModeShift))))), true);
+	
+	NSUInteger CpuCacheMode = (NSUInteger)Args.CpuCacheMode << mtlpp::ResourceCpuCacheModeShift;
+	NSUInteger StorageMode = (NSUInteger)Args.Storage << mtlpp::ResourceStorageModeShift;
+	mtlpp::ResourceOptions ResourceOptions = FAGXCommandQueue::GetCompatibleResourceOptions(mtlpp::ResourceOptions(CpuCacheMode | StorageMode | mtlpp::ResourceOptions::HazardTrackingModeUntracked));
+	
+	METAL_GPUPROFILE(FAGXScopedCPUStats CPUStat(FString::Printf(TEXT("AllocBuffer: %llu, %llu"), BufferSize, ResourceOptions)));
+	FAGXBuffer NewBuf(MTLPP_VALIDATE(mtlpp::Device, GMtlppDevice, AGXSafeGetRuntimeDebuggingLevel() >= EAGXDebugLevelValidation, NewBuffer(BufferSize, ResourceOptions), true));
+
 #if STATS || ENABLE_LOW_LEVEL_MEM_TRACKER
 	AGXLLM::LogAllocBuffer(NewBuf);
 #endif
@@ -1626,13 +1632,9 @@ void FAGXResourceHeap::ReleaseBuffer(FAGXBuffer& Buffer)
 			}
 	#endif
 			case mtlpp::StorageMode::Private:
-			{
-				Buffers[AllocPrivate].ReleasePooledResource(Buffer);
-				break;
-			}
 			case mtlpp::StorageMode::Shared:
 			{
-				Buffers[AllocShared].ReleasePooledResource(Buffer);
+				Buffers[(NSUInteger)StorageMode].ReleasePooledResource(Buffer);
 				break;
 			}
 			default:
diff --git a/Engine/Source/Runtime/Apple/AGXRHI/Private/AGXBuffer.h b/Engine/Source/Runtime/Apple/AGXRHI/Private/AGXBuffer.h
--- a/Engine/Source/Runtime/Apple/AGXRHI/Private/AGXBuffer.h
+++ b/Engine/Source/Runtime/Apple/AGXRHI/Private/AGXBuffer.h
@@ -11,7 +11,7 @@
 
 struct FAGXPooledBufferArgs
 {
-    FAGXPooledBufferArgs() : Size(0), Flags(BUF_None), Storage(mtlpp::StorageMode::Shared) {}
+    FAGXPooledBufferArgs() : Size(0), Flags(BUF_None), Storage(mtlpp::StorageMode::Shared), CpuCacheMode(mtlpp::CpuCacheMode::DefaultCache) {}
 	
     FAGXPooledBufferArgs(uint32 InSize, EBufferUsageFlags InFlags, mtlpp::StorageMode InStorage, mtlpp::CpuCacheMode InCpuCacheMode = mtlpp::CpuCacheMode::DefaultCache)
 	: Size(InSize)
diff --git a/Engine/Source/Runtime/Apple/AGXRHI/Private/AGXCommandEncoder.cpp b/Engine/Source/Runtime/Apple/AGXRHI/Private/AGXCommandEncoder.cpp
--- a/Engine/Source/Runtime/Apple/AGXRHI/Private/AGXCommandEncoder.cpp
+++ b/Engine/Source/Runtime/Apple/AGXRHI/Private/AGXCommandEncoder.cpp
@@ -1399,7 +1399,7 @@ void FAGXCommandEncoder::SetVisibilityResultMode(mtlpp::VisibilityResultMode con
 	
 #pragma mark - Public Shader Resource Mutators -
 
-void FAGXCommandEncoder::SetShaderBuffer(mtlpp::FunctionType const FunctionType, FAGXBuffer const& Buffer, NSUInteger const Offset, NSUInteger const Length, NSUInteger index, mtlpp::ResourceUsage const Usage, EPixelFormat const Format)
+void FAGXCommandEncoder::SetShaderBuffer(mtlpp::FunctionType const FunctionType, FAGXBuffer const& Buffer, NSUInteger const Offset, NSUInteger const Length, NSUInteger index, mtlpp::ResourceUsage const Usage, EPixelFormat const Format, NSUInteger const ElementRowPitch)
 {
 	check(index < ML_MaxBuffers);
     if(GetAGXDeviceContext().SupportsFeature(EAGXFeaturesSetBufferOffset) && Buffer && (ShaderBuffers[uint32(FunctionType)].Bound & (1 << index)) && ShaderBuffers[uint32(FunctionType)].Buffers[index] == Buffer)
@@ -1409,8 +1409,8 @@ void FAGXCommandEncoder::SetShaderBuffer(mtlpp::FunctionType const FunctionType,
 			FenceResource(Buffer);
 		}
 		SetShaderBufferOffset(FunctionType, Offset, Length, index);
-		ShaderBuffers[uint32(FunctionType)].Lengths[(index*2)+1] = GAGXBufferFormats[Format].DataFormat;
 		ShaderBuffers[uint32(FunctionType)].Usage[index] = Usage;
+		ShaderBuffers[uint32(FunctionType)].SetBufferMetaData(index, Length, GAGXBufferFormats[Format].DataFormat, ElementRowPitch);
 	}
     else
     {
@@ -1426,8 +1426,7 @@ void FAGXCommandEncoder::SetShaderBuffer(mtlpp::FunctionType const FunctionType,
 		ShaderBuffers[uint32(FunctionType)].Bytes[index] = nil;
 		ShaderBuffers[uint32(FunctionType)].Offsets[index] = Offset;
 		ShaderBuffers[uint32(FunctionType)].Usage[index] = Usage;
-		ShaderBuffers[uint32(FunctionType)].Lengths[index*2] = Length;
-		ShaderBuffers[uint32(FunctionType)].Lengths[(index*2)+1] = GAGXBufferFormats[Format].DataFormat;
+		ShaderBuffers[uint32(FunctionType)].SetBufferMetaData(index, Length, GAGXBufferFormats[Format].DataFormat, ElementRowPitch);
 		
 		SetShaderBufferInternal(FunctionType, index);
     }
@@ -1438,7 +1437,7 @@ void FAGXCommandEncoder::SetShaderBuffer(mtlpp::FunctionType const FunctionType,
 	}
 }
 
-void FAGXCommandEncoder::SetShaderData(mtlpp::FunctionType const FunctionType, FAGXBufferData* Data, NSUInteger const Offset, NSUInteger const Index, EPixelFormat const Format)
+void FAGXCommandEncoder::SetShaderData(mtlpp::FunctionType const FunctionType, FAGXBufferData* Data, NSUInteger const Offset, NSUInteger const Index, EPixelFormat const Format, NSUInteger const ElementRowPitch)
 {
 	check(Index < ML_MaxBuffers);
 	
@@ -1455,8 +1454,7 @@ void FAGXCommandEncoder::SetShaderData(mtlpp::FunctionType const FunctionType, F
 	ShaderBuffers[uint32(FunctionType)].Bytes[Index] = Data;
 	ShaderBuffers[uint32(FunctionType)].Offsets[Index] = Offset;
 	ShaderBuffers[uint32(FunctionType)].Usage[Index] = mtlpp::ResourceUsage::Read;
-	ShaderBuffers[uint32(FunctionType)].Lengths[Index*2] = Data ? (Data->Len - Offset) : 0;
-	ShaderBuffers[uint32(FunctionType)].Lengths[(Index*2)+1] = GAGXBufferFormats[Format].DataFormat;
+	ShaderBuffers[uint32(FunctionType)].SetBufferMetaData(Index, Data ? (Data->Len - Offset) : 0, GAGXBufferFormats[Format].DataFormat, ElementRowPitch);
 	
 	SetShaderBufferInternal(FunctionType, Index);
 }
@@ -1504,8 +1502,7 @@ void FAGXCommandEncoder::SetShaderBytes(mtlpp::FunctionType const FunctionType,
 		ShaderBuffers[uint32(FunctionType)].Bytes[Index] = nil;
 		ShaderBuffers[uint32(FunctionType)].Offsets[Index] = 0;
 		ShaderBuffers[uint32(FunctionType)].Usage[Index] = mtlpp::ResourceUsage::Read;
-		ShaderBuffers[uint32(FunctionType)].Lengths[Index*2] = Length;
-		ShaderBuffers[uint32(FunctionType)].Lengths[(Index*2)+1] = GAGXBufferFormats[PF_Unknown].DataFormat;
+		ShaderBuffers[uint32(FunctionType)].SetBufferMetaData(Index, Length, GAGXBufferFormats[PF_Unknown].DataFormat, 0);
 	}
 	else
 	{
@@ -1515,8 +1512,7 @@ void FAGXCommandEncoder::SetShaderBytes(mtlpp::FunctionType const FunctionType,
 		ShaderBuffers[uint32(FunctionType)].Bytes[Index] = nil;
 		ShaderBuffers[uint32(FunctionType)].Offsets[Index] = 0;
 		ShaderBuffers[uint32(FunctionType)].Usage[Index] = mtlpp::ResourceUsage(0);
-		ShaderBuffers[uint32(FunctionType)].Lengths[Index*2] = 0;
-		ShaderBuffers[uint32(FunctionType)].Lengths[(Index*2)+1] = GAGXBufferFormats[PF_Unknown].DataFormat;
+		ShaderBuffers[uint32(FunctionType)].SetBufferMetaData(Index, 0, GAGXBufferFormats[PF_Unknown].DataFormat, 0);
 	}
 	
 	SetShaderBufferInternal(FunctionType, Index);
@@ -1528,8 +1524,7 @@ void FAGXCommandEncoder::SetShaderBufferOffset(mtlpp::FunctionType FunctionType,
     checkf(ShaderBuffers[uint32(FunctionType)].Buffers[index] && (ShaderBuffers[uint32(FunctionType)].Bound & (1 << index)), TEXT("Buffer must already be bound"));
 	check(GetAGXDeviceContext().SupportsFeature(EAGXFeaturesSetBufferOffset));
 	ShaderBuffers[uint32(FunctionType)].Offsets[index] = Offset;
-	ShaderBuffers[uint32(FunctionType)].Lengths[index*2] = Length;
-	ShaderBuffers[uint32(FunctionType)].Lengths[(index*2)+1] = GAGXBufferFormats[PF_Unknown].DataFormat;
+	ShaderBuffers[uint32(FunctionType)].SetBufferMetaData(index, Length, GAGXBufferFormats[PF_Unknown].DataFormat, 0);
 	switch (FunctionType)
 	{
 		case mtlpp::FunctionType::Vertex:
@@ -1595,8 +1590,7 @@ void FAGXCommandEncoder::SetShaderTexture(mtlpp::FunctionType FunctionType, FAGX
 		{
 			Swizzle[0] = Swizzle[1] = Swizzle[2] = Swizzle[3] = 1;
 		}
-		FMemory::Memcpy(&ShaderBuffers[uint32(FunctionType)].Lengths[(ML_MaxBuffers*2)+(index*2)], Swizzle, sizeof(Swizzle));
-		ShaderBuffers[uint32(FunctionType)].Lengths[(ML_MaxBuffers*2)+(index*2)+1] = 0;
+		ShaderBuffers[uint32(FunctionType)].SetTextureSwizzle(index, Swizzle);
 		TextureBindingHistory.Add(ns::AutoReleased<FAGXTexture>(Texture));
 	}
 }
diff --git a/Engine/Source/Runtime/Apple/AGXRHI/Private/AGXCommandEncoder.h b/Engine/Source/Runtime/Apple/AGXRHI/Private/AGXCommandEncoder.h
--- a/Engine/Source/Runtime/Apple/AGXRHI/Private/AGXCommandEncoder.h
+++ b/Engine/Source/Runtime/Apple/AGXRHI/Private/AGXCommandEncoder.h
@@ -398,7 +398,7 @@ public:
 	 * @param Usage The resource usage mask.
 	 * @param Format The Pixel format to reinterpret the resource as.
 	 */
-	void SetShaderBuffer(mtlpp::FunctionType const FunctionType, FAGXBuffer const& Buffer, NSUInteger const Offset, NSUInteger const Length, NSUInteger const Index, mtlpp::ResourceUsage const Usage, EPixelFormat const Format = PF_Unknown);
+	void SetShaderBuffer(mtlpp::FunctionType const FunctionType, FAGXBuffer const& Buffer, NSUInteger const Offset, NSUInteger const Length, NSUInteger const Index, mtlpp::ResourceUsage const Usage, EPixelFormat const Format = PF_Unknown, NSUInteger const ElementRowPitch = 0);
 	
 	/*
 	 * Set an FAGXBufferData to the specified shader frequency at the given bind point index.
@@ -408,7 +408,7 @@ public:
 	 * @param Index The index to modify.
 	 * @param Format The pixel format to reinterpret the resource as.
 	 */
-	void SetShaderData(mtlpp::FunctionType const FunctionType, FAGXBufferData* Data, NSUInteger const Offset, NSUInteger const Index, EPixelFormat const Format = PF_Unknown);
+	void SetShaderData(mtlpp::FunctionType const FunctionType, FAGXBufferData* Data, NSUInteger const Offset, NSUInteger const Index, EPixelFormat const Format = PF_Unknown, NSUInteger const ElementRowPitch = 0);
 	
 	/*
 	 * Set bytes to the specified shader frequency at the given bind point index.
@@ -521,10 +521,34 @@ private:
         NSUInteger Offsets[ML_MaxBuffers];
 		/** The usage mask for the bound resource or 0 */
 		mtlpp::ResourceUsage Usage[ML_MaxBuffers];
-		/** The bound buffer lengths */
-		uint32 Lengths[(ML_MaxBuffers*2) + (ML_MaxTextures*2)];
+		/** The bound buffer constants */
+		struct FSizeConstants
+		{
+			union
+			{
+				uint32 	Length;
+				uint32 	Swizzle;
+			};
+			uint32 Format;
+			uint32 ElementRowPitch;
+		};
+		FSizeConstants Lengths[ML_MaxBuffers + ML_MaxTextures];
         /** A bitmask for which buffers were bound by the application where a bit value of 1 is bound and 0 is unbound. */
         uint32 Bound;
+        
+public:
+        void SetBufferMetaData(NSUInteger Index, NSUInteger Length, NSUInteger Format, NSUInteger ElementRowPitch)
+        {
+			Lengths[Index].Length = Length;
+			Lengths[Index].Format = Format;
+			Lengths[Index].ElementRowPitch = ElementRowPitch;
+        }
+        void SetTextureSwizzle(NSUInteger Index, uint8 (&Swizzle)[4])
+        {
+			FMemory::Memcpy(&Lengths[ML_MaxBuffers + Index].Swizzle, Swizzle, sizeof(Swizzle));
+			Lengths[ML_MaxBuffers + Index].Format = 0;
+			Lengths[ML_MaxBuffers + Index].ElementRowPitch = 0;
+        }
 	};
 	
 #pragma mark - Private Member Variables -
diff --git a/Engine/Source/Runtime/Apple/AGXRHI/Private/AGXStateCache.cpp b/Engine/Source/Runtime/Apple/AGXRHI/Private/AGXStateCache.cpp
--- a/Engine/Source/Runtime/Apple/AGXRHI/Private/AGXStateCache.cpp
+++ b/Engine/Source/Runtime/Apple/AGXRHI/Private/AGXStateCache.cpp
@@ -219,6 +219,7 @@ FAGXStateCache::~FAGXStateCache()
 			ShaderBuffers[Frequency].Buffers[i].Buffer = nil;
 			ShaderBuffers[Frequency].Buffers[i].Bytes = nil;
 			ShaderBuffers[Frequency].Buffers[i].Length = 0;
+			ShaderBuffers[Frequency].Buffers[i].ElementRowPitch = 0;
 			ShaderBuffers[Frequency].Buffers[i].Offset = 0;
 			ShaderBuffers[Frequency].Buffers[i].Usage = mtlpp::ResourceUsage(0);
 			ShaderBuffers[Frequency].Formats[i] = PF_Unknown;
@@ -274,6 +275,7 @@ void FAGXStateCache::Reset(void)
 			ShaderBuffers[Frequency].Buffers[i].Buffer = nil;
 			ShaderBuffers[Frequency].Buffers[i].Bytes = nil;
 			ShaderBuffers[Frequency].Buffers[i].Length = 0;
+			ShaderBuffers[Frequency].Buffers[i].ElementRowPitch = 0;
 			ShaderBuffers[Frequency].Buffers[i].Offset = 0;
 			ShaderBuffers[Frequency].Formats[i] = PF_Unknown;
 		}
@@ -1308,7 +1310,7 @@ bool FAGXStateCache::NeedsToSetRenderTarget(const FRHIRenderPassInfo& InRenderPa
 	return bAllChecksPassed == false;
 }
 
-void FAGXStateCache::SetShaderBuffer(EAGXShaderStages const Frequency, FAGXBuffer const& Buffer, FAGXBufferData* const Bytes, NSUInteger const Offset, NSUInteger const Length, NSUInteger const Index, mtlpp::ResourceUsage const Usage, EPixelFormat const Format)
+void FAGXStateCache::SetShaderBuffer(EAGXShaderStages const Frequency, FAGXBuffer const& Buffer, FAGXBufferData* const Bytes, NSUInteger const Offset, NSUInteger const Length, NSUInteger const Index, mtlpp::ResourceUsage const Usage, EPixelFormat const Format, NSUInteger const ElementRowPitch)
 {
 	check(Frequency < EAGXShaderStages::Num);
 	check(Index < ML_MaxBuffers);
@@ -1317,6 +1319,7 @@ void FAGXStateCache::SetShaderBuffer(EAGXShaderStages const Frequency, FAGXBuffe
 		ShaderBuffers[Frequency].Buffers[Index].Bytes != Bytes ||
 		ShaderBuffers[Frequency].Buffers[Index].Offset != Offset ||
 		ShaderBuffers[Frequency].Buffers[Index].Length != Length ||
+		ShaderBuffers[Frequency].Buffers[Index].ElementRowPitch != ElementRowPitch ||
 		!(ShaderBuffers[Frequency].Buffers[Index].Usage & Usage) ||
 		ShaderBuffers[Frequency].Formats[Index] != Format)
 	{
@@ -1324,6 +1327,7 @@ void FAGXStateCache::SetShaderBuffer(EAGXShaderStages const Frequency, FAGXBuffe
 		ShaderBuffers[Frequency].Buffers[Index].Bytes = Bytes;
 		ShaderBuffers[Frequency].Buffers[Index].Offset = Offset;
 		ShaderBuffers[Frequency].Buffers[Index].Length = Length;
+		ShaderBuffers[Frequency].Buffers[Index].ElementRowPitch = ElementRowPitch;
 		ShaderBuffers[Frequency].Buffers[Index].Usage = Usage;
 		
 		ShaderBuffers[Frequency].Formats[Index] = Format;
@@ -1560,8 +1564,11 @@ void FAGXStateCache::SetShaderUnorderedAccessView(EAGXShaderStages ShaderStage,
 				
 				if ((Source->Flags & (TexCreate_UAV|TexCreate_NoTiling)) == (TexCreate_UAV|TexCreate_NoTiling) && Surface->Texture.GetBuffer())
 				{
+					uint32 BytesPerRow = Surface->Texture.GetBufferBytesPerRow();
+					uint32 ElementsPerRow = BytesPerRow / GPixelFormats[(EPixelFormat)Texture->GetFormat()].BlockBytes;
+					
 					FAGXBuffer Buffer(Surface->Texture.GetBuffer(), false);
-					SetShaderBuffer(ShaderStage, Buffer, nullptr, 0, Surface->Texture.GetBuffer().GetLength(), BindIndex, mtlpp::ResourceUsage(mtlpp::ResourceUsage::Read | mtlpp::ResourceUsage::Write), PF_MAX);
+					SetShaderBuffer(ShaderStage, Buffer, nullptr, 0, Surface->Texture.GetBuffer().GetLength(), BindIndex, mtlpp::ResourceUsage(mtlpp::ResourceUsage::Read | mtlpp::ResourceUsage::Write), PF_MAX, ElementsPerRow);
 				}
 			}
 			else
@@ -2163,7 +2170,7 @@ void FAGXStateCache::CommitResourceTable(EAGXShaderStages const Frequency, mtlpp
 			FAGXBufferBinding& Binding = BufferBindings.Buffers[Index];
 			if (Binding.Buffer)
 			{
-				CommandEncoder.SetShaderBuffer(Type, Binding.Buffer, Binding.Offset, Binding.Length, Index, Binding.Usage, BufferBindings.Formats[Index]);
+				CommandEncoder.SetShaderBuffer(Type, Binding.Buffer, Binding.Offset, Binding.Length, Index, Binding.Usage, BufferBindings.Formats[Index], Binding.ElementRowPitch);
 				
 				if (Binding.Buffer.IsSingleUse())
 				{
@@ -2172,7 +2179,7 @@ void FAGXStateCache::CommitResourceTable(EAGXShaderStages const Frequency, mtlpp
 			}
 			else if (Binding.Bytes)
 			{
-				CommandEncoder.SetShaderData(Type, Binding.Bytes, Binding.Offset, Index, BufferBindings.Formats[Index]);
+				CommandEncoder.SetShaderData(Type, Binding.Bytes, Binding.Offset, Index, BufferBindings.Formats[Index], Binding.ElementRowPitch);
 			}
 		}
 	}
diff --git a/Engine/Source/Runtime/Apple/AGXRHI/Private/AGXStateCache.h b/Engine/Source/Runtime/Apple/AGXRHI/Private/AGXStateCache.h
--- a/Engine/Source/Runtime/Apple/AGXRHI/Private/AGXStateCache.h
+++ b/Engine/Source/Runtime/Apple/AGXRHI/Private/AGXStateCache.h
@@ -77,7 +77,7 @@ public:
 	 * @param Usage The resource usage flags.
 	 * @param Format The UAV pixel format.
 	 */
-	void SetShaderBuffer(EAGXShaderStages const Frequency, FAGXBuffer const& Buffer, FAGXBufferData* const Bytes, NSUInteger const Offset, NSUInteger const Length, NSUInteger const Index, mtlpp::ResourceUsage const Usage, EPixelFormat const Format = PF_Unknown);
+	void SetShaderBuffer(EAGXShaderStages const Frequency, FAGXBuffer const& Buffer, FAGXBufferData* const Bytes, NSUInteger const Offset, NSUInteger const Length, NSUInteger const Index, mtlpp::ResourceUsage const Usage, EPixelFormat const Format = PF_Unknown, NSUInteger const ElementRowPitch = 0);
 	
 	/*
 	 * Set a global texture for the specified shader frequency at the given bind point index.
@@ -195,6 +195,8 @@ private:
 		NSUInteger Offset;
 		/** The bound buffer lengths or 0. */
 		NSUInteger Length;
+		/** The bound buffer element row pitch or 0 */
+		NSUInteger ElementRowPitch;
 		/** The bound buffer usage or 0 */
 		mtlpp::ResourceUsage Usage;
 	};
diff --git a/Engine/Source/Runtime/Apple/AGXRHI/Private/AGXTexture.cpp b/Engine/Source/Runtime/Apple/AGXRHI/Private/AGXTexture.cpp
--- a/Engine/Source/Runtime/Apple/AGXRHI/Private/AGXTexture.cpp
+++ b/Engine/Source/Runtime/Apple/AGXRHI/Private/AGXTexture.cpp
@@ -675,6 +675,8 @@ FAGXSurface::FAGXSurface(ERHIResourceType ResourceType, EPixelFormat Format, uin
 , ImageSurfaceRef(nullptr)
 , bTextureView(false)
 {
+	check(SizeX > 0 && SizeY > 0 && NumMips > 0);
+	
 	// get a unique key for this surface's format
 	TMap<uint64, uint8>& PixelFormatKeyMap = GetMetalPixelFormatKeyMap();
 	if (PixelFormatKeyMap.Num() == 0)
@@ -898,7 +900,10 @@ FAGXSurface::FAGXSurface(ERHIResourceType ResourceType, EPixelFormat Format, uin
 			FAGXPooledBufferArgs Args(SizeAlign.Size, BUF_Dynamic, mtlpp::StorageMode::Private, Desc.GetCpuCacheMode());
 			FAGXBuffer Buffer = GetAGXDeviceContext().CreatePooledBuffer(Args);
 
-			Texture = Buffer.NewTexture(Desc, 0, Align(SizeAlign.Size / SizeY, 256));
+			const uint32 MinimumByteAlignment = GMtlppDevice.GetMinimumLinearTextureAlignmentForPixelFormat(MTLFormat);
+			const NSUInteger BytesPerRow = Align(Desc.GetWidth() * GPixelFormats[Format].BlockBytes, MinimumByteAlignment);
+			
+			Texture = Buffer.NewTexture(Desc, 0, BytesPerRow);
 		}
 		
 		METAL_FATAL_ASSERT(Texture, TEXT("Failed to create texture, desc %s"), *FString([Desc description]));
