Metal: Always use staging buffers to update 'Static' buffers as they could be in use by GPU

diff --git a/Engine/Source/Runtime/Apple/AGXRHI/Private/AGXVertexBuffer.cpp b/Engine/Source/Runtime/Apple/AGXRHI/Private/AGXVertexBuffer.cpp
--- a/Engine/Source/Runtime/Apple/AGXRHI/Private/AGXVertexBuffer.cpp
+++ b/Engine/Source/Runtime/Apple/AGXRHI/Private/AGXVertexBuffer.cpp
@@ -559,8 +559,11 @@ void* FAGXRHIBuffer::Lock(bool bIsOnRHIThread, EResourceLockMode InLockMode, uin
 			// cycle to next allocation
 			AdvanceBackingIndex();
 		}
-		
-		if(Mode == mtlpp::StorageMode::Private)
+
+		// Use transfer buffer for writing into 'Static' buffers as they could be in use by GPU atm
+		// Initialization of 'Static' buffers still uses direct copy when possible
+		const bool bUseTransferBuffer = (Mode == mtlpp::StorageMode::Private || (Mode == mtlpp::StorageMode::Shared && bIsStatic));
+		if(bUseTransferBuffer)
 		{
 			FAGXFrameAllocator::AllocationEntry TempBacking = GetAGXDeviceContext().GetTransferAllocator()->AcquireSpace(Len);
 			GetAGXDeviceContext().NewLock(this, TempBacking);
@@ -646,13 +649,18 @@ void FAGXRHIBuffer::Unlock()
 		check(CurrentBuffer);
 		check(LockSize > 0);
 		const bool bWriteLock = CurrentLockMode == RLM_WriteOnly;
+		const bool bIsStatic = EnumHasAnyFlags(Usage, BUF_Static);
 		
 		if(bWriteLock)
 		{
 			check(!TransferBuffer);
 			check(LockOffset == 0);
 			check(LockSize <= CurrentBuffer.GetLength());
-			if(Mode == mtlpp::StorageMode::Private)
+			
+			// Use transfer buffer for writing into 'Static' buffers as they could be in use by GPU atm
+			// Initialization of 'Static' buffers still uses direct copy when possible
+			const bool bUseTransferBuffer = (Mode == mtlpp::StorageMode::Private || (Mode == mtlpp::StorageMode::Shared && bIsStatic));
+			if(bUseTransferBuffer)
 			{
 				FAGXFrameAllocator::AllocationEntry Entry = GetAGXDeviceContext().FetchAndRemoveLock(this);
 				FAGXBuffer Transfer = Entry.Backing;
diff --git a/Engine/Source/Runtime/Apple/MetalRHI/Private/MetalVertexBuffer.cpp b/Engine/Source/Runtime/Apple/MetalRHI/Private/MetalVertexBuffer.cpp
--- a/Engine/Source/Runtime/Apple/MetalRHI/Private/MetalVertexBuffer.cpp
+++ b/Engine/Source/Runtime/Apple/MetalRHI/Private/MetalVertexBuffer.cpp
@@ -569,8 +569,11 @@ void* FMetalRHIBuffer::Lock(bool bIsOnRHIThread, EResourceLockMode InLockMode, u
 			// cycle to next allocation
 			AdvanceBackingIndex();
 		}
-		
-		if(Mode == mtlpp::StorageMode::Private)
+
+		// Use transfer buffer for writing into 'Static' buffers as they could be in use by GPU atm
+		// Initialization of 'Static' buffers still uses direct copy when possible
+		const bool bUseTransferBuffer = (Mode == mtlpp::StorageMode::Private || (Mode == mtlpp::StorageMode::Shared && bIsStatic));
+		if(bUseTransferBuffer)
 		{
 			FMetalFrameAllocator::AllocationEntry TempBacking = GetMetalDeviceContext().GetTransferAllocator()->AcquireSpace(Len);
 			GetMetalDeviceContext().NewLock(this, TempBacking);
@@ -656,13 +659,18 @@ void FMetalRHIBuffer::Unlock()
 		check(CurrentBuffer);
 		check(LockSize > 0);
 		const bool bWriteLock = CurrentLockMode == RLM_WriteOnly;
+		const bool bIsStatic = EnumHasAnyFlags(Usage, BUF_Static);
 		
 		if(bWriteLock)
 		{
 			check(!TransferBuffer);
 			check(LockOffset == 0);
 			check(LockSize <= CurrentBuffer.GetLength());
-			if(Mode == mtlpp::StorageMode::Private)
+
+			// Use transfer buffer for writing into 'Static' buffers as they could be in use by GPU atm
+			// Initialization of 'Static' buffers still uses direct copy when possible
+			const bool bUseTransferBuffer = (Mode == mtlpp::StorageMode::Private || (Mode == mtlpp::StorageMode::Shared && bIsStatic));
+			if(bUseTransferBuffer)
 			{
 				FMetalFrameAllocator::AllocationEntry Entry = GetMetalDeviceContext().FetchAndRemoveLock(this);
 				FMetalBuffer Transfer = Entry.Backing;
