Update AGXRHI to use the new TexCreate_AtomicCompatible flag inline with MetalRHI.

diff --git a/Engine/Source/Runtime/Apple/AGXRHI/Private/AGXStateCache.cpp b/Engine/Source/Runtime/Apple/AGXRHI/Private/AGXStateCache.cpp
--- a/Engine/Source/Runtime/Apple/AGXRHI/Private/AGXStateCache.cpp
+++ b/Engine/Source/Runtime/Apple/AGXRHI/Private/AGXStateCache.cpp
@@ -1562,7 +1562,8 @@ void FAGXStateCache::SetShaderUnorderedAccessView(EAGXShaderStages ShaderStage,
 				
 				SetShaderTexture(ShaderStage, Surface->Texture, BindIndex, mtlpp::ResourceUsage(mtlpp::ResourceUsage::Read | mtlpp::ResourceUsage::Write));
 				
-				if ((Source->Flags & (TexCreate_UAV|TexCreate_NoTiling)) == (TexCreate_UAV|TexCreate_NoTiling) && Surface->Texture.GetBuffer())
+				if (Surface->Texture.GetBuffer() &&
+					(EnumHasAllFlags(Source->Flags, TexCreate_UAV | TexCreate_NoTiling) || EnumHasAllFlags(Source->Flags, TexCreate_AtomicCompatible)))
 				{
 					uint32 BytesPerRow = Surface->Texture.GetBufferBytesPerRow();
 					uint32 ElementsPerRow = BytesPerRow / GPixelFormats[(EPixelFormat)Texture->GetFormat()].BlockBytes;
diff --git a/Engine/Source/Runtime/Apple/AGXRHI/Private/AGXTexture.cpp b/Engine/Source/Runtime/Apple/AGXRHI/Private/AGXTexture.cpp
--- a/Engine/Source/Runtime/Apple/AGXRHI/Private/AGXTexture.cpp
+++ b/Engine/Source/Runtime/Apple/AGXRHI/Private/AGXTexture.cpp
@@ -888,7 +888,7 @@ FAGXSurface::FAGXSurface(ERHIResourceType ResourceType, EPixelFormat Format, uin
 	{
 		const bool bBufferCompatibleOption = (Desc.GetTextureType() == mtlpp::TextureType::Texture2D || Desc.GetTextureType() == mtlpp::TextureType::TextureBuffer) && NumMips == 1;
 
-		if (!bBufferCompatibleOption || ((Flags & (TexCreate_UAV|TexCreate_NoTiling)) != (TexCreate_UAV|TexCreate_NoTiling)))
+		if (!bBufferCompatibleOption || (!EnumHasAllFlags(Flags, TexCreate_UAV|TexCreate_NoTiling) && !EnumHasAllFlags(Flags, TexCreate_AtomicCompatible)))
 		{
 			Texture = GetAGXDeviceContext().CreateTexture(this, Desc);
 		}
