Fix for validation error in Catalina and early versions of Big Sur.   For macOS < 11.5.1 change mapping for PF_FloatRGB from packed RG11B10Float to ordinary RGBA16Float.  Packed pixel formats don't appear to pass validation when used for texture_buffers on older macOS revisions under macOS 11.5.1 (unknown where the actual cut off is; 11.5.1 works; 11.4.0 does not).

diff --git a/Engine/Source/Runtime/Apple/MetalRHI/Private/MetalRHI.cpp b/Engine/Source/Runtime/Apple/MetalRHI/Private/MetalRHI.cpp
--- a/Engine/Source/Runtime/Apple/MetalRHI/Private/MetalRHI.cpp
+++ b/Engine/Source/Runtime/Apple/MetalRHI/Private/MetalRHI.cpp
@@ -604,6 +604,14 @@ FMetalDynamicRHI::FMetalDynamicRHI(ERHIFeatureLevel::Type RequestedFeatureLevel)
 	GMaxTextureMipCount = FPlatformMath::CeilLogTwo( GMaxTextureDimensions ) + 1;
 	GMaxTextureMipCount = FPlatformMath::Min<int32>( MAX_TEXTURE_MIP_COUNT, GMaxTextureMipCount );
 
+	// Using packed pixel formats (specifically MTLPixelFormatRG11B10Float) for texture_buffers on macOS <= 11.4 results in Metal validation errors:
+	// "pixel format (MTLPixelFormatRG11B10Float) cannot be written to from a shader on this device"
+	// macOS Big Sur 11.5.1 works; Big Sur 11.4.0 (and prior versions including Catalina) report the error; texture2d with buffer backing does not report this error.
+	bool bSupportsPackedPixelFormatsInTextureBuffers = true;
+#if PLATFORM_MAC
+	bSupportsPackedPixelFormatsInTextureBuffers = FPlatformMisc::MacOSXVersionCompare(11,5,1) >= 0;
+#endif
+
 	// Initialize the buffer format map - in such a way as to be able to validate it in non-shipping...
 #if METAL_DEBUG_OPTIONS
 	FMemory::Memset(GMetalBufferFormats, 255);
@@ -617,7 +625,14 @@ FMetalDynamicRHI::FMetalDynamicRHI(ERHIFeatureLevel::Type RequestedFeatureLevel)
 	GMetalBufferFormats[PF_DXT3                 ] = { mtlpp::PixelFormat::Invalid, (uint8)EMetalBufferFormat::Unknown };
 	GMetalBufferFormats[PF_DXT5                 ] = { mtlpp::PixelFormat::Invalid, (uint8)EMetalBufferFormat::Unknown };
 	GMetalBufferFormats[PF_UYVY                 ] = { mtlpp::PixelFormat::Invalid, (uint8)EMetalBufferFormat::Unknown };
-	GMetalBufferFormats[PF_FloatRGB             ] = { mtlpp::PixelFormat::RG11B10Float, (uint8)EMetalBufferFormat::RG11B10Half };
+	if(bSupportsPackedPixelFormatsInTextureBuffers)
+	{
+		GMetalBufferFormats[PF_FloatRGB        	] = { mtlpp::PixelFormat::RG11B10Float, (uint8)EMetalBufferFormat::RG11B10Half };
+	}
+	else
+	{
+		GMetalBufferFormats[PF_FloatRGB        	] = { mtlpp::PixelFormat::RGBA16Float, (uint8)EMetalBufferFormat::RGBA16Half };
+	}
 	GMetalBufferFormats[PF_FloatRGBA            ] = { mtlpp::PixelFormat::RGBA16Float, (uint8)EMetalBufferFormat::RGBA16Half };
 	GMetalBufferFormats[PF_DepthStencil         ] = { mtlpp::PixelFormat::Invalid, (uint8)EMetalBufferFormat::Unknown };
 	GMetalBufferFormats[PF_ShadowDepth          ] = { mtlpp::PixelFormat::Invalid, (uint8)EMetalBufferFormat::Unknown };
@@ -777,8 +792,17 @@ FMetalDynamicRHI::FMetalDynamicRHI(ERHIFeatureLevel::Type RequestedFeatureLevel)
     GPixelFormats[PF_DXT3				].PlatformFormat	= (uint32)mtlpp::PixelFormat::BC2_RGBA;
     GPixelFormats[PF_DXT5				].PlatformFormat	= (uint32)mtlpp::PixelFormat::BC3_RGBA;
 	
-	GPixelFormats[PF_FloatRGB			].PlatformFormat	= (uint32)mtlpp::PixelFormat::RG11B10Float;
-	GPixelFormats[PF_FloatRGB			].BlockBytes		= 4;
+	if(bSupportsPackedPixelFormatsInTextureBuffers)
+	{
+		GPixelFormats[PF_FloatRGB		].PlatformFormat	= (uint32)mtlpp::PixelFormat::RG11B10Float;
+		GPixelFormats[PF_FloatRGB		].BlockBytes		= 4;
+	}
+	else
+	{
+		GPixelFormats[PF_FloatRGB		].PlatformFormat	= (uint32)mtlpp::PixelFormat::RGBA16Float;
+		GPixelFormats[PF_FloatRGB		].BlockBytes		= 8;
+	}
+	
 	GPixelFormats[PF_FloatR11G11B10		].PlatformFormat	= (uint32)mtlpp::PixelFormat::RG11B10Float;
 	GPixelFormats[PF_FloatR11G11B10		].BlockBytes		= 4;
 	GPixelFormats[PF_FloatR11G11B10		].Supported			= true;
