Exclude MTLGPUFamilyApple6 check for tvOS

diff --git a/Engine/Source/Runtime/Apple/MetalRHI/Private/MetalRHI.cpp b/Engine/Source/Runtime/Apple/MetalRHI/Private/MetalRHI.cpp
--- a/Engine/Source/Runtime/Apple/MetalRHI/Private/MetalRHI.cpp
+++ b/Engine/Source/Runtime/Apple/MetalRHI/Private/MetalRHI.cpp
@@ -746,6 +746,7 @@ FMetalDynamicRHI::FMetalDynamicRHI(ERHIFeatureLevel::Type RequestedFeatureLevel)
 	GPixelFormats[PF_ASTC_12x12			].PlatformFormat	= (uint32)mtlpp::PixelFormat::ASTC_12x12_LDR;
 	GPixelFormats[PF_ASTC_12x12			].Supported			= true;
 
+#if !PLATFORM_TVOS
 	if([Device.GetPtr() supportsFamily:MTLGPUFamilyApple6])
 	{
 		GPixelFormats[PF_ASTC_4x4_HDR].PlatformFormat = (uint32)mtlpp::PixelFormat::ASTC_4x4_HDR;
@@ -759,6 +760,7 @@ FMetalDynamicRHI::FMetalDynamicRHI(ERHIFeatureLevel::Type RequestedFeatureLevel)
 		GPixelFormats[PF_ASTC_12x12_HDR].PlatformFormat = (uint32)mtlpp::PixelFormat::ASTC_12x12_HDR;
 		GPixelFormats[PF_ASTC_12x12_HDR].Supported = true;
 	}
+#endif
 	// used with virtual textures
 	GPixelFormats[PF_ETC2_RGB	  		].PlatformFormat	= (uint32)mtlpp::PixelFormat::ETC2_RGB8;
 	GPixelFormats[PF_ETC2_RGB			].Supported			= true;
