Workaround for crash when enabling HDR during a render pass - fetch the drawable during the resolve for destination presentable textures,  do nothing for source nil textures as there is nothing to copy from anyway.   When TexCreate_Presentable is used in RHICreateTexture2D() the actual Metal texture (drawable) is not created at that point but assigned from the CAMetalLayer at the next call to RHIBeginRenderPass() but when used during the HDR toggle we've not initialised the texture via FMetalStateCache::SetRenderPassInfo().

diff --git a/Engine/Source/Runtime/Apple/MetalRHI/Private/MetalRenderTarget.cpp b/Engine/Source/Runtime/Apple/MetalRHI/Private/MetalRenderTarget.cpp
--- a/Engine/Source/Runtime/Apple/MetalRHI/Private/MetalRenderTarget.cpp
+++ b/Engine/Source/Runtime/Apple/MetalRHI/Private/MetalRenderTarget.cpp
@@ -46,6 +46,25 @@ void FMetalRHICommandContext::RHICopyToResolveTarget(FRHITexture* SourceTextureR
 		FMetalSurface* Source = GetMetalSurfaceFromRHITexture(SourceTextureRHI);
 		FMetalSurface* Destination = GetMetalSurfaceFromRHITexture(DestTextureRHI);
 		
+		// Only valid to have nil Metal textures when they are TexCreate_Presentable
+		if(!Source->Texture)
+		{
+			// Source RHI texture is valid with no Presentable Metal texture - there is nothing to copy from
+			check(EnumHasAnyFlags(Source->Flags, TexCreate_Presentable));
+			return;
+		}
+		if(!Destination->Texture)
+		{
+			// Destination RHI texture is valid with no Presentable Metal texture - force fetch it now so we can complete the copy
+			check(EnumHasAnyFlags(Destination->Flags, TexCreate_Presentable));
+			Destination->GetDrawableTexture();
+			if(!Destination->Texture)
+			{
+				UE_LOG(LogRHI, Error, TEXT("Drawable for destination texture resolve target unavailable"));
+				return;
+			}
+		}
+
 		switch (Source->Type)
 		{
 			case RRT_Texture2D:
