Enable RHI thread for Metal in the Editor.   Was already enabled in game.

diff --git a/Engine/Config/Mac/MacEngine.ini b/Engine/Config/Mac/MacEngine.ini
--- a/Engine/Config/Mac/MacEngine.ini
+++ b/Engine/Config/Mac/MacEngine.ini
@@ -14,10 +14,6 @@ g.TimeoutForBlockOnRenderFence=300000
 framepro.ScopeMinTimeMicroseconds=10
 fx.NiagaraAllowRuntimeScalabilityChanges=1
 
-[SystemSettingsEditor]
-; Parallel rendering has not been tested in the editor so just disable it.
-r.RHICmdBypass=1
-
 [TextureStreaming]
 ; PoolSizeVRAMPercentage is how much percentage of GPU Dedicated VRAM should be used as a TexturePool cache for streaming textures (0 - unlimited streaming)
 PoolSizeVRAMPercentage=70
@@ -45,4 +41,4 @@ PlatformRequiresDataCrypto=True
 MaxShaderJobBatchSize=25
 
 [ConsoleVariables]
-r.HairStrands.Strands=0
\ No newline at end of file
+r.HairStrands.Strands=0
diff --git a/Engine/Source/Runtime/Apple/MetalRHI/Private/MetalRHI.cpp b/Engine/Source/Runtime/Apple/MetalRHI/Private/MetalRHI.cpp
--- a/Engine/Source/Runtime/Apple/MetalRHI/Private/MetalRHI.cpp
+++ b/Engine/Source/Runtime/Apple/MetalRHI/Private/MetalRHI.cpp
@@ -31,7 +31,7 @@ FMetalBufferFormat GMetalBufferFormats[PF_MAX];
 
 static bool GFormatSupportsTypedUAVLoad[PF_MAX] = { false };
 
-static TAutoConsoleVariable<int32> CVarUseRHIThread(
+static TAutoConsoleVariable<int32> CVarUseIOSRHIThread(
 													TEXT("r.Metal.IOSRHIThread"),
 													0,
 													TEXT("Controls RHIThread usage for IOS:\n")
@@ -41,16 +41,6 @@ static TAutoConsoleVariable<int32> CVarUseRHIThread(
 													ECVF_Default | ECVF_RenderThreadSafe
 													);
 
-static TAutoConsoleVariable<int32> CVarIntelUseRHIThread(
-													TEXT("r.Metal.IntelRHIThread"),
-													0,
-													TEXT("Controls RHIThread usage for Mac Intel HW:\n")
-													TEXT("\t0: No RHIThread.\n")
-													TEXT("\t1: Use RHIThread.\n")
-													TEXT("Default is 0."),
-													ECVF_Default | ECVF_RenderThreadSafe
-													);
-
 
 static void ValidateTargetedRHIFeatureLevelExists(EShaderPlatform Platform)
 {
@@ -189,7 +179,6 @@ FMetalDynamicRHI::FMetalDynamicRHI(ERHIFeatureLevel::Type RequestedFeatureLevel)
 
 	bool const bRequestedFeatureLevel = (RequestedFeatureLevel != ERHIFeatureLevel::Num);
 	bool bSupportsPointLights = false;
-	bool bSupportsRHIThread = false;
 	
 	// get the device to ask about capabilities?
 	mtlpp::Device Device = ImmediateContext.Context->GetDevice();
@@ -235,7 +224,6 @@ FMetalDynamicRHI::FMetalDynamicRHI(ERHIFeatureLevel::Type RequestedFeatureLevel)
     GConfig->GetBool(TEXT("/Script/IOSRuntimeSettings.IOSRuntimeSettings"), TEXT("bSupportsMetalMRT"), bProjectSupportsMRTs, GEngineIni);
 
 	bool const bRequestedMetalMRT = ((RequestedFeatureLevel >= ERHIFeatureLevel::SM5) || (!bRequestedFeatureLevel && FParse::Param(FCommandLine::Get(),TEXT("metalmrt"))));
-	bSupportsRHIThread = FParse::Param(FCommandLine::Get(),TEXT("rhithread"));
 
     // only allow GBuffers, etc on A8s (A7s are just not going to cut it)
     if (bProjectSupportsMRTs && bRequestedMetalMRT)
@@ -316,7 +304,6 @@ FMetalDynamicRHI::FMetalDynamicRHI(ERHIFeatureLevel::Type RequestedFeatureLevel)
 		GRHIVendorId = 0x10DE;
 		bSupportsTiledReflections = true;
 		bSupportsDistanceFields = (FPlatformMisc::MacOSXVersionCompare(10,11,4) >= 0);
-		bSupportsRHIThread = (FPlatformMisc::MacOSXVersionCompare(10,12,0) >= 0);
 	}
 	else if(GRHIAdapterName.Contains("ATi") || GRHIAdapterName.Contains("AMD"))
 	{
@@ -328,7 +315,6 @@ FMetalDynamicRHI::FMetalDynamicRHI(ERHIFeatureLevel::Type RequestedFeatureLevel)
 		}
 		bSupportsTiledReflections = true;
 		bSupportsDistanceFields = (FPlatformMisc::MacOSXVersionCompare(10,11,4) >= 0);
-		bSupportsRHIThread = true;
 		
 		// On AMD can also use completion handler time stamp if macOS < Catalina
 		GSupportsTimestampRenderQueries = true;
@@ -338,8 +324,6 @@ FMetalDynamicRHI::FMetalDynamicRHI(ERHIFeatureLevel::Type RequestedFeatureLevel)
 		bSupportsTiledReflections = false;
 		bSupportsPointLights = (FPlatformMisc::MacOSXVersionCompare(10,14,6) > 0);
 		GRHIVendorId = 0x8086;
-		// HACK: Meshes jump around in Infiltrator with RHI thread on. Needs further investigation and a real fix.
-		bSupportsRHIThread = false || CVarIntelUseRHIThread.GetValueOnAnyThread() > 0;
 		bSupportsDistanceFields = (FPlatformMisc::MacOSXVersionCompare(10,12,2) >= 0);
 		bIsIntelHaswell = (GRHIAdapterName == TEXT("Intel HD Graphics 5000") || GRHIAdapterName == TEXT("Intel Iris Graphics") || GRHIAdapterName == TEXT("Intel Iris Pro Graphics"));
 	}
@@ -349,7 +333,6 @@ FMetalDynamicRHI::FMetalDynamicRHI(ERHIFeatureLevel::Type RequestedFeatureLevel)
 		GRHIVendorId = 0x106B;
 		bSupportsTiledReflections = true;
 		bSupportsDistanceFields = true;
-		bSupportsRHIThread = true;
 		GSupportsTimestampRenderQueries = true;
 	}
 
@@ -483,12 +466,8 @@ FMetalDynamicRHI::FMetalDynamicRHI(ERHIFeatureLevel::Type RequestedFeatureLevel)
 	GRHISupportsRHIThread = false;
 	if (GMaxRHIFeatureLevel >= ERHIFeatureLevel::SM5)
 	{
+		GRHISupportsRHIThread = true;
 #if METAL_SUPPORTS_PARALLEL_RHI_EXECUTE
-#if WITH_EDITORONLY_DATA
-		GRHISupportsRHIThread = (!GIsEditor && bSupportsRHIThread);
-#else
-		GRHISupportsRHIThread = bSupportsRHIThread;
-#endif
 		GRHISupportsParallelRHIExecute = GRHISupportsRHIThread && ((!IsRHIDeviceIntel() && !IsRHIDeviceNVIDIA()) || FParse::Param(FCommandLine::Get(),TEXT("metalparallel")));
 #endif
 		GSupportsEfficientAsyncCompute = GRHISupportsParallelRHIExecute && (IsRHIDeviceAMD() || /*TODO: IsRHIDeviceApple()*/ (GRHIVendorId == 0x106B) || PLATFORM_IOS || FParse::Param(FCommandLine::Get(),TEXT("metalasynccompute"))); // Only AMD and Apple currently support async. compute and it requires parallel execution to be useful.
@@ -496,7 +475,7 @@ FMetalDynamicRHI::FMetalDynamicRHI(ERHIFeatureLevel::Type RequestedFeatureLevel)
 	}
 	else
 	{
-		GRHISupportsRHIThread = bSupportsRHIThread || (CVarUseRHIThread.GetValueOnAnyThread() > 0);
+		GRHISupportsRHIThread = FParse::Param(FCommandLine::Get(),TEXT("rhithread")) || (CVarUseIOSRHIThread.GetValueOnAnyThread() > 0);
 		GRHISupportsParallelRHIExecute = false;
 		GSupportsEfficientAsyncCompute = false;
 		GSupportsParallelOcclusionQueries = false;
diff --git a/Engine/Source/Runtime/Core/Public/Mac/MacPlatform.h b/Engine/Source/Runtime/Core/Public/Mac/MacPlatform.h
--- a/Engine/Source/Runtime/Core/Public/Mac/MacPlatform.h
+++ b/Engine/Source/Runtime/Core/Public/Mac/MacPlatform.h
@@ -84,7 +84,7 @@ typedef FMacPlatformTypes FPlatformTypes;
 #define PLATFORM_IS_ANSI_MALLOC_THREADSAFE				1
 #define PLATFORM_SUPPORTS_VIRTUAL_TEXTURE_STREAMING		1
 
-#define PLATFORM_RHITHREAD_DEFAULT_BYPASS				WITH_EDITOR
+#define PLATFORM_RHITHREAD_DEFAULT_BYPASS				0
 
 #define PLATFORM_ENABLE_POPCNT_INTRINSIC				1
 
