Workaround for RestartRenderPass() with MSAA attachment.  Allow restart when we have [multisample] store set as we can load this back in with FMetalStateCache::PrepareToRestart load action fix up.  Disallow when using memory less RT or we didn't prevsouly store the results.

diff --git a/Engine/Source/Runtime/Apple/MetalRHI/Private/MetalStateCache.cpp b/Engine/Source/Runtime/Apple/MetalRHI/Private/MetalStateCache.cpp
--- a/Engine/Source/Runtime/Apple/MetalRHI/Private/MetalStateCache.cpp
+++ b/Engine/Source/Runtime/Apple/MetalRHI/Private/MetalStateCache.cpp
@@ -562,9 +562,9 @@ bool FMetalStateCache::SetRenderPassInfo(FRHIRenderPassInfo const& InRenderTarge
 					HighLevelLoadAction == ERenderTargetLoadAction::ELoad);
 #endif
 				
+				bool bMemoryless = false;
 				if (Surface.MSAATexture && !bUseResolvedTexture)
 				{
-					bool bMemoryless = false;
 #if PLATFORM_IOS
 					if (Surface.MSAATexture.GetStorageMode() == mtlpp::StorageMode::Memoryless)
 					{
@@ -583,7 +583,6 @@ bool FMetalStateCache::SetRenderPassInfo(FRHIRenderPassInfo const& InRenderTarge
 				}
 				else
 				{
-					bool bMemoryless = false;
 #if PLATFORM_IOS
 					if (Surface.Texture.GetStorageMode() == mtlpp::StorageMode::Memoryless)
 					{
@@ -622,7 +621,9 @@ bool FMetalStateCache::SetRenderPassInfo(FRHIRenderPassInfo const& InRenderTarge
 					ColorAttachment.SetClearColor(mtlpp::ClearColor(ClearColor.R, ClearColor.G, ClearColor.B, ClearColor.A));
 				}
 
-				bCanRestartRenderPass &= (SampleCount <= 1) && (ColorAttachment.GetLoadAction() == mtlpp::LoadAction::Load) && (HighLevelStoreAction == ERenderTargetStoreAction::EStore);
+				bCanRestartRenderPass &= 	!bMemoryless &&
+											ColorAttachment.GetLoadAction() == mtlpp::LoadAction::Load &&
+											HighLevelStoreAction != ERenderTargetStoreAction::ENoAction;
 	
 				bHasValidRenderTarget = true;
 				bHasValidColorTarget = true;
@@ -850,7 +851,12 @@ bool FMetalStateCache::SetRenderPassInfo(FRHIRenderPassInfo const& InRenderTarge
 				bHasValidRenderTarget = true;
 				bFallbackDepthStencilBound = (RenderPassInfo.DepthStencilRenderTarget.DepthStencilTarget == FallbackDepthStencilSurface);
 
-				bCanRestartRenderPass &= (SampleCount <= 1) && ((RenderPassInfo.DepthStencilRenderTarget.DepthStencilTarget == FallbackDepthStencilSurface) || ((DepthAttachment.GetLoadAction() == mtlpp::LoadAction::Load) && (!RenderPassInfo.DepthStencilRenderTarget.ExclusiveDepthStencil.IsDepthWrite() || (DepthStoreAction == ERenderTargetStoreAction::EStore))));
+				bool bDepthMSAARestart = !bDepthTextureMemoryless && HighLevelStoreAction == ERenderTargetStoreAction::EMultisampleResolve;
+				bCanRestartRenderPass &=	(DepthSampleCount <= 1 || bDepthMSAARestart) &&
+											(
+												(RenderPassInfo.DepthStencilRenderTarget.DepthStencilTarget == FallbackDepthStencilSurface) ||
+												((DepthAttachment.GetLoadAction() == mtlpp::LoadAction::Load) && (bDepthMSAARestart || !RenderPassInfo.DepthStencilRenderTarget.ExclusiveDepthStencil.IsDepthWrite() || DepthStoreAction == ERenderTargetStoreAction::EStore))
+											);
 				
 				// and assign it
 				RenderPass.SetDepthAttachment(DepthAttachment);
@@ -912,7 +918,12 @@ bool FMetalStateCache::SetRenderPassInfo(FRHIRenderPassInfo const& InRenderTarge
 				
 				// @todo Stencil writes that need to persist must use ERenderTargetStoreAction::EStore on iOS.
 				// We should probably be using deferred store actions so that we can safely lazily instantiate encoders.
-				bCanRestartRenderPass &= (SampleCount <= 1) && ((RenderPassInfo.DepthStencilRenderTarget.DepthStencilTarget == FallbackDepthStencilSurface) || ((StencilAttachment.GetLoadAction() == mtlpp::LoadAction::Load) && (1 || !RenderPassInfo.DepthStencilRenderTarget.ExclusiveDepthStencil.IsStencilWrite() || (StencilStoreAction == ERenderTargetStoreAction::EStore))));
+				bool bStencilMSAARestart = !bStencilMemoryless && HighLevelStoreAction != ERenderTargetStoreAction::ENoAction;
+				bCanRestartRenderPass &= 	(bStencilMSAARestart || SampleCount <= 1) &&
+											(
+												(RenderPassInfo.DepthStencilRenderTarget.DepthStencilTarget == FallbackDepthStencilSurface) ||
+												((StencilAttachment.GetLoadAction() == mtlpp::LoadAction::Load) && (bStencilMSAARestart || !RenderPassInfo.DepthStencilRenderTarget.ExclusiveDepthStencil.IsStencilWrite() || (StencilStoreAction == ERenderTargetStoreAction::EStore)))
+											);
 				
 				// and assign it
 				RenderPass.SetStencilAttachment(StencilAttachment);
@@ -1747,67 +1758,74 @@ bool FMetalStateCache::PrepareToRestart(bool const bCurrentApplied)
 	}
 	else
 	{
-		if (SampleCount <= 1)
+		FRHIRenderPassInfo Info = GetRenderPassInfo();
+		
+		ERenderTargetActions DepthActions = GetDepthActions(Info.DepthStencilRenderTarget.Action);
+		ERenderTargetActions StencilActions = GetStencilActions(Info.DepthStencilRenderTarget.Action);
+		ERenderTargetLoadAction DepthLoadAction = GetLoadAction(DepthActions);
+		ERenderTargetStoreAction DepthStoreAction = GetStoreAction(DepthActions);
+		ERenderTargetLoadAction StencilLoadAction = GetLoadAction(StencilActions);
+		ERenderTargetStoreAction StencilStoreAction = GetStoreAction(StencilActions);
+
+		if (Info.DepthStencilRenderTarget.DepthStencilTarget)
 		{
-			FRHIRenderPassInfo Info = GetRenderPassInfo();
-			
-			ERenderTargetActions DepthActions = GetDepthActions(Info.DepthStencilRenderTarget.Action);
-			ERenderTargetActions StencilActions = GetStencilActions(Info.DepthStencilRenderTarget.Action);
-			ERenderTargetLoadAction DepthLoadAction = GetLoadAction(DepthActions);
-			ERenderTargetStoreAction DepthStoreAction = GetStoreAction(DepthActions);
-			ERenderTargetLoadAction StencilLoadAction = GetLoadAction(StencilActions);
-			ERenderTargetStoreAction StencilStoreAction = GetStoreAction(StencilActions);
-
-			if (Info.DepthStencilRenderTarget.DepthStencilTarget)
+			if(bCurrentApplied && Info.DepthStencilRenderTarget.ExclusiveDepthStencil.IsDepthWrite() && DepthStoreAction == ERenderTargetStoreAction::ENoAction)
 			{
-				if (bCurrentApplied || DepthLoadAction != ERenderTargetLoadAction::EClear)
-				{
-					DepthLoadAction = ERenderTargetLoadAction::ELoad;
-				}
-				if (Info.DepthStencilRenderTarget.ExclusiveDepthStencil.IsDepthWrite())
-				{
-					DepthStoreAction = ERenderTargetStoreAction::EStore;
-				}
+				return false;
+			}
+			if (bCurrentApplied && Info.DepthStencilRenderTarget.ExclusiveDepthStencil.IsStencilWrite() && StencilStoreAction == ERenderTargetStoreAction::ENoAction)
+			{
+				return false;
+			}
+		
+			if (bCurrentApplied || DepthLoadAction != ERenderTargetLoadAction::EClear)
+			{
+				DepthLoadAction = ERenderTargetLoadAction::ELoad;
+			}
+			if (Info.DepthStencilRenderTarget.ExclusiveDepthStencil.IsDepthWrite())
+			{
+				DepthStoreAction = ERenderTargetStoreAction::EStore;
+			}
 
-				if (bCurrentApplied || StencilLoadAction != ERenderTargetLoadAction::EClear)
-				{
-					StencilLoadAction = ERenderTargetLoadAction::ELoad;
-				}
-				if (Info.DepthStencilRenderTarget.ExclusiveDepthStencil.IsStencilWrite())
-				{
-					StencilStoreAction = ERenderTargetStoreAction::EStore;
-				}
-				
-				DepthActions = MakeRenderTargetActions(DepthLoadAction, DepthStoreAction);
-				StencilActions = MakeRenderTargetActions(StencilLoadAction, StencilStoreAction);
-				Info.DepthStencilRenderTarget.Action = MakeDepthStencilTargetActions(DepthActions, StencilActions);
+			if (bCurrentApplied || StencilLoadAction != ERenderTargetLoadAction::EClear)
+			{
+				StencilLoadAction = ERenderTargetLoadAction::ELoad;
 			}
-			
-			for (int32 RenderTargetIndex = 0; RenderTargetIndex < Info.GetNumColorRenderTargets(); RenderTargetIndex++)
+			if (Info.DepthStencilRenderTarget.ExclusiveDepthStencil.IsStencilWrite())
 			{
-				FRHIRenderPassInfo::FColorEntry& RenderTargetView = Info.ColorRenderTargets[RenderTargetIndex];
-				ERenderTargetLoadAction LoadAction = GetLoadAction(RenderTargetView.Action);
-				ERenderTargetStoreAction StoreAction = GetStoreAction(RenderTargetView.Action);
-				
-				if (!bCurrentApplied && LoadAction == ERenderTargetLoadAction::EClear)
-				{
-					StoreAction == ERenderTargetStoreAction::EStore;
-				}
-				else
-				{
-					LoadAction = ERenderTargetLoadAction::ELoad;
-				}
-				RenderTargetView.Action = MakeRenderTargetActions(LoadAction, StoreAction);
-				check(RenderTargetView.RenderTarget == nil || GetStoreAction(RenderTargetView.Action) == ERenderTargetStoreAction::EStore);
+				StencilStoreAction = ERenderTargetStoreAction::EStore;
 			}
 			
-			InvalidateRenderTargets();
-			return SetRenderPassInfo(Info, GetVisibilityResultsBuffer(), true) && CanRestartRenderPass();
+			DepthActions = MakeRenderTargetActions(DepthLoadAction, DepthStoreAction);
+			StencilActions = MakeRenderTargetActions(StencilLoadAction, StencilStoreAction);
+			Info.DepthStencilRenderTarget.Action = MakeDepthStencilTargetActions(DepthActions, StencilActions);
 		}
-		else
+		
+		for (int32 RenderTargetIndex = 0; RenderTargetIndex < Info.GetNumColorRenderTargets(); RenderTargetIndex++)
 		{
-			return false;
+			FRHIRenderPassInfo::FColorEntry& RenderTargetView = Info.ColorRenderTargets[RenderTargetIndex];
+			ERenderTargetLoadAction LoadAction = GetLoadAction(RenderTargetView.Action);
+			ERenderTargetStoreAction StoreAction = GetStoreAction(RenderTargetView.Action);
+			
+			if(bCurrentApplied && StoreAction == ERenderTargetStoreAction::ENoAction)
+			{
+				return false;
+			}
+			
+			if (!bCurrentApplied && LoadAction == ERenderTargetLoadAction::EClear)
+			{
+				StoreAction == ERenderTargetStoreAction::EStore;
+			}
+			else
+			{
+				LoadAction = ERenderTargetLoadAction::ELoad;
+			}
+			RenderTargetView.Action = MakeRenderTargetActions(LoadAction, StoreAction);
+			check(RenderTargetView.RenderTarget == nil || GetStoreAction(RenderTargetView.Action) == ERenderTargetStoreAction::EStore);
 		}
+		
+		InvalidateRenderTargets();
+		return SetRenderPassInfo(Info, GetVisibilityResultsBuffer(), true) && CanRestartRenderPass();
 	}
 }
 
