Log cleanup for Mac's Audio Mixer implementation

diff --git a/Engine/Source/Runtime/Apple/AudioMixerCoreAudio/Private/AudioMixerPlatformCoreAudio.cpp b/Engine/Source/Runtime/Apple/AudioMixerCoreAudio/Private/AudioMixerPlatformCoreAudio.cpp
--- a/Engine/Source/Runtime/Apple/AudioMixerCoreAudio/Private/AudioMixerPlatformCoreAudio.cpp
+++ b/Engine/Source/Runtime/Apple/AudioMixerCoreAudio/Private/AudioMixerPlatformCoreAudio.cpp
@@ -93,7 +93,7 @@ namespace Audio
 
 		if(Status != 0)
 		{
-			UE_LOG(LogAudioMixerAudioUnit, Display, TEXT("ERROR setting sample rate to %f"), GraphSampleRate);
+			UE_LOG(LogAudioMixerAudioUnit, Warning, TEXT("Error querying Audio Device ID: %i"), Status);
 		}
 
 		Status = AudioObjectGetPropertyData(DeviceAudioObjectID, &DevicePropertyAddress, 0, nullptr, &AudioDeviceQuerySize, &GraphSampleRate);
@@ -102,6 +102,10 @@ namespace Audio
 		{
 			UE_LOG(LogAudioMixerAudioUnit, Display, TEXT("Sample Rate: %f"), GraphSampleRate);
 		}
+		else
+		{
+			UE_LOG(LogAudioMixerAudioUnit, Warning, TEXT("Error querying Sample Rate: %i"), Status);
+		}
 
 
 		// Linear PCM stream format
@@ -170,7 +174,7 @@ namespace Audio
 		Status = AudioObjectSetPropertyData(DeviceAudioObjectID, &DevicePropertyAddress, 0, nullptr, AudioDeviceQuerySize, &AUBufferSize);
 		if(Status != 0)
 		{
-			HandleError(TEXT("Failed to set output format!"), true);
+			HandleError(TEXT("Failed to set buffer size!"), true);
 			return false;
 		}
 
