Fix for Blit encoder copy validation error on macOS Mojave.   Metal buffer Lock/unlock ops should be in multiples of 4 byte chunks on macOS (1 byte for iOS) according to the documentation.   It's not clear if this limitation is also in place for Catalina and Big Sur as they don't seem trigger the validation error.   Make sure we can apply min alignment buffer lock/unlock operations.

diff --git a/Engine/Source/Runtime/Apple/MetalRHI/Private/MetalVertexBuffer.cpp b/Engine/Source/Runtime/Apple/MetalRHI/Private/MetalVertexBuffer.cpp
--- a/Engine/Source/Runtime/Apple/MetalRHI/Private/MetalVertexBuffer.cpp
+++ b/Engine/Source/Runtime/Apple/MetalRHI/Private/MetalVertexBuffer.cpp
@@ -159,6 +159,11 @@ FMetalRHIBuffer::FMetalRHIBuffer(uint32 InSize, uint32 InUsage, ERHIResourceType
 		}
 		else
 		{
+#if PLATFORM_MAC
+			// Buffer can be blit encoder copied on lock/unlock, we need to know that the buffer size is large enough for copy operations that are in multiples of
+			// 4 bytes on macOS, iOS can be 1 byte.  Update size to know we have at least this much buffer memory, it will be larger in the end.
+			Size = Align(InSize, 4);
+#endif
 			uint32 AllocSize = Size;
 			
 			if ((InUsage & EMetalBufferUsage_LinearTex) && !FMetalCommandQueue::SupportsFeature(EMetalFeaturesTextureBuffers))
@@ -532,6 +537,11 @@ void* FMetalRHIBuffer::Lock(bool bIsOnRHIThread, EResourceLockMode InLockMode, u
 		return ((uint8*)Data->Data) + Offset;
 	}
 	
+#if PLATFORM_MAC
+	// Blit encoder validation error, lock size and subsequent blit copy unlock operations need to be in 4 byte multiples on macOS
+	InSize = FMath::Min(Align(InSize, 4), Size - Offset);
+#endif
+	
 	// The system is very naughty and does not obey this rule
 //	check(LastLockFrame == 0 || LastLockFrame != GetMetalDeviceContext().GetFrameNumberRHIThread());
 	
