Adding default pixel format capabilities for platforms that might not be able to poll for support at runtime.

diff --git a/Engine/Source/Runtime/Apple/AGXRHI/Private/AGXRHI.cpp b/Engine/Source/Runtime/Apple/AGXRHI/Private/AGXRHI.cpp
--- a/Engine/Source/Runtime/Apple/AGXRHI/Private/AGXRHI.cpp
+++ b/Engine/Source/Runtime/Apple/AGXRHI/Private/AGXRHI.cpp
@@ -29,8 +29,6 @@ bool GIsAGXInitialized = false;
 
 FAGXBufferFormat GAGXBufferFormats[PF_MAX];
 
-static bool GFormatSupportsTypedUAVLoad[PF_MAX] = { false };
-
 static TAutoConsoleVariable<int32> CVarUseRHIThread(
 													TEXT("r.AGX.IOSRHIThread"),
 													0,
@@ -900,31 +898,38 @@ FAGXDynamicRHI::FAGXDynamicRHI(ERHIFeatureLevel::Type RequestedFeatureLevel)
 	}
 #endif
 
+	RHIInitDefaultPixelFormatCapabilities();
+
+	auto AddTypedUAVSupport = [](EPixelFormat InPixelFormat)
+	{
+		EnumAddFlags(GPixelFormats[InPixelFormat].Capabilities, EPixelFormatCapabilities::TypedUAVLoad | EPixelFormatCapabilities::TypedUAVStore);
+	};
+
 	switch ([GMtlDevice readWriteTextureSupport])
 	{
-		case MTLReadWriteTextureTier2:
-			GFormatSupportsTypedUAVLoad[PF_A32B32G32R32F]			= true;
-			GFormatSupportsTypedUAVLoad[PF_R32G32B32A32_UINT]		= true;
-			GFormatSupportsTypedUAVLoad[PF_FloatRGBA]				= true;
-			GFormatSupportsTypedUAVLoad[PF_R16G16B16A16_UINT]		= true;
-			GFormatSupportsTypedUAVLoad[PF_R16G16B16A16_SINT]		= true;
-			GFormatSupportsTypedUAVLoad[PF_R8G8B8A8]				= true;
-			GFormatSupportsTypedUAVLoad[PF_R8G8B8A8_UINT]			= true;
-			GFormatSupportsTypedUAVLoad[PF_R16F]					= true;
-			GFormatSupportsTypedUAVLoad[PF_R16_UINT]				= true;
-			GFormatSupportsTypedUAVLoad[PF_R16_SINT]				= true;
-			GFormatSupportsTypedUAVLoad[PF_R8]						= true;
-			GFormatSupportsTypedUAVLoad[PF_R8_UINT]					= true;
-			// Fall through
-
-		case MTLReadWriteTextureTier1:
-			GFormatSupportsTypedUAVLoad[PF_R32_FLOAT]				= true;
-			GFormatSupportsTypedUAVLoad[PF_R32_UINT]				= true;
-			GFormatSupportsTypedUAVLoad[PF_R32_SINT]				= true;
-			// Fall through
-
-		case MTLReadWriteTextureTierNone:
-			break;
+	case MTLReadWriteTextureTier2:
+		AddTypedUAVSupport(PF_A32B32G32R32F);
+		AddTypedUAVSupport(PF_R32G32B32A32_UINT);
+		AddTypedUAVSupport(PF_FloatRGBA);
+		AddTypedUAVSupport(PF_R16G16B16A16_UINT);
+		AddTypedUAVSupport(PF_R16G16B16A16_SINT);
+		AddTypedUAVSupport(PF_R8G8B8A8);
+		AddTypedUAVSupport(PF_R8G8B8A8_UINT);
+		AddTypedUAVSupport(PF_R16F);
+		AddTypedUAVSupport(PF_R16_UINT);
+		AddTypedUAVSupport(PF_R16_SINT);
+		AddTypedUAVSupport(PF_R8);
+		AddTypedUAVSupport(PF_R8_UINT);
+		// Fall through
+
+	case MTLReadWriteTextureTier1:
+		AddTypedUAVSupport(PF_R32_FLOAT);
+		AddTypedUAVSupport(PF_R32_UINT);
+		AddTypedUAVSupport(PF_R32_SINT);
+		// Fall through
+
+	case MTLReadWriteTextureTierNone:
+		break;
 	};
 
 	// get driver version (todo: share with other RHIs)
@@ -1237,7 +1242,7 @@ void* FAGXDynamicRHI::RHIGetNativeInstance()
 
 bool FAGXDynamicRHI::RHIIsTypedUAVLoadSupported(EPixelFormat PixelFormat)
 {
-	return GFormatSupportsTypedUAVLoad[PixelFormat];
+	return EnumHasAnyFlags(GPixelFormats[PixelFormat].Capabilities, EPixelFormatCapabilities::TypedUAVLoad);
 }
 
 uint16 FAGXDynamicRHI::RHIGetPlatformTextureMaxSampleCount()
diff --git a/Engine/Source/Runtime/Apple/MetalRHI/Private/MetalRHI.cpp b/Engine/Source/Runtime/Apple/MetalRHI/Private/MetalRHI.cpp
--- a/Engine/Source/Runtime/Apple/MetalRHI/Private/MetalRHI.cpp
+++ b/Engine/Source/Runtime/Apple/MetalRHI/Private/MetalRHI.cpp
@@ -29,8 +29,6 @@ bool GIsMetalInitialized = false;
 
 FMetalBufferFormat GMetalBufferFormats[PF_MAX];
 
-static bool GFormatSupportsTypedUAVLoad[PF_MAX] = { false };
-
 static TAutoConsoleVariable<int32> CVarUseIOSRHIThread(
 													TEXT("r.Metal.IOSRHIThread"),
 													0,
@@ -943,31 +941,38 @@ FMetalDynamicRHI::FMetalDynamicRHI(ERHIFeatureLevel::Type RequestedFeatureLevel)
 	}
 #endif
 
+	RHIInitDefaultPixelFormatCapabilities();
+
+	auto AddTypedUAVSupport = [](EPixelFormat InPixelFormat)
+	{
+		EnumAddFlags(GPixelFormats[InPixelFormat].Capabilities, EPixelFormatCapabilities::TypedUAVLoad | EPixelFormatCapabilities::TypedUAVStore);
+	};
+
 	switch (Device.GetReadWriteTextureSupport())
 	{
-		case mtlpp::ReadWriteTextureTier::Tier2:
-			GFormatSupportsTypedUAVLoad[PF_A32B32G32R32F]			= true;
-			GFormatSupportsTypedUAVLoad[PF_R32G32B32A32_UINT]		= true;
-			GFormatSupportsTypedUAVLoad[PF_FloatRGBA]				= true;
-			GFormatSupportsTypedUAVLoad[PF_R16G16B16A16_UINT]		= true;
-			GFormatSupportsTypedUAVLoad[PF_R16G16B16A16_SINT]		= true;
-			GFormatSupportsTypedUAVLoad[PF_R8G8B8A8]				= true;
-			GFormatSupportsTypedUAVLoad[PF_R8G8B8A8_UINT]			= true;
-			GFormatSupportsTypedUAVLoad[PF_R16F]					= true;
-			GFormatSupportsTypedUAVLoad[PF_R16_UINT]				= true;
-			GFormatSupportsTypedUAVLoad[PF_R16_SINT]				= true;
-			GFormatSupportsTypedUAVLoad[PF_R8]						= true;
-			GFormatSupportsTypedUAVLoad[PF_R8_UINT]					= true;
-			// Fall through
-
-		case mtlpp::ReadWriteTextureTier::Tier1:
-			GFormatSupportsTypedUAVLoad[PF_R32_FLOAT]				= true;
-			GFormatSupportsTypedUAVLoad[PF_R32_UINT]				= true;
-			GFormatSupportsTypedUAVLoad[PF_R32_SINT]				= true;
-			// Fall through
-
-		case mtlpp::ReadWriteTextureTier::None:
-			break;
+	case mtlpp::ReadWriteTextureTier::Tier2:
+		AddTypedUAVSupport(PF_A32B32G32R32F);
+		AddTypedUAVSupport(PF_R32G32B32A32_UINT);
+		AddTypedUAVSupport(PF_FloatRGBA);
+		AddTypedUAVSupport(PF_R16G16B16A16_UINT);
+		AddTypedUAVSupport(PF_R16G16B16A16_SINT);
+		AddTypedUAVSupport(PF_R8G8B8A8);
+		AddTypedUAVSupport(PF_R8G8B8A8_UINT);
+		AddTypedUAVSupport(PF_R16F);
+		AddTypedUAVSupport(PF_R16_UINT);
+		AddTypedUAVSupport(PF_R16_SINT);
+		AddTypedUAVSupport(PF_R8);
+		AddTypedUAVSupport(PF_R8_UINT);
+		// Fall through
+
+	case mtlpp::ReadWriteTextureTier::Tier1:
+		AddTypedUAVSupport(PF_R32_FLOAT);
+		AddTypedUAVSupport(PF_R32_UINT);
+		AddTypedUAVSupport(PF_R32_SINT);
+		// Fall through
+
+	case mtlpp::ReadWriteTextureTier::None:
+		break;
 	};
 
 	// get driver version (todo: share with other RHIs)
@@ -1290,7 +1295,7 @@ void* FMetalDynamicRHI::RHIGetNativeInstance()
 
 bool FMetalDynamicRHI::RHIIsTypedUAVLoadSupported(EPixelFormat PixelFormat)
 {
-	return GFormatSupportsTypedUAVLoad[PixelFormat];
+	return EnumHasAnyFlags(GPixelFormats[PixelFormat].Capabilities, EPixelFormatCapabilities::TypedUAVLoad);
 }
 
 uint16 FMetalDynamicRHI::RHIGetPlatformTextureMaxSampleCount()
diff --git a/Engine/Source/Runtime/RHI/Private/RHI.cpp b/Engine/Source/Runtime/RHI/Private/RHI.cpp
--- a/Engine/Source/Runtime/RHI/Private/RHI.cpp
+++ b/Engine/Source/Runtime/RHI/Private/RHI.cpp
@@ -2526,6 +2526,30 @@ FPixelFormatInfo    GPixelFormats[PF_MAX] =
 	FPixelFormatInfo(PF_R8_SINT,            TEXT("R8_SINT"),            1,          1,          1,          1,          1,              1),
 };
 
+void RHIInitDefaultPixelFormatCapabilities()
+{
+	for (FPixelFormatInfo& Info : GPixelFormats)
+	{
+		if (Info.Supported)
+		{
+			const EPixelFormat PixelFormat = Info.UnrealFormat;
+			if (IsBlockCompressedFormat(PixelFormat))
+			{
+				// Block compressed formats should have limited capabilities
+				EnumAddFlags(Info.Capabilities, EPixelFormatCapabilities::AnyTexture | EPixelFormatCapabilities::TextureMipmaps | EPixelFormatCapabilities::TextureLoad | EPixelFormatCapabilities::TextureSample | EPixelFormatCapabilities::TextureGather);
+			}
+			else
+			{
+				EnumAddFlags(Info.Capabilities, EPixelFormatCapabilities::AllTextureFlags | EPixelFormatCapabilities::AllBufferFlags | EPixelFormatCapabilities::UAV);
+				if (!IsDepthOrStencilFormat(PixelFormat))
+				{
+					EnumRemoveFlags(Info.Capabilities, EPixelFormatCapabilities::DepthStencil);
+				}
+			}
+		}
+	}
+}
+
 static struct FValidatePixelFormats
 {
 	FValidatePixelFormats()
diff --git a/Engine/Source/Runtime/RHI/Public/RHI.h b/Engine/Source/Runtime/RHI/Public/RHI.h
--- a/Engine/Source/Runtime/RHI/Public/RHI.h
+++ b/Engine/Source/Runtime/RHI/Public/RHI.h
@@ -792,7 +792,13 @@ enum class EPixelFormatCapabilities : uint32
     TypedUAVLoad     = 1ull << 20,
 	TypedUAVStore    = 1ull << 21,
 
-	AnyTexture = Texture1D | Texture2D | Texture3D | TextureCube,
+	AnyTexture       = Texture1D | Texture2D | Texture3D | TextureCube,
+
+	AllTextureFlags  = AnyTexture | RenderTarget | DepthStencil | TextureMipmaps | TextureLoad | TextureSample | TextureGather | TextureAtomics | TextureBlendable,
+	AllBufferFlags   = Buffer | VertexBuffer | IndexBuffer | BufferLoad | BufferStore | BufferAtomics,
+	AllUAVFlags      = UAV | TypedUAVLoad | TypedUAVStore,
+
+	AllFlags         = AllTextureFlags | AllBufferFlags | AllUAVFlags
 };
 ENUM_CLASS_FLAGS(EPixelFormatCapabilities);
 
@@ -830,6 +836,9 @@ struct FPixelFormatInfo
 
 extern RHI_API FPixelFormatInfo GPixelFormats[PF_MAX];		// Maps members of EPixelFormat to a FPixelFormatInfo describing the format.
 
+/** Initialize the 'best guess' pixel format capabilities. Platform formats and support must be filled out before calling this. */
+extern RHI_API void RHIInitDefaultPixelFormatCapabilities();
+
 inline bool PixelFormatHasCapabilities(EPixelFormat InFormat, EPixelFormatCapabilities InCapabilities)
 {
 	return EnumHasAllFlags(GPixelFormats[InFormat].Capabilities, InCapabilities);
