Remove references to UE4 in Metal and AGX RHIs.

diff --git a/Engine/Source/Runtime/Apple/AGXRHI/Private/AGXContext.h b/Engine/Source/Runtime/Apple/AGXRHI/Private/AGXContext.h
--- a/Engine/Source/Runtime/Apple/AGXRHI/Private/AGXContext.h
+++ b/Engine/Source/Runtime/Apple/AGXRHI/Private/AGXContext.h
@@ -43,7 +43,7 @@ public:
 
 	/**
 	 * Do anything necessary to prepare for any kind of draw call 
-	 * @param PrimitiveType The UE4 primitive type for the draw call, needed to compile the correct render pipeline.
+	 * @param PrimitiveType The UnrealEngine primitive type for the draw call, needed to compile the correct render pipeline.
 	 * @returns True if the preparation completed and the draw call can be encoded, false to skip.
 	 */
 	bool PrepareToDraw(uint32 PrimitiveType);
diff --git a/Engine/Source/Runtime/Apple/AGXRHI/Private/AGXTexture.cpp b/Engine/Source/Runtime/Apple/AGXRHI/Private/AGXTexture.cpp
--- a/Engine/Source/Runtime/Apple/AGXRHI/Private/AGXTexture.cpp
+++ b/Engine/Source/Runtime/Apple/AGXRHI/Private/AGXTexture.cpp
@@ -2637,7 +2637,7 @@ void FAGXDynamicRHI::RHICopySubTextureRegion(FRHITexture2D* SourceTexture, FRHIT
 		}
 		else
 		{
-			UE_LOG(LogAGX, Error, TEXT("RHICopySubTextureRegion Source (UE4 %d: MTL %d) <-> Destination (UE4 %d: MTL %d) texture format mismatch"), (uint32)SourceTexture->GetFormat(), (uint32)MetalSrcTexture->Surface.Texture.GetPixelFormat(), (uint32)DestinationTexture->GetFormat(), (uint32)MetalDestTexture->Surface.Texture.GetPixelFormat());
+			UE_LOG(LogAGX, Error, TEXT("RHICopySubTextureRegion Source (UnrealEngine %d: MTL %d) <-> Destination (UnrealEngine %d: MTL %d) texture format mismatch"), (uint32)SourceTexture->GetFormat(), (uint32)MetalSrcTexture->Surface.Texture.GetPixelFormat(), (uint32)DestinationTexture->GetFormat(), (uint32)MetalDestTexture->Surface.Texture.GetPixelFormat());
 		}
 	}
 }
@@ -2768,7 +2768,7 @@ void FAGXRHICommandContext::RHICopyTexture(FRHITexture* SourceTextureRHI, FRHITe
 		}
 		else
 		{
-			UE_LOG(LogAGX, Error, TEXT("RHICopyTexture Source (UE4 %d: MTL %d) <-> Destination (UE4 %d: MTL %d) texture format mismatch"), (uint32)SourceTextureRHI->GetFormat(), (uint32)MetalSrcTexture->Texture.GetPixelFormat(), (uint32)DestTextureRHI->GetFormat(), (uint32)MetalDestTexture->Texture.GetPixelFormat());
+			UE_LOG(LogAGX, Error, TEXT("RHICopyTexture Source (UnrealEngine %d: MTL %d) <-> Destination (UnrealEngine %d: MTL %d) texture format mismatch"), (uint32)SourceTextureRHI->GetFormat(), (uint32)MetalSrcTexture->Texture.GetPixelFormat(), (uint32)DestTextureRHI->GetFormat(), (uint32)MetalDestTexture->Texture.GetPixelFormat());
 		}
 	}
 }
diff --git a/Engine/Source/Runtime/Apple/AGXRHI/Private/AGXVertexDeclaration.cpp b/Engine/Source/Runtime/Apple/AGXRHI/Private/AGXVertexDeclaration.cpp
--- a/Engine/Source/Runtime/Apple/AGXRHI/Private/AGXVertexDeclaration.cpp
+++ b/Engine/Source/Runtime/Apple/AGXRHI/Private/AGXVertexDeclaration.cpp
@@ -147,7 +147,7 @@ void FAGXVertexDeclaration::GenerateLayout(const FVertexDeclarationElementList&
 				Stride = TranslateElementTypeToSize(Element.Type);
 			}
 
-			// look for any unset strides coming from UE4 (this can be removed when all are fixed)
+			// look for any unset strides coming from the Engine (this can be removed when all are fixed)
 			if (Element.Stride == 0xFFFF)
 			{
 				NSLog(@"Setting illegal stride - break here if you want to find out why, but this won't break until we try to render with it");
diff --git a/Engine/Source/Runtime/Apple/MetalRHI/Private/MetalCommandEncoder.h b/Engine/Source/Runtime/Apple/MetalRHI/Private/MetalCommandEncoder.h
--- a/Engine/Source/Runtime/Apple/MetalRHI/Private/MetalCommandEncoder.h
+++ b/Engine/Source/Runtime/Apple/MetalRHI/Private/MetalCommandEncoder.h
@@ -115,7 +115,7 @@ enum EMetalCommandEncoderType
 /**
  * FMetalCommandEncoder:
  *	Wraps the details of switching between different command encoders on the command-buffer, allowing for restoration of the render encoder if needed.
- * 	UE4 expects the API to serialise commands in-order, but Metal expects applications to work with command-buffers directly so we need to implement 
+ * 	UnrealEngine expects the API to serialise commands in-order, but Metal expects applications to work with command-buffers directly so we need to implement 
  *	the RHI semantics by switching between encoder types. This class hides the ugly details. Eventually it might be possible to move some of the operations
  *	into pre- & post- command-buffers so that we avoid encoder switches but that will take changes to the RHI and high-level code too, so it won't happen soon.
  */
diff --git a/Engine/Source/Runtime/Apple/MetalRHI/Private/MetalContext.h b/Engine/Source/Runtime/Apple/MetalRHI/Private/MetalContext.h
--- a/Engine/Source/Runtime/Apple/MetalRHI/Private/MetalContext.h
+++ b/Engine/Source/Runtime/Apple/MetalRHI/Private/MetalContext.h
@@ -45,7 +45,7 @@ public:
 
 	/**
 	 * Do anything necessary to prepare for any kind of draw call 
-	 * @param PrimitiveType The UE4 primitive type for the draw call, needed to compile the correct render pipeline.
+	 * @param PrimitiveType The UnrealEngine primitive type for the draw call, needed to compile the correct render pipeline.
 	 * @returns True if the preparation completed and the draw call can be encoded, false to skip.
 	 */
 	bool PrepareToDraw(uint32 PrimitiveType);
diff --git a/Engine/Source/Runtime/Apple/MetalRHI/Private/MetalTexture.cpp b/Engine/Source/Runtime/Apple/MetalRHI/Private/MetalTexture.cpp
--- a/Engine/Source/Runtime/Apple/MetalRHI/Private/MetalTexture.cpp
+++ b/Engine/Source/Runtime/Apple/MetalRHI/Private/MetalTexture.cpp
@@ -2670,7 +2670,7 @@ void FMetalDynamicRHI::RHICopySubTextureRegion(FRHITexture2D* SourceTexture, FRH
 		}
 		else
 		{
-			UE_LOG(LogMetal, Error, TEXT("RHICopySubTextureRegion Source (UE4 %d: MTL %d) <-> Destination (UE4 %d: MTL %d) texture format mismatch"), (uint32)SourceTexture->GetFormat(), (uint32)MetalSrcTexture->Surface.Texture.GetPixelFormat(), (uint32)DestinationTexture->GetFormat(), (uint32)MetalDestTexture->Surface.Texture.GetPixelFormat());
+			UE_LOG(LogMetal, Error, TEXT("RHICopySubTextureRegion Source (UnrealEngine %d: MTL %d) <-> Destination (UnrealEngine %d: MTL %d) texture format mismatch"), (uint32)SourceTexture->GetFormat(), (uint32)MetalSrcTexture->Surface.Texture.GetPixelFormat(), (uint32)DestinationTexture->GetFormat(), (uint32)MetalDestTexture->Surface.Texture.GetPixelFormat());
 		}
 	}
 }
@@ -2801,7 +2801,7 @@ void FMetalRHICommandContext::RHICopyTexture(FRHITexture* SourceTextureRHI, FRHI
 		}
 		else
 		{
-			UE_LOG(LogMetal, Error, TEXT("RHICopyTexture Source (UE4 %d: MTL %d) <-> Destination (UE4 %d: MTL %d) texture format mismatch"), (uint32)SourceTextureRHI->GetFormat(), (uint32)MetalSrcTexture->Texture.GetPixelFormat(), (uint32)DestTextureRHI->GetFormat(), (uint32)MetalDestTexture->Texture.GetPixelFormat());
+			UE_LOG(LogMetal, Error, TEXT("RHICopyTexture Source (UnrealEngine %d: MTL %d) <-> Destination (UnrealEngine %d: MTL %d) texture format mismatch"), (uint32)SourceTextureRHI->GetFormat(), (uint32)MetalSrcTexture->Texture.GetPixelFormat(), (uint32)DestTextureRHI->GetFormat(), (uint32)MetalDestTexture->Texture.GetPixelFormat());
 		}
 	}
 }
diff --git a/Engine/Source/Runtime/Apple/MetalRHI/Private/MetalVertexDeclaration.cpp b/Engine/Source/Runtime/Apple/MetalRHI/Private/MetalVertexDeclaration.cpp
--- a/Engine/Source/Runtime/Apple/MetalRHI/Private/MetalVertexDeclaration.cpp
+++ b/Engine/Source/Runtime/Apple/MetalRHI/Private/MetalVertexDeclaration.cpp
@@ -147,7 +147,7 @@ void FMetalVertexDeclaration::GenerateLayout(const FVertexDeclarationElementList
 				Stride = TranslateElementTypeToSize(Element.Type);
 			}
 
-			// look for any unset strides coming from UE4 (this can be removed when all are fixed)
+			// look for any unset strides coming from the Engine (this can be removed when all are fixed)
 			if (Element.Stride == 0xFFFF)
 			{
 				NSLog(@"Setting illegal stride - break here if you want to find out why, but this won't break until we try to render with it");
