AGXRHI: scrub manual debugging hooks

Summary
This PR scrubs a lot of manual debugging functionality from the AGXRHI.

Details
- Remove entire debug classes for command encoders, command buffer markers, command buffers
- Remove unused / untested / stale fence code
- Remove remaining remnants of initial argument buffer implementation (that was partially (but not completely) yanked from Metal RHI)
- Remove MTLHeap code (unused / untested / broken)

+ Import platform config files. This will pave the road for removing the PLATFORM_XYZ code from the platform independent layer.
+ Add support for missing pixel formats.

PR #8773: AGXRHI: scrub manual debugging hooks (Contributed by [REDACTED_PERSON])

#rnx

diff --git a/Engine/Source/Runtime/Apple/AGXRHI/Private/AGXBlitCommandEncoder.cpp b/Engine/Source/Runtime/Apple/AGXRHI/Private/AGXBlitCommandEncoder.cpp
deleted file mode 100644
diff --git a/Engine/Source/Runtime/Apple/AGXRHI/Private/AGXBlitCommandEncoder.h b/Engine/Source/Runtime/Apple/AGXRHI/Private/AGXBlitCommandEncoder.h
deleted file mode 100644
diff --git a/Engine/Source/Runtime/Apple/AGXRHI/Private/AGXBuffer.cpp b/Engine/Source/Runtime/Apple/AGXRHI/Private/AGXBuffer.cpp
--- a/Engine/Source/Runtime/Apple/AGXRHI/Private/AGXBuffer.cpp
+++ b/Engine/Source/Runtime/Apple/AGXRHI/Private/AGXBuffer.cpp
@@ -5,7 +5,6 @@
 #include "Templates/AlignmentTemplates.h"
 #include "AGXLLM.h"
 #include <objc/runtime.h>
-#include "AGXCommandBuffer.h"
 #include "AGXProfiler.h"
 #include "AGXRenderPass.h"
 
@@ -172,7 +171,6 @@ FAGXSubBufferHeap::FAGXSubBufferHeap(NSUInteger Size, NSUInteger Alignment, mtlp
 , UsedSize(0)
 {
 	Options = (mtlpp::ResourceOptions)FAGXCommandQueue::GetCompatibleResourceOptions(Options);
-	static bool bSupportsHeaps = GetAGXDeviceContext().SupportsFeature(EAGXFeaturesHeaps);
 	NSUInteger FullSize = Align(Size, Alignment);
 	METAL_GPUPROFILE(FAGXScopedCPUStats CPUStat(FString::Printf(TEXT("AllocBuffer: %llu, %llu"), FullSize, Options)));
 	
@@ -181,43 +179,22 @@ FAGXSubBufferHeap::FAGXSubBufferHeap(NSUInteger Size, NSUInteger Alignment, mtlp
 	check(Storage != mtlpp::StorageMode::Managed /* Managed memory cannot be safely suballocated! When you overwrite existing data the GPU buffer is immediately disposed of! */);
 #endif
 
-	if (bSupportsHeaps && Storage == mtlpp::StorageMode::Private)
-	{
-		mtlpp::HeapDescriptor Desc;
-		Desc.SetSize(FullSize);
-		Desc.SetStorageMode(Storage);
-		ParentHeap = GMtlppDevice.NewHeap(Desc);
-		check(ParentHeap.GetPtr());
-#if STATS || ENABLE_LOW_LEVEL_MEM_TRACKER
-		AGXLLM::LogAllocHeap(ParentHeap);
-#endif
-	}
-	else
-	{
-		ParentBuffer = MTLPP_VALIDATE(mtlpp::Device, GMtlppDevice, AGXSafeGetRuntimeDebuggingLevel() >= EAGXDebugLevelValidation, NewBuffer(FullSize, Options));
-		check(ParentBuffer.GetPtr());
-		check(ParentBuffer.GetLength() >= FullSize);
+	ParentBuffer = MTLPP_VALIDATE(mtlpp::Device, GMtlppDevice, AGXSafeGetRuntimeDebuggingLevel() >= EAGXDebugLevelValidation, NewBuffer(FullSize, Options));
+	check(ParentBuffer.GetPtr());
+	check(ParentBuffer.GetLength() >= FullSize);
 #if STATS || ENABLE_LOW_LEVEL_MEM_TRACKER
-		AGXLLM::LogAllocBuffer(ParentBuffer);
+	AGXLLM::LogAllocBuffer(ParentBuffer);
 #endif
-		FreeRanges.Add(ns::Range(0, FullSize));
-	}
+	FreeRanges.Add(ns::Range(0, FullSize));
+
 	INC_MEMORY_STAT_BY(STAT_AGXBufferUnusedMemory, FullSize);
 	INC_MEMORY_STAT_BY(STAT_AGXHeapBufferUnusedMemory, FullSize);
 }
 
 FAGXSubBufferHeap::~FAGXSubBufferHeap()
 {
-	if (ParentHeap)
-	{
-		DEC_MEMORY_STAT_BY(STAT_AGXBufferUnusedMemory, ParentHeap.GetSize());
-		DEC_MEMORY_STAT_BY(STAT_AGXHeapBufferUnusedMemory, ParentHeap.GetSize());
-	}
-	else
-	{
-		DEC_MEMORY_STAT_BY(STAT_AGXBufferUnusedMemory, ParentBuffer.GetLength());
-		DEC_MEMORY_STAT_BY(STAT_AGXHeapBufferUnusedMemory, ParentBuffer.GetLength());
-	}
+	DEC_MEMORY_STAT_BY(STAT_AGXBufferUnusedMemory, ParentBuffer.GetLength());
+	DEC_MEMORY_STAT_BY(STAT_AGXHeapBufferUnusedMemory, ParentBuffer.GetLength());
 }
 
 void FAGXSubBufferHeap::SetOwner(ns::Range const& Range, FAGXRHIBuffer* Owner, bool bIsSwap)
@@ -251,13 +228,6 @@ void FAGXSubBufferHeap::FreeRange(ns::Range const& Range)
 			}
 		}
 	}
-	if (ParentHeap)
-	{
-		INC_MEMORY_STAT_BY(STAT_AGXBufferUnusedMemory, Range.Length);
-		INC_MEMORY_STAT_BY(STAT_AGXHeapBufferUnusedMemory, Range.Length);
-		DEC_MEMORY_STAT_BY(STAT_AGXHeapBufferMemory, Range.Length);
-	}
-	else
 	{
 #if METAL_DEBUG_OPTIONS
 		if (GIsRHIInitialized)
@@ -324,62 +294,27 @@ void FAGXSubBufferHeap::FreeRange(ns::Range const& Range)
 
 ns::String FAGXSubBufferHeap::GetLabel() const
 {
-	if (ParentHeap)
-	{
-		return ParentHeap.GetLabel();
-	}
-	else
-	{
-		return ParentBuffer.GetLabel();
-	}
+	return ParentBuffer.GetLabel();
 }
 
 mtlpp::StorageMode FAGXSubBufferHeap::GetStorageMode() const
 {
-	if (ParentHeap)
-	{
-		return ParentHeap.GetStorageMode();
-	}
-	else
-	{
-		return ParentBuffer.GetStorageMode();
-	}
+	return ParentBuffer.GetStorageMode();
 }
 
 mtlpp::CpuCacheMode FAGXSubBufferHeap::GetCpuCacheMode() const
 {
-	if (ParentHeap)
-	{
-		return ParentHeap.GetCpuCacheMode();
-	}
-	else
-	{
-		return ParentBuffer.GetCpuCacheMode();
-	}
+	return ParentBuffer.GetCpuCacheMode();
 }
 
 NSUInteger FAGXSubBufferHeap::GetSize() const
 {
-	if (ParentHeap)
-	{
-		return ParentHeap.GetSize();
-	}
-	else
-	{
-		return ParentBuffer.GetLength();
-	}
+	return ParentBuffer.GetLength();
 }
 
 NSUInteger FAGXSubBufferHeap::GetUsedSize() const
 {
-	if (ParentHeap)
-	{
-		return ParentHeap.GetUsedSize();
-	}
-	else
-	{
-		return UsedSize;
-	}
+	return UsedSize;
 }
 
 int64 FAGXSubBufferHeap::NumCurrentAllocations() const
@@ -389,50 +324,24 @@ int64 FAGXSubBufferHeap::NumCurrentAllocations() const
 
 void FAGXSubBufferHeap::SetLabel(const ns::String& label)
 {
-	if (ParentHeap)
-	{
-		ParentHeap.SetLabel(label);
-	}
-	else
-	{
-		ParentBuffer.SetLabel(label);
-	}
+	ParentBuffer.SetLabel(label);
 }
 
 NSUInteger FAGXSubBufferHeap::MaxAvailableSize() const
 {
-	if (ParentHeap)
+	if (UsedSize < GetSize())
 	{
-		return ParentHeap.MaxAvailableSizeWithAlignment(MinAlign);
+		return FreeRanges.Last().Length;
 	}
 	else
 	{
-		if (UsedSize < GetSize())
-		{
-			return FreeRanges.Last().Length;
-		}
-		else
-		{
-			return 0;
-		}
+		return 0;
 	}
 }
 
 bool FAGXSubBufferHeap::CanAllocateSize(NSUInteger Size) const
 {
-	if (ParentHeap)
-	{
-		NSUInteger Storage = (NSUInteger(GetStorageMode()) << mtlpp::ResourceStorageModeShift);
-		NSUInteger Cache = (NSUInteger(GetCpuCacheMode()) << mtlpp::ResourceCpuCacheModeShift);
-		mtlpp::ResourceOptions Opt = mtlpp::ResourceOptions(Storage | Cache);
-		
-		NSUInteger Align = GMtlppDevice.HeapBufferSizeAndAlign(Size, Opt).Align;
-		return Size <= ParentHeap.MaxAvailableSizeWithAlignment(Align);
-	}
-	else
-	{
-		return Size <= MaxAvailableSize();
-	}
+	return Size <= MaxAvailableSize();
 }
 
 FAGXBuffer FAGXSubBufferHeap::NewBuffer(NSUInteger length)
@@ -440,21 +349,6 @@ FAGXBuffer FAGXSubBufferHeap::NewBuffer(NSUInteger length)
 	NSUInteger Size = Align(length, MinAlign);
 	FAGXBuffer Result;
 	
-	if (ParentHeap)
-	{
-		NSUInteger Storage = (NSUInteger(GetStorageMode()) << mtlpp::ResourceStorageModeShift);
-		NSUInteger Cache = (NSUInteger(GetCpuCacheMode()) << mtlpp::ResourceCpuCacheModeShift);
-		mtlpp::ResourceOptions Opt = mtlpp::ResourceOptions(Storage | Cache);
-		
-		Result = FAGXBuffer(ParentHeap.NewBuffer(Size, Opt), this);
-#if STATS || ENABLE_LOW_LEVEL_MEM_TRACKER
-		AGXLLM::LogAllocBuffer(Result);
-#endif
-		DEC_MEMORY_STAT_BY(STAT_AGXBufferUnusedMemory, Result.GetLength());
-		DEC_MEMORY_STAT_BY(STAT_AGXHeapBufferUnusedMemory, Result.GetLength());
-		INC_MEMORY_STAT_BY(STAT_AGXHeapBufferMemory, Result.GetLength());
-	}
-	else
 	{
 		check(ParentBuffer && ParentBuffer.GetPtr());
 	
@@ -512,14 +406,7 @@ FAGXBuffer FAGXSubBufferHeap::NewBuffer(NSUInteger length)
 
 mtlpp::PurgeableState FAGXSubBufferHeap::SetPurgeableState(mtlpp::PurgeableState state)
 {
-	if (ParentHeap)
-	{
-		return ParentHeap.SetPurgeableState(state);
-	}
-	else
-	{
-		return ParentBuffer.SetPurgeableState(state);
-	}
+	return ParentBuffer.SetPurgeableState(state);
 }
 
 #pragma mark --
@@ -658,163 +545,84 @@ FAGXSubBufferMagazine::FAGXSubBufferMagazine(NSUInteger Size, NSUInteger ChunkSi
 , UsedSize(0)
 {
 	Options = (mtlpp::ResourceOptions)FAGXCommandQueue::GetCompatibleResourceOptions(Options);
-    static bool bSupportsHeaps = GetAGXDeviceContext().SupportsFeature(EAGXFeaturesHeaps);
-    mtlpp::StorageMode Storage = (mtlpp::StorageMode)((Options & mtlpp::ResourceStorageModeMask) >> mtlpp::ResourceStorageModeShift);
-    if (PLATFORM_IOS && bSupportsHeaps && Storage == mtlpp::StorageMode::Private)
-    {
-        MinAlign = GMtlppDevice.HeapBufferSizeAndAlign(BlockSize, Options).Align;
-    }
-    
-    NSUInteger FullSize = Align(Size, MinAlign);
+	mtlpp::StorageMode Storage = (mtlpp::StorageMode)((Options & mtlpp::ResourceStorageModeMask) >> mtlpp::ResourceStorageModeShift);
+	NSUInteger FullSize = Align(Size, MinAlign);
 	METAL_GPUPROFILE(FAGXScopedCPUStats CPUStat(FString::Printf(TEXT("AllocBuffer: %llu, %llu"), FullSize, Options)));
 	
 #if PLATFORM_MAC
 	check(Storage != mtlpp::StorageMode::Managed /* Managed memory cannot be safely suballocated! When you overwrite existing data the GPU buffer is immediately disposed of! */);
 #endif
 
-	if (bSupportsHeaps && Storage == mtlpp::StorageMode::Private)
-	{
-		mtlpp::HeapDescriptor Desc;
-		Desc.SetSize(FullSize);
-		Desc.SetStorageMode(Storage);
-		ParentHeap = GMtlppDevice.NewHeap(Desc);
-		check(ParentHeap.GetPtr());
-		METAL_FATAL_ASSERT(ParentHeap, TEXT("Failed to create heap of size %u and resource options %u"), Size, (uint32)Options);
-#if STATS || ENABLE_LOW_LEVEL_MEM_TRACKER
-		AGXLLM::LogAllocHeap(ParentHeap);
-#endif
-	}
-	else
-	{
-		ParentBuffer = MTLPP_VALIDATE(mtlpp::Device, GMtlppDevice, AGXSafeGetRuntimeDebuggingLevel() >= EAGXDebugLevelValidation, NewBuffer(FullSize, Options));
-		check(ParentBuffer.GetPtr());
-		check(ParentBuffer.GetLength() >= FullSize);
-		METAL_FATAL_ASSERT(ParentBuffer, TEXT("Failed to create heap of size %u and resource options %u"), Size, (uint32)Options);
+	ParentBuffer = MTLPP_VALIDATE(mtlpp::Device, GMtlppDevice, AGXSafeGetRuntimeDebuggingLevel() >= EAGXDebugLevelValidation, NewBuffer(FullSize, Options));
+	check(ParentBuffer.GetPtr());
+	check(ParentBuffer.GetLength() >= FullSize);
+	METAL_FATAL_ASSERT(ParentBuffer, TEXT("Failed to create heap of size %u and resource options %u"), Size, (uint32)Options);
 #if STATS || ENABLE_LOW_LEVEL_MEM_TRACKER
-		AGXLLM::LogAllocBuffer(ParentBuffer);
+	AGXLLM::LogAllocBuffer(ParentBuffer);
 #endif
-		
-		INC_MEMORY_STAT_BY(STAT_AGXBufferUnusedMemory, FullSize);
-		INC_MEMORY_STAT_BY(STAT_AGXMagazineBufferUnusedMemory, FullSize);
-        uint32 BlockCount = FullSize / ChunkSize;
-        Blocks.AddZeroed(BlockCount);
-	}
+	
+	INC_MEMORY_STAT_BY(STAT_AGXBufferUnusedMemory, FullSize);
+	INC_MEMORY_STAT_BY(STAT_AGXMagazineBufferUnusedMemory, FullSize);
+	uint32 BlockCount = FullSize / ChunkSize;
+	Blocks.AddZeroed(BlockCount);
 }
 
 FAGXSubBufferMagazine::~FAGXSubBufferMagazine()
 {
-	if (ParentHeap)
-	{
-		DEC_MEMORY_STAT_BY(STAT_AGXBufferUnusedMemory, ParentHeap.GetSize());
-		DEC_MEMORY_STAT_BY(STAT_AGXMagazineBufferUnusedMemory, ParentHeap.GetSize());
-	}
-	else
-	{
-		DEC_MEMORY_STAT_BY(STAT_AGXBufferUnusedMemory, ParentBuffer.GetLength());
-		DEC_MEMORY_STAT_BY(STAT_AGXMagazineBufferUnusedMemory, ParentBuffer.GetLength());
-	}
+	DEC_MEMORY_STAT_BY(STAT_AGXBufferUnusedMemory, ParentBuffer.GetLength());
+	DEC_MEMORY_STAT_BY(STAT_AGXMagazineBufferUnusedMemory, ParentBuffer.GetLength());
 }
 
 void FAGXSubBufferMagazine::FreeRange(ns::Range const& Range)
 {
 	FPlatformAtomics::InterlockedDecrement(&OutstandingAllocs);
-	if (ParentHeap)
+
+#if METAL_DEBUG_OPTIONS
+	if (GIsRHIInitialized)
 	{
-		INC_MEMORY_STAT_BY(STAT_AGXBufferUnusedMemory, Range.Length);
-		INC_MEMORY_STAT_BY(STAT_AGXMagazineBufferUnusedMemory, Range.Length);
-		DEC_MEMORY_STAT_BY(STAT_AGXMagazineBufferMemory, Range.Length);
+		MTLPP_VALIDATE_ONLY(mtlpp::Buffer, ParentBuffer, AGXSafeGetRuntimeDebuggingLevel() >= EAGXDebugLevelValidation, ReleaseRange(Range));
+		FAGXBuffer Buf(ParentBuffer.NewBuffer(Range), false);
+		GetAGXDeviceContext().ValidateIsInactiveBuffer(Buf);
 	}
-	else
-	{
-#if METAL_DEBUG_OPTIONS
-		if (GIsRHIInitialized)
-		{
-			MTLPP_VALIDATE_ONLY(mtlpp::Buffer, ParentBuffer, AGXSafeGetRuntimeDebuggingLevel() >= EAGXDebugLevelValidation, ReleaseRange(Range));
-			FAGXBuffer Buf(ParentBuffer.NewBuffer(Range), false);
-			GetAGXDeviceContext().ValidateIsInactiveBuffer(Buf);
-		}
 #endif
+
+	uint32 BlockIndex = Range.Location / Range.Length;
+	FPlatformAtomics::AtomicStore(&Blocks[BlockIndex], 0);
+	FPlatformAtomics::InterlockedAdd(&UsedSize, -((int64)Range.Length));
 	
-		uint32 BlockIndex = Range.Location / Range.Length;
-		FPlatformAtomics::AtomicStore(&Blocks[BlockIndex], 0);
-		FPlatformAtomics::InterlockedAdd(&UsedSize, -((int64)Range.Length));
-		
-		INC_MEMORY_STAT_BY(STAT_AGXBufferUnusedMemory, Range.Length);
-		INC_MEMORY_STAT_BY(STAT_AGXMagazineBufferUnusedMemory, Range.Length);
-		DEC_MEMORY_STAT_BY(STAT_AGXMagazineBufferMemory, Range.Length);
-	}
+	INC_MEMORY_STAT_BY(STAT_AGXBufferUnusedMemory, Range.Length);
+	INC_MEMORY_STAT_BY(STAT_AGXMagazineBufferUnusedMemory, Range.Length);
+	DEC_MEMORY_STAT_BY(STAT_AGXMagazineBufferMemory, Range.Length);
 }
 
 ns::String FAGXSubBufferMagazine::GetLabel() const
 {
-	if (ParentHeap)
-	{
-		return ParentHeap.GetLabel();
-	}
-	else
-	{
-		return ParentBuffer.GetLabel();
-	}
+	return ParentBuffer.GetLabel();
 }
 
 mtlpp::StorageMode FAGXSubBufferMagazine::GetStorageMode() const
 {
-	if (ParentHeap)
-	{
-		return ParentHeap.GetStorageMode();
-	}
-	else
-	{
-		return ParentBuffer.GetStorageMode();
-	}
+	return ParentBuffer.GetStorageMode();
 }
 
 mtlpp::CpuCacheMode FAGXSubBufferMagazine::GetCpuCacheMode() const
 {
-	if (ParentHeap)
-	{
-		return ParentHeap.GetCpuCacheMode();
-	}
-	else
-	{
-		return ParentBuffer.GetCpuCacheMode();
-	}
+	return ParentBuffer.GetCpuCacheMode();
 }
 
 NSUInteger FAGXSubBufferMagazine::GetSize() const
 {
-	if (ParentHeap)
-	{
-		return ParentHeap.GetSize();
-	}
-	else
-	{
-		return ParentBuffer.GetLength();
-	}
+	return ParentBuffer.GetLength();
 }
 
 NSUInteger FAGXSubBufferMagazine::GetUsedSize() const
 {
-	if (ParentHeap)
-	{
-		return ParentHeap.GetUsedSize();
-	}
-	else
-	{
-		return (NSUInteger)FPlatformAtomics::AtomicRead(&UsedSize);
-	}
+	return (NSUInteger)FPlatformAtomics::AtomicRead(&UsedSize);
 }
 
 NSUInteger FAGXSubBufferMagazine::GetFreeSize() const
 {
-	if (ParentHeap)
-	{
-		return ParentHeap.MaxAvailableSizeWithAlignment(MinAlign);
-	}
-	else
-	{
-		return GetSize() - GetUsedSize();
-	}
+	return GetSize() - GetUsedSize();
 }
 
 int64 FAGXSubBufferMagazine::NumCurrentAllocations() const
@@ -829,14 +637,7 @@ bool FAGXSubBufferMagazine::CanAllocateSize(NSUInteger Size) const
 
 void FAGXSubBufferMagazine::SetLabel(const ns::String& label)
 {
-	if (ParentHeap)
-	{
-		ParentHeap.SetLabel(label);
-	}
-	else
-	{
-		ParentBuffer.SetLabel(label);
-	}
+	ParentBuffer.SetLabel(label);
 }
 
 FAGXBuffer FAGXSubBufferMagazine::NewBuffer()
@@ -844,36 +645,19 @@ FAGXBuffer FAGXSubBufferMagazine::NewBuffer()
 	NSUInteger Size = BlockSize;
 	FAGXBuffer Result;
 
-	if (ParentHeap)
-	{
-		NSUInteger Storage = (NSUInteger(GetStorageMode()) << mtlpp::ResourceStorageModeShift);
-		NSUInteger Cache = (NSUInteger(GetCpuCacheMode()) << mtlpp::ResourceCpuCacheModeShift);
-		mtlpp::ResourceOptions Opt = mtlpp::ResourceOptions(Storage | Cache);
-		
-		Result = FAGXBuffer(ParentHeap.NewBuffer(Size, Opt), this);
-#if STATS || ENABLE_LOW_LEVEL_MEM_TRACKER
-		AGXLLM::LogAllocBuffer(Result);
-#endif
-		DEC_MEMORY_STAT_BY(STAT_AGXBufferUnusedMemory, Result.GetLength());
-		DEC_MEMORY_STAT_BY(STAT_AGXMagazineBufferUnusedMemory, Result.GetLength());
-		INC_MEMORY_STAT_BY(STAT_AGXMagazineBufferMemory, Result.GetLength());
-	}
-	else
+	check(ParentBuffer && ParentBuffer.GetPtr());
+	
+	for (uint32 i = 0; i < Blocks.Num(); i++)
 	{
-		check(ParentBuffer && ParentBuffer.GetPtr());
-		
-		for (uint32 i = 0; i < Blocks.Num(); i++)
+		if (FPlatformAtomics::InterlockedCompareExchange(&Blocks[i], 1, 0) == 0)
 		{
-			if (FPlatformAtomics::InterlockedCompareExchange(&Blocks[i], 1, 0) == 0)
-			{
-				ns::Range Range(i * BlockSize, BlockSize);
-				FPlatformAtomics::InterlockedAdd(&UsedSize, ((int64)Range.Length));
-				DEC_MEMORY_STAT_BY(STAT_AGXBufferUnusedMemory, Range.Length);
-				DEC_MEMORY_STAT_BY(STAT_AGXMagazineBufferUnusedMemory, Range.Length);
-				INC_MEMORY_STAT_BY(STAT_AGXMagazineBufferMemory, Range.Length);
-				Result = FAGXBuffer(MTLPP_VALIDATE(mtlpp::Buffer, ParentBuffer, AGXSafeGetRuntimeDebuggingLevel() >= EAGXDebugLevelValidation, NewBuffer(Range)), this);
-				break;
-			}
+			ns::Range Range(i * BlockSize, BlockSize);
+			FPlatformAtomics::InterlockedAdd(&UsedSize, ((int64)Range.Length));
+			DEC_MEMORY_STAT_BY(STAT_AGXBufferUnusedMemory, Range.Length);
+			DEC_MEMORY_STAT_BY(STAT_AGXMagazineBufferUnusedMemory, Range.Length);
+			INC_MEMORY_STAT_BY(STAT_AGXMagazineBufferMemory, Range.Length);
+			Result = FAGXBuffer(MTLPP_VALIDATE(mtlpp::Buffer, ParentBuffer, AGXSafeGetRuntimeDebuggingLevel() >= EAGXDebugLevelValidation, NewBuffer(Range)), this);
+			break;
 		}
 	}
 
@@ -884,14 +668,7 @@ FAGXBuffer FAGXSubBufferMagazine::NewBuffer()
 
 mtlpp::PurgeableState FAGXSubBufferMagazine::SetPurgeableState(mtlpp::PurgeableState state)
 {
-	if (ParentHeap)
-	{
-		return ParentHeap.SetPurgeableState(state);
-	}
-	else
-	{
-		return ParentBuffer.SetPurgeableState(state);
-	}
+	return ParentBuffer.SetPurgeableState(state);
 }
 
 FAGXRingBufferRef::FAGXRingBufferRef(FAGXBuffer Buf)
@@ -1388,63 +1165,13 @@ FAGXResourceHeap::TextureHeapSize FAGXResourceHeap::TextureSizeToIndex(uint32 Si
 	return (TextureHeapSize)Lower;
 }
 
-mtlpp::Heap FAGXResourceHeap::GetTextureHeap(mtlpp::TextureDescriptor Desc, mtlpp::SizeAndAlign Size)
-{
-	mtlpp::Heap Result;
-	static bool bTextureHeaps = FParse::Param(FCommandLine::Get(),TEXT("agxtextureheaps"));
-	if (FAGXCommandQueue::SupportsFeature(EAGXFeaturesHeaps) && bTextureHeaps && Size.Size <= HeapTextureHeapSizes[MaxTextureSize])
-	{
-		FAGXResourceHeap::TextureHeapSize HeapIndex = TextureSizeToIndex(Size.Size);
-
-		EAGXHeapTextureUsage UsageIndex = EAGXHeapTextureUsageNum;
-		mtlpp::StorageMode StorageMode = Desc.GetStorageMode();
-		mtlpp::CpuCacheMode CPUMode = Desc.GetCpuCacheMode();
-		if ((Desc.GetUsage() & mtlpp::TextureUsage::RenderTarget) && StorageMode == mtlpp::StorageMode::Private && CPUMode == mtlpp::CpuCacheMode::DefaultCache)
-		{
-			UsageIndex = PLATFORM_MAC ? EAGXHeapTextureUsageNum : EAGXHeapTextureUsageRenderTarget;
-		}
-		else if (StorageMode == mtlpp::StorageMode::Private && CPUMode == mtlpp::CpuCacheMode::WriteCombined)
-		{
-			UsageIndex = EAGXHeapTextureUsageResource;
-		}
-		
-		if (UsageIndex < EAGXHeapTextureUsageNum)
-		{
-			for (mtlpp::Heap& Heap : TextureHeaps[UsageIndex][HeapIndex])
-			{
-				if (Heap.MaxAvailableSizeWithAlignment(Size.Align) >= Size.Size)
-				{
-					Result = Heap;
-					break;
-				}
-			}
-			if (!Result)
-			{
-				MTLHeapDescriptor* HeapDesc = [[MTLHeapDescriptor alloc] init];
-				[HeapDesc setSize:HeapTextureHeapSizes[HeapIndex]];
-				[HeapDesc setStorageMode:(MTLStorageMode)Desc.GetStorageMode()];
-				[HeapDesc setCpuCacheMode:(MTLCPUCacheMode)Desc.GetCpuCacheMode()];
-				Result = mtlpp::Heap([GMtlDevice newHeapWithDescriptor:HeapDesc], nullptr, ns::Ownership::Assign);
-				[HeapDesc release];
-#if STATS || ENABLE_LOW_LEVEL_MEM_TRACKER
-				AGXLLM::LogAllocHeap(Result);
-#endif
-				TextureHeaps[UsageIndex][HeapIndex].Add(Result);
-			}
-			check(Result);
-		}
-	}
-	return Result;
-}
-
 FAGXBuffer FAGXResourceHeap::CreateBuffer(uint32 Size, uint32 Alignment, EBufferUsageFlags Flags, mtlpp::ResourceOptions Options, bool bForceUnique)
 {
 	LLM_SCOPE_METAL(ELLMTagAGX::Buffers);
 	LLM_PLATFORM_SCOPE_METAL(ELLMTagAGX::Buffers);
 	
-	static bool bSupportsHeaps = GetAGXDeviceContext().SupportsFeature(EAGXFeaturesHeaps);
 	static bool bSupportsBufferSubAllocation = FAGXCommandQueue::SupportsFeature(EAGXFeaturesBufferSubAllocation);
-	bForceUnique |= (!bSupportsBufferSubAllocation && !bSupportsHeaps);
+	bForceUnique |= (!bSupportsBufferSubAllocation);
 	
 	uint32 Usage = EnumHasAnyFlags(Flags, BUF_Static) ? UsageStatic : UsageDynamic;
 	
@@ -1656,18 +1383,7 @@ FAGXTexture FAGXResourceHeap::CreateTexture(mtlpp::TextureDescriptor Desc, FAGXS
 	LLM_SCOPE_METAL(ELLMTagAGX::Textures);
 	LLM_PLATFORM_SCOPE_METAL(ELLMTagAGX::Textures);
 	
-	mtlpp::SizeAndAlign Res = GMtlppDevice.HeapTextureSizeAndAlign(Desc);
-	mtlpp::Heap Heap = GetTextureHeap(Desc, Res);
-	if (Heap)
-	{
-		METAL_GPUPROFILE(FAGXScopedCPUStats CPUStat(FString::Printf(TEXT("AllocTexture: %s"), TEXT("")/**FString([Desc.GetPtr() description])*/)));
-		FAGXTexture Texture = Heap.NewTexture(Desc);
-#if STATS || ENABLE_LOW_LEVEL_MEM_TRACKER
-		AGXLLM::LogAllocTexture(Desc, Texture);
-#endif
-		return Texture;
-	}
-	else if (Desc.GetUsage() & mtlpp::TextureUsage::RenderTarget)
+	if (Desc.GetUsage() & mtlpp::TextureUsage::RenderTarget)
 	{
 		LLM_PLATFORM_SCOPE_METAL(ELLMTagAGX::RenderTargets);
 		return TargetPool.CreateTexture(Desc);
@@ -1680,7 +1396,7 @@ FAGXTexture FAGXResourceHeap::CreateTexture(mtlpp::TextureDescriptor Desc, FAGXS
 
 void FAGXResourceHeap::ReleaseTexture(FAGXSurface* Surface, FAGXTexture& Texture)
 {
-	if (Texture && !Texture.GetBuffer() && !Texture.GetParentTexture() && !Texture.GetHeap())
+	if (Texture && !Texture.GetBuffer() && !Texture.GetParentTexture())
 	{
         if (Texture.GetUsage() & mtlpp::TextureUsage::RenderTarget)
         {
diff --git a/Engine/Source/Runtime/Apple/AGXRHI/Private/AGXBuffer.h b/Engine/Source/Runtime/Apple/AGXRHI/Private/AGXBuffer.h
--- a/Engine/Source/Runtime/Apple/AGXRHI/Private/AGXBuffer.h
+++ b/Engine/Source/Runtime/Apple/AGXRHI/Private/AGXBuffer.h
@@ -65,7 +65,6 @@ private:
 	NSUInteger MinAlign;
 	NSUInteger UsedSize;
 	mtlpp::Buffer ParentBuffer;
-	mutable mtlpp::Heap ParentHeap;
 	TArray<ns::Range> FreeRanges;
     TArray<Allocation> AllocRanges;
 };
@@ -125,7 +124,6 @@ private:
 	int64 volatile OutstandingAllocs;
 	int64 volatile UsedSize;
 	mtlpp::Buffer ParentBuffer;
-	mutable mtlpp::Heap ParentHeap;
 	TArray<int8> Blocks;
 };
 
@@ -428,8 +426,6 @@ private:
 	uint32 GetHeapIndex(uint32 Size);
 	TextureHeapSize TextureSizeToIndex(uint32 Size);
 	
-	mtlpp::Heap GetTextureHeap(mtlpp::TextureDescriptor Desc, mtlpp::SizeAndAlign Size);
-	
 private:
 	static uint32 MagazineSizes[NumMagazineSizes];
 	static uint32 HeapSizes[NumHeapSizes];
@@ -458,6 +454,4 @@ private:
 	/** We can reuse texture allocations as well, to minimize their performance impact */
 	FAGXTexturePool TexturePool;
 	FAGXTexturePool TargetPool;
-	
-	TArray<mtlpp::Heap> TextureHeaps[EAGXHeapTextureUsageNum][NumTextureHeapSizes];
 };
diff --git a/Engine/Source/Runtime/Apple/AGXRHI/Private/AGXCommandBuffer.cpp b/Engine/Source/Runtime/Apple/AGXRHI/Private/AGXCommandBuffer.cpp
deleted file mode 100644
diff --git a/Engine/Source/Runtime/Apple/AGXRHI/Private/AGXCommandBuffer.h b/Engine/Source/Runtime/Apple/AGXRHI/Private/AGXCommandBuffer.h
deleted file mode 100644
diff --git a/Engine/Source/Runtime/Apple/AGXRHI/Private/AGXCommandEncoder.cpp b/Engine/Source/Runtime/Apple/AGXRHI/Private/AGXCommandEncoder.cpp
--- a/Engine/Source/Runtime/Apple/AGXRHI/Private/AGXCommandEncoder.cpp
+++ b/Engine/Source/Runtime/Apple/AGXRHI/Private/AGXCommandEncoder.cpp
@@ -9,9 +9,6 @@
 #include "AGXGraphicsPipelineState.h"
 #include "AGXCommandBufferFence.h"
 #include "AGXCommandEncoder.h"
-#include "AGXCommandBuffer.h"
-#include "AGXComputeCommandEncoder.h"
-#include "AGXRenderCommandEncoder.h"
 #include "AGXProfiler.h"
 #include "MetalShaderResources.h"
 
@@ -21,194 +18,6 @@ const uint32 EncoderRingBufferSize = 1024 * 1024;
 extern int32 GAGXBufferScribble;
 #endif
 
-static TCHAR const* const GAGXCommandDataTypeName[] = {
-	TEXT("DrawPrimitive"),
-	TEXT("DrawPrimitiveIndexed"),
-	TEXT("DrawPrimitivePatch"),
-	TEXT("DrawPrimitiveIndirect"),
-	TEXT("DrawPrimitiveIndexedIndirect"),
-	TEXT("Dispatch"),
-	TEXT("DispatchIndirect"),
-};
-
-
-FString FAGXCommandData::ToString() const
-{
-	FString Result;
-	if ((uint32)CommandType < (uint32)FAGXCommandData::Type::Num)
-	{
-		Result = GAGXCommandDataTypeName[(uint32)CommandType];
-		switch(CommandType)
-		{
-			case FAGXCommandData::Type::DrawPrimitive:
-				Result += FString::Printf(TEXT(" BaseInstance: %u InstanceCount: %u VertexCount: %u VertexStart: %u"), Draw.BaseInstance, Draw.InstanceCount, Draw.VertexCount, Draw.VertexStart);
-				break;
-			case FAGXCommandData::Type::DrawPrimitiveIndexed:
-				Result += FString::Printf(TEXT(" BaseInstance: %u BaseVertex: %u IndexCount: %u IndexStart: %u InstanceCount: %u"), DrawIndexed.BaseInstance, DrawIndexed.BaseVertex, DrawIndexed.IndexCount, DrawIndexed.IndexStart, DrawIndexed.InstanceCount);
-				break;
-			case FAGXCommandData::Type::DrawPrimitivePatch:
-				Result += FString::Printf(TEXT(" BaseInstance: %u InstanceCount: %u PatchCount: %u PatchStart: %u"), DrawPatch.BaseInstance, DrawPatch.InstanceCount, DrawPatch.PatchCount, DrawPatch.PatchStart);
-				break;
-			case FAGXCommandData::Type::Dispatch:
-				Result += FString::Printf(TEXT(" X: %u Y: %u Z: %u"), (uint32)Dispatch.threadgroupsPerGrid[0], (uint32)Dispatch.threadgroupsPerGrid[1], (uint32)Dispatch.threadgroupsPerGrid[2]);
-				break;
-			case FAGXCommandData::Type::DispatchIndirect:
-				Result += FString::Printf(TEXT(" Buffer: %p Offset: %u"), (void*)DispatchIndirect.ArgumentBuffer, (uint32)DispatchIndirect.ArgumentOffset);
-				break;
-			case FAGXCommandData::Type::DrawPrimitiveIndirect:
-			case FAGXCommandData::Type::DrawPrimitiveIndexedIndirect:
-			case FAGXCommandData::Type::Num:
-			default:
-				break;
-		}
-	}
-	return Result;
-};
-
-struct FAGXCommandContextDebug
-{
-	TArray<FAGXCommandDebug> Commands;
-	TSet<TRefCountPtr<FAGXGraphicsPipelineState>> PSOs;
-	TSet<TRefCountPtr<FAGXComputeShader>> ComputeShaders;
-	FAGXBuffer DebugBuffer;
-};
-
-@interface FAGXCommandBufferDebug : FApplePlatformObject
-{
-	@public
-	TArray<FAGXCommandContextDebug> Contexts;
-	uint32 Index;
-}
-@end
-@implementation FAGXCommandBufferDebug
-APPLE_PLATFORM_OBJECT_ALLOC_OVERRIDES(FAGXCommandBufferDebug)
-- (instancetype)init
-{
-	id Self = [super init];
-	if (Self)
-	{
-		Index = ~0u;
-	}
-	return Self;
-}
-- (void)dealloc
-{
-	Contexts.Empty();
-	[super dealloc];
-}
-@end
-
-char const* FAGXCommandBufferMarkers::kTableAssociationKey = "FAGXCommandBufferMarkers::kTableAssociationKey";
-
-FAGXCommandBufferMarkers::FAGXCommandBufferMarkers(void)
-: ns::Object<FAGXCommandBufferDebug*, ns::CallingConvention::ObjectiveC>(nil)
-{
-	
-}
-
-FAGXCommandBufferMarkers::FAGXCommandBufferMarkers(mtlpp::CommandBuffer& CmdBuf)
-: ns::Object<FAGXCommandBufferDebug*, ns::CallingConvention::ObjectiveC>([FAGXCommandBufferDebug new], ns::Ownership::Assign)
-{
-	CmdBuf.SetAssociatedObject<FAGXCommandBufferMarkers>(FAGXCommandBufferMarkers::kTableAssociationKey, *this);
-	m_ptr->Contexts.SetNum(1);
-}
-
-
-FAGXCommandBufferMarkers::FAGXCommandBufferMarkers(FAGXCommandBufferDebug* CmdBuf)
-: ns::Object<FAGXCommandBufferDebug*, ns::CallingConvention::ObjectiveC>(CmdBuf)
-{
-	
-}
-
-void FAGXCommandBufferMarkers::AllocateContexts(uint32 NumContexts)
-{
-	if (m_ptr && m_ptr->Contexts.Num() < NumContexts)
-	{
-		m_ptr->Contexts.SetNum(NumContexts);
-	}
-}
-
-uint32 FAGXCommandBufferMarkers::AddCommand(uint32 CmdBufIndex, uint32 Encoder, uint32 ContextIndex, FAGXBuffer& DebugBuffer, FAGXGraphicsPipelineState* PSO, FAGXComputeShader* ComputeShader, FAGXCommandData& Data)
-{
-	uint32 Num = 0;
-	if (m_ptr)
-	{
-		if (m_ptr->Index == ~0u)
-		{
-			m_ptr->Index = CmdBufIndex;
-		}
-		
-		FAGXCommandContextDebug& Context = m_ptr->Contexts[ContextIndex];
-		if (Context.DebugBuffer != DebugBuffer)
-		{
-			Context.DebugBuffer = DebugBuffer;
-		}
-		
-		if (PSO)
-			Context.PSOs.Add(PSO);
-		if (ComputeShader)
-			Context.ComputeShaders.Add(ComputeShader);
-		
-		Num = Context.Commands.Num();
-		FAGXCommandDebug Command;
-        Command.CmdBufIndex = CmdBufIndex;
-		Command.Encoder = Encoder;
-		Command.Index = Num;
-		Command.PSO = PSO;
-		Command.ComputeShader = ComputeShader;
-		Command.Data = Data;
-		Context.Commands.Add(Command);
-	}
-	return Num;
-}
-
-TArray<FAGXCommandDebug>* FAGXCommandBufferMarkers::GetCommands(uint32 ContextIndex)
-{
-	TArray<FAGXCommandDebug>* Result = nullptr;
-	if (m_ptr)
-	{
-		FAGXCommandContextDebug& Context = m_ptr->Contexts[ContextIndex];
-		Result = &Context.Commands;
-	}
-	return Result;
-}
-
-ns::AutoReleased<FAGXBuffer> FAGXCommandBufferMarkers::GetDebugBuffer(uint32 ContextIndex)
-{
-	ns::AutoReleased<FAGXBuffer> Buffer;
-	if (m_ptr)
-	{
-		FAGXCommandContextDebug& Context = m_ptr->Contexts[ContextIndex];
-		Buffer = Context.DebugBuffer;
-	}
-	return Buffer;
-}
-
-uint32 FAGXCommandBufferMarkers::NumContexts() const
-{
-	uint32 Num = 0;
-	if (m_ptr)
-	{
-		Num = m_ptr->Contexts.Num();
-	}
-	return Num;
-}
-
-uint32 FAGXCommandBufferMarkers::GetIndex() const
-{
-	uint32 Num = 0;
-	if (m_ptr)
-	{
-		Num = m_ptr->Index;
-	}
-	return Num;
-}
-
-FAGXCommandBufferMarkers FAGXCommandBufferMarkers::Get(mtlpp::CommandBuffer const& CmdBuf)
-{
-	return CmdBuf.GetAssociatedObject<FAGXCommandBufferMarkers>(FAGXCommandBufferMarkers::kTableAssociationKey);
-}
-
 #pragma mark - Public C++ Boilerplate -
 
 FAGXCommandEncoder::FAGXCommandEncoder(FAGXCommandList& CmdList, EAGXCommandEncoderType InType)
@@ -216,16 +25,10 @@ FAGXCommandEncoder::FAGXCommandEncoder(FAGXCommandList& CmdList, EAGXCommandEnco
 , bSupportsMetalFeaturesSetBytes(CmdList.GetCommandQueue().SupportsFeature(EAGXFeaturesSetBytes))
 , RingBuffer(EncoderRingBufferSize, BufferOffsetAlignment, FAGXCommandQueue::GetCompatibleResourceOptions((mtlpp::ResourceOptions)(mtlpp::ResourceOptions::HazardTrackingModeUntracked | BUFFER_RESOURCE_STORAGE_MANAGED)))
 , RenderPassDesc(nil)
-, EncoderFence(nil)
 #if ENABLE_METAL_GPUPROFILE
 , CommandBufferStats(nullptr)
 #endif
-#if METAL_DEBUG_OPTIONS
-, WaitCount(0)
-, UpdateCount(0)
-#endif
 , DebugGroups([NSMutableArray new])
-, FenceStage(mtlpp::RenderStages::Fragment)
 , EncoderNum(0)
 , CmdBufIndex(0)
 , Type(InType)
@@ -379,12 +182,6 @@ void FAGXCommandEncoder::StartCommandBuffer(void)
 	{
 		CmdBufIndex++;
 		CommandBuffer = CommandList.GetCommandQueue().CreateCommandBuffer();
-		METAL_DEBUG_LAYER(EAGXDebugLevelFastValidation, CommandBufferDebug = FAGXCommandBufferDebugging::Get(CommandBuffer));
-		
-		if (GAGXCommandBufferDebuggingEnabled)
-		{
-			CommandBufferMarkers = FAGXCommandBufferMarkers(CommandBuffer);
-		}
 		
 		if ([DebugGroups count] > 0)
 		{
@@ -539,11 +336,6 @@ mtlpp::BlitCommandEncoder& FAGXCommandEncoder::GetBlitCommandEncoder(void)
 	return BlitCommandEncoder;
 }
 
-TRefCountPtr<FAGXFence> const& FAGXCommandEncoder::GetEncoderFence(void) const
-{
-	return EncoderFence;
-}
-	
 #pragma mark - Public Command Encoder Mutators -
 
 void FAGXCommandEncoder::BeginParallelRenderCommandEncoding(uint32 NumChildren)
@@ -553,14 +345,10 @@ void FAGXCommandEncoder::BeginParallelRenderCommandEncoding(uint32 NumChildren)
 	check(CommandBuffer);
 	check(IsRenderCommandEncoderActive() == false && IsComputeCommandEncoderActive() == false && IsBlitCommandEncoderActive() == false);
 	
-	FenceResources.Append(TransitionedResources);
-	
 	ParallelRenderCommandEncoder = MTLPP_VALIDATE(mtlpp::CommandBuffer, CommandBuffer, AGXSafeGetRuntimeDebuggingLevel() >= EAGXDebugLevelValidation, ParallelRenderCommandEncoder(RenderPassDesc));
-	METAL_DEBUG_LAYER(EAGXDebugLevelFastValidation, ParallelEncoderDebug = FAGXParallelRenderCommandEncoderDebugging(ParallelRenderCommandEncoder, RenderPassDesc, CommandBufferDebug));
 	
 	EncoderNum++;
 	
-	check(!EncoderFence);
 	NSString* Label = nil;
 	
 	if(GetEmitDrawEvents())
@@ -573,7 +361,6 @@ void FAGXCommandEncoder::BeginParallelRenderCommandEncoding(uint32 NumChildren)
 			for (NSString* Group in DebugGroups)
 			{
 				ParallelRenderCommandEncoder.PushDebugGroup(Group);
-				METAL_DEBUG_LAYER(EAGXDebugLevelFastValidation, RenderEncoderDebug.PushDebugGroup(Group));
 			}
 		}
 	}
@@ -582,7 +369,6 @@ void FAGXCommandEncoder::BeginParallelRenderCommandEncoding(uint32 NumChildren)
 	{
 		mtlpp::RenderCommandEncoder CommandEncoder = MTLPP_VALIDATE(mtlpp::ParallelRenderCommandEncoder, ParallelRenderCommandEncoder, AGXSafeGetRuntimeDebuggingLevel() >= EAGXDebugLevelValidation, GetRenderCommandEncoder());
 		ChildRenderCommandEncoders.Add(CommandEncoder);
-		METAL_DEBUG_LAYER(EAGXDebugLevelFastValidation, ParallelEncoderDebug.GetRenderCommandEncoderDebugger(CommandEncoder));
 	}
 }
 
@@ -592,21 +378,16 @@ void FAGXCommandEncoder::BeginRenderCommandEncoding(void)
 	check(CommandList.IsParallel() || CommandBuffer);
 	check(IsRenderCommandEncoderActive() == false && IsComputeCommandEncoderActive() == false && IsBlitCommandEncoderActive() == false);
 	
-	FenceResources.Append(TransitionedResources);
-	
 	if (!CommandList.IsParallel() || Type == EAGXCommandEncoderPrologue)
 	{
 		RenderCommandEncoder = MTLPP_VALIDATE(mtlpp::CommandBuffer, CommandBuffer, AGXSafeGetRuntimeDebuggingLevel() >= EAGXDebugLevelValidation, RenderCommandEncoder(RenderPassDesc));
-		METAL_DEBUG_LAYER(EAGXDebugLevelFastValidation, RenderEncoderDebug = FAGXRenderCommandEncoderDebugging(RenderCommandEncoder, RenderPassDesc, CommandBufferDebug));
 		EncoderNum++;	
 	}
 	else
 	{
 		RenderCommandEncoder = GetAGXDeviceContext().GetParallelRenderCommandEncoder(CommandList.GetParallelIndex(), ParallelRenderCommandEncoder, CommandBuffer);
-		METAL_DEBUG_LAYER(EAGXDebugLevelFastValidation, RenderEncoderDebug = FAGXRenderCommandEncoderDebugging::Get(RenderCommandEncoder));
 	}
 	
-	check(!EncoderFence);
 	NSString* Label = nil;
 	
 	if(GetEmitDrawEvents())
@@ -619,15 +400,9 @@ void FAGXCommandEncoder::BeginRenderCommandEncoding(void)
 			for (NSString* Group in DebugGroups)
 			{
 				RenderCommandEncoder.PushDebugGroup(Group);
-				METAL_DEBUG_LAYER(EAGXDebugLevelFastValidation, RenderEncoderDebug.PushDebugGroup(Group));
 			}
 		}
 	}
-	
-	if (CommandList.IsImmediate())
-	{
-		EncoderFence = CommandList.GetCommandQueue().CreateFence(Label);
-	}
 }
 
 void FAGXCommandEncoder::BeginComputeCommandEncoding(mtlpp::DispatchType DispatchType)
@@ -635,9 +410,6 @@ void FAGXCommandEncoder::BeginComputeCommandEncoding(mtlpp::DispatchType Dispatc
 	check(CommandBuffer);
 	check(IsRenderCommandEncoderActive() == false && IsComputeCommandEncoderActive() == false && IsBlitCommandEncoderActive() == false);
 	
-	FenceResources.Append(TransitionedResources);
-	TransitionedResources.Empty();
-	
 	if (DispatchType == mtlpp::DispatchType::Serial)
 	{
 		ComputeCommandEncoder = MTLPP_VALIDATE(mtlpp::CommandBuffer, CommandBuffer, AGXSafeGetRuntimeDebuggingLevel() >= EAGXDebugLevelValidation, ComputeCommandEncoder());
@@ -646,11 +418,9 @@ void FAGXCommandEncoder::BeginComputeCommandEncoding(mtlpp::DispatchType Dispatc
 	{
 		ComputeCommandEncoder = MTLPP_VALIDATE(mtlpp::CommandBuffer, CommandBuffer, AGXSafeGetRuntimeDebuggingLevel() >= EAGXDebugLevelValidation, ComputeCommandEncoder(DispatchType));
 	}
-	METAL_DEBUG_LAYER(EAGXDebugLevelFastValidation, ComputeEncoderDebug = FAGXComputeCommandEncoderDebugging(ComputeCommandEncoder, CommandBufferDebug));
 
 	EncoderNum++;
 	
-	check(!EncoderFence);
 	NSString* Label = nil;
 	
 	if(GetEmitDrawEvents())
@@ -663,12 +433,9 @@ void FAGXCommandEncoder::BeginComputeCommandEncoding(mtlpp::DispatchType Dispatc
 			for (NSString* Group in DebugGroups)
 			{
 				ComputeCommandEncoder.PushDebugGroup(Group);
-				METAL_DEBUG_LAYER(EAGXDebugLevelFastValidation, ComputeEncoderDebug.PushDebugGroup(Group));
 			}
 		}
 	}
-	
-	EncoderFence = CommandList.GetCommandQueue().CreateFence(Label);
 }
 
 void FAGXCommandEncoder::BeginBlitCommandEncoding(void)
@@ -676,15 +443,10 @@ void FAGXCommandEncoder::BeginBlitCommandEncoding(void)
 	check(CommandBuffer);
 	check(IsRenderCommandEncoderActive() == false && IsComputeCommandEncoderActive() == false && IsBlitCommandEncoderActive() == false);
 	
-	FenceResources.Append(TransitionedResources);
-	TransitionedResources.Empty();
-	
 	BlitCommandEncoder = MTLPP_VALIDATE(mtlpp::CommandBuffer, CommandBuffer, AGXSafeGetRuntimeDebuggingLevel() >= EAGXDebugLevelValidation, BlitCommandEncoder());
-	METAL_DEBUG_LAYER(EAGXDebugLevelFastValidation, BlitEncoderDebug = FAGXBlitCommandEncoderDebugging(BlitCommandEncoder, CommandBufferDebug));
 	
 	EncoderNum++;
 	
-	check(!EncoderFence);
 	NSString* Label = nil;
 	
 	if(GetEmitDrawEvents())
@@ -697,25 +459,19 @@ void FAGXCommandEncoder::BeginBlitCommandEncoding(void)
 			for (NSString* Group in DebugGroups)
 			{
 				BlitCommandEncoder.PushDebugGroup(Group);
-				METAL_DEBUG_LAYER(EAGXDebugLevelFastValidation, BlitEncoderDebug.PushDebugGroup(Group));
 			}
 		}
 	}
-	
-	EncoderFence = CommandList.GetCommandQueue().CreateFence(Label);
 }
 
-TRefCountPtr<FAGXFence> FAGXCommandEncoder::EndEncoding(void)
+void FAGXCommandEncoder::EndEncoding(void)
 {
-	static bool bSupportsFences = CommandList.GetCommandQueue().SupportsFeature(EAGXFeaturesFences);
-	TRefCountPtr<FAGXFence> Fence = nullptr;
 	@autoreleasepool
 	{
 		if(IsRenderCommandEncoderActive())
 		{
 			if (RenderCommandEncoder)
 			{
-				check(!bSupportsFences || EncoderFence || !CommandList.IsImmediate());
 				if (ParallelRenderCommandEncoder.GetPtr() == nil)
 				{
 					check(RenderPassDesc);
@@ -744,52 +500,15 @@ TRefCountPtr<FAGXFence> FAGXCommandEncoder::EndEncoding(void)
 					}
 				}
 				
-				for (FAGXFence* FragFence : FragmentFences)
-				{
-					if (FragFence->NeedsWait(mtlpp::RenderStages::Fragment))
-					{
-						mtlpp::Fence FragmentFence = FragFence->Get(mtlpp::RenderStages::Fragment);
-						mtlpp::Fence FragInnerFence = METAL_DEBUG_OPTION(CommandList.GetCommandQueue().GetRuntimeDebuggingLevel() >= EAGXDebugLevelValidation ? mtlpp::Fence(((FAGXDebugFence*)FragmentFence.GetPtr()).Inner) :) FragmentFence;
-						
-						RenderCommandEncoder.WaitForFence(FragInnerFence, FenceStage);
-						METAL_DEBUG_LAYER(EAGXDebugLevelFastValidation, RenderEncoderDebug.AddWaitFence(FragmentFence));
-						FragFence->Wait(mtlpp::RenderStages::Fragment);
-					}
-				}
-				FragmentFences.Empty();
-				
-				if (FenceStage == mtlpp::RenderStages::Vertex)
-				{
-					FenceResources.Empty();
-					FenceStage = mtlpp::RenderStages::Fragment;
-				}
-				
-				if (EncoderFence && EncoderFence->NeedsWrite(mtlpp::RenderStages::Fragment))
-				{
-					Fence = EncoderFence;
-				}
-				UpdateFence(EncoderFence);
-				
-#if METAL_DEBUG_OPTIONS
-				if (bSupportsFences && AGXSafeGetRuntimeDebuggingLevel() >= EAGXDebugLevelFastValidation && (!WaitCount || !UpdateCount))
-				{
-					UE_LOG(LogAGX, Error, TEXT("%s has incorrect fence waits (%u) vs. updates (%u)."), *FString(RenderCommandEncoder.GetLabel()), WaitCount, UpdateCount);
-					
-				}
-				WaitCount = 0;
-				UpdateCount = 0;
-#endif
 				RenderCommandEncoder.EndEncoding();
-				METAL_DEBUG_LAYER(EAGXDebugLevelFastValidation, RenderEncoderDebug.EndEncoder());
 				RenderCommandEncoder = nil;
-				EncoderFence = nullptr;
 			}
 			
 			if (ParallelRenderCommandEncoder && IsParallel())
 			{
 				RingBuffer.Commit(CommandBuffer);
 				
-			#if METAL_DEBUG_OPTIONS
+#if METAL_DEBUG_OPTIONS
 				if(CommandList.GetCommandQueue().GetRuntimeDebuggingLevel() >= EAGXDebugLevelValidation)
 				{
 					for (FAGXBuffer const& Buffer : ActiveBuffers)
@@ -806,7 +525,7 @@ TRefCountPtr<FAGXFence> FAGXCommandEncoder::EndEncoding(void)
 											}
 										});
 				}
-			#endif
+#endif // METAL_DEBUG_OPTIONS
 
 				BufferBindingHistory.Empty();
 				TextureBindingHistory.Empty();
@@ -848,7 +567,6 @@ TRefCountPtr<FAGXFence> FAGXCommandEncoder::EndEncoding(void)
 				}
 
 				ParallelRenderCommandEncoder.EndEncoding();
-				METAL_DEBUG_LAYER(EAGXDebugLevelFastValidation, ParallelEncoderDebug.EndEncoder());
 				ParallelRenderCommandEncoder = nil;
 
 				ChildRenderCommandEncoders.Empty();
@@ -856,83 +574,13 @@ TRefCountPtr<FAGXFence> FAGXCommandEncoder::EndEncoding(void)
 		}
 		else if(IsComputeCommandEncoderActive())
 		{
-			check(!bSupportsFences || EncoderFence);
-			
-			for (FAGXFence* FragFence : FragmentFences)
-			{
-				if (FragFence->NeedsWait(mtlpp::RenderStages::Fragment))
-				{
-					mtlpp::Fence FragmentFence = FragFence->Get(mtlpp::RenderStages::Fragment);
-					mtlpp::Fence FragInnerFence = METAL_DEBUG_OPTION(CommandList.GetCommandQueue().GetRuntimeDebuggingLevel() >= EAGXDebugLevelValidation ? mtlpp::Fence(((FAGXDebugFence*)FragmentFence.GetPtr()).Inner) :) FragmentFence;
-					
-					ComputeCommandEncoder.WaitForFence(FragInnerFence);
-					METAL_DEBUG_LAYER(EAGXDebugLevelFastValidation, ComputeEncoderDebug.AddWaitFence(FragmentFence));
-					FragFence->Wait(mtlpp::RenderStages::Fragment);
-				}
-			}
-			FragmentFences.Empty();
-			FenceResources.Empty();
-			FenceStage = mtlpp::RenderStages::Fragment;
-			
-			if (EncoderFence && EncoderFence->NeedsWrite(mtlpp::RenderStages::Fragment))
-			{
-				Fence = EncoderFence;
-			}
-			UpdateFence(EncoderFence);
-			
-#if METAL_DEBUG_OPTIONS
-			if (bSupportsFences && AGXSafeGetRuntimeDebuggingLevel() >= EAGXDebugLevelFastValidation && (!WaitCount || !UpdateCount))
-			{
-				UE_LOG(LogAGX, Error, TEXT("%s has incorrect fence waits (%u) vs. updates (%u)."), *FString(ComputeCommandEncoder.GetLabel()), WaitCount, UpdateCount);
-				
-			}
-			WaitCount = 0;
-			UpdateCount = 0;
-#endif
 			ComputeCommandEncoder.EndEncoding();
-			METAL_DEBUG_LAYER(EAGXDebugLevelFastValidation, ComputeEncoderDebug.EndEncoder());
 			ComputeCommandEncoder = nil;
-			EncoderFence = nullptr;
 		}
 		else if(IsBlitCommandEncoderActive())
 		{
-			// check(!bSupportsFences || EncoderFence);
-			
-			for (FAGXFence* FragFence : FragmentFences)
-			{
-				if (FragFence->NeedsWait(mtlpp::RenderStages::Fragment))
-				{
-					mtlpp::Fence FragmentFence = FragFence->Get(mtlpp::RenderStages::Fragment);
-					mtlpp::Fence FragInnerFence = METAL_DEBUG_OPTION(CommandList.GetCommandQueue().GetRuntimeDebuggingLevel() >= EAGXDebugLevelValidation ? mtlpp::Fence(((FAGXDebugFence*)FragmentFence.GetPtr()).Inner) :) FragmentFence;
-					
-					BlitCommandEncoder.WaitForFence(FragInnerFence);
-					METAL_DEBUG_LAYER(EAGXDebugLevelFastValidation, BlitEncoderDebug.AddWaitFence(FragmentFence));
-					FragFence->Wait(mtlpp::RenderStages::Fragment);
-				}
-			}
-			FragmentFences.Empty();
-			FenceResources.Empty();
-			FenceStage = mtlpp::RenderStages::Fragment;
-			
-			if (EncoderFence && EncoderFence->NeedsWrite(mtlpp::RenderStages::Fragment))
-			{
-				Fence = EncoderFence;
-			}
-			UpdateFence(EncoderFence);
-			
-#if METAL_DEBUG_OPTIONS
-			if (bSupportsFences && AGXSafeGetRuntimeDebuggingLevel() >= EAGXDebugLevelFastValidation && (!WaitCount || !UpdateCount))
-			{
-				UE_LOG(LogAGX, Error, TEXT("%s has incorrect fence waits (%u) vs. updates (%u)."), *FString(BlitCommandEncoder.GetLabel()), WaitCount, UpdateCount);
-				
-			}
-			WaitCount = 0;
-			UpdateCount = 0;
-#endif
 			BlitCommandEncoder.EndEncoding();
-			METAL_DEBUG_LAYER(EAGXDebugLevelFastValidation, BlitEncoderDebug.EndEncoder());
 			BlitCommandEncoder = nil;
-			EncoderFence = nullptr;
 		}
 	}
 	
@@ -948,7 +596,6 @@ TRefCountPtr<FAGXFence> FAGXCommandEncoder::EndEncoding(void)
 		FMemory::Memzero(ShaderBuffers[Frequency].Usage);
 		ShaderBuffers[Frequency].Bound = 0;
 	}
-    return Fence;
 }
 
 void FAGXCommandEncoder::InsertCommandBufferFence(FAGXCommandBufferFence& Fence, mtlpp::CommandBufferHandler Handler)
@@ -972,176 +619,6 @@ void FAGXCommandEncoder::AddCompletionHandler(mtlpp::CommandBufferHandler Handle
 	Block_release(HeapHandler);
 }
 
-void FAGXCommandEncoder::UpdateFence(FAGXFence* Fence)
-{
-	check(IsRenderCommandEncoderActive() || IsComputeCommandEncoderActive() || IsBlitCommandEncoderActive());
-	static bool bSupportsFences = CommandList.GetCommandQueue().SupportsFeature(EAGXFeaturesFences);
-	if ((bSupportsFences METAL_DEBUG_OPTION(|| CommandList.GetCommandQueue().GetRuntimeDebuggingLevel() >= EAGXDebugLevelValidation)) && Fence)
-	{
-		mtlpp::Fence VertexFence = Fence->Get(mtlpp::RenderStages::Vertex);
-		mtlpp::Fence InnerFence = METAL_DEBUG_OPTION(CommandList.GetCommandQueue().GetRuntimeDebuggingLevel() >= EAGXDebugLevelValidation ? mtlpp::Fence(((FAGXDebugFence*)VertexFence.GetPtr()).Inner) :) VertexFence;
-		if (RenderCommandEncoder)
-		{
-			mtlpp::Fence FragmentFence = Fence->Get(mtlpp::RenderStages::Fragment);
-			mtlpp::Fence FragInnerFence = METAL_DEBUG_OPTION(CommandList.GetCommandQueue().GetRuntimeDebuggingLevel() >= EAGXDebugLevelValidation ? mtlpp::Fence(((FAGXDebugFence*)FragmentFence.GetPtr()).Inner) :) FragmentFence;
-			
-			if (Fence->NeedsWrite(mtlpp::RenderStages::Vertex))
-			{
-				RenderCommandEncoder.UpdateFence(InnerFence, (mtlpp::RenderStages)(mtlpp::RenderStages::Vertex));
-				METAL_DEBUG_LAYER(EAGXDebugLevelFastValidation, RenderEncoderDebug.AddUpdateFence(VertexFence));
-				Fence->Write(mtlpp::RenderStages::Vertex);
-				METAL_DEBUG_LAYER(EAGXDebugLevelFastValidation, UpdateCount++);
-			}
-
-			if (Fence->NeedsWrite(mtlpp::RenderStages::Fragment))
-			{
-				RenderCommandEncoder.UpdateFence(FragInnerFence, (mtlpp::RenderStages)(mtlpp::RenderStages::Fragment));
-				METAL_DEBUG_LAYER(EAGXDebugLevelFastValidation, RenderEncoderDebug.AddUpdateFence(FragmentFence));
-				Fence->Write(mtlpp::RenderStages::Fragment);
-				METAL_DEBUG_LAYER(EAGXDebugLevelFastValidation, UpdateCount++);
-			}
-		}
-		else if (ComputeCommandEncoder && Fence->NeedsWrite(mtlpp::RenderStages::Vertex))
-		{
-			ComputeCommandEncoder.UpdateFence(InnerFence);
-			METAL_DEBUG_LAYER(EAGXDebugLevelFastValidation, ComputeEncoderDebug.AddUpdateFence(VertexFence));
-			Fence->Write(mtlpp::RenderStages::Vertex);
-			METAL_DEBUG_LAYER(EAGXDebugLevelFastValidation, UpdateCount++);
-		}
-		else if (BlitCommandEncoder && Fence->NeedsWrite(mtlpp::RenderStages::Vertex))
-		{
-			BlitCommandEncoder.UpdateFence(InnerFence);
-			METAL_DEBUG_LAYER(EAGXDebugLevelFastValidation, BlitEncoderDebug.AddUpdateFence(VertexFence));
-			Fence->Write(mtlpp::RenderStages::Vertex);
-			METAL_DEBUG_LAYER(EAGXDebugLevelFastValidation, UpdateCount++);
-		}
-	}
-}
-
-void FAGXCommandEncoder::WaitForFence(FAGXFence* Fence)
-{
-	check(IsRenderCommandEncoderActive() || IsComputeCommandEncoderActive() || IsBlitCommandEncoderActive());
-	static bool bSupportsFences = CommandList.GetCommandQueue().SupportsFeature(EAGXFeaturesFences);
-	if ((bSupportsFences METAL_DEBUG_OPTION(|| CommandList.GetCommandQueue().GetRuntimeDebuggingLevel() >= EAGXDebugLevelValidation)) && Fence)
-	{
-		if (Fence->NeedsWait(mtlpp::RenderStages::Vertex))
-		{
-			METAL_DEBUG_LAYER(EAGXDebugLevelFastValidation, WaitCount++);
-			
-			mtlpp::Fence VertexFence = Fence->Get(mtlpp::RenderStages::Vertex);
-			mtlpp::Fence InnerFence = METAL_DEBUG_OPTION(CommandList.GetCommandQueue().GetRuntimeDebuggingLevel() >= EAGXDebugLevelValidation ? mtlpp::Fence(((FAGXDebugFence*)VertexFence.GetPtr()).Inner) :) VertexFence;
-			if (RenderCommandEncoder)
-			{
-				RenderCommandEncoder.WaitForFence(InnerFence, (mtlpp::RenderStages)(mtlpp::RenderStages::Vertex));
-				METAL_DEBUG_LAYER(EAGXDebugLevelFastValidation, RenderEncoderDebug.AddWaitFence(VertexFence));
-				Fence->Wait(mtlpp::RenderStages::Vertex);
-			}
-			else if (ComputeCommandEncoder)
-			{
-				ComputeCommandEncoder.WaitForFence(InnerFence);
-				METAL_DEBUG_LAYER(EAGXDebugLevelFastValidation, ComputeEncoderDebug.AddWaitFence(VertexFence));
-				Fence->Wait(mtlpp::RenderStages::Vertex);
-			}
-			else if (BlitCommandEncoder)
-			{
-				BlitCommandEncoder.WaitForFence(InnerFence);
-				METAL_DEBUG_LAYER(EAGXDebugLevelFastValidation, BlitEncoderDebug.AddWaitFence(VertexFence));
-				Fence->Wait(mtlpp::RenderStages::Vertex);
-			}
-		}
-		if (Fence->NeedsWait(mtlpp::RenderStages::Fragment))
-		{
-			if (FenceStage == mtlpp::RenderStages::Vertex || BlitCommandEncoder)
-			{
-				mtlpp::Fence FragmentFence = Fence->Get(mtlpp::RenderStages::Fragment);
-				mtlpp::Fence FragInnerFence = METAL_DEBUG_OPTION(CommandList.GetCommandQueue().GetRuntimeDebuggingLevel() >= EAGXDebugLevelValidation ? mtlpp::Fence(((FAGXDebugFence*)FragmentFence.GetPtr()).Inner) :) FragmentFence;
-				if (RenderCommandEncoder)
-				{
-					RenderCommandEncoder.WaitForFence(FragInnerFence, (mtlpp::RenderStages)(mtlpp::RenderStages::Vertex));
-					METAL_DEBUG_LAYER(EAGXDebugLevelFastValidation, RenderEncoderDebug.AddWaitFence(FragmentFence));
-					Fence->Wait(mtlpp::RenderStages::Fragment);
-				}
-				else if (ComputeCommandEncoder)
-				{
-					ComputeCommandEncoder.WaitForFence(FragInnerFence);
-					METAL_DEBUG_LAYER(EAGXDebugLevelFastValidation, ComputeEncoderDebug.AddWaitFence(FragmentFence));
-					Fence->Wait(mtlpp::RenderStages::Fragment);
-				}
-				else if (BlitCommandEncoder)
-				{
-					BlitCommandEncoder.WaitForFence(FragInnerFence);
-					METAL_DEBUG_LAYER(EAGXDebugLevelFastValidation, BlitEncoderDebug.AddWaitFence(FragmentFence));
-					Fence->Wait(mtlpp::RenderStages::Fragment);
-				}
-				METAL_DEBUG_LAYER(EAGXDebugLevelFastValidation, WaitCount++);
-			}
-			else
-			{
-				METAL_DEBUG_LAYER(EAGXDebugLevelFastValidation, WaitCount++);
-				FragmentFences.Add(Fence);
-			}
-		}
-	}
-}
-
-void FAGXCommandEncoder::WaitAndUpdateFence(FAGXFence* Fence)
-{
-	check(IsRenderCommandEncoderActive() || IsComputeCommandEncoderActive() || IsBlitCommandEncoderActive());
-	static bool bSupportsFences = CommandList.GetCommandQueue().SupportsFeature(EAGXFeaturesFences);
-	if ((bSupportsFences METAL_DEBUG_OPTION(|| CommandList.GetCommandQueue().GetRuntimeDebuggingLevel() >= EAGXDebugLevelValidation)) && Fence)
-	{
-		METAL_DEBUG_LAYER(EAGXDebugLevelFastValidation, WaitCount++);
-		METAL_DEBUG_LAYER(EAGXDebugLevelFastValidation, UpdateCount++);
-		
-		mtlpp::Fence VertexFence = Fence->Get(mtlpp::RenderStages::Vertex);
-		mtlpp::Fence InnerFence = METAL_DEBUG_OPTION(CommandList.GetCommandQueue().GetRuntimeDebuggingLevel() >= EAGXDebugLevelValidation ? mtlpp::Fence(((FAGXDebugFence*)VertexFence.GetPtr()).Inner) :) VertexFence;
-		if (RenderCommandEncoder)
-		{
-			mtlpp::Fence FragmentFence = Fence->Get(mtlpp::RenderStages::Fragment);
-			mtlpp::Fence FragInnerFence = METAL_DEBUG_OPTION(CommandList.GetCommandQueue().GetRuntimeDebuggingLevel() >= EAGXDebugLevelValidation ? mtlpp::Fence(((FAGXDebugFence*)FragmentFence.GetPtr()).Inner) :) FragmentFence;
-			
-			METAL_DEBUG_LAYER(EAGXDebugLevelFastValidation, WaitCount++);
-			METAL_DEBUG_LAYER(EAGXDebugLevelFastValidation, UpdateCount++);
-			
-			RenderCommandEncoder.WaitForFence(InnerFence, (mtlpp::RenderStages)(mtlpp::RenderStages::Vertex));
-			METAL_DEBUG_LAYER(EAGXDebugLevelFastValidation, RenderEncoderDebug.AddWaitFence(VertexFence));
-			Fence->Wait(mtlpp::RenderStages::Vertex);
-			
-			RenderCommandEncoder.WaitForFence(FragInnerFence, (mtlpp::RenderStages)(mtlpp::RenderStages::Fragment));
-			METAL_DEBUG_LAYER(EAGXDebugLevelFastValidation, RenderEncoderDebug.AddWaitFence(FragmentFence));
-			Fence->Wait(mtlpp::RenderStages::Fragment);
-			
-			RenderCommandEncoder.UpdateFence(InnerFence, (mtlpp::RenderStages)(mtlpp::RenderStages::Fragment));
-			METAL_DEBUG_LAYER(EAGXDebugLevelFastValidation, RenderEncoderDebug.AddUpdateFence(VertexFence));
-			Fence->Write(mtlpp::RenderStages::Vertex);
-			
-			RenderCommandEncoder.UpdateFence(FragInnerFence, (mtlpp::RenderStages)(mtlpp::RenderStages::Vertex));
-			METAL_DEBUG_LAYER(EAGXDebugLevelFastValidation, RenderEncoderDebug.AddUpdateFence(FragmentFence));
-			Fence->Write(mtlpp::RenderStages::Fragment);
-		}
-		else if (ComputeCommandEncoder)
-		{
-			ComputeCommandEncoder.WaitForFence(InnerFence);
-			METAL_DEBUG_LAYER(EAGXDebugLevelFastValidation, ComputeEncoderDebug.AddWaitFence(VertexFence));
-			Fence->Wait(mtlpp::RenderStages::Vertex);
-			
-			ComputeCommandEncoder.UpdateFence(InnerFence);
-			METAL_DEBUG_LAYER(EAGXDebugLevelFastValidation, ComputeEncoderDebug.AddUpdateFence(VertexFence));
-			Fence->Write(mtlpp::RenderStages::Vertex);
-		}
-		else if (BlitCommandEncoder)
-		{
-			BlitCommandEncoder.WaitForFence(InnerFence);
-			METAL_DEBUG_LAYER(EAGXDebugLevelFastValidation, BlitEncoderDebug.AddWaitFence(VertexFence));
-			Fence->Wait(mtlpp::RenderStages::Vertex);
-			
-			BlitCommandEncoder.UpdateFence(InnerFence);
-			METAL_DEBUG_LAYER(EAGXDebugLevelFastValidation, BlitEncoderDebug.AddUpdateFence(VertexFence));
-			Fence->Write(mtlpp::RenderStages::Vertex);
-		}
-	}
-}
-
 #pragma mark - Public Debug Support -
 
 void FAGXCommandEncoder::InsertDebugSignpost(ns::String const& String)
@@ -1151,7 +628,6 @@ void FAGXCommandEncoder::InsertDebugSignpost(ns::String const& String)
 		if (RenderCommandEncoder)
 		{
 			RenderCommandEncoder.InsertDebugSignpost(String);
-			METAL_DEBUG_LAYER(EAGXDebugLevelFastValidation, RenderEncoderDebug.InsertDebugSignpost(String));
 		}
 		else if (ParallelRenderCommandEncoder && !IsParallel())
 		{
@@ -1161,12 +637,10 @@ void FAGXCommandEncoder::InsertDebugSignpost(ns::String const& String)
 		else if (ComputeCommandEncoder)
 		{
 			ComputeCommandEncoder.InsertDebugSignpost(String);
-			METAL_DEBUG_LAYER(EAGXDebugLevelFastValidation, ComputeEncoderDebug.InsertDebugSignpost(String));
 		}
 		else if (BlitCommandEncoder)
 		{
 			BlitCommandEncoder.InsertDebugSignpost(String);
-			METAL_DEBUG_LAYER(EAGXDebugLevelFastValidation, BlitEncoderDebug.InsertDebugSignpost(String));
 		}
 	}
 }
@@ -1179,7 +653,6 @@ void FAGXCommandEncoder::PushDebugGroup(ns::String const& String)
 		if (RenderCommandEncoder)
 		{
 			RenderCommandEncoder.PushDebugGroup(String);
-			METAL_DEBUG_LAYER(EAGXDebugLevelFastValidation, RenderEncoderDebug.PushDebugGroup(String));
 		}
 		else if (ParallelRenderCommandEncoder && !IsParallel())
 		{
@@ -1189,12 +662,10 @@ void FAGXCommandEncoder::PushDebugGroup(ns::String const& String)
 		else if (ComputeCommandEncoder)
 		{
 			ComputeCommandEncoder.PushDebugGroup(String);
-			METAL_DEBUG_LAYER(EAGXDebugLevelFastValidation, ComputeEncoderDebug.PushDebugGroup(String));
 		}
 		else if (BlitCommandEncoder)
 		{
 			BlitCommandEncoder.PushDebugGroup(String);
-			METAL_DEBUG_LAYER(EAGXDebugLevelFastValidation, BlitEncoderDebug.PushDebugGroup(String));
 		}
 	}
 }
@@ -1207,7 +678,6 @@ void FAGXCommandEncoder::PopDebugGroup(void)
 		if (RenderCommandEncoder)
 		{
 			RenderCommandEncoder.PopDebugGroup();
-			METAL_DEBUG_LAYER(EAGXDebugLevelFastValidation, RenderEncoderDebug.PopDebugGroup());
 		}
 		else if (ParallelRenderCommandEncoder && !IsParallel())
 		{
@@ -1217,21 +687,14 @@ void FAGXCommandEncoder::PopDebugGroup(void)
 		else if (ComputeCommandEncoder)
 		{
 			ComputeCommandEncoder.PopDebugGroup();
-			METAL_DEBUG_LAYER(EAGXDebugLevelFastValidation, ComputeEncoderDebug.PopDebugGroup());
 		}
 		else if (BlitCommandEncoder)
 		{
 			BlitCommandEncoder.PopDebugGroup();
-			METAL_DEBUG_LAYER(EAGXDebugLevelFastValidation, BlitEncoderDebug.PopDebugGroup());
 		}
 	}
 }
 
-FAGXCommandBufferMarkers& FAGXCommandEncoder::GetMarkers(void)
-{
-	return CommandBufferMarkers;
-}
-
 #if ENABLE_METAL_GPUPROFILE
 FAGXCommandBufferStats* FAGXCommandEncoder::GetCommandBufferStats(void)
 {
@@ -1292,7 +755,6 @@ void FAGXCommandEncoder::SetRenderPipelineState(FAGXShaderPipeline* PipelineStat
 {
 	check (RenderCommandEncoder);
 	{
-		METAL_DEBUG_LAYER(EAGXDebugLevelFastValidation, RenderEncoderDebug.SetPipeline(PipelineState));
 		RenderCommandEncoder.SetRenderPipelineState(PipelineState->RenderPipelineState);
 	}
 }
@@ -1376,7 +838,6 @@ void FAGXCommandEncoder::SetDepthStencilState(mtlpp::DepthStencilState const& In
     check (RenderCommandEncoder);
 	{
 		RenderCommandEncoder.SetDepthStencilState(InDepthStencilState);
-		METAL_DEBUG_LAYER(EAGXDebugLevelFastValidation, RenderEncoderDebug.SetDepthStencilState(InDepthStencilState));
 	}
 }
 
@@ -1473,17 +934,14 @@ void FAGXCommandEncoder::SetShaderBytes(mtlpp::FunctionType const FunctionType,
 			{
 				case mtlpp::FunctionType::Vertex:
 					check(RenderCommandEncoder);
-					METAL_DEBUG_LAYER(EAGXDebugLevelFastValidation, RenderEncoderDebug.SetBytes(EAGXShaderVertex, Bytes, Length, Index));
 					RenderCommandEncoder.SetVertexData(Bytes, Length, Index);
 					break;
 				case mtlpp::FunctionType::Fragment:
 					check(RenderCommandEncoder);
-					METAL_DEBUG_LAYER(EAGXDebugLevelFastValidation, RenderEncoderDebug.SetBytes(EAGXShaderFragment, Bytes, Length, Index));
 					RenderCommandEncoder.SetFragmentData(Bytes, Length, Index);
 					break;
 				case mtlpp::FunctionType::Kernel:
 					check(ComputeCommandEncoder);
-					METAL_DEBUG_LAYER(EAGXDebugLevelFastValidation, ComputeEncoderDebug.SetBytes(Bytes, Length, Index));
 					ComputeCommandEncoder.SetBytes(Bytes, Length, Index);
 					break;
 				default:
@@ -1530,17 +988,14 @@ void FAGXCommandEncoder::SetShaderBufferOffset(mtlpp::FunctionType FunctionType,
 		case mtlpp::FunctionType::Vertex:
 			check (RenderCommandEncoder);
 			RenderCommandEncoder.SetVertexBufferOffset(Offset + ShaderBuffers[uint32(FunctionType)].Buffers[index].GetOffset(), index);
-			METAL_DEBUG_LAYER(EAGXDebugLevelFastValidation, RenderEncoderDebug.SetBufferOffset(EAGXShaderVertex, Offset + ShaderBuffers[uint32(FunctionType)].Buffers[index].GetOffset(), index));
 			break;
 		case mtlpp::FunctionType::Fragment:
 			check(RenderCommandEncoder);
 			RenderCommandEncoder.SetFragmentBufferOffset(Offset + ShaderBuffers[uint32(FunctionType)].Buffers[index].GetOffset(), index);
-			METAL_DEBUG_LAYER(EAGXDebugLevelFastValidation, RenderEncoderDebug.SetBufferOffset(EAGXShaderFragment, Offset + ShaderBuffers[uint32(FunctionType)].Buffers[index].GetOffset(), index));
 			break;
 		case mtlpp::FunctionType::Kernel:
 			check (ComputeCommandEncoder);
 			ComputeCommandEncoder.SetBufferOffset(Offset + ShaderBuffers[uint32(FunctionType)].Buffers[index].GetOffset(), index);
-			METAL_DEBUG_LAYER(EAGXDebugLevelFastValidation, ComputeEncoderDebug.SetBufferOffset(Offset + ShaderBuffers[uint32(FunctionType)].Buffers[index].GetOffset(), index));
 			break;
 		default:
 			check(false);
@@ -1556,21 +1011,15 @@ void FAGXCommandEncoder::SetShaderTexture(mtlpp::FunctionType FunctionType, FAGX
 		case mtlpp::FunctionType::Vertex:
 			check (RenderCommandEncoder);
 			FenceResource(Texture);
-			// MTLPP_VALIDATE(mtlpp::RenderCommandEncoder, RenderCommandEncoder, AGXSafeGetRuntimeDebuggingLevel() >= EAGXDebugLevelValidation, UseResource(Texture, mtlpp::ResourceUsage::Read));
-			METAL_DEBUG_LAYER(EAGXDebugLevelFastValidation, RenderEncoderDebug.SetTexture(EAGXShaderVertex, Texture, index));
 			RenderCommandEncoder.SetVertexTexture(Texture, index);
 			break;
 		case mtlpp::FunctionType::Fragment:
 			check(RenderCommandEncoder);
-			// MTLPP_VALIDATE(mtlpp::RenderCommandEncoder, RenderCommandEncoder, AGXSafeGetRuntimeDebuggingLevel() >= EAGXDebugLevelValidation, UseResource(Texture, mtlpp::ResourceUsage::Read));
-			METAL_DEBUG_LAYER(EAGXDebugLevelFastValidation, RenderEncoderDebug.SetTexture(EAGXShaderFragment, Texture, index));
 			RenderCommandEncoder.SetFragmentTexture(Texture, index);
 			break;
 		case mtlpp::FunctionType::Kernel:
 			check (ComputeCommandEncoder);
 			FenceResource(Texture);
-			// MTLPP_VALIDATE(mtlpp::ComputeCommandEncoder, ComputeCommandEncoder, AGXSafeGetRuntimeDebuggingLevel() >= EAGXDebugLevelValidation, UseResource(Texture, mtlpp::ResourceUsage::Read));
-			METAL_DEBUG_LAYER(EAGXDebugLevelFastValidation, ComputeEncoderDebug.SetTexture(Texture, index));
 			ComputeCommandEncoder.SetTexture(Texture, index);
 			break;
 		default:
@@ -1602,17 +1051,14 @@ void FAGXCommandEncoder::SetShaderSamplerState(mtlpp::FunctionType FunctionType,
 	{
 		case mtlpp::FunctionType::Vertex:
        		check (RenderCommandEncoder);
-			METAL_DEBUG_LAYER(EAGXDebugLevelFastValidation, RenderEncoderDebug.SetSamplerState(EAGXShaderVertex, Sampler, index));
 			RenderCommandEncoder.SetVertexSamplerState(Sampler, index);
 			break;
 		case mtlpp::FunctionType::Fragment:
 			check (RenderCommandEncoder);
-			METAL_DEBUG_LAYER(EAGXDebugLevelFastValidation, RenderEncoderDebug.SetSamplerState(EAGXShaderFragment, Sampler, index));
 			RenderCommandEncoder.SetFragmentSamplerState(Sampler, index);
 			break;
 		case mtlpp::FunctionType::Kernel:
 			check (ComputeCommandEncoder);
-			METAL_DEBUG_LAYER(EAGXDebugLevelFastValidation, ComputeEncoderDebug.SetSamplerState(Sampler, index));
 			ComputeCommandEncoder.SetSamplerState(Sampler, index);
 			break;
 		default:
@@ -1629,23 +1075,8 @@ void FAGXCommandEncoder::SetShaderSideTable(mtlpp::FunctionType const FunctionTy
 	}
 }
 
-void FAGXCommandEncoder::UseIndirectArgumentResource(FAGXTexture const& Texture, mtlpp::ResourceUsage const Usage)
-{
-	FenceResource(Texture);
-	UseResource(Texture, Usage);
-	TextureBindingHistory.Add(ns::AutoReleased<FAGXTexture>(Texture));
-}
-
-void FAGXCommandEncoder::UseIndirectArgumentResource(FAGXBuffer const& Buffer, mtlpp::ResourceUsage const Usage)
-{
-	FenceResource(Buffer);
-	UseResource(Buffer, Usage);
-	BufferBindingHistory.Add(ns::AutoReleased<FAGXBuffer>(Buffer));
-}
-
 void FAGXCommandEncoder::TransitionResources(mtlpp::Resource const& Resource)
 {
-	TransitionedResources.Add(Resource.GetPtr());
 }
 
 #pragma mark - Public Compute State Mutators -
@@ -1654,7 +1085,6 @@ void FAGXCommandEncoder::SetComputePipelineState(FAGXShaderPipeline* State)
 {
 	check (ComputeCommandEncoder);
 	{
-		METAL_DEBUG_LAYER(EAGXDebugLevelFastValidation, ComputeEncoderDebug.SetPipeline(State));
 		ComputeCommandEncoder.SetComputePipelineState(State->ComputePipelineState);
 	}
 }
@@ -1682,110 +1112,10 @@ bool FAGXCommandEncoder::HasBufferBindingHistory(FAGXBuffer const& Buffer) const
 
 void FAGXCommandEncoder::FenceResource(mtlpp::Texture const& Resource)
 {
-	mtlpp::Resource::Type Res = Resource.GetPtr();
-	ns::AutoReleased<mtlpp::Texture> Parent = Resource.GetParentTexture();
-	ns::AutoReleased<mtlpp::Buffer> Buffer = Resource.GetBuffer();
-	if (Parent)
-	{
-		Res = Parent.GetPtr();
-	}
-	else if (Buffer)
-	{
-		Res = Buffer.GetPtr();
-	}
-	if (FenceStage == mtlpp::RenderStages::Vertex || FenceResources.Contains(Res))
-	{
-		FenceStage = mtlpp::RenderStages::Vertex;
-		
-		for (FAGXFence* FragFence : FragmentFences)
-		{
-			if (FragFence->NeedsWait(mtlpp::RenderStages::Fragment))
-			{
-				mtlpp::Fence FragmentFence = FragFence->Get(mtlpp::RenderStages::Fragment);
-				mtlpp::Fence FragInnerFence = METAL_DEBUG_OPTION(CommandList.GetCommandQueue().GetRuntimeDebuggingLevel() >= EAGXDebugLevelValidation ? mtlpp::Fence(((FAGXDebugFence*)FragmentFence.GetPtr()).Inner) :) FragmentFence;
-				
-				if (RenderCommandEncoder)
-				{
-					RenderCommandEncoder.WaitForFence(FragInnerFence, (mtlpp::RenderStages)(mtlpp::RenderStages::Vertex));
-					METAL_DEBUG_LAYER(EAGXDebugLevelFastValidation, RenderEncoderDebug.AddWaitFence(FragmentFence));
-					FragFence->Wait(mtlpp::RenderStages::Fragment);
-				}
-				else if (ComputeCommandEncoder)
-				{
-					ComputeCommandEncoder.WaitForFence(FragInnerFence);
-					METAL_DEBUG_LAYER(EAGXDebugLevelFastValidation, ComputeEncoderDebug.AddWaitFence(FragmentFence));
-					FragFence->Wait(mtlpp::RenderStages::Fragment);
-				}
-				else if (BlitCommandEncoder)
-				{
-					BlitCommandEncoder.WaitForFence(FragInnerFence);
-					METAL_DEBUG_LAYER(EAGXDebugLevelFastValidation, BlitEncoderDebug.AddWaitFence(FragmentFence));
-					FragFence->Wait(mtlpp::RenderStages::Fragment);
-				}
-				METAL_DEBUG_LAYER(EAGXDebugLevelFastValidation, WaitCount++);
-			}
-		}
-		FragmentFences.Empty();
-	}
 }
 
 void FAGXCommandEncoder::FenceResource(mtlpp::Buffer const& Resource)
 {
-	mtlpp::Resource::Type Res = Resource.GetPtr();
-	if (FenceStage == mtlpp::RenderStages::Vertex || FenceResources.Contains(Res))
-	{
-		FenceStage = mtlpp::RenderStages::Vertex;
-		
-		for (FAGXFence* FragFence : FragmentFences)
-		{
-			if (FragFence->NeedsWait(mtlpp::RenderStages::Fragment))
-			{
-				mtlpp::Fence FragmentFence = FragFence->Get(mtlpp::RenderStages::Fragment);
-				mtlpp::Fence FragInnerFence = METAL_DEBUG_OPTION(CommandList.GetCommandQueue().GetRuntimeDebuggingLevel() >= EAGXDebugLevelValidation ? mtlpp::Fence(((FAGXDebugFence*)FragmentFence.GetPtr()).Inner) :) FragmentFence;
-				
-				if (RenderCommandEncoder)
-				{
-					RenderCommandEncoder.WaitForFence(FragInnerFence, (mtlpp::RenderStages)(mtlpp::RenderStages::Vertex));
-					METAL_DEBUG_LAYER(EAGXDebugLevelFastValidation, RenderEncoderDebug.AddWaitFence(FragmentFence));
-					FragFence->Wait(mtlpp::RenderStages::Fragment);
-				}
-				else if (ComputeCommandEncoder)
-				{
-					ComputeCommandEncoder.WaitForFence(FragInnerFence);
-					METAL_DEBUG_LAYER(EAGXDebugLevelFastValidation, ComputeEncoderDebug.AddWaitFence(FragmentFence));
-					FragFence->Wait(mtlpp::RenderStages::Fragment);
-				}
-				else if (BlitCommandEncoder)
-				{
-					BlitCommandEncoder.WaitForFence(FragInnerFence);
-					METAL_DEBUG_LAYER(EAGXDebugLevelFastValidation, BlitEncoderDebug.AddWaitFence(FragmentFence));
-					FragFence->Wait(mtlpp::RenderStages::Fragment);
-				}
-			}
-		}
-		FragmentFences.Empty();
-	}
-}
-
-void FAGXCommandEncoder::UseResource(mtlpp::Resource const& Resource, mtlpp::ResourceUsage const Usage)
-{
-	static bool UseResourceAvailable = FAGXCommandQueue::SupportsFeature(EAGXFeaturesIABs);
-	if (UseResourceAvailable || AGXSafeGetRuntimeDebuggingLevel() >= EAGXDebugLevelValidation)
-	{
-		mtlpp::ResourceUsage Current = ResourceUsage.FindRef(Resource.GetPtr());
-		if (Current != Usage)
-		{
-			ResourceUsage.Add(Resource.GetPtr(), Usage);
-			if (RenderCommandEncoder)
-			{
-				MTLPP_VALIDATE(mtlpp::RenderCommandEncoder, RenderCommandEncoder, AGXSafeGetRuntimeDebuggingLevel() >= EAGXDebugLevelValidation, UseResource(Resource, Usage));
-			}
-			else if (ComputeCommandEncoder)
-			{
-				MTLPP_VALIDATE(mtlpp::ComputeCommandEncoder, ComputeCommandEncoder, AGXSafeGetRuntimeDebuggingLevel() >= EAGXDebugLevelValidation, UseResource(Resource, Usage));
-			}
-		}
-	}
 }
 
 void FAGXCommandEncoder::SetShaderBufferInternal(mtlpp::FunctionType Function, uint32 Index)
@@ -1820,23 +1150,17 @@ void FAGXCommandEncoder::SetShaderBufferInternal(mtlpp::FunctionType Function, u
 				ShaderBuffers[uint32(Function)].Bound |= (1 << Index);
 				check(RenderCommandEncoder);
 				FenceResource(Buffer);
-				// MTLPP_VALIDATE(mtlpp::RenderCommandEncoder, RenderCommandEncoder, AGXSafeGetRuntimeDebuggingLevel() >= EAGXDebugLevelValidation, UseResource(Buffer, mtlpp::ResourceUsage::Read));
-				METAL_DEBUG_LAYER(EAGXDebugLevelFastValidation, RenderEncoderDebug.SetBuffer(EAGXShaderVertex, Buffer, Offset, Index));
 				RenderCommandEncoder.SetVertexBuffer(Buffer, Offset, Index);
 				break;
 			case mtlpp::FunctionType::Fragment:
 				ShaderBuffers[uint32(Function)].Bound |= (1 << Index);
 				check(RenderCommandEncoder);
-				// MTLPP_VALIDATE(mtlpp::RenderCommandEncoder, RenderCommandEncoder, AGXSafeGetRuntimeDebuggingLevel() >= EAGXDebugLevelValidation, UseResource(Buffer, mtlpp::ResourceUsage::Read));
-				METAL_DEBUG_LAYER(EAGXDebugLevelFastValidation, RenderEncoderDebug.SetBuffer(EAGXShaderFragment, Buffer, Offset, Index));
 				RenderCommandEncoder.SetFragmentBuffer(Buffer, Offset, Index);
 				break;
 			case mtlpp::FunctionType::Kernel:
 				ShaderBuffers[uint32(Function)].Bound |= (1 << Index);
 				check(ComputeCommandEncoder);
 				FenceResource(Buffer);
-				// MTLPP_VALIDATE(mtlpp::ComputeCommandEncoder, ComputeCommandEncoder, AGXSafeGetRuntimeDebuggingLevel() >= EAGXDebugLevelValidation, UseResource(Buffer, mtlpp::ResourceUsage::Read));
-				METAL_DEBUG_LAYER(EAGXDebugLevelFastValidation, ComputeEncoderDebug.SetBuffer(Buffer, Offset, Index));
 				ComputeCommandEncoder.SetBuffer(Buffer, Offset, Index);
 				break;
 			default:
@@ -1862,19 +1186,16 @@ void FAGXCommandEncoder::SetShaderBufferInternal(mtlpp::FunctionType Function, u
 			case mtlpp::FunctionType::Vertex:
 				ShaderBuffers[uint32(Function)].Bound |= (1 << Index);
 				check(RenderCommandEncoder);
-				METAL_DEBUG_LAYER(EAGXDebugLevelFastValidation, RenderEncoderDebug.SetBytes(EAGXShaderVertex, Bytes, Len, Index));
 				RenderCommandEncoder.SetVertexData(Bytes, Len, Index);
 				break;
 			case mtlpp::FunctionType::Fragment:
 				ShaderBuffers[uint32(Function)].Bound |= (1 << Index);
 				check(RenderCommandEncoder);
-				METAL_DEBUG_LAYER(EAGXDebugLevelFastValidation, RenderEncoderDebug.SetBytes(EAGXShaderFragment, Bytes, Len, Index));
 				RenderCommandEncoder.SetFragmentData(Bytes, Len, Index);
 				break;
 			case mtlpp::FunctionType::Kernel:
 				ShaderBuffers[uint32(Function)].Bound |= (1 << Index);
 				check(ComputeCommandEncoder);
-				METAL_DEBUG_LAYER(EAGXDebugLevelFastValidation, ComputeEncoderDebug.SetBytes(Bytes, Len, Index));
 				ComputeCommandEncoder.SetBytes(Bytes, Len, Index);
 				break;
 			default:
diff --git a/Engine/Source/Runtime/Apple/AGXRHI/Private/AGXCommandEncoder.h b/Engine/Source/Runtime/Apple/AGXRHI/Private/AGXCommandEncoder.h
--- a/Engine/Source/Runtime/Apple/AGXRHI/Private/AGXCommandEncoder.h
+++ b/Engine/Source/Runtime/Apple/AGXRHI/Private/AGXCommandEncoder.h
@@ -4,11 +4,6 @@
 
 #include <Metal/Metal.h>
 #include "AGXBuffer.h"
-#include "AGXDebugCommandEncoder.h"
-#include "AGXBlitCommandEncoder.h"
-#include "AGXComputeCommandEncoder.h"
-#include "AGXRenderCommandEncoder.h"
-#include "AGXFence.h"
 #include "AGXProfiler.h"
 #include "command_buffer.hpp"
 
@@ -40,67 +35,6 @@ enum EAGXSubmitFlags
 	EAGXSubmitFlagsLastCommandBuffer = 1 << 6,
 };
 
-@class FAGXCommandBufferDebug;
-
-struct FAGXCommandData
-{
-	enum class Type
-	{
-		DrawPrimitive,
-		DrawPrimitiveIndexed,
-		DrawPrimitivePatch,
-		DrawPrimitiveIndirect,
-		DrawPrimitiveIndexedIndirect,
-		Dispatch,
-		DispatchIndirect,
-		Num,
-	};
-	struct DispatchIndirectArgs
-	{
-		id ArgumentBuffer;
-		NSUInteger ArgumentOffset;
-	};
-	Type CommandType;
-	union
-	{
-		mtlpp::DrawPrimitivesIndirectArguments Draw;
-		mtlpp::DrawIndexedPrimitivesIndirectArguments DrawIndexed;
-		mtlpp::DrawPatchIndirectArguments DrawPatch;
-		MTLDispatchThreadgroupsIndirectArguments Dispatch;
-		DispatchIndirectArgs DispatchIndirect;
-	};
-	FString ToString() const;
-};
-
-struct FAGXCommandDebug
-{
-    uint32 CmdBufIndex;
-	uint32 Encoder;
-	uint32 Index;
-	FAGXGraphicsPipelineState* PSO;
-	FAGXComputeShader* ComputeShader;
-	FAGXCommandData Data;
-};
-
-class FAGXCommandBufferMarkers : public ns::Object<FAGXCommandBufferDebug*, ns::CallingConvention::ObjectiveC>
-{
-public:
-	FAGXCommandBufferMarkers(void);
-	FAGXCommandBufferMarkers(mtlpp::CommandBuffer& CmdBuf);
-	FAGXCommandBufferMarkers(FAGXCommandBufferDebug* CmdBuf);
-	
-	void AllocateContexts(uint32 NumContexts);
-	uint32 AddCommand(uint32 CmdBufIndex, uint32 Encoder, uint32 ContextIndex, FAGXBuffer& DebugBuffer, FAGXGraphicsPipelineState* PSO, FAGXComputeShader* ComputeShader, FAGXCommandData& Data);
-	TArray<FAGXCommandDebug>* GetCommands(uint32 Context);
-	ns::AutoReleased<FAGXBuffer> GetDebugBuffer(uint32 ContextIndex);
-	uint32 NumContexts() const;
-	uint32 GetIndex() const;
-	
-	static FAGXCommandBufferMarkers Get(mtlpp::CommandBuffer const& CmdBuf);
-	
-	static char const* kTableAssociationKey;
-};
-
 /**
  * EAGXCommandEncoderType:
  *   EAGXCommandEncoderCurrent: The primary encoder that is used for draw calls & dispatches
@@ -208,9 +142,6 @@ public:
 	/** @returns The active blit command encoder or nil if there isn't one. */
 	mtlpp::BlitCommandEncoder& GetBlitCommandEncoder(void);
 	
-	/** @returns The MTLFence for the current encoder or nil if there isn't one. */
-	TRefCountPtr<FAGXFence> const& GetEncoderFence(void) const;
-	
 	/** @returns The number of encoded passes in the command buffer. */
 	uint32 NumEncodedPasses(void) const { return EncoderNum; }
 	
@@ -234,7 +165,7 @@ public:
 	void BeginBlitCommandEncoding(void);
 	
 	/** Declare that all command generation from this encoder is complete, and detach from the MTLCommandBuffer if there is an encoder active or does nothing if there isn't. */
-	TRefCountPtr<FAGXFence> EndEncoding(void);
+	void EndEncoding(void);
 	
 	/** Initialises a fence for the current command-buffer optionally adding a command-buffer completion handler to the command-buffer */
 	void InsertCommandBufferFence(FAGXCommandBufferFence& Fence, mtlpp::CommandBufferHandler Handler);
@@ -242,18 +173,6 @@ public:
 	/** Adds a command-buffer completion handler to the command-buffer */
 	void AddCompletionHandler(mtlpp::CommandBufferHandler Handler);
 	
-	/** Update the event to capture all GPU work so far enqueued by this encoder. */
-	void UpdateFence(FAGXFence* Fence);
-	
-	/** Prevent further GPU work until the event is reached. */
-	void WaitForFence(FAGXFence* Fence);
-	
-	/**
-	 * Prevent further GPU work until the event is reached and then kick the fence again so future encoders may wait on it.
-	 * Unlike the functions above this is not guarded and the caller must ensure proper matching waits & updates.
-	 */
-	void WaitAndUpdateFence(FAGXFence* Fence);
-
 #pragma mark - Public Debug Support -
 	
 	/*
@@ -271,31 +190,11 @@ public:
 	/* Pop the latest named string off of the stack. */
 	void PopDebugGroup(void);
 	
-	/** Returns the object that records debug command markers into the command-buffer. */
-	FAGXCommandBufferMarkers& GetMarkers(void);
-	
 #if ENABLE_METAL_GPUPROFILE
 	/* Get the command-buffer stats object. */
 	FAGXCommandBufferStats* GetCommandBufferStats(void);
 #endif
 
-#if MTLPP_CONFIG_VALIDATE && METAL_DEBUG_OPTIONS
-	/** @returns The active render command encoder or nil if there isn't one. */
-	FAGXCommandBufferDebugging& GetCommandBufferDebugging(void) { return CommandBufferDebug; } 
-	
-	/** @returns The active render command encoder or nil if there isn't one. */
-	FAGXRenderCommandEncoderDebugging& GetRenderCommandEncoderDebugging(void) { return RenderEncoderDebug; } 
-	
-	/** @returns The active compute command encoder or nil if there isn't one. */
-	FAGXComputeCommandEncoderDebugging& GetComputeCommandEncoderDebugging(void) { return ComputeEncoderDebug; }
-	
-	/** @returns The active blit command encoder or nil if there isn't one. */
-	FAGXBlitCommandEncoderDebugging& GetBlitCommandEncoderDebugging(void) { return BlitEncoderDebug; }
-	
-	/** @returns The active blit command encoder or nil if there isn't one. */
-	FAGXParallelRenderCommandEncoderDebugging& GetParallelRenderCommandEncoderDebugging(void) { return ParallelEncoderDebug; }
-#endif
-	
 #pragma mark - Public Render State Mutators -
 
 	/**
@@ -443,20 +342,6 @@ public:
 	 */
 	void SetShaderSideTable(mtlpp::FunctionType const FunctionType, NSUInteger const Index);
 	
-    /*
-     * Inform the driver that an indirect argument resource will be used and what the usage mode will be.
-     * @param Texture The texture that will be used as an indirect argument.
-     * @param Usage The usage mode for the texture.
-     */
-	void UseIndirectArgumentResource(FAGXTexture const& Texture, mtlpp::ResourceUsage const Usage);
-    
-    /*
-     * Inform the driver that an indirect argument resource will be used and what the usage mode will be.
-     * @param Buffer The buffer that will be used as an indirect argument.
-     * @param Usage The usage mode for the texture.
-     */
-	void UseIndirectArgumentResource(FAGXBuffer const& Buffer, mtlpp::ResourceUsage const Usage);
-
 	/*
 	 * Transition resource so that we can barrier fragment->vertex stages.
 	 * @param Resource The resource we are going to make readable having been written.
@@ -503,8 +388,6 @@ private:
 	
 	void FenceResource(mtlpp::Texture const& Resource);
 	void FenceResource(mtlpp::Buffer const& Resource);
-
-	void UseResource(mtlpp::Resource const& Resource, mtlpp::ResourceUsage const Usage);
 	
 #pragma mark - Private Type Declarations -
 private:
@@ -573,22 +456,10 @@ public:
 	mtlpp::ComputeCommandEncoder ComputeCommandEncoder;
 	mtlpp::BlitCommandEncoder BlitCommandEncoder;
 	TArray<mtlpp::RenderCommandEncoder> ChildRenderCommandEncoders;
-	FAGXCommandBufferMarkers CommandBufferMarkers;
-	
-	METAL_DEBUG_ONLY(FAGXCommandBufferDebugging CommandBufferDebug);
-	METAL_DEBUG_ONLY(FAGXRenderCommandEncoderDebugging RenderEncoderDebug);
-	METAL_DEBUG_ONLY(FAGXComputeCommandEncoderDebugging ComputeEncoderDebug);
-	METAL_DEBUG_ONLY(FAGXBlitCommandEncoderDebugging BlitEncoderDebug);
-	METAL_DEBUG_ONLY(FAGXParallelRenderCommandEncoderDebugging ParallelEncoderDebug);
 	
-	TRefCountPtr<FAGXFence> EncoderFence;
 #if ENABLE_METAL_GPUPROFILE
 	FAGXCommandBufferStats* CommandBufferStats;
 #endif
-#if METAL_DEBUG_OPTIONS
-	uint8 WaitCount;
-	uint8 UpdateCount;
-#endif
 
 	TArray<ns::Object<mtlpp::CommandBufferHandler>> CompletionHandlers;
 	NSMutableArray* DebugGroups;
@@ -600,12 +471,6 @@ public:
 	
 	TMap<mtlpp::Resource::Type, mtlpp::ResourceUsage> ResourceUsage;
 	
-	TSet<mtlpp::Resource::Type> TransitionedResources;
-	TSet<mtlpp::Resource::Type> FenceResources;
-	
-	TSet<TRefCountPtr<FAGXFence>> FragmentFences;
-	
-	mtlpp::RenderStages FenceStage;
 	uint32 EncoderNum;
 	uint32 CmdBufIndex;
 	EAGXCommandEncoderType Type;
diff --git a/Engine/Source/Runtime/Apple/AGXRHI/Private/AGXCommandList.cpp b/Engine/Source/Runtime/Apple/AGXRHI/Private/AGXCommandList.cpp
--- a/Engine/Source/Runtime/Apple/AGXRHI/Private/AGXCommandList.cpp
+++ b/Engine/Source/Runtime/Apple/AGXRHI/Private/AGXCommandList.cpp
@@ -10,7 +10,6 @@
 #include "AGXCommandList.h"
 #include "AGXCommandQueue.h"
 #include "AGXProfiler.h"
-#include "AGXCommandBuffer.h"
 #include "AGXCommandEncoder.h"
 #include "ns.hpp"
 
@@ -60,115 +59,6 @@ static void ReportMetalCommandBufferFailure(mtlpp::CommandBuffer const& Complete
     if (bDoCheck && !GIsSuspended && !GIsRenderingThreadSuspended)
 #endif
     {
-		FAGXCommandBufferMarkers Markers = FAGXCommandBufferMarkers::Get(CompletedBuffer);
-		if (Markers)
-		{
-			uint32 CommandBufferIndex = Markers.GetIndex();
-			FAGXDebugInfo BrokenDraw;
-			BrokenDraw.EncoderIndex = UINT32_MAX;
-			BrokenDraw.CommandIndex = UINT32_MAX;
-			BrokenDraw.CmdBuffIndex = UINT32_MAX;
-			BrokenDraw.ContextIndex = UINT32_MAX;
-            BrokenDraw.CommandBuffer = 0;
-			uint32 BrokenContext = 0;
-			bool bFoundBrokenDraw = false;
-			for (uint32 i = 0; (CommandBufferIndex != ~0u) && !bFoundBrokenDraw && i < Markers.NumContexts(); i++)
-			{
-				ns::AutoReleased<FAGXBuffer> DebugBuffer = Markers.GetDebugBuffer(i);
-				TArray<FAGXCommandDebug>* Commands = Markers.GetCommands(i);
-				if (DebugBuffer && Commands && Commands->Num())
-				{
-					uint32 DebugLength = DebugBuffer.GetLength();
-					uint32 DebugCount = DebugLength / sizeof(FAGXDebugInfo);
-					check(DebugCount >= 1);
-					
-					FAGXDebugInfo* DebugArray = (FAGXDebugInfo*)DebugBuffer.GetContents();
-					FAGXDebugInfo CurrentDraw = DebugArray[0];
-					
-					// On TBDR we find the disjoint where one tile has progressed further than another
-					// We find the earliest failure in the probably vain hope that this is the tile that actually failed
-					// There's actually no guarantee of that, it might be one of the later tiles that exploded
-					for (uint32 j = 0; j < DebugCount; j++)
-					{
-						if (CommandBufferIndex == DebugArray[j].CmdBuffIndex && DebugArray[j].EncoderIndex < CurrentDraw.EncoderIndex && DebugArray[j].CommandIndex < CurrentDraw.CommandIndex)
-						{
-							CurrentDraw = DebugArray[j];
-							break;
-						}
-					}
-
-					// Find the first command to fail - which depends on the order of encoders - parallel contexts make this more complicated
-					if (CurrentDraw.CmdBuffIndex == BrokenDraw.CmdBuffIndex && CurrentDraw.CommandIndex < Commands->Num() && CurrentDraw.EncoderIndex < BrokenDraw.EncoderIndex)
-					{
-						BrokenDraw = CurrentDraw;
-                        BrokenContext = i;
-                        bFoundBrokenDraw = true;
-					}
-                    else if(Commands->Num() > 1)
-                    {
-						BrokenDraw.EncoderIndex = 0;
-						BrokenDraw.CommandIndex = 0;
-						BrokenDraw.CmdBuffIndex = CommandBufferIndex;
-						BrokenDraw.ContextIndex = 0;
-                        BrokenContext = i;
-						BrokenDraw.CommandBuffer = (uintptr_t)CompletedBuffer.GetPtr();
-                        BrokenDraw = DebugArray[1];
-                        bFoundBrokenDraw = true;
-                    }
-				}
-			}
-			if (bFoundBrokenDraw)
-			{
-				id<MTLCommandBuffer> Ptr = reinterpret_cast<id<MTLCommandBuffer>>(BrokenDraw.CommandBuffer);
-				UE_LOG(LogAGX, Error, TEXT("GPU last wrote Command Buffer: %u (%llx) Encoder Index: %d Context Index: %d Draw Index: %d PSO: VS: %u_%u, PS: %u_%u."), BrokenDraw.CmdBuffIndex, BrokenDraw.CommandBuffer, BrokenDraw.EncoderIndex, BrokenDraw.ContextIndex, BrokenDraw.CommandIndex, BrokenDraw.PSOSignature[0], BrokenDraw.PSOSignature[1], BrokenDraw.PSOSignature[2], BrokenDraw.PSOSignature[3]);
-			
-				ns::AutoReleased<FAGXBuffer> DebugBuffer = Markers.GetDebugBuffer(BrokenContext);
-				TArray<FAGXCommandDebug>* Commands = Markers.GetCommands(BrokenContext);
-				if (Commands->Num())
-				{
-					UE_LOG(LogAGX, Error, TEXT("Failed executing following commands:"));
-					uint32 StartIdx = BrokenDraw.CommandIndex;
-					for (uint32 CmdIdx = StartIdx; CmdIdx < Commands->Num(); CmdIdx++)
-					{
-						FAGXCommandDebug& Command = (*Commands)[CmdIdx];
-						FString VSHash = Command.PSO->VertexShader->GetHash().ToString();
-						uint32 VSSig[2] = { Command.PSO->VertexShader->SourceLen, Command.PSO->VertexShader->SourceCRC };
-						
-						FString PSHash;
-						uint32 PSSig[2] = { 0, 0 };
-						if (Command.PSO->PixelShader.IsValid())
-						{
-							PSHash = Command.PSO->PixelShader->GetHash().ToString();
-							PSSig[0] = Command.PSO->PixelShader->SourceLen;
-							PSSig[1] = Command.PSO->PixelShader->SourceCRC;
-						}
-						
-						UE_LOG(LogAGX, Error, TEXT("Command Buffer: %d (%p) Encoder: %d Command: %d: %s PSO: VS: %s (%u_%u), PS: %s (%u_%u)"), Command.CmdBufIndex, CompletedBuffer.GetPtr(), Command.Encoder, Command.Index, *Command.Data.ToString(), *VSHash, VSSig[0], VSSig[1], *PSHash, PSSig[0], PSSig[1]);
-					}
-				}
-                if (DebugBuffer)
-                {
-                    uint32 DebugLength = DebugBuffer.GetLength();
-                    uint32 DebugCount = DebugLength / sizeof(FAGXDebugInfo);
-                    check(DebugCount >= 1);
-                    
-                    FAGXDebugInfo* DebugArray = (FAGXDebugInfo*)DebugBuffer.GetContents();
-                    for (uint32 i = 0; i < DebugCount; i++)
-                    {
-                        FAGXDebugInfo& Command = DebugArray[i];
-						
-						// Stop when nothing has been written into the buffer
-						if (Command.CmdBuffIndex == 0 && Command.CommandBuffer == 0 && Command.EncoderIndex == 0 && Command.CommandIndex == 0 && Command.PSOSignature[0] == 0 && Command.PSOSignature[1] == 0 && Command.PSOSignature[2] == 0 && Command.PSOSignature[3] == 0)
-						{
-							break;
-						}
-						
-                        UE_LOG(LogAGX, Error, TEXT("Command Buffer: %d (%p) Debug Buffer: %p Tile: %u Context: %d Encoder: %d Command: %d PSO: VS: %u_%u, PS: %u_%u"), Command.CmdBuffIndex, Command.CommandBuffer, DebugBuffer.GetPtr(), i, Command.ContextIndex, Command.EncoderIndex, Command.CommandIndex, Command.PSOSignature[0], Command.PSOSignature[1], Command.PSOSignature[2], Command.PSOSignature[3]);
-                    }
-                }
-			}
-		}
-		
 #if PLATFORM_IOS
         UE_LOG(LogAGX, Warning, TEXT("Command Buffer %s Failed with %s Error! Error Domain: %s Code: %d Description %s %s %s"), *LabelString, ErrorType, *DomainString, Code, *ErrorString, *FailureString, *RecoveryString);
         FIOSPlatformMisc::GPUAssert();
diff --git a/Engine/Source/Runtime/Apple/AGXRHI/Private/AGXCommandQueue.cpp b/Engine/Source/Runtime/Apple/AGXRHI/Private/AGXCommandQueue.cpp
--- a/Engine/Source/Runtime/Apple/AGXRHI/Private/AGXCommandQueue.cpp
+++ b/Engine/Source/Runtime/Apple/AGXRHI/Private/AGXCommandQueue.cpp
@@ -7,7 +7,6 @@
 #include "AGXRHIPrivate.h"
 
 #include "AGXCommandQueue.h"
-#include "AGXCommandBuffer.h"
 #include "AGXCommandList.h"
 #include "AGXProfiler.h"
 #include "Misc/ConfigCacheIni.h"
@@ -17,7 +16,6 @@
 NSUInteger FAGXCommandQueue::PermittedOptions = 0;
 uint64 FAGXCommandQueue::Features = 0;
 extern mtlpp::VertexFormat GAGXFColorVertexFormat;
-bool GAGXCommandBufferDebuggingEnabled = 0;
 
 #pragma mark - Public C++ Boilerplate -
 
@@ -26,23 +24,17 @@ FAGXCommandQueue::FAGXCommandQueue(uint32 const MaxNumCommandBuffers /* = 0 */)
 , RuntimeDebuggingLevel(EAGXDebugLevelOff)
 {
 	int32 MaxShaderVersion = 0;
-	int32 IndirectArgumentTier = 0;
-#if PLATFORM_MAC
 	int32 DefaultMaxShaderVersion = 5; // MSL v2.2
-	int32 MinShaderVersion = 4; // MSL v2.1
-    const TCHAR* const Settings = TEXT("/Script/MacTargetPlatform.MacTargetSettings");
+	int32 MinShaderVersion = 5; // MSL v2.2
+
+#if PLATFORM_MAC
+	const TCHAR* const Settings = TEXT("/Script/MacTargetPlatform.MacTargetSettings");
 #else
-	int32 DefaultMaxShaderVersion = 2;
-	int32 MinShaderVersion = 2;
-    const TCHAR* const Settings = TEXT("/Script/IOSRuntimeSettings.IOSRuntimeSettings");
+	const TCHAR* const Settings = TEXT("/Script/IOSRuntimeSettings.IOSRuntimeSettings");
 #endif
     if(!GConfig->GetInt(Settings, TEXT("MaxShaderLanguageVersion"), MaxShaderVersion, GEngineIni))
-    {
-        MaxShaderVersion = DefaultMaxShaderVersion;
-    }
-	if(!GConfig->GetInt(Settings, TEXT("IndirectArgumentTier"), IndirectArgumentTier, GEngineIni))
 	{
-		IndirectArgumentTier = 0;
+		MaxShaderVersion = DefaultMaxShaderVersion;
 	}
 	MaxShaderVersion = FMath::Max(MinShaderVersion, MaxShaderVersion);
 	AGXValidateVersion(MaxShaderVersion);
@@ -57,151 +49,74 @@ FAGXCommandQueue::FAGXCommandQueue(uint32 const MaxNumCommandBuffers /* = 0 */)
 	}
 	check(CommandQueue);
 #if PLATFORM_IOS
-	NSOperatingSystemVersion Vers = [[NSProcessInfo processInfo] operatingSystemVersion];
-	if(Vers.majorVersion >= 9)
-	{
-		Features = EAGXFeaturesSetBufferOffset | EAGXFeaturesSetBytes;
+	NSOperatingSystemVersion Vers = [[NSProcessInfo processInfo]operatingSystemVersion];
+	Features = EAGXFeaturesSetBufferOffset | EAGXFeaturesSetBytes;
 
 #if PLATFORM_TVOS
-        Features &= ~(EAGXFeaturesSetBytes);
-		
-		if ([GMtlDevice supportsFeatureSet:MTLFeatureSet_tvOS_GPUFamily2_v1])
-		{
-			Features |= EAGXFeaturesCountingQueries | EAGXFeaturesBaseVertexInstance | EAGXFeaturesIndirectBuffer | EAGXFeaturesMSAADepthResolve | EAGXFeaturesMSAAStoreAndResolve;
-		}
-		
-		if(Vers.majorVersion > 10)
-		{
-			Features |= EAGXFeaturesPrivateBufferSubAllocation;
-			
-			if(Vers.majorVersion >= 11)
-			{
-				Features |= EAGXFeaturesGPUCaptureManager | EAGXFeaturesBufferSubAllocation | EAGXFeaturesParallelRenderEncoders | EAGXFeaturesPipelineBufferMutability;
-				
-				if (MaxShaderVersion >= 3)
-				{
-					GAGXFColorVertexFormat = mtlpp::VertexFormat::UChar4Normalized_BGRA;
-				}
+	Features &= ~(EAGXFeaturesSetBytes);
+
+	if ([GMtlDevice supportsFeatureSet:MTLFeatureSet_tvOS_GPUFamily2_v1])
+	{
+		Features |= EAGXFeaturesCountingQueries | EAGXFeaturesBaseVertexInstance | EAGXFeaturesIndirectBuffer | EAGXFeaturesMSAADepthResolve | EAGXFeaturesMSAAStoreAndResolve;
+	}
+
+	Features |= EAGXFeaturesPrivateBufferSubAllocation | EAGXFeaturesGPUCaptureManager | EAGXFeaturesBufferSubAllocation | EAGXFeaturesParallelRenderEncoders | EAGXFeaturesPipelineBufferMutability | EAGXFeaturesMaxThreadsPerThreadgroup | EAGXFeaturesTextureBuffers;
+	GAGXFColorVertexFormat = mtlpp::VertexFormat::UChar4Normalized_BGRA;
 
-				if (Vers.majorVersion >= 12)
-				{
-					Features |= EAGXFeaturesMaxThreadsPerThreadgroup;
-					
-					if (FParse::Param(FCommandLine::Get(),TEXT("metalfence")))
-					{
-					Features |= EAGXFeaturesFences;
-					}
-					
-					if (FParse::Param(FCommandLine::Get(),TEXT("metalheap")))
-					{
-					Features |= EAGXFeaturesHeaps;
-					}
-					
-					if (MaxShaderVersion >= 4)
-					{
-						Features |= EAGXFeaturesTextureBuffers;
-					}
-				}
-			}
-		}
 #else
-		if ([GMtlDevice supportsFeatureSet:MTLFeatureSet_iOS_GPUFamily3_v1])
-		{
-			Features |= EAGXFeaturesCountingQueries | EAGXFeaturesBaseVertexInstance | EAGXFeaturesIndirectBuffer | EAGXFeaturesMSAADepthResolve;
-		}
-		
-		if ([GMtlDevice supportsFeatureSet:MTLFeatureSet_iOS_GPUFamily3_v2] || [GMtlDevice supportsFeatureSet:MTLFeatureSet_iOS_GPUFamily2_v3] || [GMtlDevice supportsFeatureSet:MTLFeatureSet_iOS_GPUFamily1_v3])
-		{
-			if (FParse::Param(FCommandLine::Get(),TEXT("metalfence")))
-			{
-				Features |= EAGXFeaturesFences;
-			}
-			
-			if (FParse::Param(FCommandLine::Get(),TEXT("metalheap")))
-			{
-				Features |= EAGXFeaturesHeaps;
-			}
-		}
-		
-		if ([GMtlDevice supportsFeatureSet:MTLFeatureSet_iOS_GPUFamily3_v2])
-		{
-			Features |= EAGXFeaturesMSAAStoreAndResolve;
-		}
+	if ([GMtlDevice supportsFeatureSet:MTLFeatureSet_iOS_GPUFamily3_v1])
+	{
+		Features |= EAGXFeaturesCountingQueries | EAGXFeaturesBaseVertexInstance | EAGXFeaturesIndirectBuffer | EAGXFeaturesMSAADepthResolve;
+	}
 		
-		if(Vers.majorVersion > 10 || (Vers.majorVersion == 10 && Vers.minorVersion >= 3))
-        {
-			// Turning the below option on will allocate more buffer memory which isn't generally desirable on iOS
-			// Features |= EAGXFeaturesEfficientBufferBlits;
-			
-			// These options are fine however as thye just change how we allocate small buffers
-            Features |= EAGXFeaturesBufferSubAllocation;
-			Features |= EAGXFeaturesPrivateBufferSubAllocation;
-			
-			if(Vers.majorVersion >= 11)
-			{
-				if (MaxShaderVersion >= 3)
-				{
-					GAGXFColorVertexFormat = mtlpp::VertexFormat::UChar4Normalized_BGRA;
-				}
-				
-				Features |= EAGXFeaturesPresentMinDuration | EAGXFeaturesGPUCaptureManager | EAGXFeaturesBufferSubAllocation | EAGXFeaturesParallelRenderEncoders | EAGXFeaturesPipelineBufferMutability;
-                
-				// Turn on Texture Buffers! These are faster on the GPU as we don't need to do out-of-bounds tests but require Metal 2.1 and macOS 10.14
-				if (Vers.majorVersion >= 12)
-				{
-					Features |= EAGXFeaturesMaxThreadsPerThreadgroup;
-                    if (!FParse::Param(FCommandLine::Get(),TEXT("nometalfence")))
-                    {
-                        Features |= EAGXFeaturesFences;
-                    }
-                    
-                    if (!FParse::Param(FCommandLine::Get(),TEXT("nometalheap")))
-                    {
-                        Features |= EAGXFeaturesHeaps;
-                    }
-					
-					if (MaxShaderVersion >= 4)
-					{
-						Features |= EAGXFeaturesTextureBuffers;
-					}
-                    
-                    if ([GMtlDevice supportsFeatureSet:MTLFeatureSet_iOS_GPUFamily4_v1])
-                    {
-                        Features |= EAGXFeaturesTileShaders;
+	if ([GMtlDevice supportsFeatureSet:MTLFeatureSet_iOS_GPUFamily3_v2])
+	{
+		Features |= EAGXFeaturesMSAAStoreAndResolve;
+	}
+
+	// Turning the below option on will allocate more buffer memory which isn't generally desirable on iOS
+	// Features |= EAGXFeaturesEfficientBufferBlits;
+
+	// These options are fine however as thye just change how we allocate small buffers
+	Features |= EAGXFeaturesBufferSubAllocation;
+	Features |= EAGXFeaturesPrivateBufferSubAllocation;
+	GAGXFColorVertexFormat = mtlpp::VertexFormat::UChar4Normalized_BGRA;
+
+	Features |= EAGXFeaturesPresentMinDuration
+		| EAGXFeaturesGPUCaptureManager
+		| EAGXFeaturesBufferSubAllocation
+		| EAGXFeaturesParallelRenderEncoders
+		| EAGXFeaturesPipelineBufferMutability;
+
+	Features |= EAGXFeaturesMaxThreadsPerThreadgroup;
+	Features |= EAGXFeaturesTextureBuffers;
+
+	if ([GMtlDevice supportsFeatureSet : MTLFeatureSet_iOS_GPUFamily4_v1])
+	{
+		Features |= EAGXFeaturesTileShaders;
                         
-                        // The below implies tile shaders which are necessary to order the draw calls and generate a buffer that shows what PSOs/draws ran on each tile.
-                        IConsoleVariable* GPUCrashDebuggingCVar = IConsoleManager::Get().FindConsoleVariable(TEXT("r.GPUCrashDebugging"));
-                        GAGXCommandBufferDebuggingEnabled = (GPUCrashDebuggingCVar && GPUCrashDebuggingCVar->GetInt() != 0) || FParse::Param(FCommandLine::Get(),TEXT("metalgpudebug"));
-                    }
-                    
-					if ([GMtlDevice supportsFeatureSet:MTLFeatureSet_iOS_GPUFamily5_v1])
-					{
-						Features |= EAGXFeaturesLayeredRendering;
-					}
-				}
-			}
-        }
-#endif
 	}
-	else if(Vers.majorVersion == 8 && Vers.minorVersion >= 3)
+                    
+	if ([GMtlDevice supportsFeatureSet:MTLFeatureSet_iOS_GPUFamily5_v1])
 	{
-		Features = EAGXFeaturesSetBufferOffset;
+		Features |= EAGXFeaturesLayeredRendering;
 	}
+#endif
 #else // Assume that Mac & other platforms all support these from the start. They can diverge later.
-	const bool bIsNVIDIA = [[GMtlDevice name] rangeOfString:@"Nvidia" options:NSCaseInsensitiveSearch].location != NSNotFound;
+	const bool bIsNVIDIA = [[GMtlDevice name]rangeOfString:@"Nvidia" options:NSCaseInsensitiveSearch] .location != NSNotFound;
 	Features = EAGXFeaturesCountingQueries | EAGXFeaturesBaseVertexInstance | EAGXFeaturesIndirectBuffer | EAGXFeaturesLayeredRendering | EAGXFeaturesCubemapArrays;
 	if (!bIsNVIDIA)
 	{
 		Features |= EAGXFeaturesSetBufferOffset;
 	}
 	if ([GMtlDevice supportsFeatureSet:MTLFeatureSet_macOS_GPUFamily1_v2])
-    {
+	{
         Features |= (   EAGXFeaturesSetBytes
-					  | EAGXFeaturesMSAADepthResolve
-					  | EAGXFeaturesMSAAStoreAndResolve
-					  | EAGXFeaturesEfficientBufferBlits
-					  | EAGXFeaturesBufferSubAllocation
-					  | EAGXFeaturesPrivateBufferSubAllocation
+			| EAGXFeaturesMSAADepthResolve
+			| EAGXFeaturesMSAAStoreAndResolve
+			| EAGXFeaturesEfficientBufferBlits
+			| EAGXFeaturesBufferSubAllocation
+			| EAGXFeaturesPrivateBufferSubAllocation
 					  | EAGXFeaturesMaxThreadsPerThreadgroup );
 		
 		GAGXFColorVertexFormat = mtlpp::VertexFormat::UChar4Normalized_BGRA;
@@ -211,68 +126,18 @@ FAGXCommandQueue::FAGXCommandQueue(uint32 const MaxNumCommandBuffers /* = 0 */)
 			Features |= EAGXFeaturesParallelRenderEncoders;
 		}
 
-		{
-			if (MaxShaderVersion >= 4)
-			{
-				Features |= EAGXFeaturesTextureBuffers;
-            }
-            if (IndirectArgumentTier >= 1)
-            {
-                Features |= EAGXFeaturesIABs;
-				
-				if (IndirectArgumentTier >= 2)
-				{
-					Features |= EAGXFeaturesTier2IABs;
-				}
-            }
-            
-            IConsoleVariable* GPUCrashDebuggingCVar = IConsoleManager::Get().FindConsoleVariable(TEXT("r.GPUCrashDebugging"));
-            GAGXCommandBufferDebuggingEnabled = (GPUCrashDebuggingCVar && GPUCrashDebuggingCVar->GetInt() != 0) || FParse::Param(FCommandLine::Get(),TEXT("metalgpudebug"));
-            
-            // The editor spawns so many viewports and preview icons that we can run out of hardware fences!
-			// Need to figure out a way to safely flush the rendering and reuse the fences when that happens.
-#if WITH_EDITORONLY_DATA
-			if (!GIsEditor)
-#endif
-			{
-				if (FParse::Param(FCommandLine::Get(),TEXT("metalfence")))
-				{
-					Features |= EAGXFeaturesFences;
-				}
-				
-				// There are still too many driver bugs to use MTLHeap on macOS - nothing works without causing random, undebuggable GPU hangs that completely deadlock the Mac and don't generate any validation errors or command-buffer failures
-				if (FParse::Param(FCommandLine::Get(),TEXT("forcemetalheap")))
-				{
-					Features |= EAGXFeaturesHeaps;
-				}
-			}
-		}
+		Features |= EAGXFeaturesTextureBuffers;
     }
     else if ([[GMtlDevice name] rangeOfString:@"Nvidia" options:NSCaseInsensitiveSearch].location != NSNotFound)
-    {
+	{
 		// Using set*Bytes fixes bugs on Nvidia for 10.11 so we should use it...
-    	Features |= EAGXFeaturesSetBytes;
-    }
+		Features |= EAGXFeaturesSetBytes;
+	}
     
     if ([GMtlDevice supportsFeatureSet:MTLFeatureSet_macOS_GPUFamily1_v3])
-    {
-        Features |= EAGXFeaturesMultipleViewports | EAGXFeaturesPipelineBufferMutability | EAGXFeaturesGPUCaptureManager;
-		
-		if (FParse::Param(FCommandLine::Get(),TEXT("metalfence")))
-		{
-			Features |= EAGXFeaturesFences;
-		}
-		
-		if (FParse::Param(FCommandLine::Get(),TEXT("metalheap")))
-		{
-			Features |= EAGXFeaturesHeaps;
-		}
-		
-		if (FParse::Param(FCommandLine::Get(),TEXT("metaliabs")))
-		{
-			Features |= EAGXFeaturesIABs;
-		}
-    }
+	{
+		Features |= EAGXFeaturesMultipleViewports | EAGXFeaturesPipelineBufferMutability | EAGXFeaturesGPUCaptureManager;
+	}
 #endif
 	
 #if !UE_BUILD_SHIPPING
@@ -292,20 +157,13 @@ FAGXCommandQueue::FAGXCommandQueue(uint32 const MaxNumCommandBuffers /* = 0 */)
 	PermittedOptions = 0;
 	PermittedOptions |= mtlpp::ResourceOptions::CpuCacheModeDefaultCache;
 	PermittedOptions |= mtlpp::ResourceOptions::CpuCacheModeWriteCombined;
-	{
-		PermittedOptions |= mtlpp::ResourceOptions::StorageModeShared;
-		PermittedOptions |= mtlpp::ResourceOptions::StorageModePrivate;
+	PermittedOptions |= mtlpp::ResourceOptions::StorageModeShared;
+	PermittedOptions |= mtlpp::ResourceOptions::StorageModePrivate;
 #if PLATFORM_MAC
-		PermittedOptions |= mtlpp::ResourceOptions::StorageModeManaged;
+	PermittedOptions |= mtlpp::ResourceOptions::StorageModeManaged;
 #else
-		PermittedOptions |= mtlpp::ResourceOptions::StorageModeMemoryless;
+	PermittedOptions |= mtlpp::ResourceOptions::StorageModeMemoryless;
 #endif
-		// You can't use HazardUntracked under the validation layer due to bugs in the layer when trying to create linear-textures/texture-buffers
-		if ((Features & EAGXFeaturesFences) && !(Features & EAGXFeaturesValidation))
-		{
-			PermittedOptions |= mtlpp::ResourceOptions::HazardTrackingModeUntracked;
-		}
-	}
 }
 
 FAGXCommandQueue::~FAGXCommandQueue(void)
@@ -340,7 +198,6 @@ mtlpp::CommandBuffer FAGXCommandQueue::CreateCommandBuffer(void)
 		
 		if (RuntimeDebuggingLevel > EAGXDebugLevelOff)
 		{			
-			METAL_DEBUG_ONLY(FAGXCommandBufferDebugging AddDebugging(CmdBuffer));
 			MTLPP_VALIDATION(mtlpp::CommandBufferValidationTable ValidatedCommandBuffer(CmdBuffer));
 		}
 	}
@@ -385,40 +242,6 @@ void FAGXCommandQueue::SubmitCommandBuffers(TArray<mtlpp::CommandBuffer> BufferL
 	}
 }
 
-FAGXFence* FAGXCommandQueue::CreateFence(ns::String const& Label) const
-{
-	if ((Features & EAGXFeaturesFences) != 0)
-	{
-		FAGXFence* InternalFence = FAGXFencePool::Get().AllocateFence();
-		for (uint32 i = mtlpp::RenderStages::Vertex; InternalFence && i <= mtlpp::RenderStages::Fragment; i++)
-		{
-			mtlpp::Fence InnerFence = InternalFence->Get((mtlpp::RenderStages)i);
-			NSString* String = nil;
-			if (GetEmitDrawEvents())
-			{
-				String = [NSString stringWithFormat:@"%u %p: %@", i, InnerFence.GetPtr(), Label.GetPtr()];
-			}
-	#if METAL_DEBUG_OPTIONS
-			if (RuntimeDebuggingLevel >= EAGXDebugLevelValidation)
-			{
-				FAGXDebugFence* Fence = (FAGXDebugFence*)InnerFence.GetPtr();
-				Fence.label = String;
-			}
-			else
-	#endif
-			if(InnerFence && String)
-			{
-				InnerFence.SetLabel(String);
-			}
-		}
-		return InternalFence;
-	}
-	else
-	{
-		return nullptr;
-	}
-}
-
 void FAGXCommandQueue::GetCommittedCommandBufferFences(TArray<mtlpp::CommandBufferFence>& Fences)
 {
 	TArray<mtlpp::CommandBufferFence*> Temp;
@@ -449,7 +272,7 @@ mtlpp::ResourceOptions FAGXCommandQueue::GetCompatibleResourceOptions(mtlpp::Res
 void FAGXCommandQueue::InsertDebugCaptureBoundary(void)
 {
 	PRAGMA_DISABLE_DEPRECATION_WARNINGS
-	[CommandQueue insertDebugCaptureBoundary];
+		[CommandQueue insertDebugCaptureBoundary];
 	PRAGMA_ENABLE_DEPRECATION_WARNINGS
 }
 
diff --git a/Engine/Source/Runtime/Apple/AGXRHI/Private/AGXCommandQueue.h b/Engine/Source/Runtime/Apple/AGXRHI/Private/AGXCommandQueue.h
--- a/Engine/Source/Runtime/Apple/AGXRHI/Private/AGXCommandQueue.h
+++ b/Engine/Source/Runtime/Apple/AGXRHI/Private/AGXCommandQueue.h
@@ -36,9 +36,9 @@ typedef NS_OPTIONS(uint64, EAGXFeatures)
 	/** Supports detailed statistics */
 	EAGXFeaturesStatistics = 1 << 8,
 	/** Supports the explicit MTLHeap APIs */
-	EAGXFeaturesHeaps = 1 << 9,
+	EAGXFeaturesHeaps_REMOVED = 1 << 9,
 	/** Supports the explicit MTLFence APIs */
-	EAGXFeaturesFences = 1 << 10,
+	EAGXFeaturesFences_REMOVED = 1 << 10,
 	/** Supports MSAA Depth Resolves */
 	EAGXFeaturesMSAADepthResolve = 1 << 11,
 	/** Supports Store & Resolve in a single store action */
@@ -66,7 +66,7 @@ typedef NS_OPTIONS(uint64, EAGXFeatures)
 	/** Supports parallel render encoders */
 	EAGXFeaturesParallelRenderEncoders = 1llu << 23llu,
 	/** Supports indirect argument buffers */
-	EAGXFeaturesIABs = 1llu << 24llu,
+	EAGXFeaturesIABs_REMOVED = 1llu << 24llu,
 	/** Supports specifying the mutability of buffers bound to PSOs */
     EAGXFeaturesPipelineBufferMutability = 1llu << 25llu,
     /** Supports tile shaders */
@@ -74,7 +74,7 @@ typedef NS_OPTIONS(uint64, EAGXFeatures)
 	/** Unused Reserved Bit */
 	EAGXFeaturesUnusedReservedBit27 = 1llu << 27llu, // was EAGXFeaturesSeparateTessellation
 	/** Supports indirect argument buffers Tier 2 */
-	EAGXFeaturesTier2IABs = 1llu << 28llu,
+	EAGXFeaturesTier2IABs_REMOVED = 1llu << 28llu,
 };
 
 /**
@@ -118,9 +118,6 @@ public:
 	 */
 	void SubmitCommandBuffers(TArray<mtlpp::CommandBuffer> BufferList, uint32 Index, uint32 Count);
 
-	/** @returns Creates a new MTLFence or nil if this is unsupported */
-	FAGXFence* CreateFence(ns::String const& Label) const;
-	
 	/** @params Fences An array of command-buffer fences for the committed command-buffers */
 	void GetCommittedCommandBufferFences(TArray<mtlpp::CommandBufferFence>& Fences);
 	
diff --git a/Engine/Source/Runtime/Apple/AGXRHI/Private/AGXCommands.cpp b/Engine/Source/Runtime/Apple/AGXRHI/Private/AGXCommands.cpp
--- a/Engine/Source/Runtime/Apple/AGXRHI/Private/AGXCommands.cpp
+++ b/Engine/Source/Runtime/Apple/AGXRHI/Private/AGXCommands.cpp
@@ -15,7 +15,6 @@
 #include "ShaderParameterUtils.h"
 #include "SceneUtils.h"
 #include "AGXProfiler.h"
-#include "AGXCommandBuffer.h"
 #include "StaticBoundShaderState.h"
 #include "EngineGlobals.h"
 #include "PipelineStateCache.h"
diff --git a/Engine/Source/Runtime/Apple/AGXRHI/Private/AGXComputeCommandEncoder.cpp b/Engine/Source/Runtime/Apple/AGXRHI/Private/AGXComputeCommandEncoder.cpp
deleted file mode 100644
diff --git a/Engine/Source/Runtime/Apple/AGXRHI/Private/AGXComputeCommandEncoder.h b/Engine/Source/Runtime/Apple/AGXRHI/Private/AGXComputeCommandEncoder.h
deleted file mode 100644
diff --git a/Engine/Source/Runtime/Apple/AGXRHI/Private/AGXContext.cpp b/Engine/Source/Runtime/Apple/AGXRHI/Private/AGXContext.cpp
--- a/Engine/Source/Runtime/Apple/AGXRHI/Private/AGXContext.cpp
+++ b/Engine/Source/Runtime/Apple/AGXRHI/Private/AGXContext.cpp
@@ -15,7 +15,6 @@
 
 #include "AGXContext.h"
 #include "AGXProfiler.h"
-#include "AGXCommandBuffer.h"
 
 #include "AGXFrameAllocator.h"
 
@@ -163,6 +162,11 @@ mtlpp::Device GMtlppDevice;
 #pragma mark - AGX Device Context Support Routines
 
 
+uint32 AGXSafeGetRuntimeDebuggingLevel()
+{
+	return GIsRHIInitialized ? GetAGXDeviceContext().GetCommandQueue().GetRuntimeDebuggingLevel() : GAGXRuntimeDebugLevel;
+}
+
 #if PLATFORM_MAC
 static id<NSObject> GAGXDeviceObserver = nil;
 
@@ -379,11 +383,6 @@ FAGXDeviceContext* FAGXDeviceContext::CreateDeviceContext()
 	FAGXCommandQueue* Queue = new FAGXCommandQueue(GAGXCommandQueueSize);
 	check(Queue);
 	
-	if (FAGXCommandQueue::SupportsFeature(EAGXFeaturesFences))
-	{
-		FAGXFencePool::Get().Init();
-	}
-	
 	return new FAGXDeviceContext(DeviceIndex, Queue);
 }
 
@@ -518,7 +517,7 @@ void FAGXDeviceContext::ClearFreeList()
 				{
 					ScribbleBuffer(Buffer);
 				}
-				if (GAGXResourcePurgeOnDelete && !Buffer.GetHeap() && !Buffer.GetParentBuffer())
+				if (GAGXResourcePurgeOnDelete && !Buffer.GetParentBuffer())
 				{
 					Buffer.SetPurgeableState(mtlpp::PurgeableState::Empty);
 				}
@@ -530,7 +529,7 @@ void FAGXDeviceContext::ClearFreeList()
                 if (!Texture.GetBuffer() && !Texture.GetParentTexture())
 				{
 #if METAL_DEBUG_OPTIONS
-					if (GAGXResourcePurgeOnDelete && !Texture.GetHeap())
+					if (GAGXResourcePurgeOnDelete)
 					{
 						Texture.SetPurgeableState(mtlpp::PurgeableState::Empty);
 					}
@@ -538,10 +537,6 @@ void FAGXDeviceContext::ClearFreeList()
 					Heap.ReleaseTexture(nullptr, Texture);
 				}
 			}
-			for ( FAGXFence* Fence : Pair->FenceFreeList )
-			{
-				FAGXFencePool::Get().ReleaseFence(Fence);
-			}
 			delete Pair;
 			DelayedFreeLists.RemoveAt(Index, 1, false);
 		}
@@ -648,7 +643,7 @@ bool FAGXDeviceContext::FAGXDelayedFreeList::IsComplete() const
 	return bFinished;
 }
 
-void FAGXDeviceContext::FlushFreeList(bool const bFlushFences)
+void FAGXDeviceContext::FlushFreeList()
 {
 	FAGXDelayedFreeList* NewList = new FAGXDelayedFreeList;
 	
@@ -660,25 +655,6 @@ void FAGXDeviceContext::FlushFreeList(bool const bFlushFences)
 	NewList->UsedBuffers = MoveTemp(UsedBuffers);
 	NewList->UsedTextures = MoveTemp(UsedTextures);
 	NewList->ObjectFreeList = ObjectFreeList;
-	if (bFlushFences)
-	{
-		TArray<FAGXFence*> Fences;
-		FenceFreeList.PopAll(Fences);
-		for (FAGXFence* Fence : Fences)
-		{
-			if(!UsedFences.Contains(Fence))
-			{
-				UsedFences.Add(Fence);
-			}
-		}
-		NewList->FenceFreeList = MoveTemp(UsedFences);
-	}
-#if METAL_DEBUG_OPTIONS
-	if (FrameFences.Num())
-	{
-		FrameFences.Empty();
-	}
-#endif
 	ObjectFreeList.Empty(ObjectFreeList.Num());
 	FreeListMutex.Unlock();
 	
@@ -783,28 +759,11 @@ void FAGXDeviceContext::ReleaseTexture(FAGXTexture& Texture)
 	}
 }
 
-void FAGXDeviceContext::ReleaseFence(FAGXFence* Fence)
-{
-#if METAL_DEBUG_OPTIONS
-	if(GetCommandList().GetCommandQueue().GetRuntimeDebuggingLevel() >= EAGXDebugLevelValidation)
-	{
-		FScopeLock Lock(&FreeListMutex);
-		FrameFences.Add(Fence);
-	}
-#endif
-	
-	if (GIsAGXInitialized) // @todo zebra: there seems to be some race condition at exit when the framerate is very low
-	{
-		check(Fence);
-		FenceFreeList.Push(Fence);
-	}
-}
-
 FAGXTexture FAGXDeviceContext::CreateTexture(FAGXSurface* Surface, mtlpp::TextureDescriptor Descriptor)
 {
 	FAGXTexture Tex = Heap.CreateTexture(Descriptor, Surface);
 #if METAL_DEBUG_OPTIONS
-	if (GAGXResourcePurgeOnDelete && !Tex.GetHeap())
+	if (GAGXResourcePurgeOnDelete)
 	{
 		Tex.SetPurgeableState(mtlpp::PurgeableState::NonVolatile);
 	}
@@ -829,7 +788,7 @@ FAGXBuffer FAGXDeviceContext::CreatePooledBuffer(FAGXPooledBufferArgs const& Arg
     FAGXBuffer Buffer = Heap.CreateBuffer(Args.Size, RequestedBufferOffsetAlignment, Args.Flags, FAGXCommandQueue::GetCompatibleResourceOptions((mtlpp::ResourceOptions)(CpuResourceOption | mtlpp::ResourceOptions::HazardTrackingModeUntracked | ((NSUInteger)Args.Storage << mtlpp::ResourceStorageModeShift))));
 	check(Buffer && Buffer.GetPtr());
 #if METAL_DEBUG_OPTIONS
-	if (GAGXResourcePurgeOnDelete && !Buffer.GetHeap())
+	if (GAGXResourcePurgeOnDelete)
 	{
 		Buffer.SetPurgeableState(mtlpp::PurgeableState::NonVolatile);
 	}
@@ -854,18 +813,15 @@ void FAGXDeviceContext::ReleaseBuffer(FAGXBuffer& Buffer)
 
 struct FAGXRHICommandUpdateFence final : public FRHICommand<FAGXRHICommandUpdateFence>
 {
-	TRefCountPtr<FAGXFence> Fence;
 	uint32 Num;
 	
-	FORCEINLINE_DEBUGGABLE FAGXRHICommandUpdateFence(TRefCountPtr<FAGXFence>& InFence, uint32 InNum)
-	: Fence(InFence)
-	, Num(InNum)
+	FORCEINLINE_DEBUGGABLE FAGXRHICommandUpdateFence(uint32 InNum)
+	: Num(InNum)
 	{
 	}
 	
 	void Execute(FRHICommandListBase& CmdList)
 	{
-		GetAGXDeviceContext().SetParallelPassFences(nil, Fence);
 		GetAGXDeviceContext().FinishFrame(true);
 		GetAGXDeviceContext().BeginParallelRenderCommandEncoding(Num);
 	}
@@ -886,11 +842,6 @@ FAGXRHICommandContext* FAGXDeviceContext::AcquireContext(int32 NewIndex, int32 N
 	}
 	check(Context);
 	
-	if(ParallelFences.Num() < NewNum)
-	{
-		ParallelFences.AddDefaulted(NewNum - ParallelFences.Num());
-	}
-	
 	NSString* StartLabel = nil;
 	NSString* EndLabel = nil;
 #if METAL_DEBUG_OPTIONS
@@ -898,23 +849,16 @@ FAGXRHICommandContext* FAGXDeviceContext::AcquireContext(int32 NewIndex, int32 N
 	EndLabel = [NSString stringWithFormat:@"End Parallel Context Index %d Num %d", NewIndex, NewNum];
 #endif
 	
-	TRefCountPtr<FAGXFence> StartFence = (NewIndex == 0 ? CommandList.GetCommandQueue().CreateFence(StartLabel) : ParallelFences[NewIndex - 1].GetReference());
-	TRefCountPtr<FAGXFence> EndFence = CommandList.GetCommandQueue().CreateFence(EndLabel);
-	ParallelFences[NewIndex] = EndFence;
-	
-	// Give the context the fences so that we can properly order the parallel contexts.
-	Context->GetInternalContext().SetParallelPassFences(StartFence, EndFence);
-	
 	if (NewIndex == 0)
 	{
 		if (FRHICommandListExecutor::GetImmediateCommandList().Bypass() || !IsRunningRHIInSeparateThread())
 		{
-			FAGXRHICommandUpdateFence UpdateCommand(StartFence, NewNum);
+			FAGXRHICommandUpdateFence UpdateCommand(NewNum);
 			UpdateCommand.Execute(FRHICommandListExecutor::GetImmediateCommandList());
 		}
 		else
 		{
-			new (FRHICommandListExecutor::GetImmediateCommandList().AllocCommand<FAGXRHICommandUpdateFence>()) FAGXRHICommandUpdateFence(StartFence, NewNum);
+			new (FRHICommandListExecutor::GetImmediateCommandList().AllocCommand<FAGXRHICommandUpdateFence>()) FAGXRHICommandUpdateFence(NewNum);
 			FRHICommandListExecutor::GetImmediateCommandList().RHIThreadFence(true);
 			FRHICommandListExecutor::GetImmediateCommandList().ImmediateFlush(EImmediateFlushType::DispatchToRHIThread);
 		}
@@ -1018,14 +962,12 @@ bool FAGXDeviceContext::ValidateIsInactiveBuffer(FAGXBuffer const& Buffer)
 uint32 FAGXContext::CurrentContextTLSSlot = FPlatformTLS::AllocTlsSlot();
 #endif
 
-FAGXContext::FAGXContext(FAGXCommandQueue& Queue, bool const bIsImmediate)
+FAGXContext::FAGXContext(FAGXCommandQueue& Queue, bool bIsImmediate)
 : CommandQueue(Queue)
 , CommandList(Queue, bIsImmediate)
 , StateCache(bIsImmediate)
 , RenderPass(CommandList, StateCache)
 , QueryBuffer(new FAGXQueryBufferPool(this))
-, StartFence(nil)
-, EndFence(nil)
 , NumParallelContextsInPass(0)
 {
 	// create a semaphore for multi-buffering the command buffer
@@ -1089,24 +1031,7 @@ void FAGXContext::MakeCurrent(FAGXContext* Context)
 }
 #endif
 
-void FAGXContext::SetParallelPassFences(FAGXFence* Start, FAGXFence* End)
-{
-	check(!StartFence && !EndFence);
-	StartFence = Start;
-	EndFence = End;
-}
-
-TRefCountPtr<FAGXFence> const& FAGXContext::GetParallelPassStartFence(void) const
-{
-	return StartFence;
-}
-
-TRefCountPtr<FAGXFence> const& FAGXContext::GetParallelPassEndFence(void) const
-{
-	return EndFence;
-}
-
-void FAGXContext::InitFrame(bool const bImmediateContext, uint32 Index, uint32 Num)
+void FAGXContext::InitFrame(bool bImmediateContext, uint32 Index, uint32 Num)
 {
 #if ENABLE_METAL_GPUPROFILE
 	FPlatformTLS::SetTlsValue(CurrentContextTLSSlot, this);
@@ -1129,23 +1054,14 @@ void FAGXContext::InitFrame(bool const bImmediateContext, uint32 Index, uint32 N
 	RenderPass.ShrinkRingBuffers();
 	
 	// Begin the render pass frame.
-	RenderPass.Begin(StartFence);
-	
-	// Unset the start fence, the render-pass owns it and we can consider it encoded now!
-	StartFence = nullptr;
+	RenderPass.Begin();
 	
 	// make sure first SetRenderTarget goes through
 	StateCache.InvalidateRenderTargets();
 }
 
-void FAGXContext::FinishFrame(bool const bImmediateContext)
+void FAGXContext::FinishFrame(bool bImmediateContext)
 {
-	// Ensure that we update the end fence for parallel contexts.
-	RenderPass.Update(EndFence);
-	
-	// Unset the end fence, the render-pass owns it and we can consider it encoded now!
-	EndFence = nullptr;
-	
 	// End the render pass
 	RenderPass.End();
 	
@@ -1402,7 +1318,7 @@ bool FAGXContext::PrepareToDraw(uint32 PrimitiveType)
 	return true;
 }
 
-void FAGXContext::SetRenderPassInfo(const FRHIRenderPassInfo& RenderTargetsInfo, bool const bRestart)
+void FAGXContext::SetRenderPassInfo(const FRHIRenderPassInfo& RenderTargetsInfo, bool bRestart)
 {
 	if (CommandList.IsParallel())
 	{
@@ -1611,7 +1527,7 @@ void FAGXContext::AsyncGenerateMipmapsForTexture(FAGXTexture const& Texture)
 	RenderPass.AsyncGenerateMipmapsForTexture(Texture);
 }
 
-void FAGXContext::SubmitAsyncCommands(mtlpp::CommandBufferHandler ScheduledHandler, mtlpp::CommandBufferHandler CompletionHandler, bool const bWait)
+void FAGXContext::SubmitAsyncCommands(mtlpp::CommandBufferHandler ScheduledHandler, mtlpp::CommandBufferHandler CompletionHandler, bool bWait)
 {
 	RenderPass.AddAsyncCommandBufferHandlers(ScheduledHandler, CompletionHandler);
 	if (bWait)
@@ -1701,8 +1617,7 @@ void FAGXDeviceContext::SetParallelRenderPassDescriptor(FRHIRenderPassInfo const
 
 	if (!RenderPass.IsWithinParallelPass())
 	{
-		RenderPass.Begin(EndFence);
-		EndFence = nullptr;
+		RenderPass.Begin();
 		StateCache.InvalidateRenderTargets();
 		SetRenderPassInfo(TargetInfo, false);
 	}
@@ -1724,8 +1639,7 @@ void FAGXDeviceContext::EndParallelRenderCommandEncoding(void)
 	if (FPlatformAtomics::InterlockedDecrement(&ActiveParallelContexts) == 0)
 	{
 		RenderPass.EndRenderPass();
-		RenderPass.Begin(StartFence, true);
-		StartFence = nullptr;
+		RenderPass.Begin(true);
 		FPlatformAtomics::AtomicStore(&NumParallelContextsInPass, 0);
 	}
 }
@@ -1775,12 +1689,6 @@ public:
 		{
 			check(Index == NewIndex && Num == NewNum);
 			
-			if (Index == (Num - 1))
-			{
-				TRefCountPtr<FAGXFence> Fence = CmdContext->GetInternalContext().GetParallelPassEndFence();
-				GetAGXDeviceContext().SetParallelPassFences(Fence, nil);
-			}
-
 			CmdContext->GetInternalContext().FinishFrame(false);
 			GetAGXDeviceContext().EndParallelRenderCommandEncoding();
 
diff --git a/Engine/Source/Runtime/Apple/AGXRHI/Private/AGXContext.h b/Engine/Source/Runtime/Apple/AGXRHI/Private/AGXContext.h
--- a/Engine/Source/Runtime/Apple/AGXRHI/Private/AGXContext.h
+++ b/Engine/Source/Runtime/Apple/AGXRHI/Private/AGXContext.h
@@ -29,7 +29,7 @@ class FAGXContext
 {
 	friend class FAGXCommandContextContainer;
 public:
-	FAGXContext(FAGXCommandQueue& Queue, bool const bIsImmediate);
+	FAGXContext(FAGXCommandQueue& Queue, bool bIsImmediate);
 	virtual ~FAGXContext();
 	
 	FAGXCommandQueue& GetCommandQueue();
@@ -51,7 +51,7 @@ public:
 	/**
 	 * Set the color, depth and stencil render targets, and then make the new command buffer/encoder
 	 */
-	void SetRenderPassInfo(const FRHIRenderPassInfo& RenderTargetsInfo, bool const bRestart = false);
+	void SetRenderPassInfo(const FRHIRenderPassInfo& RenderTargetsInfo, bool bRestart = false);
 	
 	/**
 	 * Allocate from a dynamic ring buffer - by default align to the allowed alignment for offset field when setting buffers
@@ -99,7 +99,7 @@ public:
 	
     void AsyncGenerateMipmapsForTexture(FAGXTexture const& Texture);
     
-	void SubmitAsyncCommands(mtlpp::CommandBufferHandler ScheduledHandler, mtlpp::CommandBufferHandler CompletionHandler, bool const bWait);
+	void SubmitAsyncCommands(mtlpp::CommandBufferHandler ScheduledHandler, mtlpp::CommandBufferHandler CompletionHandler, bool bWait);
 	
 	void SynchronizeTexture(FAGXTexture const& Texture, uint32 Slice, uint32 Level);
 	
@@ -118,12 +118,8 @@ public:
 	static FAGXContext* GetCurrentContext();
 #endif
 	
-	void SetParallelPassFences(FAGXFence* Start, FAGXFence* End);
-	TRefCountPtr<FAGXFence> const& GetParallelPassStartFence(void) const;
-	TRefCountPtr<FAGXFence> const& GetParallelPassEndFence(void) const;
-	
-	void InitFrame(bool const bImmediateContext, uint32 Index, uint32 Num);
-	void FinishFrame(bool const bImmediateContext);
+	void InitFrame(bool bImmediateContext, uint32 Index, uint32 Num);
+	void FinishFrame(bool bImmediateContext);
 
 	// Track Write->Read transitions for TBDR Fragment->Verex fencing
 	void TransitionResource(FRHIUnorderedAccessView* InResource);
@@ -151,12 +147,6 @@ protected:
 	/** A pool of buffers for writing visibility query results. */
 	TSharedPtr<FAGXQueryBufferPool, ESPMode::ThreadSafe> QueryBuffer;
 	
-	/** Initial fence to wait on for parallel contexts */
-	TRefCountPtr<FAGXFence> StartFence;
-	
-	/** Fence to update at the end for parallel contexts */
-	TRefCountPtr<FAGXFence> EndFence;
-	
 #if ENABLE_METAL_GPUPROFILE
 	/** the slot to store a per-thread context ref */
 	static uint32 CurrentContextTLSSlot;
@@ -197,10 +187,9 @@ public:
 	void ReleaseObject(id Object);
 	void ReleaseTexture(FAGXSurface* Surface, FAGXTexture& Texture);
 	void ReleaseTexture(FAGXTexture& Texture);
-	void ReleaseFence(FAGXFence* Fence);
 	
 	void BeginFrame();
-	void FlushFreeList(bool const bFlushFences = true);
+	void FlushFreeList();
 	void ClearFreeList();
 	void DrainHeap();
 	void EndFrame();
@@ -271,8 +260,6 @@ private:
 	/** Free lists for releasing objects only once it is safe to do so */
 	TSet<FAGXBuffer> UsedBuffers;
 	TSet<FAGXTexture> UsedTextures;
-	TSet<FAGXFence*> UsedFences;
-	TLockFreePointerListLIFO<FAGXFence> FenceFreeList;
 	TSet<id> ObjectFreeList;
 	struct FAGXDelayedFreeList
 	{
@@ -280,7 +267,6 @@ private:
 		TArray<mtlpp::CommandBufferFence> Fences;
 		TSet<FAGXBuffer> UsedBuffers;
 		TSet<FAGXTexture> UsedTextures;
-		TSet<FAGXFence*> FenceFreeList;
 		TSet<id> ObjectFreeList;
 #if METAL_DEBUG_OPTIONS
 		int32 DeferCount;
@@ -295,9 +281,6 @@ private:
 	TMap<FAGXRHIBuffer*, FAGXFrameAllocator::AllocationEntry> OutstandingLocks;
 	
 #if METAL_DEBUG_OPTIONS
-	/** The list of fences for the current frame */
-	TArray<FAGXFence*> FrameFences;
-    
     FCriticalSection ActiveBuffersMutex;
     
     /** These are the active buffers that cannot be CPU modified */
@@ -307,9 +290,6 @@ private:
 	/** Free-list of contexts for parallel encoding */
 	TLockFreePointerListLIFO<FAGXRHICommandContext> ParallelContexts;
 	
-	/** Fences for parallel execution */
-	TArray<TRefCountPtr<FAGXFence>> ParallelFences;
-	
 	/** Critical section for FreeList */
 	FCriticalSection FreeListMutex;
 	
diff --git a/Engine/Source/Runtime/Apple/AGXRHI/Private/AGXDebugCommandEncoder.cpp b/Engine/Source/Runtime/Apple/AGXRHI/Private/AGXDebugCommandEncoder.cpp
deleted file mode 100644
diff --git a/Engine/Source/Runtime/Apple/AGXRHI/Private/AGXDebugCommandEncoder.h b/Engine/Source/Runtime/Apple/AGXRHI/Private/AGXDebugCommandEncoder.h
deleted file mode 100644
diff --git a/Engine/Source/Runtime/Apple/AGXRHI/Private/AGXFence.cpp b/Engine/Source/Runtime/Apple/AGXRHI/Private/AGXFence.cpp
deleted file mode 100644
diff --git a/Engine/Source/Runtime/Apple/AGXRHI/Private/AGXFence.h b/Engine/Source/Runtime/Apple/AGXRHI/Private/AGXFence.h
deleted file mode 100644
diff --git a/Engine/Source/Runtime/Apple/AGXRHI/Private/AGXLLM.cpp b/Engine/Source/Runtime/Apple/AGXRHI/Private/AGXLLM.cpp
--- a/Engine/Source/Runtime/Apple/AGXRHI/Private/AGXLLM.cpp
+++ b/Engine/Source/Runtime/Apple/AGXRHI/Private/AGXLLM.cpp
@@ -274,33 +274,6 @@ void AGXLLM::LogAllocBuffer(mtlpp::Buffer const& Buffer)
 	}
 }
 
-void AGXLLM::LogAllocHeap(mtlpp::Heap const& Heap)
-{
-	void* Ptr = (void*)Heap.GetPtr();
-	uint64 Size = Heap.GetSize();
-	
-	INC_MEMORY_STAT_BY(STAT_AGXHeapMemory, Size);
-	INC_DWORD_STAT(STAT_AGXHeapCount);
-	
-	LLM(FLowLevelMemTracker::Get().OnLowLevelAlloc(ELLMTracker::Platform, Ptr, Size, ELLMTag::Untagged, ELLMAllocType::System));
-	// Assign a dealloc handler to untrack the memory - but don't track the dispatch block!
-	{
-		LLM_SCOPED_PAUSE_TRACKING(ELLMAllocType::System);
-		
-		objc_setAssociatedObject(Heap.GetPtr(), (void*)&AGXLLM::LogAllocHeap,
-								 [[[FAGXDeallocHandler alloc] initWithBlock:^{
-			LLM_SCOPE_METAL(ELLMTagAGX::Heaps);
-			LLM_PLATFORM_SCOPE_METAL(ELLMTagAGX::Heaps);
-			
-			LLM(FLowLevelMemTracker::Get().OnLowLevelFree(ELLMTracker::Platform, Ptr, ELLMAllocType::System));
-			
-			DEC_MEMORY_STAT_BY(STAT_AGXHeapMemory, Size);
-			DEC_DWORD_STAT(STAT_AGXHeapCount);
-		}] autorelease],
-		OBJC_ASSOCIATION_RETAIN);
-	}
-}
-
 void AGXLLM::LogAliasTexture(mtlpp::Texture const& Texture)
 {
 	objc_setAssociatedObject(Texture.GetPtr(), (void*)&AGXLLM::LogAllocTexture, nil, OBJC_ASSOCIATION_RETAIN);
diff --git a/Engine/Source/Runtime/Apple/AGXRHI/Private/AGXLLM.h b/Engine/Source/Runtime/Apple/AGXRHI/Private/AGXLLM.h
--- a/Engine/Source/Runtime/Apple/AGXRHI/Private/AGXLLM.h
+++ b/Engine/Source/Runtime/Apple/AGXRHI/Private/AGXLLM.h
@@ -48,7 +48,6 @@ namespace AGXLLM
 {
 	void LogAllocTexture(mtlpp::TextureDescriptor const& Desc, mtlpp::Texture const& Texture);
 	void LogAllocBuffer(mtlpp::Buffer const& Buffer);
-	void LogAllocHeap(mtlpp::Heap const& Heap);
 	
 	void LogAliasTexture(mtlpp::Texture const& Texture);
 	void LogAliasBuffer(mtlpp::Buffer const& Buffer);
diff --git a/Engine/Source/Runtime/Apple/AGXRHI/Private/AGXPipeline.cpp b/Engine/Source/Runtime/Apple/AGXRHI/Private/AGXPipeline.cpp
--- a/Engine/Source/Runtime/Apple/AGXRHI/Private/AGXPipeline.cpp
+++ b/Engine/Source/Runtime/Apple/AGXRHI/Private/AGXPipeline.cpp
@@ -13,7 +13,6 @@
 #include "MetalShaderResources.h"
 #include "AGXProfiler.h"
 #include "AGXCommandQueue.h"
-#include "AGXCommandBuffer.h"
 #include "RenderUtils.h"
 #include "Misc/ScopeRWLock.h"
 #include "HAL/PThreadEvent.h"
@@ -36,122 +35,6 @@ static uint32 BlendBitOffsets[] = { Offset_BlendState0, Offset_BlendState1, Offs
 static uint32 RTBitOffsets[] = { Offset_RenderTargetFormat0, Offset_RenderTargetFormat1, Offset_RenderTargetFormat2, Offset_RenderTargetFormat3, Offset_RenderTargetFormat4, Offset_RenderTargetFormat5, Offset_RenderTargetFormat6, Offset_RenderTargetFormat7 };
 static_assert(Offset_RasterEnd < 64 && Offset_End < 128, "Offset_RasterEnd must be < 64 && Offset_End < 128");
 
-static float RoundUpNearestEven(const float f)
-{
-	const float ret = FMath::CeilToFloat(f);
-	const float isOdd = (float)(((int)ret) & 1);
-	return ret + isOdd;
-}
-
-// A tile-based or vertex-based debug shader for trying to emulate Aftermath style failure reporting
-static NSString* GAGXDebugShader = @"#include <metal_stdlib>\n"
-"#include <metal_compute>\n"
-"\n"
-"using namespace metal;\n"
-"\n"
-"struct DebugInfo\n"
-"{\n"
-"   uint CmdBuffIndex;\n"
-"	uint EncoderIndex;\n"
-"   uint ContextIndex;\n"
-"   uint CommandIndex;\n"
-"   uint CommandBuffer[2];\n"
-"	uint PSOSignature[4];\n"
-"};\n"
-"\n"
-#if !PLATFORM_MAC
-"// Executes once per-tile\n"
-"kernel void Main_Debug(constant DebugInfo *debugTable [[ buffer(0) ]], device DebugInfo* debugBuffer [[ buffer(1) ]], uint2 threadgroup_position_in_grid [[ threadgroup_position_in_grid ]], uint2 threadgroups_per_grid [[ threadgroups_per_grid ]])\n"
-"{\n"
-"	// Write Pass, Draw indices\n"
-"	// Write Vertex+Fragment PSO sig (in form VertexLen, VertexCRC, FragLen, FragCRC)\n"
-"   uint tile_index = threadgroup_position_in_grid.x + (threadgroup_position_in_grid.y * threadgroups_per_grid.x);"
-"	debugBuffer[tile_index] = debugTable[0];\n"
-"}";
-#else
-"// Executes once as a point draw call\n"
-"vertex void Main_Debug(constant DebugInfo *debugTable [[ buffer(0) ]], device DebugInfo* debugBuffer [[ buffer(1) ]])\n"
-"{\n"
-"	// Write Pass, Draw indices\n"
-"	// Write Vertex+Fragment PSO sig (in form VertexLen, VertexCRC, FragLen, FragCRC)\n"
-"	debugBuffer[0] = debugTable[0];\n"
-"}";
-#endif
-
-// A compute debug shader for trying to emulate Aftermath style failure reporting
-static NSString* GAGXDebugMarkerComputeShader = @"#include <metal_stdlib>\n"
-"#include <metal_compute>\n"
-"\n"
-"using namespace metal;\n"
-"\n"
-"struct DebugInfo\n"
-"{\n"
-"   uint CmdBuffIndex;\n"
-"	uint EncoderIndex;\n"
-"   uint ContextIndex;\n"
-"   uint CommandIndex;\n"
-"   uint CommandBuffer[2];\n"
-"	uint PSOSignature[4];\n"
-"};\n"
-"\n"
-"// Executes once\n"
-"kernel void Main_Debug(constant DebugInfo *debugTable [[ buffer(0) ]], device DebugInfo* debugBuffer [[ buffer(1) ]])\n"
-"{\n"
-"	// Write Pass, Draw indices\n"
-"	// Write Vertex+Fragment PSO sig (in form VertexLen, VertexCRC, FragLen, FragCRC)\n"
-"	debugBuffer[0] = debugTable[0];\n"
-"}";
-
-struct FAGXHelperFunctions
-{
-    mtlpp::Library DebugShadersLib;
-    mtlpp::Function DebugFunc;
-	
-	mtlpp::Library DebugComputeShadersLib;
-	mtlpp::Function DebugComputeFunc;
-	mtlpp::ComputePipelineState DebugComputeState;
-	
-    FAGXHelperFunctions()
-    {
-#if !PLATFORM_TVOS
-        if (GAGXCommandBufferDebuggingEnabled)
-        {
-            mtlpp::CompileOptions CompileOptions;
-            ns::AutoReleasedError Error;
-
-			DebugShadersLib = GMtlppDevice.NewLibrary(GAGXDebugShader, CompileOptions, &Error);
-            DebugFunc = DebugShadersLib.NewFunction(@"Main_Debug");
-			
-			DebugComputeShadersLib = GMtlppDevice.NewLibrary(GAGXDebugMarkerComputeShader, CompileOptions, &Error);
-			DebugComputeFunc = DebugComputeShadersLib.NewFunction(@"Main_Debug");
-			
-			DebugComputeState = GMtlppDevice.NewComputePipelineState(DebugComputeFunc, &Error);
-        }
-#endif
-    }
-    
-    static FAGXHelperFunctions& Get()
-    {
-        static FAGXHelperFunctions sSelf;
-        return sSelf;
-    }
-    
-    mtlpp::Function GetDebugFunction()
-    {
-        return DebugFunc;
-    }
-	
-	mtlpp::ComputePipelineState GetDebugComputeState()
-	{
-		return DebugComputeState;
-	}
-};
-
-mtlpp::ComputePipelineState AGXGetMetalDebugComputeState()
-{
-	return FAGXHelperFunctions::Get().GetDebugComputeState();
-}
-
 struct FAGXGraphicsPipelineKey
 {
 	FAGXRenderPipelineHash RenderPipelineHash;
@@ -443,10 +326,8 @@ APPLE_PLATFORM_OBJECT_ALLOC_OVERRIDES(FAGXShaderPipeline)
 	{
 		RenderPipelineReflection = mtlpp::RenderPipelineReflection(nil);
 		ComputePipelineReflection = mtlpp::ComputePipelineReflection(nil);
-		StreamPipelineReflection = mtlpp::RenderPipelineReflection(nil);
 #if METAL_DEBUG_OPTIONS
 		RenderDesc = mtlpp::RenderPipelineDescriptor(nil);
-		StreamDesc = mtlpp::RenderPipelineDescriptor(nil);
 		ComputeDesc = mtlpp::ComputePipelineDescriptor(nil);
 #endif
 	}
@@ -457,8 +338,8 @@ APPLE_PLATFORM_OBJECT_ALLOC_OVERRIDES(FAGXShaderPipeline)
 {
 	if (RenderPipelineReflection)
 	{
-		[self initResourceMask:EAGXShaderVertex];
-		[self initResourceMask:EAGXShaderFragment];
+		[self initResourceMask:SF_Vertex];
+		[self initResourceMask:SF_Pixel];
 		
 		if (AGXSafeGetRuntimeDebuggingLevel() < EAGXDebugLevelValidation)
 		{
@@ -467,29 +348,20 @@ APPLE_PLATFORM_OBJECT_ALLOC_OVERRIDES(FAGXShaderPipeline)
 	}
 	if (ComputePipelineReflection)
 	{
-		[self initResourceMask:EAGXShaderCompute];
+		[self initResourceMask:SF_Compute];
 		
 		if (AGXSafeGetRuntimeDebuggingLevel() < EAGXDebugLevelValidation)
 		{
 			ComputePipelineReflection = mtlpp::ComputePipelineReflection(nil);
 		}
 	}
-	if (StreamPipelineReflection)
-	{
-		[self initResourceMask:EAGXShaderStream];
-		
-		if (AGXSafeGetRuntimeDebuggingLevel() < EAGXDebugLevelValidation)
-		{
-			StreamPipelineReflection = mtlpp::RenderPipelineReflection(nil);
-		}
-	}
 }
-- (void)initResourceMask:(EAGXShaderFrequency)Frequency
+- (void)initResourceMask:(EShaderFrequency)Frequency
 {
 	NSArray<MTLArgument*>* Arguments = nil;
 	switch(Frequency)
 	{
-		case EAGXShaderVertex:
+		case SF_Vertex:
 		{
 			MTLRenderPipelineReflection* Reflection = RenderPipelineReflection;
 			check(Reflection);
@@ -497,7 +369,7 @@ APPLE_PLATFORM_OBJECT_ALLOC_OVERRIDES(FAGXShaderPipeline)
 			Arguments = Reflection.vertexArguments;
 			break;
 		}
-		case EAGXShaderFragment:
+		case SF_Pixel:
 		{
 			MTLRenderPipelineReflection* Reflection = RenderPipelineReflection;
 			check(Reflection);
@@ -505,7 +377,7 @@ APPLE_PLATFORM_OBJECT_ALLOC_OVERRIDES(FAGXShaderPipeline)
 			Arguments = Reflection.fragmentArguments;
 			break;
 		}
-		case EAGXShaderCompute:
+		case SF_Compute:
 		{
 			MTLComputePipelineReflection* Reflection = ComputePipelineReflection;
 			check(Reflection);
@@ -513,14 +385,6 @@ APPLE_PLATFORM_OBJECT_ALLOC_OVERRIDES(FAGXShaderPipeline)
 			Arguments = Reflection.arguments;
 			break;
 		}
-		case EAGXShaderStream:
-		{
-			MTLRenderPipelineReflection* Reflection = StreamPipelineReflection;
-			check(Reflection);
-			
-			Arguments = Reflection.vertexArguments;
-			break;
-		}
 		default:
 			check(false);
 			break;
@@ -795,7 +659,7 @@ static bool ConfigureRenderPipelineDescriptor(mtlpp::RenderPipelineDescriptor& R
 		FMetalShaderBindings& VertexBindings = VertexShader->Bindings;
 		int8 VertexSideTable = VertexShader->SideTableBinding;
 		{
-			uint32 ImmutableBuffers = VertexBindings.ConstantBuffers | VertexBindings.ArgumentBuffers;
+			uint32 ImmutableBuffers = VertexBindings.ConstantBuffers;
 			while(ImmutableBuffers)
 			{
 				uint32 Index = __builtin_ctz(ImmutableBuffers);
@@ -817,7 +681,7 @@ static bool ConfigureRenderPipelineDescriptor(mtlpp::RenderPipelineDescriptor& R
 		if (PixelShader)
 		{
 			ns::AutoReleased<ns::Array<mtlpp::PipelineBufferDescriptor>> FragmentPipelineBuffers = RenderPipelineDesc.GetFragmentBuffers();
-			uint32 ImmutableBuffers = PixelShader->Bindings.ConstantBuffers | PixelShader->Bindings.ArgumentBuffers;
+			uint32 ImmutableBuffers = PixelShader->Bindings.ConstantBuffers;
 			while(ImmutableBuffers)
 			{
 				uint32 Index = __builtin_ctz(ImmutableBuffers);
@@ -928,27 +792,6 @@ static FAGXShaderPipeline* CreateMTLRenderPipeline(bool const bSync, FAGXGraphic
         Pipeline->FragmentSource = PixelShader ? PixelShader->GetSourceCode() : nil;
 #endif
 		
-#if !PLATFORM_TVOS
-		if (GAGXCommandBufferDebuggingEnabled)
-		{
-#if PLATFORM_MAC
-			DebugPipelineDesc.SetVertexFunction(FAGXHelperFunctions::Get().GetDebugFunction());
-			DebugPipelineDesc.SetRasterizationEnabled(false);
-#else
-			DebugPipelineDesc.SetTileFunction(FAGXHelperFunctions::Get().GetDebugFunction());
-			DebugPipelineDesc.SetRasterSampleCount(RenderPipelineDesc.GetSampleCount());
-			DebugPipelineDesc.SetThreadgroupSizeMatchesTileSize(false);
-#endif
-#if ENABLE_METAL_GPUPROFILE
-			DebugPipelineDesc.SetLabel(@"Main_Debug");
-#endif
-
-			ns::AutoReleasedError RenderError;
-			METAL_GPUPROFILE(FAGXScopedCPUStats CPUStat(FString::Printf(TEXT("NewDebugPipeline: %s"), TEXT("")/**FString([RenderPipelineDesc.GetPtr() description])*/)));
-			Pipeline->DebugPipelineState = GMtlppDevice.NewRenderPipelineState(DebugPipelineDesc, mtlpp::PipelineOption::NoPipelineOption, Reflection, nullptr);
-		}
-#endif
-        
 #if METAL_DEBUG_OPTIONS
         if (GFrameCounter > 3)
         {
diff --git a/Engine/Source/Runtime/Apple/AGXRHI/Private/AGXPipeline.h b/Engine/Source/Runtime/Apple/AGXRHI/Private/AGXPipeline.h
--- a/Engine/Source/Runtime/Apple/AGXRHI/Private/AGXPipeline.h
+++ b/Engine/Source/Runtime/Apple/AGXRHI/Private/AGXPipeline.h
@@ -3,10 +3,6 @@
 #pragma once
 
 #include "AGXRHIPrivate.h"
-#if METAL_DEBUG_OPTIONS
-#include "AGXDebugCommandEncoder.h"
-#endif
-
 #include "ShaderPipelineCache.h"
 
 enum EAGXPipelineHashBits
@@ -71,24 +67,20 @@ private:
 @public
 	mtlpp::RenderPipelineState RenderPipelineState;
 	mtlpp::ComputePipelineState ComputePipelineState;
-	mtlpp::RenderPipelineState StreamPipelineState;
-	mtlpp::RenderPipelineState DebugPipelineState;
-	TArray<uint32> BufferDataSizes[EAGXShaderStagesNum];
-	TMap<uint8, uint8> TextureTypes[EAGXShaderStagesNum];
-	FAGXDebugShaderResourceMask ResourceMask[EAGXShaderStagesNum];
+	TArray<uint32> BufferDataSizes[SF_NumStandardFrequencies];
+	TMap<uint8, uint8> TextureTypes[SF_NumStandardFrequencies];
+	FAGXDebugShaderResourceMask ResourceMask[SF_NumStandardFrequencies];
 	mtlpp::RenderPipelineReflection RenderPipelineReflection;
-	mtlpp::RenderPipelineReflection StreamPipelineReflection;
 	mtlpp::ComputePipelineReflection ComputePipelineReflection;
 #if METAL_DEBUG_OPTIONS
 	ns::String VertexSource;
 	ns::String FragmentSource;
 	ns::String ComputeSource;
 	mtlpp::RenderPipelineDescriptor RenderDesc;
-	mtlpp::RenderPipelineDescriptor StreamDesc;
 	mtlpp::ComputePipelineDescriptor ComputeDesc;
 #endif
 }
 - (instancetype)init;
 - (void)initResourceMask;
-- (void)initResourceMask:(EAGXShaderFrequency)Frequency;
+- (void)initResourceMask:(EShaderFrequency)Frequency;
 @end
diff --git a/Engine/Source/Runtime/Apple/AGXRHI/Private/AGXProfiler.cpp b/Engine/Source/Runtime/Apple/AGXRHI/Private/AGXProfiler.cpp
--- a/Engine/Source/Runtime/Apple/AGXRHI/Private/AGXProfiler.cpp
+++ b/Engine/Source/Runtime/Apple/AGXRHI/Private/AGXProfiler.cpp
@@ -4,7 +4,6 @@
 #include "AGXProfiler.h"
 #include "EngineGlobals.h"
 #include "StaticBoundShaderState.h"
-#include "AGXCommandBuffer.h"
 #include "HAL/FileManager.h"
 
 DEFINE_STAT(STAT_AGXUniformMemAlloc);
diff --git a/Engine/Source/Runtime/Apple/AGXRHI/Private/AGXRHI.cpp b/Engine/Source/Runtime/Apple/AGXRHI/Private/AGXRHI.cpp
--- a/Engine/Source/Runtime/Apple/AGXRHI/Private/AGXRHI.cpp
+++ b/Engine/Source/Runtime/Apple/AGXRHI/Private/AGXRHI.cpp
@@ -23,6 +23,8 @@
 #include "Engine/RendererSettings.h"
 #include "AGXTransitionData.h"
 
+#include COMPILED_PLATFORM_HEADER(PlatformAGXConfig.h)
+
 DEFINE_LOG_CATEGORY(LogAGX)
 
 bool GIsAGXInitialized = false;
@@ -49,6 +51,10 @@ static TAutoConsoleVariable<int32> CVarIntelUseRHIThread(
 													ECVF_Default | ECVF_RenderThreadSafe
 													);
 
+id<MTLDevice> AGXUtilGetDevice()
+{
+	return GMtlDevice;
+}
 
 static void ValidateTargetedRHIFeatureLevelExists(EShaderPlatform Platform)
 {
@@ -181,10 +187,6 @@ FAGXDynamicRHI::FAGXDynamicRHI(ERHIFeatureLevel::Type RequestedFeatureLevel)
 	// Metal always needs a render target to render with fragment shaders!
 	GRHIRequiresRenderTargetForPixelShaderUAVs = true;
 
-	//@todo-rco: Query name from API
-	GRHIAdapterName = TEXT("AGX");
-	GRHIVendorId = 1; // non-zero to avoid asserts
-
 	bool const bRequestedFeatureLevel = (RequestedFeatureLevel != ERHIFeatureLevel::Num);
 	bool bSupportsPointLights = false;
 	bool bSupportsRHIThread = false;
@@ -279,7 +281,8 @@ FAGXDynamicRHI::FAGXDynamicRHI(ERHIFeatureLevel::Type RequestedFeatureLevel)
 	GShaderPlatformForFeatureLevel[ERHIFeatureLevel::SM5] = (GMaxRHIFeatureLevel >= ERHIFeatureLevel::SM5) ? GMaxRHIShaderPlatform : SP_NumPlatforms;
 
 #else // PLATFORM_IOS
-                
+	FPlatformAGXConfig::PopulateRHIGlobals();
+
 	uint32 DeviceIndex = ((FAGXDeviceContext*)ImmediateContext.Context)->GetDeviceIndex();
 	
 	TArray<FMacPlatformMisc::FGPUDescriptor> const& GPUs = FPlatformMisc::GetGPUDescriptors();
@@ -289,8 +292,6 @@ FAGXDynamicRHI::FAGXDynamicRHI(ERHIFeatureLevel::Type RequestedFeatureLevel)
 	bool bSupportsD24S8 = false;
 	bool bSupportsD16 = false;
 	
-	GRHIAdapterName = FString([GMtlDevice name]);
-	
 	// However they don't all support other features depending on the version of the OS.
 	bool bSupportsTiledReflections = false;
 	bool bSupportsDistanceFields = false;
@@ -341,7 +342,6 @@ FAGXDynamicRHI::FAGXDynamicRHI(ERHIFeatureLevel::Type RequestedFeatureLevel)
 	else if(GRHIAdapterName.Contains("Apple"))
 	{
 		bSupportsPointLights = true;
-		GRHIVendorId = 0x106B;
 		bSupportsTiledReflections = true;
 		bSupportsDistanceFields = true;
 		bSupportsRHIThread = true;
@@ -561,6 +561,8 @@ FAGXDynamicRHI::FAGXDynamicRHI(ERHIFeatureLevel::Type RequestedFeatureLevel)
 	GSupportsWideMRT = true;
 	GSupportsSeparateRenderTargetBlendState = (GMaxRHIFeatureLevel >= ERHIFeatureLevel::SM5);
 
+	// TODO: GSupportsTransientResourceAliasing requires support for MTLHeap and MTLFence
+
 	GRHISupportsPipelineFileCache = true;
 
 #if PLATFORM_MAC
@@ -662,42 +664,55 @@ FAGXDynamicRHI::FAGXDynamicRHI(ERHIFeatureLevel::Type RequestedFeatureLevel)
 	GAGXBufferFormats[PF_R5G6B5_UNORM         ] = { mtlpp::PixelFormat::Invalid, (uint8)EMetalBufferFormat::R5G6B5Unorm };
 	GAGXBufferFormats[PF_B5G5R5A1_UNORM       ] = { mtlpp::PixelFormat::Invalid, (uint8)EMetalBufferFormat::B5G5R5A1Unorm };
 	GAGXBufferFormats[PF_R8G8B8A8             ] = { mtlpp::PixelFormat::RGBA8Unorm, (uint8)EMetalBufferFormat::RGBA8Unorm };
-	GAGXBufferFormats[PF_A8R8G8B8				] = { mtlpp::PixelFormat::RGBA8Unorm, (uint8)EMetalBufferFormat::RGBA8Unorm }; // mtlpp::PixelFormat::BGRA8Unorm/EMetalBufferFormat::BGRA8Unorm,  < We don't support this as a vertex-format so we have code to swizzle in the shader
-	GAGXBufferFormats[PF_BC4					] = { mtlpp::PixelFormat::Invalid, (uint8)EMetalBufferFormat::Unknown };
+	GAGXBufferFormats[PF_A8R8G8B8             ] = { mtlpp::PixelFormat::RGBA8Unorm, (uint8)EMetalBufferFormat::RGBA8Unorm }; // mtlpp::PixelFormat::BGRA8Unorm/EMetalBufferFormat::BGRA8Unorm,  < We don't support this as a vertex-format so we have code to swizzle in the shader
+	GAGXBufferFormats[PF_BC4                  ] = { mtlpp::PixelFormat::Invalid, (uint8)EMetalBufferFormat::Unknown };
 	GAGXBufferFormats[PF_R8G8                 ] = { mtlpp::PixelFormat::RG8Unorm, (uint8)EMetalBufferFormat::RG8Unorm };
-	GAGXBufferFormats[PF_ATC_RGB				] = { mtlpp::PixelFormat::Invalid, (uint8)EMetalBufferFormat::Unknown };
-	GAGXBufferFormats[PF_ATC_RGBA_E			] = { mtlpp::PixelFormat::Invalid, (uint8)EMetalBufferFormat::Unknown };
-	GAGXBufferFormats[PF_ATC_RGBA_I			] = { mtlpp::PixelFormat::Invalid, (uint8)EMetalBufferFormat::Unknown };
-	GAGXBufferFormats[PF_X24_G8				] = { mtlpp::PixelFormat::Invalid, (uint8)EMetalBufferFormat::Unknown };
-	GAGXBufferFormats[PF_ETC1					] = { mtlpp::PixelFormat::Invalid, (uint8)EMetalBufferFormat::Unknown };
-	GAGXBufferFormats[PF_ETC2_RGB				] = { mtlpp::PixelFormat::Invalid, (uint8)EMetalBufferFormat::Unknown };
-	GAGXBufferFormats[PF_ETC2_RGBA			] = { mtlpp::PixelFormat::Invalid, (uint8)EMetalBufferFormat::Unknown };
-	GAGXBufferFormats[PF_R32G32B32A32_UINT	] = { mtlpp::PixelFormat::RGBA32Uint, (uint8)EMetalBufferFormat::RGBA32Uint };
-	GAGXBufferFormats[PF_R16G16_UINT			] = { mtlpp::PixelFormat::RG16Uint, (uint8)EMetalBufferFormat::RG16Uint };
-	GAGXBufferFormats[PF_R32G32_UINT			] = { mtlpp::PixelFormat::RG32Uint, (uint8)EMetalBufferFormat::RG32Uint };
+	GAGXBufferFormats[PF_ATC_RGB              ] = { mtlpp::PixelFormat::Invalid, (uint8)EMetalBufferFormat::Unknown };
+	GAGXBufferFormats[PF_ATC_RGBA_E           ] = { mtlpp::PixelFormat::Invalid, (uint8)EMetalBufferFormat::Unknown };
+	GAGXBufferFormats[PF_ATC_RGBA_I           ] = { mtlpp::PixelFormat::Invalid, (uint8)EMetalBufferFormat::Unknown };
+	GAGXBufferFormats[PF_X24_G8               ] = { mtlpp::PixelFormat::Invalid, (uint8)EMetalBufferFormat::Unknown };
+	GAGXBufferFormats[PF_ETC1                 ] = { mtlpp::PixelFormat::Invalid, (uint8)EMetalBufferFormat::Unknown };
+	GAGXBufferFormats[PF_ETC2_RGB             ] = { mtlpp::PixelFormat::Invalid, (uint8)EMetalBufferFormat::Unknown };
+	GAGXBufferFormats[PF_ETC2_RGBA            ] = { mtlpp::PixelFormat::Invalid, (uint8)EMetalBufferFormat::Unknown };
+	GAGXBufferFormats[PF_R32G32B32A32_UINT    ] = { mtlpp::PixelFormat::RGBA32Uint, (uint8)EMetalBufferFormat::RGBA32Uint };
+	GAGXBufferFormats[PF_R16G16_UINT          ] = { mtlpp::PixelFormat::RG16Uint, (uint8)EMetalBufferFormat::RG16Uint };
+	GAGXBufferFormats[PF_R32G32_UINT          ] = { mtlpp::PixelFormat::RG32Uint, (uint8)EMetalBufferFormat::RG32Uint };
 	GAGXBufferFormats[PF_ASTC_4x4             ] = { mtlpp::PixelFormat::Invalid, (uint8)EMetalBufferFormat::Unknown };
 	GAGXBufferFormats[PF_ASTC_6x6             ] = { mtlpp::PixelFormat::Invalid, (uint8)EMetalBufferFormat::Unknown };
 	GAGXBufferFormats[PF_ASTC_8x8             ] = { mtlpp::PixelFormat::Invalid, (uint8)EMetalBufferFormat::Unknown };
 	GAGXBufferFormats[PF_ASTC_10x10           ] = { mtlpp::PixelFormat::Invalid, (uint8)EMetalBufferFormat::Unknown };
 	GAGXBufferFormats[PF_ASTC_12x12           ] = { mtlpp::PixelFormat::Invalid, (uint8)EMetalBufferFormat::Unknown };
-	GAGXBufferFormats[PF_BC6H					] = { mtlpp::PixelFormat::Invalid, (uint8)EMetalBufferFormat::Unknown };
-	GAGXBufferFormats[PF_BC7					] = { mtlpp::PixelFormat::Invalid, (uint8)EMetalBufferFormat::Unknown };
-	GAGXBufferFormats[PF_R8_UINT				] = { mtlpp::PixelFormat::R8Uint, (uint8)EMetalBufferFormat::R8Uint };
-	GAGXBufferFormats[PF_R8					] = { mtlpp::PixelFormat::R8Unorm, (uint8)EMetalBufferFormat::R8Unorm };
-	GAGXBufferFormats[PF_L8					] = { mtlpp::PixelFormat::Invalid, (uint8)EMetalBufferFormat::R8Unorm };
-	GAGXBufferFormats[PF_XGXR8				] = { mtlpp::PixelFormat::Invalid, (uint8)EMetalBufferFormat::Unknown };
-	GAGXBufferFormats[PF_R8G8B8A8_UINT		] = { mtlpp::PixelFormat::RGBA8Uint, (uint8)EMetalBufferFormat::RGBA8Uint };
-	GAGXBufferFormats[PF_R8G8B8A8_SNORM		] = { mtlpp::PixelFormat::RGBA8Snorm, (uint8)EMetalBufferFormat::RGBA8Snorm };
-	GAGXBufferFormats[PF_R16G16B16A16_UNORM	] = { mtlpp::PixelFormat::RGBA16Unorm, (uint8)EMetalBufferFormat::RGBA16Unorm };
-	GAGXBufferFormats[PF_R16G16B16A16_SNORM	] = { mtlpp::PixelFormat::RGBA16Snorm, (uint8)EMetalBufferFormat::RGBA16Snorm };
-	GAGXBufferFormats[PF_PLATFORM_HDR_0		] = { mtlpp::PixelFormat::Invalid, (uint8)EMetalBufferFormat::Unknown };
-	GAGXBufferFormats[PF_PLATFORM_HDR_1		] = { mtlpp::PixelFormat::Invalid, (uint8)EMetalBufferFormat::Unknown };
-	GAGXBufferFormats[PF_PLATFORM_HDR_2		] = { mtlpp::PixelFormat::Invalid, (uint8)EMetalBufferFormat::Unknown };
-	GAGXBufferFormats[PF_NV12					] = { mtlpp::PixelFormat::Invalid, (uint8)EMetalBufferFormat::Unknown };
+	GAGXBufferFormats[PF_ASTC_4x4_HDR         ] = { mtlpp::PixelFormat::Invalid, (uint8)EMetalBufferFormat::Unknown };
+	GAGXBufferFormats[PF_ASTC_6x6_HDR         ] = { mtlpp::PixelFormat::Invalid, (uint8)EMetalBufferFormat::Unknown };
+	GAGXBufferFormats[PF_ASTC_8x8_HDR         ] = { mtlpp::PixelFormat::Invalid, (uint8)EMetalBufferFormat::Unknown };
+	GAGXBufferFormats[PF_ASTC_10x10_HDR       ] = { mtlpp::PixelFormat::Invalid, (uint8)EMetalBufferFormat::Unknown };
+	GAGXBufferFormats[PF_ASTC_12x12_HDR       ] = { mtlpp::PixelFormat::Invalid, (uint8)EMetalBufferFormat::Unknown };
+	GAGXBufferFormats[PF_BC6H                 ] = { mtlpp::PixelFormat::Invalid, (uint8)EMetalBufferFormat::Unknown };
+	GAGXBufferFormats[PF_BC7                  ] = { mtlpp::PixelFormat::Invalid, (uint8)EMetalBufferFormat::Unknown };
+	GAGXBufferFormats[PF_R8_UINT              ] = { mtlpp::PixelFormat::R8Uint, (uint8)EMetalBufferFormat::R8Uint };
+	GAGXBufferFormats[PF_R8                   ] = { mtlpp::PixelFormat::R8Unorm, (uint8)EMetalBufferFormat::R8Unorm };
+	GAGXBufferFormats[PF_L8                   ] = { mtlpp::PixelFormat::Invalid, (uint8)EMetalBufferFormat::R8Unorm };
+	GAGXBufferFormats[PF_XGXR8                ] = { mtlpp::PixelFormat::Invalid, (uint8)EMetalBufferFormat::Unknown };
+	GAGXBufferFormats[PF_R8G8B8A8_UINT        ] = { mtlpp::PixelFormat::RGBA8Uint, (uint8)EMetalBufferFormat::RGBA8Uint };
+	GAGXBufferFormats[PF_R8G8B8A8_SNORM       ] = { mtlpp::PixelFormat::RGBA8Snorm, (uint8)EMetalBufferFormat::RGBA8Snorm };
+	GAGXBufferFormats[PF_R16G16B16A16_UNORM   ] = { mtlpp::PixelFormat::RGBA16Unorm, (uint8)EMetalBufferFormat::RGBA16Unorm };
+	GAGXBufferFormats[PF_R16G16B16A16_SNORM   ] = { mtlpp::PixelFormat::RGBA16Snorm, (uint8)EMetalBufferFormat::RGBA16Snorm };
+	GAGXBufferFormats[PF_PLATFORM_HDR_0       ] = { mtlpp::PixelFormat::Invalid, (uint8)EMetalBufferFormat::Unknown };
+	GAGXBufferFormats[PF_PLATFORM_HDR_1       ] = { mtlpp::PixelFormat::Invalid, (uint8)EMetalBufferFormat::Unknown };
+	GAGXBufferFormats[PF_PLATFORM_HDR_2       ] = { mtlpp::PixelFormat::Invalid, (uint8)EMetalBufferFormat::Unknown };
+	GAGXBufferFormats[PF_NV12                 ] = { mtlpp::PixelFormat::Invalid, (uint8)EMetalBufferFormat::Unknown };
 	
-	GAGXBufferFormats[PF_ETC2_R11_EAC			] = { mtlpp::PixelFormat::Invalid, (uint8)EMetalBufferFormat::Unknown };
-	GAGXBufferFormats[PF_ETC2_RG11_EAC		] = { mtlpp::PixelFormat::Invalid, (uint8)EMetalBufferFormat::Unknown };
-		
+	GAGXBufferFormats[PF_ETC2_R11_EAC         ] = { mtlpp::PixelFormat::Invalid, (uint8)EMetalBufferFormat::Unknown };
+	GAGXBufferFormats[PF_ETC2_RG11_EAC        ] = { mtlpp::PixelFormat::Invalid, (uint8)EMetalBufferFormat::Unknown };
+
+	GAGXBufferFormats[PF_G16R16_SNORM         ] = { mtlpp::PixelFormat::RG16Snorm, (uint8)EMetalBufferFormat::RG16Snorm };
+	GAGXBufferFormats[PF_R8G8_UINT            ] = { mtlpp::PixelFormat::RG8Uint, (uint8)EMetalBufferFormat::RG8Uint };
+	GAGXBufferFormats[PF_R32G32B32_UINT       ] = { mtlpp::PixelFormat::Invalid, (uint8)EMetalBufferFormat::Unknown };
+	GAGXBufferFormats[PF_R32G32B32_SINT       ] = { mtlpp::PixelFormat::Invalid, (uint8)EMetalBufferFormat::Unknown };
+	GAGXBufferFormats[PF_R32G32B32F           ] = { mtlpp::PixelFormat::Invalid, (uint8)EMetalBufferFormat::Unknown };
+	GAGXBufferFormats[PF_R8_SINT              ] = { mtlpp::PixelFormat::R8Sint, (uint8)EMetalBufferFormat::R8Sint };
+	GAGXBufferFormats[PF_R64_UINT             ] = { mtlpp::PixelFormat::Invalid, (uint8)EMetalBufferFormat::Unknown };
+
 	// Initialize the platform pixel format map.
 	GPixelFormats[PF_Unknown			].PlatformFormat	= (uint32)mtlpp::PixelFormat::Invalid;
 	GPixelFormats[PF_A32B32G32R32F		].PlatformFormat	= (uint32)mtlpp::PixelFormat::RGBA32Float;
@@ -734,6 +749,21 @@ FAGXDynamicRHI::FAGXDynamicRHI(ERHIFeatureLevel::Type RequestedFeatureLevel)
 	GPixelFormats[PF_ASTC_12x12			].PlatformFormat	= (uint32)mtlpp::PixelFormat::ASTC_12x12_LDR;
 	GPixelFormats[PF_ASTC_12x12			].Supported			= true;
 
+#if !PLATFORM_TVOS
+	if ([GMtlDevice supportsFamily:MTLGPUFamilyApple6])
+	{
+		GPixelFormats[PF_ASTC_4x4_HDR].PlatformFormat		= (uint32)mtlpp::PixelFormat::ASTC_4x4_HDR;
+		GPixelFormats[PF_ASTC_4x4_HDR].Supported			= true;
+		GPixelFormats[PF_ASTC_6x6_HDR].PlatformFormat		= (uint32)mtlpp::PixelFormat::ASTC_6x6_HDR;
+		GPixelFormats[PF_ASTC_6x6_HDR].Supported			= true;
+		GPixelFormats[PF_ASTC_8x8_HDR].PlatformFormat		= (uint32)mtlpp::PixelFormat::ASTC_8x8_HDR;
+		GPixelFormats[PF_ASTC_8x8_HDR].Supported			= true;
+		GPixelFormats[PF_ASTC_10x10_HDR].PlatformFormat		= (uint32)mtlpp::PixelFormat::ASTC_10x10_HDR;
+		GPixelFormats[PF_ASTC_10x10_HDR].Supported			= true;
+		GPixelFormats[PF_ASTC_12x12_HDR].PlatformFormat		= (uint32)mtlpp::PixelFormat::ASTC_12x12_HDR;
+		GPixelFormats[PF_ASTC_12x12_HDR].Supported			= true;
+	}
+#endif
 	// used with virtual textures
 	GPixelFormats[PF_ETC2_RGB	  		].PlatformFormat	= (uint32)mtlpp::PixelFormat::ETC2_RGB8;
 	GPixelFormats[PF_ETC2_RGB			].Supported			= true;
@@ -794,6 +824,7 @@ FAGXDynamicRHI::FAGXDynamicRHI(ERHIFeatureLevel::Type RequestedFeatureLevel)
 	
 	GPixelFormats[PF_FloatRGB			].PlatformFormat	= (uint32)mtlpp::PixelFormat::RG11B10Float;
 	GPixelFormats[PF_FloatRGB			].BlockBytes		= 4;
+	
 	GPixelFormats[PF_FloatR11G11B10		].PlatformFormat	= (uint32)mtlpp::PixelFormat::RG11B10Float;
 	GPixelFormats[PF_FloatR11G11B10		].BlockBytes		= 4;
 	GPixelFormats[PF_FloatR11G11B10		].Supported			= true;
@@ -889,6 +920,18 @@ FAGXDynamicRHI::FAGXDynamicRHI(ERHIFeatureLevel::Type RequestedFeatureLevel)
 	GPixelFormats[PF_NV12				].PlatformFormat	= (uint32)mtlpp::PixelFormat::Invalid;
 	GPixelFormats[PF_NV12				].Supported			= false;
 	
+	GPixelFormats[PF_G16R16_SNORM		].PlatformFormat	= (uint32)mtlpp::PixelFormat::RG16Snorm;
+	GPixelFormats[PF_R8G8_UINT			].PlatformFormat	= (uint32)mtlpp::PixelFormat::RG8Uint;
+	GPixelFormats[PF_R32G32B32_UINT		].PlatformFormat	= (uint32)mtlpp::PixelFormat::Invalid;
+	GPixelFormats[PF_R32G32B32_UINT		].Supported			= false;
+	GPixelFormats[PF_R32G32B32_SINT		].PlatformFormat	= (uint32)mtlpp::PixelFormat::Invalid;
+	GPixelFormats[PF_R32G32B32_SINT		].Supported			= false;
+	GPixelFormats[PF_R32G32B32F			].PlatformFormat	= (uint32)mtlpp::PixelFormat::Invalid;
+	GPixelFormats[PF_R32G32B32F			].Supported			= false;
+	GPixelFormats[PF_R8_SINT			].PlatformFormat	= (uint32)mtlpp::PixelFormat::R8Sint;
+	GPixelFormats[PF_R64_UINT			].PlatformFormat	= (uint32)mtlpp::PixelFormat::Invalid;
+	GPixelFormats[PF_R64_UINT			].Supported			= false;
+
 #if METAL_DEBUG_OPTIONS
 	for (uint32 i = 0; i < PF_MAX; i++)
 	{
@@ -935,10 +978,6 @@ FAGXDynamicRHI::FAGXDynamicRHI(ERHIFeatureLevel::Type RequestedFeatureLevel)
 	{
 		FGPUDriverInfo GPUDriverInfo = FPlatformMisc::GetGPUDriverInfo(GRHIAdapterName);
 		
-		GRHIAdapterUserDriverVersion = GPUDriverInfo.UserDriverVersion;
-		GRHIAdapterInternalDriverVersion = GPUDriverInfo.InternalDriverVersion;
-		GRHIAdapterDriverDate = GPUDriverInfo.DriverDate;
-		
 		UE_LOG(LogAGX, Display, TEXT("    Adapter Name: %s"), *GRHIAdapterName);
 		UE_LOG(LogAGX, Display, TEXT("  Driver Version: %s (internal:%s, unified:%s)"), *GRHIAdapterUserDriverVersion, *GRHIAdapterInternalDriverVersion, *GPUDriverInfo.GetUnifiedDriverVersion());
 		UE_LOG(LogAGX, Display, TEXT("     Driver Date: %s"), *GRHIAdapterDriverDate);
@@ -1212,7 +1251,7 @@ bool FAGXDynamicRHI::RHIGetAvailableResolutions(FScreenResolutionArray& Resoluti
 void FAGXDynamicRHI::RHIFlushResources()
 {
 	@autoreleasepool {
-		((FAGXDeviceContext*)ImmediateContext.Context)->FlushFreeList(false);
+		((FAGXDeviceContext*)ImmediateContext.Context)->FlushFreeList();
 		ImmediateContext.Context->SubmitCommandBufferAndWait();
 		((FAGXDeviceContext*)ImmediateContext.Context)->ClearFreeList();
         ((FAGXDeviceContext*)ImmediateContext.Context)->DrainHeap();
diff --git a/Engine/Source/Runtime/Apple/AGXRHI/Private/AGXRHIContext.cpp b/Engine/Source/Runtime/Apple/AGXRHI/Private/AGXRHIContext.cpp
--- a/Engine/Source/Runtime/Apple/AGXRHI/Private/AGXRHIContext.cpp
+++ b/Engine/Source/Runtime/Apple/AGXRHI/Private/AGXRHIContext.cpp
@@ -53,19 +53,6 @@ void AGXSafeReleaseMetalBuffer(FAGXBuffer& Buffer)
 	}
 }
 
-void AGXSafeReleaseMetalFence(FAGXFence* Object)
-{
-	if(GIsAGXInitialized && GDynamicRHI && Object)
-	{
-		FAGXRHICommandContext* Context = static_cast<FAGXRHICommandContext*>(RHIGetDefaultContext());
-		if(Context)
-		{
-			((FAGXDeviceContext&)Context->GetInternalContext()).ReleaseFence(Object);
-			return;
-		}
-	}
-}
-
 FAGXRHICommandContext::FAGXRHICommandContext(class FAGXProfiler* InProfiler, FAGXContext* WrapContext)
 : Context(WrapContext)
 , Profiler(InProfiler)
@@ -86,10 +73,6 @@ FAGXRHICommandContext::~FAGXRHICommandContext()
 FAGXRHIComputeContext::FAGXRHIComputeContext(class FAGXProfiler* InProfiler, FAGXContext* WrapContext)
 : FAGXRHICommandContext(InProfiler, WrapContext)
 {
-	if (FAGXCommandQueue::SupportsFeature(EAGXFeaturesFences) && FApplePlatformMisc::IsOSAtLeastVersion((uint32[]){10, 14, 0}, (uint32[]){12, 0, 0}, (uint32[]){12, 0, 0}))
-	{
-		WrapContext->GetCurrentRenderPass().SetDispatchType(mtlpp::DispatchType::Concurrent);
-	}
 }
 
 FAGXRHIComputeContext::~FAGXRHIComputeContext()
diff --git a/Engine/Source/Runtime/Apple/AGXRHI/Private/AGXRHIPrivate.h b/Engine/Source/Runtime/Apple/AGXRHI/Private/AGXRHIPrivate.h
--- a/Engine/Source/Runtime/Apple/AGXRHI/Private/AGXRHIPrivate.h
+++ b/Engine/Source/Runtime/Apple/AGXRHI/Private/AGXRHIPrivate.h
@@ -15,11 +15,16 @@
 #import <Metal/Metal.h>
 #import <QuartzCore/CAMetalLayer.h>
 
+#define AGXRHI_TRUE  1
+#define AGXRHI_FALSE 0
+
 // Metal C++ wrapper
 THIRD_PARTY_INCLUDES_START
 #include "mtlpp.hpp"
 THIRD_PARTY_INCLUDES_END
 
+extern id<MTLDevice> AGXUtilGetDevice();
+
 // The Metal device object.
 extern id<MTLDevice> GMtlDevice;
 
@@ -52,6 +57,54 @@ const uint32 AGXBufferBytesSize = BufferOffsetAlignment * 32;
 #include "AGXDynamicRHI.h"
 #include "RHI.h"
 
+/**
+ * EAGXDebugLevel: Level of AGX RHI debug features to be enabled.
+ */
+enum EAGXDebugLevel
+{
+	EAGXDebugLevelOff,
+	EAGXDebugLevelFastValidation,
+	EAGXDebugLevelResetOnBind,
+	EAGXDebugLevelConditionalSubmit,
+	EAGXDebugLevelValidation,
+	EAGXDebugLevelLogOperations,
+	EAGXDebugLevelWaitForComplete,
+};
+
+/**
+ * The sampler, buffer and texture resource limits as defined here:
+ * https://developer.apple.com/library/ios/documentation/Miscellaneous/Conceptual/MetalProgrammingGuide/Render-Ctx/Render-Ctx.html
+ */
+#if PLATFORM_IOS
+#define METAL_MAX_BUFFERS 31
+#define METAL_MAX_TEXTURES 31
+typedef uint32 FAGXTextureMask;
+#elif PLATFORM_MAC
+#define METAL_MAX_BUFFERS 31
+#define METAL_MAX_TEXTURES 128
+typedef __uint128_t FAGXTextureMask;
+#else
+#error "Unsupported Platform!"
+#endif
+typedef uint32 FAGXBufferMask;
+typedef uint16 FAGXSamplerMask;
+
+enum EAGXLimits
+{
+	ML_MaxSamplers = 16, /** Maximum number of samplers */
+	ML_MaxBuffers = METAL_MAX_BUFFERS, /** Maximum number of buffers */
+	ML_MaxTextures = METAL_MAX_TEXTURES, /** Maximum number of textures - there are more textures available on Mac than iOS */
+	ML_MaxViewports = 16 /** Technically this may be different at runtime, but this is the likely absolute upper-bound */
+};
+
+/** A structure for quick mask-testing of shader-stage resource bindings */
+struct FAGXDebugShaderResourceMask
+{
+	FAGXTextureMask TextureMask;
+	FAGXBufferMask BufferMask;
+	FAGXSamplerMask SamplerMask;
+};
+
 #define BUFFER_CACHE_MODE mtlpp::ResourceOptions::CpuCacheModeDefaultCache
 
 #if PLATFORM_MAC
@@ -108,8 +161,6 @@ extern FAGXBufferFormat GAGXBufferFormats[PF_MAX];
 #define METAL_DEBUG_LAYER(Level, Code)
 #endif
 
-extern bool GAGXCommandBufferDebuggingEnabled;
-
 /** Set to 1 to enable GPU events in Xcode frame debugger */
 #ifndef ENABLE_METAL_GPUEVENTS_IN_TEST
 	#define ENABLE_METAL_GPUEVENTS_IN_TEST 0
@@ -139,19 +190,6 @@ extern bool GAGXCommandBufferDebuggingEnabled;
 	#define METAL_IGNORED(Func)
 #endif
 
-struct FAGXDebugInfo
-{
-	uint32 CmdBuffIndex;
-	uint32 EncoderIndex;
-	uint32 ContextIndex;
-	uint32 CommandIndex;
-	uint64 CommandBuffer;
-	uint32 PSOSignature[4];
-};
-
-// Get a compute pipeline state used to implement some debug features.
-mtlpp::ComputePipelineState AGXGetMetalDebugComputeState();
-
 // Access the internal context for the device-owning DynamicRHI object
 FAGXDeviceContext& GetAGXDeviceContext();
 
@@ -164,9 +202,6 @@ void AGXSafeReleaseMetalTexture(FAGXTexture& Object);
 // Safely release a metal buffer, correctly handling the case where the RHI has been destructed first
 void AGXSafeReleaseMetalBuffer(FAGXBuffer& Buffer);
 
-// Safely release a fence, correctly handling cases where fences aren't supported or the debug implementation is used.
-void AGXSafeReleaseMetalFence(class FAGXFence* Object);
-
 // Safely release a render pass descriptor so that it may be reused.
 void AGXSafeReleaseMetalRenderPassDescriptor(mtlpp::RenderPassDescriptor& Desc);
 
@@ -235,7 +270,7 @@ extern uint32 AGXSafeGetRuntimeDebuggingLevel();
 
 extern int32 GAGXBufferZeroFill;
 
-mtlpp::LanguageVersion AGXValidateVersion(uint8 Version);
+mtlpp::LanguageVersion AGXValidateVersion(uint32 Version);
 
 // Needs to be the same as EShaderFrequency when all stages are supported, but unlike EShaderFrequency you can compile out stages.
 enum EAGXShaderStages
diff --git a/Engine/Source/Runtime/Apple/AGXRHI/Private/AGXRHIRenderQuery.cpp b/Engine/Source/Runtime/Apple/AGXRHI/Private/AGXRHIRenderQuery.cpp
--- a/Engine/Source/Runtime/Apple/AGXRHI/Private/AGXRHIRenderQuery.cpp
+++ b/Engine/Source/Runtime/Apple/AGXRHI/Private/AGXRHIRenderQuery.cpp
@@ -9,7 +9,6 @@
 #include "AGXCommandBufferFence.h"
 #include "AGXProfiler.h"
 #include "AGXLLM.h"
-#include "AGXCommandBuffer.h"
 #include "HAL/PThreadEvent.h"
 
 
diff --git a/Engine/Source/Runtime/Apple/AGXRHI/Private/AGXRenderCommandEncoder.cpp b/Engine/Source/Runtime/Apple/AGXRHI/Private/AGXRenderCommandEncoder.cpp
deleted file mode 100644
diff --git a/Engine/Source/Runtime/Apple/AGXRHI/Private/AGXRenderCommandEncoder.h b/Engine/Source/Runtime/Apple/AGXRHI/Private/AGXRenderCommandEncoder.h
deleted file mode 100644
diff --git a/Engine/Source/Runtime/Apple/AGXRHI/Private/AGXRenderPass.cpp b/Engine/Source/Runtime/Apple/AGXRHI/Private/AGXRenderPass.cpp
--- a/Engine/Source/Runtime/Apple/AGXRHI/Private/AGXRenderPass.cpp
+++ b/Engine/Source/Runtime/Apple/AGXRHI/Private/AGXRenderPass.cpp
@@ -10,7 +10,6 @@
 #include "AGXGraphicsPipelineState.h"
 #include "AGXVertexDeclaration.h"
 #include "AGXRenderPass.h"
-#include "AGXCommandBuffer.h"
 #include "AGXProfiler.h"
 #include "AGXFrameAllocator.h"
 
@@ -53,7 +52,6 @@ FAGXRenderPass::~FAGXRenderPass(void)
 {
 	check(!CurrentEncoder.GetCommandBuffer());
 	check(!PrologueEncoder.GetCommandBuffer());
-	check(!PassStartFence);
 }
 
 void FAGXRenderPass::SetDispatchType(mtlpp::DispatchType Type)
@@ -61,27 +59,8 @@ void FAGXRenderPass::SetDispatchType(mtlpp::DispatchType Type)
 	ComputeDispatchType = Type;
 }
 
-void FAGXRenderPass::Begin(FAGXFence* Fence, bool const bParallelBegin)
+void FAGXRenderPass::Begin(bool bParallelBegin)
 {
-	if (!bParallelBegin || !FAGXCommandQueue::SupportsFeature(EAGXFeaturesParallelRenderEncoders))
-	{
-		check(!PassStartFence || Fence == nullptr);
-		if (Fence)
-		{
-			PassStartFence = Fence;
-			PrologueStartEncoderFence = Fence;
-		}
-	}
-	else
-	{
-		check(!ParallelPassEndFence || Fence == nullptr);
-		if (Fence)
-		{
-			ParallelPassEndFence = Fence;
-			PrologueStartEncoderFence = Fence;
-		}
-	}
-	
 	if (!CmdList.IsParallel() && !CurrentEncoder.GetCommandBuffer())
 	{
 		CurrentEncoder.StartCommandBuffer();
@@ -89,58 +68,14 @@ void FAGXRenderPass::Begin(FAGXFence* Fence, bool const bParallelBegin)
 	}
 }
 
-void FAGXRenderPass::Wait(FAGXFence* Fence)
-{
-	if (Fence)
-	{
-		if (PrologueEncoder.IsBlitCommandEncoderActive() || PrologueEncoder.IsComputeCommandEncoderActive())
-		{
-			PrologueEncoder.WaitForFence(Fence);
-			METAL_DEBUG_LAYER(EAGXDebugLevelValidation, FAGXFence::ValidateUsage(Fence));
-		}
-		else if (CurrentEncoder.IsRenderCommandEncoderActive() || CurrentEncoder.IsBlitCommandEncoderActive() || CurrentEncoder.IsComputeCommandEncoderActive())
-		{
-			CurrentEncoder.WaitForFence(Fence);
-			METAL_DEBUG_LAYER(EAGXDebugLevelValidation, FAGXFence::ValidateUsage(Fence));
-		}
-        else
-        {
-            PassStartFence = Fence;
-			PrologueStartEncoderFence = Fence;
-        }
-	}
-		
-}
-
-void FAGXRenderPass::Update(FAGXFence* Fence)
-{
-	if (Fence)
-	{
-		// Force an encoder - possibly consuming the start fence so that we get the proper order
-		// the higher-level can generate empty contexts but we have no sane way to deal with that.
-		if (!CurrentEncoder.IsRenderCommandEncoderActive() && !CurrentEncoder.IsBlitCommandEncoderActive() && !CurrentEncoder.IsComputeCommandEncoderActive())
-		{
-			ConditionalSwitchToBlit();
-		}
-		CurrentEncoder.UpdateFence(Fence);
-		State.FlushVisibilityResults(CurrentEncoder);
-		TRefCountPtr<FAGXFence> NewFence = CurrentEncoder.EndEncoding();
-		check(!CurrentEncoderFence || !NewFence);
-		if (NewFence)
-		{
-			CurrentEncoderFence = NewFence;
-		}
-	}
-}
-
-TRefCountPtr<FAGXFence> const& FAGXRenderPass::Submit(EAGXSubmitFlags Flags)
+void FAGXRenderPass::Submit(EAGXSubmitFlags Flags)
 {
 	if (CurrentEncoder.GetCommandBuffer() || (Flags & EAGXSubmitFlagsAsyncCommandBuffer))
 	{
 		if (PrologueEncoder.IsBlitCommandEncoderActive() || PrologueEncoder.IsComputeCommandEncoderActive())
 		{
 			check(PrologueEncoder.GetCommandBuffer());
-			PrologueEncoderFence = PrologueEncoder.EndEncoding();
+			PrologueEncoder.EndEncoding();
 		}
 		if (PrologueEncoder.GetCommandBuffer())
 		{
@@ -175,7 +110,7 @@ TRefCountPtr<FAGXFence> const& FAGXRenderPass::Submit(EAGXSubmitFlags Flags)
                 State.SetRenderStoreActions(CurrentEncoder, (Flags & EAGXSubmitFlagsBreakCommandBuffer));
 				State.FlushVisibilityResults(CurrentEncoder);
             }
-            CurrentEncoderFence = CurrentEncoder.EndEncoding();
+            CurrentEncoder.EndEncoding();
         }
 		
         CurrentEncoder.CommitCommandBuffer(Flags);
@@ -187,8 +122,6 @@ TRefCountPtr<FAGXFence> const& FAGXRenderPass::Submit(EAGXSubmitFlags Flags)
 		PrologueEncoder.Reset();
 		CurrentEncoder.Reset();
 	}
-	
-	return CurrentEncoderFence;
 }
 
 void FAGXRenderPass::BeginParallelRenderPass(mtlpp::RenderPassDescriptor RenderPass, uint32 NumParallelContextsInPass)
@@ -202,12 +135,12 @@ void FAGXRenderPass::BeginParallelRenderPass(mtlpp::RenderPassDescriptor RenderP
 	{
 		if (PrologueEncoder.IsBlitCommandEncoderActive() || PrologueEncoder.IsComputeCommandEncoderActive())
 		{
-			PrologueEncoderFence = PrologueEncoder.EndEncoding();
+			PrologueEncoder.EndEncoding();
 		}
 		if (CurrentEncoder.IsRenderCommandEncoderActive() || CurrentEncoder.IsBlitCommandEncoderActive() || CurrentEncoder.IsComputeCommandEncoderActive())
 		{
 			State.FlushVisibilityResults(CurrentEncoder);
-			CurrentEncoderFence = CurrentEncoder.EndEncoding();
+			CurrentEncoder.EndEncoding();
 		}
 
 		CurrentEncoder.SetRenderPassDescriptor(RenderPass);
@@ -235,12 +168,12 @@ void FAGXRenderPass::BeginRenderPass(mtlpp::RenderPassDescriptor RenderPass)
 	// EndEncoding should provide the encoder fence...
 	if (PrologueEncoder.IsBlitCommandEncoderActive() || PrologueEncoder.IsComputeCommandEncoderActive())
 	{
-		PrologueEncoderFence = PrologueEncoder.EndEncoding();
+		PrologueEncoder.EndEncoding();
 	}
 	if (CurrentEncoder.IsRenderCommandEncoderActive() || CurrentEncoder.IsBlitCommandEncoderActive() || CurrentEncoder.IsComputeCommandEncoderActive())
 	{
 		State.FlushVisibilityResults(CurrentEncoder);
-		CurrentEncoderFence = CurrentEncoder.EndEncoding();
+		CurrentEncoder.EndEncoding();
 	}
 	State.SetStateDirty();
 	State.SetRenderTargetsActive(true);
@@ -252,30 +185,6 @@ void FAGXRenderPass::BeginRenderPass(mtlpp::RenderPassDescriptor RenderPass)
 	if (!GAGXDeferRenderPasses || !State.CanRestartRenderPass() || CmdList.IsParallel())
 	{
 		CurrentEncoder.BeginRenderCommandEncoding();
-		if (PassStartFence)
-		{
-			CurrentEncoder.WaitForFence(PassStartFence);
-			PassStartFence = nullptr;
-		}
-		if (ParallelPassEndFence)
-		{
-			CurrentEncoder.WaitForFence(ParallelPassEndFence);
-			ParallelPassEndFence = nullptr;
-		}
-		if (CurrentEncoderFence)
-		{
-			CurrentEncoder.WaitForFence(CurrentEncoderFence);
-			CurrentEncoderFence = nullptr;
-		}
-		if (PrologueEncoderFence)
-		{
-			// Consume on the current encoder but do not invalidate
-			CurrentEncoder.WaitForFence(PrologueEncoderFence);
-		}
-		if (PrologueEncoder.IsBlitCommandEncoderActive() || PrologueEncoder.IsComputeCommandEncoderActive())
-		{
-			CurrentEncoder.WaitForFence(PrologueEncoder.GetEncoderFence());
-		}
 		State.SetRenderStoreActions(CurrentEncoder, false);
 		check(CurrentEncoder.IsRenderCommandEncoderActive());
 	}
@@ -366,37 +275,13 @@ void FAGXRenderPass::RestartRenderPass(mtlpp::RenderPassDescriptor RenderPass)
 			State.SetRenderStoreActions(CurrentEncoder, true);
 			State.FlushVisibilityResults(CurrentEncoder);
 		}
-		CurrentEncoderFence = CurrentEncoder.EndEncoding();
+		CurrentEncoder.EndEncoding();
 	}
 	State.SetStateDirty();
 	State.SetRenderTargetsActive(true);
 	
 	CurrentEncoder.SetRenderPassDescriptor(RenderPassDesc);
 	CurrentEncoder.BeginRenderCommandEncoding();
-	if (PassStartFence)
-	{
-		CurrentEncoder.WaitForFence(PassStartFence);
-		PassStartFence = nullptr;
-	}
-	if (ParallelPassEndFence)
-	{
-		CurrentEncoder.WaitForFence(ParallelPassEndFence);
-		ParallelPassEndFence = nullptr;
-	}
-	if (CurrentEncoderFence)
-	{
-		CurrentEncoder.WaitForFence(CurrentEncoderFence);
-		CurrentEncoderFence = nullptr;
-	}
-	if (PrologueEncoderFence)
-	{
-		// Consume on the current encoder but do not invalidate
-		CurrentEncoder.WaitForFence(PrologueEncoderFence);
-	}
-	if (PrologueEncoder.IsBlitCommandEncoderActive() || PrologueEncoder.IsComputeCommandEncoderActive())
-	{
-		CurrentEncoder.WaitForFence(PrologueEncoder.GetEncoderFence());
-	}
 	State.SetRenderStoreActions(CurrentEncoder, false);
 	
 	check(CurrentEncoder.IsRenderCommandEncoderActive());
@@ -418,20 +303,7 @@ void FAGXRenderPass::DrawPrimitive(uint32 PrimitiveType, uint32 BaseVertexIndex,
 	
 	METAL_GPUPROFILE(FAGXProfiler::GetProfiler()->EncodeDraw(CurrentEncoder.GetCommandBufferStats(), __FUNCTION__, NumPrimitives, NumVertices, NumInstances));
 	CurrentEncoder.GetRenderCommandEncoder().Draw(AGXTranslatePrimitiveType(PrimitiveType), BaseVertexIndex, NumVertices, NumInstances);
-	METAL_DEBUG_LAYER(EAGXDebugLevelFastValidation, CurrentEncoder.GetRenderCommandEncoderDebugging().Draw(AGXTranslatePrimitiveType(PrimitiveType), BaseVertexIndex, NumVertices, NumInstances));
 
-	if (GAGXCommandBufferDebuggingEnabled)
-	{
-		FAGXCommandData Data;
-		Data.CommandType = FAGXCommandData::Type::DrawPrimitive;
-		Data.Draw.BaseInstance = 0;
-		Data.Draw.InstanceCount = NumInstances;
-		Data.Draw.VertexCount = NumVertices;
-		Data.Draw.VertexStart = BaseVertexIndex;
-		
-		InsertDebugDraw(Data);
-	}
-	
 	ConditionalSubmit();	
 }
 
@@ -450,15 +322,6 @@ void FAGXRenderPass::DrawPrimitiveIndirect(uint32 PrimitiveType, FAGXVertexBuffe
 		
 		METAL_GPUPROFILE(FAGXProfiler::GetProfiler()->EncodeDraw(CurrentEncoder.GetCommandBufferStats(), __FUNCTION__, 1, 1, 1));
 		CurrentEncoder.GetRenderCommandEncoder().Draw(AGXTranslatePrimitiveType(PrimitiveType), TheBackingBuffer, ArgumentOffset);
-		METAL_DEBUG_LAYER(EAGXDebugLevelFastValidation, CurrentEncoder.GetRenderCommandEncoderDebugging().Draw(AGXTranslatePrimitiveType(PrimitiveType), TheBackingBuffer, ArgumentOffset));
-
-		if (GAGXCommandBufferDebuggingEnabled)
-		{
-			FAGXCommandData Data;
-			Data.CommandType = FAGXCommandData::Type::DrawPrimitiveIndirect;
-
-			InsertDebugDraw(Data);
-		}
 
 		ConditionalSubmit();
 	}
@@ -530,25 +393,10 @@ void FAGXRenderPass::DrawIndexedPrimitive(FAGXBuffer const& IndexBuffer, uint32
 	if (GRHISupportsBaseVertexIndex && GRHISupportsFirstInstance)
 	{
 		CurrentEncoder.GetRenderCommandEncoder().DrawIndexed(AGXTranslatePrimitiveType(PrimitiveType), NumIndices, ((IndexStride == 2) ? mtlpp::IndexType::UInt16 : mtlpp::IndexType::UInt32), IndexBuffer, StartIndex * IndexStride, NumInstances, BaseVertexIndex, FirstInstance);
-		METAL_DEBUG_LAYER(EAGXDebugLevelFastValidation, CurrentEncoder.GetRenderCommandEncoderDebugging().DrawIndexed(AGXTranslatePrimitiveType(PrimitiveType), NumIndices, ((IndexStride == 2) ? mtlpp::IndexType::UInt16 : mtlpp::IndexType::UInt32), IndexBuffer, StartIndex * IndexStride, NumInstances, BaseVertexIndex, FirstInstance));
 	}
 	else
 	{
 		CurrentEncoder.GetRenderCommandEncoder().DrawIndexed(AGXTranslatePrimitiveType(PrimitiveType), NumIndices, ((IndexStride == 2) ? mtlpp::IndexType::UInt16 : mtlpp::IndexType::UInt32), IndexBuffer, StartIndex * IndexStride, NumInstances);
-		METAL_DEBUG_LAYER(EAGXDebugLevelFastValidation, CurrentEncoder.GetRenderCommandEncoderDebugging().DrawIndexed(AGXTranslatePrimitiveType(PrimitiveType), NumIndices, ((IndexStride == 2) ? mtlpp::IndexType::UInt16 : mtlpp::IndexType::UInt32), IndexBuffer, StartIndex * IndexStride, NumInstances));
-	}
-	
-	if (GAGXCommandBufferDebuggingEnabled)
-	{
-		FAGXCommandData Data;
-		Data.CommandType = FAGXCommandData::Type::DrawPrimitiveIndexed;
-		Data.DrawIndexed.BaseInstance = FirstInstance;
-		Data.DrawIndexed.BaseVertex = BaseVertexIndex;
-		Data.DrawIndexed.IndexCount = NumIndices;
-		Data.DrawIndexed.IndexStart = StartIndex;
-		Data.DrawIndexed.InstanceCount = NumInstances;
-		
-		InsertDebugDraw(Data);
 	}
 	
 	ConditionalSubmit();
@@ -575,15 +423,7 @@ void FAGXRenderPass::DrawIndexedIndirect(FAGXIndexBuffer* IndexBuffer, uint32 Pr
 		
 		METAL_GPUPROFILE(FAGXProfiler::GetProfiler()->EncodeDraw(CurrentEncoder.GetCommandBufferStats(), __FUNCTION__, 1, 1, 1));
 		CurrentEncoder.GetRenderCommandEncoder().DrawIndexed(AGXTranslatePrimitiveType(PrimitiveType), (mtlpp::IndexType)IndexBuffer->IndexType, TheBackingIndexBuffer, 0, TheBackingBuffer, (DrawArgumentsIndex * 5 * sizeof(uint32)));
-		METAL_DEBUG_LAYER(EAGXDebugLevelFastValidation, CurrentEncoder.GetRenderCommandEncoderDebugging().DrawIndexed(AGXTranslatePrimitiveType(PrimitiveType), (mtlpp::IndexType)IndexBuffer->IndexType, TheBackingIndexBuffer, 0, TheBackingBuffer, (DrawArgumentsIndex * 5 * sizeof(uint32))));
 
-		if (GAGXCommandBufferDebuggingEnabled)
-		{
-			FAGXCommandData Data;
-			Data.CommandType = FAGXCommandData::Type::DrawPrimitiveIndexedIndirect;
-			
-			InsertDebugDraw(Data);
-		}
 		ConditionalSubmit();
 	}
 	else
@@ -610,16 +450,7 @@ void FAGXRenderPass::DrawIndexedPrimitiveIndirect(uint32 PrimitiveType,FAGXIndex
 		
 		METAL_GPUPROFILE(FAGXProfiler::GetProfiler()->EncodeDraw(CurrentEncoder.GetCommandBufferStats(), __FUNCTION__, 1, 1, 1));
 		CurrentEncoder.GetRenderCommandEncoder().DrawIndexed(AGXTranslatePrimitiveType(PrimitiveType), (mtlpp::IndexType)IndexBuffer->IndexType, TheBackingIndexBuffer, 0, TheBackingBuffer, ArgumentOffset);
-		METAL_DEBUG_LAYER(EAGXDebugLevelFastValidation, CurrentEncoder.GetRenderCommandEncoderDebugging().DrawIndexed(AGXTranslatePrimitiveType(PrimitiveType), (mtlpp::IndexType)IndexBuffer->IndexType, TheBackingIndexBuffer, 0, TheBackingBuffer, ArgumentOffset));
 
-		if (GAGXCommandBufferDebuggingEnabled)
-		{
-			FAGXCommandData Data;
-			Data.CommandType = FAGXCommandData::Type::DrawPrimitiveIndirect;
-			
-			InsertDebugDraw(Data);
-		}
-		
 		ConditionalSubmit();
 	}
 	else
@@ -647,7 +478,6 @@ void FAGXRenderPass::Dispatch(uint32 ThreadGroupCountX, uint32 ThreadGroupCountY
 		check(ComputeShader->NumThreadsX > 0 && ComputeShader->NumThreadsY > 0 && ComputeShader->NumThreadsZ > 0);
 		mtlpp::Size Threadgroups = mtlpp::Size(ThreadGroupCountX, ThreadGroupCountY, ThreadGroupCountZ);
 		PrologueEncoder.GetComputeCommandEncoder().DispatchThreadgroups(Threadgroups, ThreadgroupCounts);
-		METAL_DEBUG_LAYER(EAGXDebugLevelFastValidation, PrologueEncoder.GetComputeCommandEncoderDebugging().DispatchThreadgroups(Threadgroups, ThreadgroupCounts));
 		
 		ConditionalSubmit();
 	}
@@ -668,18 +498,6 @@ void FAGXRenderPass::Dispatch(uint32 ThreadGroupCountX, uint32 ThreadGroupCountY
 		check(ComputeShader->NumThreadsX > 0 && ComputeShader->NumThreadsY > 0 && ComputeShader->NumThreadsZ > 0);
 		mtlpp::Size Threadgroups = mtlpp::Size(ThreadGroupCountX, ThreadGroupCountY, ThreadGroupCountZ);
 		CurrentEncoder.GetComputeCommandEncoder().DispatchThreadgroups(Threadgroups, ThreadgroupCounts);
-		METAL_DEBUG_LAYER(EAGXDebugLevelFastValidation, CurrentEncoder.GetComputeCommandEncoderDebugging().DispatchThreadgroups(Threadgroups, ThreadgroupCounts));
-		
-		if (GAGXCommandBufferDebuggingEnabled)
-		{
-			FAGXCommandData Data;
-			Data.CommandType = FAGXCommandData::Type::Dispatch;
-			Data.Dispatch.threadgroupsPerGrid[0] = ThreadGroupCountX;
-			Data.Dispatch.threadgroupsPerGrid[1] = ThreadGroupCountY;
-			Data.Dispatch.threadgroupsPerGrid[2] = ThreadGroupCountZ;
-			
-			InsertDebugDispatch(Data);
-		}
 		
 		ConditionalSubmit();
 	}
@@ -706,7 +524,6 @@ void FAGXRenderPass::DispatchIndirect(FAGXVertexBuffer* ArgumentBuffer, uint32 A
 		check(ComputeShader->NumThreadsX > 0 && ComputeShader->NumThreadsY > 0 && ComputeShader->NumThreadsZ > 0);
 		
 		PrologueEncoder.GetComputeCommandEncoder().DispatchThreadgroupsWithIndirectBuffer(ArgumentBuffer->GetCurrentBuffer(), ArgumentOffset, ThreadgroupCounts);
-		METAL_DEBUG_LAYER(EAGXDebugLevelFastValidation, PrologueEncoder.GetComputeCommandEncoderDebugging().DispatchThreadgroupsWithIndirectBuffer(ArgumentBuffer->GetCurrentBuffer(), ArgumentOffset, ThreadgroupCounts));
 		
 		ConditionalSubmit();
 	}
@@ -726,23 +543,12 @@ void FAGXRenderPass::DispatchIndirect(FAGXVertexBuffer* ArgumentBuffer, uint32 A
 		check(ComputeShader->NumThreadsX > 0 && ComputeShader->NumThreadsY > 0 && ComputeShader->NumThreadsZ > 0);
 
 		CurrentEncoder.GetComputeCommandEncoder().DispatchThreadgroupsWithIndirectBuffer(ArgumentBuffer->GetCurrentBuffer(), ArgumentOffset, ThreadgroupCounts);
-		METAL_DEBUG_LAYER(EAGXDebugLevelFastValidation, CurrentEncoder.GetComputeCommandEncoderDebugging().DispatchThreadgroupsWithIndirectBuffer(ArgumentBuffer->GetCurrentBuffer(), ArgumentOffset, ThreadgroupCounts));
 
-		if (GAGXCommandBufferDebuggingEnabled)
-		{
-			FAGXCommandData Data;
-			Data.CommandType = FAGXCommandData::Type::DispatchIndirect;
-			Data.DispatchIndirect.ArgumentBuffer = ArgumentBuffer->GetCurrentBuffer();
-			Data.DispatchIndirect.ArgumentOffset = ArgumentOffset;
-			
-			InsertDebugDispatch(Data);
-		}
-	
 		ConditionalSubmit();
 	}
 }
 
-TRefCountPtr<FAGXFence> const& FAGXRenderPass::EndRenderPass(void)
+void FAGXRenderPass::EndRenderPass(void)
 {
 	if (bWithinRenderPass)
 	{
@@ -752,7 +558,6 @@ TRefCountPtr<FAGXFence> const& FAGXRenderPass::EndRenderPass(void)
 		// This just calls End - it exists only to enforce assumptions
 		End();
 	}
-	return CurrentEncoderFence;
 }
 
 void FAGXRenderPass::InsertTextureBarrier()
@@ -775,7 +580,6 @@ void FAGXRenderPass::CopyFromTextureToBuffer(FAGXTexture const& Texture, uint32
 	METAL_GPUPROFILE(FAGXProfiler::GetProfiler()->EncodeBlit(CurrentEncoder.GetCommandBufferStats(), __FUNCTION__));
 	{
 		MTLPP_VALIDATE(mtlpp::BlitCommandEncoder, Encoder, AGXSafeGetRuntimeDebuggingLevel() >= EAGXDebugLevelValidation, Copy(Texture, sourceSlice, sourceLevel, sourceOrigin, sourceSize, toBuffer, destinationOffset, destinationBytesPerRow, destinationBytesPerImage, options));
-		METAL_DEBUG_LAYER(EAGXDebugLevelFastValidation, CurrentEncoder.GetBlitCommandEncoderDebugging().Copy(Texture, sourceSlice, sourceLevel, sourceOrigin, sourceSize, toBuffer, destinationOffset, destinationBytesPerRow, destinationBytesPerImage, options));
 	}
 	ConditionalSubmit();
 }
@@ -791,12 +595,10 @@ void FAGXRenderPass::CopyFromBufferToTexture(FAGXBuffer const& Buffer, uint32 so
 	if (options == mtlpp::BlitOption::None)
 	{
 		MTLPP_VALIDATE(mtlpp::BlitCommandEncoder, Encoder, AGXSafeGetRuntimeDebuggingLevel() >= EAGXDebugLevelValidation, Copy(Buffer, sourceOffset, sourceBytesPerRow, sourceBytesPerImage, sourceSize, toTexture, destinationSlice, destinationLevel, destinationOrigin));
-		METAL_DEBUG_LAYER(EAGXDebugLevelFastValidation, CurrentEncoder.GetBlitCommandEncoderDebugging().Copy(Buffer, sourceOffset, sourceBytesPerRow, sourceBytesPerImage, sourceSize, toTexture, destinationSlice, destinationLevel, destinationOrigin));
 	}
 	else
 	{
 		MTLPP_VALIDATE(mtlpp::BlitCommandEncoder, Encoder, AGXSafeGetRuntimeDebuggingLevel() >= EAGXDebugLevelValidation, Copy(Buffer, sourceOffset, sourceBytesPerRow, sourceBytesPerImage, sourceSize, toTexture, destinationSlice, destinationLevel, destinationOrigin, options));
-		METAL_DEBUG_LAYER(EAGXDebugLevelFastValidation, CurrentEncoder.GetBlitCommandEncoderDebugging().Copy(Buffer, sourceOffset, sourceBytesPerRow, sourceBytesPerImage, sourceSize, toTexture, destinationSlice, destinationLevel, destinationOrigin, options));
 	}
 	ConditionalSubmit();
 }
@@ -809,7 +611,6 @@ void FAGXRenderPass::CopyFromTextureToTexture(FAGXTexture const& Texture, uint32
 	
 	METAL_GPUPROFILE(FAGXProfiler::GetProfiler()->EncodeBlit(CurrentEncoder.GetCommandBufferStats(), __FUNCTION__));
 	MTLPP_VALIDATE(mtlpp::BlitCommandEncoder, Encoder, AGXSafeGetRuntimeDebuggingLevel() >= EAGXDebugLevelValidation, Copy(Texture, sourceSlice, sourceLevel, sourceOrigin, sourceSize, toTexture, destinationSlice, destinationLevel, destinationOrigin));
-	METAL_DEBUG_LAYER(EAGXDebugLevelFastValidation, CurrentEncoder.GetBlitCommandEncoderDebugging().Copy(Texture, sourceSlice, sourceLevel, sourceOrigin, sourceSize, toTexture, destinationSlice, destinationLevel, destinationOrigin));
 	ConditionalSubmit();
 }
 
@@ -821,7 +622,6 @@ void FAGXRenderPass::CopyFromBufferToBuffer(FAGXBuffer const& SourceBuffer, NSUI
 	
 	METAL_GPUPROFILE(FAGXProfiler::GetProfiler()->EncodeBlit(CurrentEncoder.GetCommandBufferStats(), __FUNCTION__));
 	MTLPP_VALIDATE(mtlpp::BlitCommandEncoder, Encoder, AGXSafeGetRuntimeDebuggingLevel() >= EAGXDebugLevelValidation, Copy(SourceBuffer, SourceOffset, DestinationBuffer, DestinationOffset, Size));
-	METAL_DEBUG_LAYER(EAGXDebugLevelFastValidation, CurrentEncoder.GetBlitCommandEncoderDebugging().Copy(SourceBuffer, SourceOffset, DestinationBuffer, DestinationOffset, Size));
 	ConditionalSubmit();
 }
 
@@ -833,7 +633,6 @@ void FAGXRenderPass::PresentTexture(FAGXTexture const& Texture, uint32 sourceSli
 	
 	METAL_GPUPROFILE(FAGXProfiler::GetProfiler()->EncodeBlit(CurrentEncoder.GetCommandBufferStats(), __FUNCTION__));
 	MTLPP_VALIDATE(mtlpp::BlitCommandEncoder, Encoder, AGXSafeGetRuntimeDebuggingLevel() >= EAGXDebugLevelValidation, Copy(Texture, sourceSlice, sourceLevel, sourceOrigin, sourceSize, toTexture, destinationSlice, destinationLevel, destinationOrigin));
-	METAL_DEBUG_LAYER(EAGXDebugLevelFastValidation, CurrentEncoder.GetBlitCommandEncoderDebugging().Copy(Texture, sourceSlice, sourceLevel, sourceOrigin, sourceSize, toTexture, destinationSlice, destinationLevel, destinationOrigin));
 }
 
 void FAGXRenderPass::SynchronizeTexture(FAGXTexture const& Texture, uint32 Slice, uint32 Level)
@@ -846,7 +645,6 @@ void FAGXRenderPass::SynchronizeTexture(FAGXTexture const& Texture, uint32 Slice
 	
 	// METAL_GPUPROFILE(FAGXProfiler::GetProfiler()->EncodeBlit(CurrentEncoder.GetCommandBufferStats(), __FUNCTION__));
 	MTLPP_VALIDATE(mtlpp::BlitCommandEncoder, Encoder, AGXSafeGetRuntimeDebuggingLevel() >= EAGXDebugLevelValidation, Synchronize(Texture, Slice, Level));
-	METAL_DEBUG_LAYER(EAGXDebugLevelFastValidation, CurrentEncoder.GetBlitCommandEncoderDebugging().Synchronize(Texture, Slice, Level));
 	ConditionalSubmit();
 #endif
 }
@@ -861,7 +659,6 @@ void FAGXRenderPass::SynchroniseResource(mtlpp::Resource const& Resource)
 	
 	// METAL_GPUPROFILE(FAGXProfiler::GetProfiler()->EncodeBlit(CurrentEncoder.GetCommandBufferStats(), __FUNCTION__));
 	MTLPP_VALIDATE(mtlpp::BlitCommandEncoder, Encoder, AGXSafeGetRuntimeDebuggingLevel() >= EAGXDebugLevelValidation, Synchronize(Resource));
-	METAL_DEBUG_LAYER(EAGXDebugLevelFastValidation, CurrentEncoder.GetBlitCommandEncoderDebugging().Synchronize(Resource));
 	ConditionalSubmit();
 #endif
 }
@@ -871,27 +668,23 @@ void FAGXRenderPass::FillBuffer(FAGXBuffer const& Buffer, ns::Range Range, uint8
 	check(Buffer);
 	
 	mtlpp::BlitCommandEncoder TargetEncoder;
-	METAL_DEBUG_ONLY(FAGXBlitCommandEncoderDebugging Debugging);
 	bool bAsync = !CurrentEncoder.HasBufferBindingHistory(Buffer);
 	if(bAsync)
 	{
 		ConditionalSwitchToAsyncBlit();
 		TargetEncoder = PrologueEncoder.GetBlitCommandEncoder();
 		METAL_GPUPROFILE(FAGXProfiler::GetProfiler()->EncodeBlit(PrologueEncoder.GetCommandBufferStats(), FString::Printf(TEXT("FillBuffer: %p %llu %llu"), Buffer.GetPtr(), Buffer.GetOffset() + Range.Location, Range.Length)));
-		METAL_DEBUG_LAYER(EAGXDebugLevelFastValidation, Debugging = PrologueEncoder.GetBlitCommandEncoderDebugging());
 	}
 	else
 	{
 		ConditionalSwitchToBlit();
 		TargetEncoder = CurrentEncoder.GetBlitCommandEncoder();
 		METAL_GPUPROFILE(FAGXProfiler::GetProfiler()->EncodeBlit(CurrentEncoder.GetCommandBufferStats(), FString::Printf(TEXT("FillBuffer: %p %llu %llu"), Buffer.GetPtr(), Buffer.GetOffset() + Range.Location, Range.Length)));
-		METAL_DEBUG_LAYER(EAGXDebugLevelFastValidation, Debugging = CurrentEncoder.GetBlitCommandEncoderDebugging());
 	}
 	
 	check(TargetEncoder.GetPtr());
 	
 	MTLPP_VALIDATE(mtlpp::BlitCommandEncoder, TargetEncoder, AGXSafeGetRuntimeDebuggingLevel() >= EAGXDebugLevelValidation, Fill(Buffer, Range, Value));
-	METAL_DEBUG_LAYER(EAGXDebugLevelFastValidation, (bAsync ? PrologueEncoder.GetBlitCommandEncoderDebugging() : CurrentEncoder.GetBlitCommandEncoderDebugging()).Fill(Buffer, Range, Value));
 	
 	if (!bAsync)
 	{
@@ -902,21 +695,18 @@ void FAGXRenderPass::FillBuffer(FAGXBuffer const& Buffer, ns::Range Range, uint8
 bool FAGXRenderPass::AsyncCopyFromBufferToTexture(FAGXBuffer const& Buffer, uint32 sourceOffset, uint32 sourceBytesPerRow, uint32 sourceBytesPerImage, mtlpp::Size sourceSize, FAGXTexture const& toTexture, uint32 destinationSlice, uint32 destinationLevel, mtlpp::Origin destinationOrigin, mtlpp::BlitOption options)
 {
 	mtlpp::BlitCommandEncoder TargetEncoder;
-	METAL_DEBUG_ONLY(FAGXBlitCommandEncoderDebugging Debugging);
 	bool bAsync = !CurrentEncoder.HasTextureBindingHistory(toTexture);
 	if(bAsync)
 	{
 		ConditionalSwitchToAsyncBlit();
 		TargetEncoder = PrologueEncoder.GetBlitCommandEncoder();
 		METAL_GPUPROFILE(FAGXProfiler::GetProfiler()->EncodeBlit(PrologueEncoder.GetCommandBufferStats(), __FUNCTION__));
-		METAL_DEBUG_LAYER(EAGXDebugLevelFastValidation, Debugging = PrologueEncoder.GetBlitCommandEncoderDebugging());
 	}
 	else
 	{
 		ConditionalSwitchToBlit();
 		TargetEncoder = CurrentEncoder.GetBlitCommandEncoder();
 		METAL_GPUPROFILE(FAGXProfiler::GetProfiler()->EncodeBlit(CurrentEncoder.GetCommandBufferStats(), __FUNCTION__));
-		METAL_DEBUG_LAYER(EAGXDebugLevelFastValidation, Debugging = CurrentEncoder.GetBlitCommandEncoderDebugging());
 	}
 	
 	check(TargetEncoder.GetPtr());
@@ -924,12 +714,10 @@ bool FAGXRenderPass::AsyncCopyFromBufferToTexture(FAGXBuffer const& Buffer, uint
 	if (options == mtlpp::BlitOption::None)
 	{
 		MTLPP_VALIDATE(mtlpp::BlitCommandEncoder, TargetEncoder, AGXSafeGetRuntimeDebuggingLevel() >= EAGXDebugLevelValidation, Copy(Buffer, sourceOffset, sourceBytesPerRow, sourceBytesPerImage, sourceSize, toTexture, destinationSlice, destinationLevel, destinationOrigin));
-		METAL_DEBUG_LAYER(EAGXDebugLevelFastValidation, Debugging.Copy(Buffer, sourceOffset, sourceBytesPerRow, sourceBytesPerImage, sourceSize, toTexture, destinationSlice, destinationLevel, destinationOrigin));
 	}
 	else
 	{
 		MTLPP_VALIDATE(mtlpp::BlitCommandEncoder, TargetEncoder, AGXSafeGetRuntimeDebuggingLevel() >= EAGXDebugLevelValidation, Copy(Buffer, sourceOffset, sourceBytesPerRow, sourceBytesPerImage, sourceSize, toTexture, destinationSlice, destinationLevel, destinationOrigin, options));
-		METAL_DEBUG_LAYER(EAGXDebugLevelFastValidation, Debugging.Copy(Buffer, sourceOffset, sourceBytesPerRow, sourceBytesPerImage, sourceSize, toTexture, destinationSlice, destinationLevel, destinationOrigin, options));
 	}
 	
 	return bAsync;
@@ -938,27 +726,23 @@ bool FAGXRenderPass::AsyncCopyFromBufferToTexture(FAGXBuffer const& Buffer, uint
 bool FAGXRenderPass::AsyncCopyFromTextureToTexture(FAGXTexture const& Texture, uint32 sourceSlice, uint32 sourceLevel, mtlpp::Origin sourceOrigin, mtlpp::Size sourceSize, FAGXTexture const& toTexture, uint32 destinationSlice, uint32 destinationLevel, mtlpp::Origin destinationOrigin)
 {
 	mtlpp::BlitCommandEncoder TargetEncoder;
-	METAL_DEBUG_ONLY(FAGXBlitCommandEncoderDebugging Debugging);
 	bool bAsync = !CurrentEncoder.HasTextureBindingHistory(toTexture);
 	if(bAsync)
 	{
 		ConditionalSwitchToAsyncBlit();
 		TargetEncoder = PrologueEncoder.GetBlitCommandEncoder();
 		METAL_GPUPROFILE(FAGXProfiler::GetProfiler()->EncodeBlit(PrologueEncoder.GetCommandBufferStats(), __FUNCTION__));
-		METAL_DEBUG_LAYER(EAGXDebugLevelFastValidation, Debugging = PrologueEncoder.GetBlitCommandEncoderDebugging());
 	}
 	else
 	{
 		ConditionalSwitchToBlit();
 		TargetEncoder = CurrentEncoder.GetBlitCommandEncoder();
 		METAL_GPUPROFILE(FAGXProfiler::GetProfiler()->EncodeBlit(CurrentEncoder.GetCommandBufferStats(), __FUNCTION__));
-		METAL_DEBUG_LAYER(EAGXDebugLevelFastValidation, Debugging = CurrentEncoder.GetBlitCommandEncoderDebugging());
 	}
 	
 	check(TargetEncoder.GetPtr());
 	
 	MTLPP_VALIDATE(mtlpp::BlitCommandEncoder, TargetEncoder, AGXSafeGetRuntimeDebuggingLevel() >= EAGXDebugLevelValidation, Copy(Texture, sourceSlice, sourceLevel, sourceOrigin, sourceSize, toTexture, destinationSlice, destinationLevel, destinationOrigin));
-	METAL_DEBUG_LAYER(EAGXDebugLevelFastValidation, Debugging.Copy(Texture, sourceSlice, sourceLevel, sourceOrigin, sourceSize, toTexture, destinationSlice, destinationLevel, destinationOrigin));
 	
 	return bAsync;
 }
@@ -971,27 +755,23 @@ bool FAGXRenderPass::CanAsyncCopyToBuffer(FAGXBuffer const& DestinationBuffer)
 void FAGXRenderPass::AsyncCopyFromBufferToBuffer(FAGXBuffer const& SourceBuffer, NSUInteger SourceOffset, FAGXBuffer const& DestinationBuffer, NSUInteger DestinationOffset, NSUInteger Size)
 {
 	mtlpp::BlitCommandEncoder TargetEncoder;
-	METAL_DEBUG_ONLY(FAGXBlitCommandEncoderDebugging Debugging);
 	bool bAsync = !CurrentEncoder.HasBufferBindingHistory(DestinationBuffer);
 	if(bAsync)
 	{
 		ConditionalSwitchToAsyncBlit();
 		TargetEncoder = PrologueEncoder.GetBlitCommandEncoder();
 		METAL_GPUPROFILE(FAGXProfiler::GetProfiler()->EncodeBlit(PrologueEncoder.GetCommandBufferStats(), FString::Printf(TEXT("AsyncCopyFromBufferToBuffer: %p %llu %llu"), DestinationBuffer.GetPtr(), DestinationBuffer.GetOffset() + DestinationOffset, Size)));
-		METAL_DEBUG_LAYER(EAGXDebugLevelFastValidation, Debugging = PrologueEncoder.GetBlitCommandEncoderDebugging());
 	}
 	else
 	{
 		ConditionalSwitchToBlit();
 		TargetEncoder = CurrentEncoder.GetBlitCommandEncoder();
 		METAL_GPUPROFILE(FAGXProfiler::GetProfiler()->EncodeBlit(CurrentEncoder.GetCommandBufferStats(), FString::Printf(TEXT("AsyncCopyFromBufferToBuffer: %p %llu %llu"), DestinationBuffer.GetPtr(), DestinationBuffer.GetOffset() + DestinationOffset, Size)));
-		METAL_DEBUG_LAYER(EAGXDebugLevelFastValidation, Debugging = CurrentEncoder.GetBlitCommandEncoderDebugging());
 	}
 	
 	check(TargetEncoder.GetPtr());
 	
     MTLPP_VALIDATE(mtlpp::BlitCommandEncoder, TargetEncoder, AGXSafeGetRuntimeDebuggingLevel() >= EAGXDebugLevelValidation, Copy(SourceBuffer, SourceOffset, DestinationBuffer, DestinationOffset, Size));
-	METAL_DEBUG_LAYER(EAGXDebugLevelFastValidation, Debugging.Copy(SourceBuffer, SourceOffset, DestinationBuffer, DestinationOffset, Size));
 }
 
 FAGXBuffer FAGXRenderPass::AllocateTemporyBufferForCopy(FAGXBuffer const& DestinationBuffer, NSUInteger Size, NSUInteger Align)
@@ -1019,15 +799,14 @@ void FAGXRenderPass::AsyncGenerateMipmapsForTexture(FAGXTexture const& Texture)
 	
 	METAL_GPUPROFILE(FAGXProfiler::GetProfiler()->EncodeBlit(CurrentEncoder.GetCommandBufferStats(), __FUNCTION__));
 	MTLPP_VALIDATE(mtlpp::BlitCommandEncoder, Encoder, AGXSafeGetRuntimeDebuggingLevel() >= EAGXDebugLevelValidation, GenerateMipmaps(Texture));
-	METAL_DEBUG_LAYER(EAGXDebugLevelFastValidation, PrologueEncoder.GetBlitCommandEncoderDebugging().GenerateMipmaps(Texture));
 }
 
-TRefCountPtr<FAGXFence> const& FAGXRenderPass::End(void)
+void FAGXRenderPass::End(void)
 {
 	// EndEncoding should provide the encoder fence...
 	if (PrologueEncoder.IsBlitCommandEncoderActive() || PrologueEncoder.IsComputeCommandEncoderActive())
 	{
-		PrologueEncoderFence = PrologueEncoder.EndEncoding();
+		PrologueEncoder.EndEncoding();
 	}
 	
 	if (CmdList.IsImmediate() && IsWithinParallelPass() && CurrentEncoder.IsParallelRenderCommandEncoderActive())
@@ -1036,35 +815,18 @@ TRefCountPtr<FAGXFence> const& FAGXRenderPass::End(void)
 		CurrentEncoder.EndEncoding();
 		
 		ConditionalSwitchToBlit();
-		CurrentEncoderFence = CurrentEncoder.EndEncoding();
-		ParallelPassEndFence = nullptr;
-		PassStartFence = nullptr;
+		CurrentEncoder.EndEncoding();
 	}
 	else if (CurrentEncoder.IsRenderCommandEncoderActive() || CurrentEncoder.IsBlitCommandEncoderActive() || CurrentEncoder.IsComputeCommandEncoderActive())
 	{
 		State.FlushVisibilityResults(CurrentEncoder);
-		check(!CurrentEncoderFence);
-		check(!PassStartFence);
-		check(!ParallelPassEndFence);
-		CurrentEncoderFence = CurrentEncoder.EndEncoding();
-	}
-	else if (PassStartFence || ParallelPassEndFence)
-	{
-		ConditionalSwitchToBlit();
-		CurrentEncoderFence = CurrentEncoder.EndEncoding();
-		ParallelPassEndFence = nullptr;
-		PassStartFence = nullptr;
+		CurrentEncoder.EndEncoding();
 	}
 	
-	check(!PassStartFence);
-	check(!ParallelPassEndFence);
-	
 	State.SetRenderTargetsActive(false);
 	
 	RenderPassDesc = nil;
 	bWithinRenderPass = false;
-	
-	return CurrentEncoderFence;
 }
 
 void FAGXRenderPass::InsertCommandBufferFence(FAGXCommandBufferFence& Fence, mtlpp::CommandBufferHandler Handler)
@@ -1105,16 +867,13 @@ void FAGXRenderPass::InsertDebugEncoder()
 	check(NewBuf);
 	
 	mtlpp::BlitCommandEncoder TargetEncoder;
-	METAL_DEBUG_ONLY(FAGXBlitCommandEncoderDebugging Debugging);
 	ConditionalSwitchToBlit();
 	TargetEncoder = CurrentEncoder.GetBlitCommandEncoder();
 	METAL_GPUPROFILE(FAGXProfiler::GetProfiler()->EncodeBlit(CurrentEncoder.GetCommandBufferStats(), __FUNCTION__));
-	METAL_DEBUG_LAYER(EAGXDebugLevelFastValidation, Debugging = CurrentEncoder.GetBlitCommandEncoderDebugging());
 	
 	check(TargetEncoder.GetPtr());
 	
 	MTLPP_VALIDATE(mtlpp::BlitCommandEncoder, TargetEncoder, AGXSafeGetRuntimeDebuggingLevel() >= EAGXDebugLevelValidation, Fill(NewBuf, ns::Range(0, BufferOffsetAlignment), 0xff));
-	METAL_DEBUG_LAYER(EAGXDebugLevelFastValidation, CurrentEncoder.GetBlitCommandEncoderDebugging().Fill(NewBuf, ns::Range(0, BufferOffsetAlignment), 0xff));
 	
 	ConditionalSubmit();
 }
@@ -1182,7 +941,7 @@ void FAGXRenderPass::ConditionalSwitchToRender(void)
 	
 	if (CurrentEncoder.IsComputeCommandEncoderActive() || CurrentEncoder.IsBlitCommandEncoderActive())
 	{
-		CurrentEncoderFence = CurrentEncoder.EndEncoding();
+		CurrentEncoder.EndEncoding();
 	}
 	
 	if (!CurrentEncoder.IsRenderCommandEncoderActive())
@@ -1207,7 +966,7 @@ void FAGXRenderPass::ConditionalSwitchToCompute(void)
 			State.SetRenderStoreActions(CurrentEncoder, true);
 			State.FlushVisibilityResults(CurrentEncoder);
 		}
-		CurrentEncoderFence = CurrentEncoder.EndEncoding();
+		CurrentEncoder.EndEncoding();
 		State.SetRenderTargetsActive(false);
 	}
 	
@@ -1215,33 +974,9 @@ void FAGXRenderPass::ConditionalSwitchToCompute(void)
 	{
 		State.SetStateDirty();
 		CurrentEncoder.BeginComputeCommandEncoding(ComputeDispatchType);
-		if (PassStartFence)
-		{
-			CurrentEncoder.WaitForFence(PassStartFence);
-			PassStartFence = nullptr;
-		}
-		if (ParallelPassEndFence)
-		{
-			CurrentEncoder.WaitForFence(ParallelPassEndFence);
-			ParallelPassEndFence = nullptr;
-		}
-		if (CurrentEncoderFence)
-		{
-			CurrentEncoder.WaitForFence(CurrentEncoderFence);
-			CurrentEncoderFence = nullptr;
-		}
-		if (PrologueEncoderFence)
-		{
-			CurrentEncoder.WaitForFence(PrologueEncoderFence);
-		}
 	}
 	
 	check(CurrentEncoder.IsComputeCommandEncoderActive());
-	
-	if (PrologueEncoder.IsBlitCommandEncoderActive() || PrologueEncoder.IsComputeCommandEncoderActive())
-	{
-		CurrentEncoder.WaitForFence(PrologueEncoder.GetEncoderFence());
-	}
 }
 
 void FAGXRenderPass::ConditionalSwitchToBlit(void)
@@ -1258,40 +993,16 @@ void FAGXRenderPass::ConditionalSwitchToBlit(void)
 			State.SetRenderStoreActions(CurrentEncoder, true);
 			State.FlushVisibilityResults(CurrentEncoder);
 		}
-		CurrentEncoderFence = CurrentEncoder.EndEncoding();
+		CurrentEncoder.EndEncoding();
 		State.SetRenderTargetsActive(false);
 	}
 	
 	if (!CurrentEncoder.IsBlitCommandEncoderActive())
 	{
 		CurrentEncoder.BeginBlitCommandEncoding();
-		if (PassStartFence)
-		{
-			CurrentEncoder.WaitForFence(PassStartFence);
-			PassStartFence = nullptr;
-		}
-		if (ParallelPassEndFence)
-		{
-			CurrentEncoder.WaitForFence(ParallelPassEndFence);
-			ParallelPassEndFence = nullptr;
-		}
-		if (CurrentEncoderFence)
-		{
-			CurrentEncoder.WaitForFence(CurrentEncoderFence);
-			CurrentEncoderFence = nullptr;
-		}
-		if (PrologueEncoderFence)
-		{
-			CurrentEncoder.WaitForFence(PrologueEncoderFence);
-		}
 	}
 	
 	check(CurrentEncoder.IsBlitCommandEncoderActive());
-	
-	if (PrologueEncoder.IsBlitCommandEncoderActive() || PrologueEncoder.IsComputeCommandEncoderActive())
-	{
-		CurrentEncoder.WaitForFence(PrologueEncoder.GetEncoderFence());
-	}
 }
 
 void FAGXRenderPass::ConditionalSwitchToAsyncBlit(void)
@@ -1300,7 +1011,7 @@ void FAGXRenderPass::ConditionalSwitchToAsyncBlit(void)
 	
 	if (PrologueEncoder.IsComputeCommandEncoderActive() || PrologueEncoder.IsRenderCommandEncoderActive())
 	{
-		PrologueEncoderFence = PrologueEncoder.EndEncoding();
+		PrologueEncoder.EndEncoding();
 	}
 	
 	if (!PrologueEncoder.IsBlitCommandEncoderActive())
@@ -1310,51 +1021,6 @@ void FAGXRenderPass::ConditionalSwitchToAsyncBlit(void)
 			PrologueEncoder.StartCommandBuffer();
 		}
 		PrologueEncoder.BeginBlitCommandEncoding();
-		if (PrologueStartEncoderFence)
-		{
-			if (PrologueStartEncoderFence->NeedsWait(mtlpp::RenderStages::Vertex))
-			{
-				PrologueEncoder.WaitForFence(PrologueStartEncoderFence);
-			}
-			else
-			{
-				PrologueEncoder.WaitAndUpdateFence(PrologueStartEncoderFence);
-			}
-			PrologueStartEncoderFence = nullptr;
-		}
-		if (PrologueEncoderFence)
-		{
-			if (PrologueEncoderFence->NeedsWait(mtlpp::RenderStages::Vertex))
-			{
-				PrologueEncoder.WaitForFence(PrologueEncoderFence);
-			}
-			else
-			{
-				PrologueEncoder.WaitAndUpdateFence(PrologueEncoderFence);
-			}
-			PrologueEncoderFence = nullptr;
-		}
-#if METAL_DEBUG_OPTIONS
-//		if (GetEmitDrawEvents() && PrologueEncoder.GetEncoderFence())
-//		{
-//			for (uint32 i = mtlpp::RenderStages::Vertex; i <= mtlpp::RenderStages::Fragment && PrologueEncoder.GetEncoderFence()->Get((mtlpp::RenderStages)i).GetPtr(); i++)
-//			{
-//				if (CmdList.GetCommandQueue().GetRuntimeDebuggingLevel() >= EAGXDebugLevelValidation)
-//				{
-//					PrologueEncoder.GetEncoderFence()->Get((mtlpp::RenderStages)i).GetPtr().label = [NSString stringWithFormat:@"Prologue %@", PrologueEncoderFence->Get((mtlpp::RenderStages)i).GetLabel().GetPtr()];
-//				}
-//				else
-//				{
-//					PrologueEncoder.GetEncoderFence()->Get((mtlpp::RenderStages)i).SetLabel([NSString stringWithFormat:@"Prologue %@", PrologueEncoderFence->Get((mtlpp::RenderStages)i).GetLabel().GetPtr()]);
-//				}
-//			}
-//		}
-#endif
-
-		if (CurrentEncoder.IsRenderCommandEncoderActive() || CurrentEncoder.IsComputeCommandEncoderActive() || CurrentEncoder.IsBlitCommandEncoderActive())
-		{
-			CurrentEncoder.WaitForFence(PrologueEncoder.GetEncoderFence());
-		}
 	}
 	
 	check(PrologueEncoder.IsBlitCommandEncoderActive());
@@ -1366,7 +1032,7 @@ void FAGXRenderPass::ConditionalSwitchToAsyncCompute(void)
 	
 	if (PrologueEncoder.IsBlitCommandEncoderActive() || PrologueEncoder.IsRenderCommandEncoderActive())
 	{
-		PrologueEncoderFence = PrologueEncoder.EndEncoding();
+		PrologueEncoder.EndEncoding();
 	}
 	
 	if (!PrologueEncoder.IsComputeCommandEncoderActive())
@@ -1377,64 +1043,6 @@ void FAGXRenderPass::ConditionalSwitchToAsyncCompute(void)
 		}
 		State.SetStateDirty();
 		PrologueEncoder.BeginComputeCommandEncoding(ComputeDispatchType);
-		
-		if (PrologueStartEncoderFence)
-		{
-			if (PrologueStartEncoderFence->NeedsWait(mtlpp::RenderStages::Vertex))
-			{
-				PrologueEncoder.WaitForFence(PrologueStartEncoderFence);
-			}
-			else
-			{
-				PrologueEncoder.WaitAndUpdateFence(PrologueStartEncoderFence);
-			}
-			PrologueStartEncoderFence = nullptr;
-		}
-		if (PrologueEncoderFence)
-		{
-			if (PrologueEncoderFence->NeedsWait(mtlpp::RenderStages::Vertex))
-			{
-				PrologueEncoder.WaitForFence(PrologueEncoderFence);
-			}
-			else
-			{
-				PrologueEncoder.WaitAndUpdateFence(PrologueEncoderFence);
-			}
-			PrologueEncoderFence = nullptr;
-		}
-		if (PassStartFence)
-		{
-			if(PassStartFence->NeedsWait(mtlpp::RenderStages::Vertex))
-			{
-				PrologueEncoder.WaitForFence(PassStartFence);
-			}
-			else
-			{
-				PrologueEncoder.WaitAndUpdateFence(PassStartFence);
-			}
-			PassStartFence = nullptr;
-		}
-#if METAL_DEBUG_OPTIONS
-//		if (GetEmitDrawEvents() && PrologueEncoder.GetEncoderFence())
-//		{
-//			for (uint32 i = mtlpp::RenderStages::Vertex; i <= mtlpp::RenderStages::Fragment && PrologueEncoder.GetEncoderFence()->Get((mtlpp::RenderStages)i).GetPtr(); i++)
-//			{
-//				if (CmdList.GetCommandQueue().GetRuntimeDebuggingLevel() >= EAGXDebugLevelValidation)
-//				{
-//					PrologueEncoder.GetEncoderFence()->Get((mtlpp::RenderStages)i).GetPtr().label = [NSString stringWithFormat:@"Prologue %@", PrologueEncoderFence->Get((mtlpp::RenderStages)i).GetLabel().GetPtr()];
-//				}
-//				else
-//				{
-//					PrologueEncoder.GetEncoderFence()->Get((mtlpp::RenderStages)i).SetLabel([NSString stringWithFormat:@"Prologue %@", PrologueEncoderFence->Get((mtlpp::RenderStages)i).GetLabel().GetPtr()]);
-//				}
-//			}
-//		}
-#endif
-		
-		if (CurrentEncoder.IsRenderCommandEncoderActive() || CurrentEncoder.IsComputeCommandEncoderActive() || CurrentEncoder.IsBlitCommandEncoderActive())
-		{
-			CurrentEncoder.WaitForFence(PrologueEncoder.GetEncoderFence());
-		}
 	}
 	
 	check(PrologueEncoder.IsComputeCommandEncoderActive());
@@ -1648,120 +1256,3 @@ uint32 FAGXRenderPass::GetCommandBufferIndex(void) const
 @protocol IMTLComputeCommandEncoder
 - (void)memoryBarrierWithResources:(const id<MTLResource>[])resources count:(NSUInteger)count;
 @end
-
-void FAGXRenderPass::InsertDebugDraw(FAGXCommandData& Data)
-{
-#if !PLATFORM_TVOS
-	if (GAGXCommandBufferDebuggingEnabled && (!FAGXCommandQueue::SupportsFeature(EAGXFeaturesValidation) || State.GetVisibilityResultMode() == mtlpp::VisibilityResultMode::Disabled))
-	{
-		FAGXGraphicsPipelineState* BoundShaderState = State.GetGraphicsPSO();
-		
-		uint32 NumCommands = CurrentEncoder.GetMarkers().AddCommand(GetCommandBufferIndex(), GetEncoderIndex(), CmdList.GetParallelIndex(), State.GetDebugBuffer(), BoundShaderState, nullptr, Data);
-		
-		if ((NumCommands % GAGXDebugOpsCount) == 0)
-		{
-			FAGXDebugInfo DebugInfo;
-			DebugInfo.EncoderIndex = GetEncoderIndex();
-			DebugInfo.ContextIndex = CmdList.GetParallelIndex();
-			DebugInfo.CommandIndex = NumCommands;
-			DebugInfo.CmdBuffIndex = GetCommandBufferIndex();
-			DebugInfo.CommandBuffer = reinterpret_cast<uintptr_t>(CurrentEncoder.GetCommandBuffer().GetPtr());
-			DebugInfo.PSOSignature[0] = BoundShaderState->VertexShader->SourceLen;
-			DebugInfo.PSOSignature[1] = BoundShaderState->VertexShader->SourceCRC;
-			if (IsValidRef(BoundShaderState->PixelShader))
-			{
-				DebugInfo.PSOSignature[2] = BoundShaderState->PixelShader->SourceLen;
-				DebugInfo.PSOSignature[3] = BoundShaderState->PixelShader->SourceCRC;
-			}
-			else
-			{
-				DebugInfo.PSOSignature[2] = 0;
-				DebugInfo.PSOSignature[3] = 0;
-			}
-			
-			FAGXShaderPipeline* PSO = State.GetPipelineState();
-			
-			CurrentEncoder.GetRenderCommandEncoder().SetRenderPipelineState(PSO->DebugPipelineState);
-
-			mtlpp::VisibilityResultMode VisMode = State.GetVisibilityResultMode();
-			uint32 VisibilityOffset = State.GetVisibilityResultOffset();
-			if (VisMode != mtlpp::VisibilityResultMode::Disabled)
-			{
-				CurrentEncoder.GetRenderCommandEncoder().SetVisibilityResultMode(mtlpp::VisibilityResultMode::Disabled, 0);
-			}
-
-		#if PLATFORM_MAC
-			id<MTLBuffer> DebugBufferPtr = State.GetDebugBuffer().GetPtr();
-			[(id<IMTLRenderCommandEncoder>)CurrentEncoder.GetRenderCommandEncoder().GetPtr() memoryBarrierWithResources:&DebugBufferPtr count:1 afterStages:MTLRenderStageFragment beforeStages:MTLRenderStageVertex];
-			
-			CurrentEncoder.SetShaderBytes(mtlpp::FunctionType::Vertex, (uint8 const*)&DebugInfo, sizeof(DebugInfo), 0);
-			State.SetShaderBufferDirty(EAGXShaderStages::Vertex, 0);
-
-			CurrentEncoder.SetShaderBuffer(mtlpp::FunctionType::Vertex, State.GetDebugBuffer(), 0, State.GetDebugBuffer().GetLength(), 1, mtlpp::ResourceUsage::Write);
-			State.SetShaderBufferDirty(EAGXShaderStages::Vertex, 1);
-
-			CurrentEncoder.GetRenderCommandEncoder().Draw(mtlpp::PrimitiveType::Point, 0, 1);
-			
-			[(id<IMTLRenderCommandEncoder>)CurrentEncoder.GetRenderCommandEncoder().GetPtr() memoryBarrierWithResources:&DebugBufferPtr count:1 afterStages:MTLRenderStageVertex beforeStages:MTLRenderStageVertex];
-		#else
-			CurrentEncoder.GetRenderCommandEncoder().SetTileData((uint8 const*)&DebugInfo, sizeof(DebugInfo), 0);
-			CurrentEncoder.GetRenderCommandEncoder().SetTileBuffer(State.GetDebugBuffer(), 0, 1);
-			mtlpp::Size ThreadsPerTile(1, 1, 1);
-			CurrentEncoder.GetRenderCommandEncoder().DispatchThreadsPerTile(ThreadsPerTile);
-		#endif
-
-			CurrentEncoder.GetRenderCommandEncoder().SetRenderPipelineState(PSO->RenderPipelineState);
-
-			if (VisMode != mtlpp::VisibilityResultMode::Disabled)
-			{
-				CurrentEncoder.GetRenderCommandEncoder().SetVisibilityResultMode(VisMode, VisibilityOffset);
-			}
-		}
-	}
-#endif
-}
-
-void FAGXRenderPass::InsertDebugDispatch(FAGXCommandData& Data)
-{
-#if !PLATFORM_TVOS
-	if (GAGXCommandBufferDebuggingEnabled)
-	{
-		FAGXComputeShader* BoundShaderState = State.GetComputeShader();
-		
-		uint32 NumCommands = CurrentEncoder.GetMarkers().AddCommand(GetCommandBufferIndex(), GetEncoderIndex(), CmdList.GetParallelIndex(), State.GetDebugBuffer(), nullptr, BoundShaderState, Data);
-		
-		if ((NumCommands % GAGXDebugOpsCount) == 0)
-		{
-			FAGXDebugInfo DebugInfo;
-			DebugInfo.EncoderIndex = GetEncoderIndex();
-			DebugInfo.ContextIndex = CmdList.GetParallelIndex();
-			DebugInfo.CommandIndex = NumCommands;
-			DebugInfo.CmdBuffIndex = GetCommandBufferIndex();
-			DebugInfo.CommandBuffer = reinterpret_cast<uintptr_t>(CurrentEncoder.GetCommandBuffer().GetPtr());
-			DebugInfo.PSOSignature[0] = BoundShaderState->SourceLen;
-			DebugInfo.PSOSignature[1] = BoundShaderState->SourceCRC;
-			DebugInfo.PSOSignature[2] = 0;
-			DebugInfo.PSOSignature[3] = 0;
-			
-			CurrentEncoder.GetComputeCommandEncoder().SetComputePipelineState(AGXGetMetalDebugComputeState());
-			
-			id<MTLBuffer> DebugBufferPtr = State.GetDebugBuffer().GetPtr();
-			[(id<IMTLComputeCommandEncoder>)CurrentEncoder.GetComputeCommandEncoder().GetPtr() memoryBarrierWithResources:&DebugBufferPtr count:1];
-			
-			CurrentEncoder.SetShaderBytes(mtlpp::FunctionType::Kernel, (uint8 const*)&DebugInfo, sizeof(DebugInfo), 0);
-			State.SetShaderBufferDirty(EAGXShaderStages::Compute, 0);
-			
-			CurrentEncoder.SetShaderBuffer(mtlpp::FunctionType::Kernel, State.GetDebugBuffer(), 0, State.GetDebugBuffer().GetLength(), 1, mtlpp::ResourceUsage::Write);
-			State.SetShaderBufferDirty(EAGXShaderStages::Compute, 1);
-			
-			mtlpp::Size ThreadsPerTile(1, 1, 1);
-			CurrentEncoder.GetComputeCommandEncoder().DispatchThreads(ThreadsPerTile, ThreadsPerTile);
-			
-			[(id<IMTLComputeCommandEncoder>)CurrentEncoder.GetComputeCommandEncoder().GetPtr() memoryBarrierWithResources:&DebugBufferPtr count:1];
-			
-			FAGXShaderPipeline* Pipeline = BoundShaderState->GetPipeline();
-			CurrentEncoder.GetComputeCommandEncoder().SetComputePipelineState(Pipeline->ComputePipelineState);
-		}
-	}
-#endif
-}
diff --git a/Engine/Source/Runtime/Apple/AGXRHI/Private/AGXRenderPass.h b/Engine/Source/Runtime/Apple/AGXRHI/Private/AGXRenderPass.h
--- a/Engine/Source/Runtime/Apple/AGXRHI/Private/AGXRenderPass.h
+++ b/Engine/Source/Runtime/Apple/AGXRHI/Private/AGXRenderPass.h
@@ -5,7 +5,6 @@
 #include "AGXRHIPrivate.h"
 #include "AGXCommandEncoder.h"
 #include "AGXState.h"
-#include "AGXFence.h"
 
 class FAGXCommandList;
 class FAGXCommandQueue;
@@ -24,11 +23,7 @@ public:
 #pragma mark -
 	void SetDispatchType(mtlpp::DispatchType Type);
 	
-    void Begin(FAGXFence* Fence, bool const bParallelBegin = false);
-	
-	void Wait(FAGXFence* Fence);
-
-	void Update(FAGXFence* Fence);
+    void Begin(bool const bParallelBegin = false);
 	
     void BeginParallelRenderPass(mtlpp::RenderPassDescriptor RenderPass, uint32 NumParallelContextsInPass);
 
@@ -51,7 +46,7 @@ public:
     
     void DispatchIndirect(FAGXVertexBuffer* ArgumentBufferRHI, uint32 ArgumentOffset);
     
-    TRefCountPtr<FAGXFence> const& EndRenderPass(void);
+    void EndRenderPass(void);
     
     void CopyFromTextureToBuffer(FAGXTexture const& Texture, uint32 sourceSlice, uint32 sourceLevel, mtlpp::Origin sourceOrigin, mtlpp::Size sourceSize, FAGXBuffer const& toBuffer, uint32 destinationOffset, uint32 destinationBytesPerRow, uint32 destinationBytesPerImage, mtlpp::BlitOption options);
     
@@ -81,9 +76,9 @@ public:
 	
 	void AsyncGenerateMipmapsForTexture(FAGXTexture const& Texture);
 	
-    TRefCountPtr<FAGXFence> const& Submit(EAGXSubmitFlags SubmissionFlags);
+    void Submit(EAGXSubmitFlags SubmissionFlags);
     
-    TRefCountPtr<FAGXFence> const& End(void);
+    void End(void);
 	
 	void InsertCommandBufferFence(FAGXCommandBufferFence& Fence, mtlpp::CommandBufferHandler Handler);
 	
@@ -175,9 +170,6 @@ private:
 	uint32 GetEncoderIndex(void) const;
 	uint32 GetCommandBufferIndex(void) const;
 	
-	void InsertDebugDraw(FAGXCommandData& Data);
-	void InsertDebugDispatch(FAGXCommandData& Data);
-
 private:
 	FAGXCommandList& CmdList;
     FAGXStateCache& State;
@@ -195,15 +187,6 @@ private:
 	// Disjoint ranges *are* permitted!
 	TMap<id<MTLBuffer>, TArray<NSRange>> OutstandingBufferUploads;
 
-	// Fences for the current command encoder chain
-	TRefCountPtr<FAGXFence> PassStartFence;
-	TRefCountPtr<FAGXFence> CurrentEncoderFence;
-	TRefCountPtr<FAGXFence> ParallelPassEndFence;
-
-	// Fences for the prologue command encoder chain
-	TRefCountPtr<FAGXFence> PrologueStartEncoderFence;
-	TRefCountPtr<FAGXFence> PrologueEncoderFence;
-    
     mtlpp::RenderPassDescriptor RenderPassDesc;
     
 	mtlpp::DispatchType ComputeDispatchType;
diff --git a/Engine/Source/Runtime/Apple/AGXRHI/Private/AGXRenderTarget.cpp b/Engine/Source/Runtime/Apple/AGXRHI/Private/AGXRenderTarget.cpp
--- a/Engine/Source/Runtime/Apple/AGXRHI/Private/AGXRenderTarget.cpp
+++ b/Engine/Source/Runtime/Apple/AGXRHI/Private/AGXRenderTarget.cpp
@@ -7,7 +7,6 @@
 #include "AGXRHIPrivate.h"
 #include "ScreenRendering.h"
 #include "AGXProfiler.h"
-#include "AGXCommandBuffer.h"
 #include "ResolveShader.h"
 #include "PipelineStateCache.h"
 #include "Math/PackedVector.h"
diff --git a/Engine/Source/Runtime/Apple/AGXRHI/Private/AGXShaders.cpp b/Engine/Source/Runtime/Apple/AGXRHI/Private/AGXShaders.cpp
--- a/Engine/Source/Runtime/Apple/AGXRHI/Private/AGXShaders.cpp
+++ b/Engine/Source/Runtime/Apple/AGXRHI/Private/AGXShaders.cpp
@@ -17,7 +17,6 @@
 #include "MetalShaderResources.h"
 #include "AGXResources.h"
 #include "AGXProfiler.h"
-#include "AGXCommandBuffer.h"
 #include "Serialization/MemoryReader.h"
 #include "Misc/FileHelper.h"
 #include "Misc/ScopeRWLock.h"
@@ -45,84 +44,25 @@ NSString* AGXDecodeMetalSourceCode(uint32 CodeSize, TArray<uint8> const& Compres
 	return GlslCodeNSString;
 }
 
-mtlpp::LanguageVersion AGXValidateVersion(uint8 Version)
+mtlpp::LanguageVersion AGXValidateVersion(uint32 Version)
 {
-	static uint32 MetalMacOSVersions[][3] = {
-		{10,11,6},
-		{10,11,6},
-		{10,12,6},
-		{10,13,0},
-		{10,14,0},
-		{10,15,0},
-		{11, 0,0},
-	};
-	static uint32 MetaliOSVersions[][3] = {
-		{ 8,0,0},
-		{ 9,0,0},
-		{10,0,0},
-		{11,0,0},
-		{12,0,0},
-		{13,0,0},
-		{14,0,0},
-	};
-	static TCHAR const* StandardNames[] =
-	{
-		TEXT("Metal 1.0"),
-		TEXT("Metal 1.1"),
-		TEXT("Metal 1.2"),
-		TEXT("Metal 2.0"),
-		TEXT("Metal 2.1"),
-		TEXT("Metal 2.2"),
-		TEXT("Metal 2.3"),
-	};
-	
-	mtlpp::LanguageVersion Result = mtlpp::LanguageVersion::Version1_1;
+	mtlpp::LanguageVersion Result = mtlpp::LanguageVersion::Version2_2;
+
 	switch(Version)
 	{
+		case 7:
+			Result = mtlpp::LanguageVersion::Version2_4;
+			break;
 		case 6:
 			Result = mtlpp::LanguageVersion::Version2_3;
 			break;
 		case 5:
 			Result = mtlpp::LanguageVersion::Version2_2;
 			break;
-		case 4:
-			Result = mtlpp::LanguageVersion::Version2_1;
-			break;
-		case 3:
-			Result = mtlpp::LanguageVersion::Version2_0;
-			break;
-		case 2:
-			Result = mtlpp::LanguageVersion::Version1_2;
-			break;
-		case 1:
-			Result = mtlpp::LanguageVersion::Version1_1;
-			break;
-		case 0:
 		default:
-#if PLATFORM_MAC
-			Result = mtlpp::LanguageVersion::Version1_1;
-#else
-			Result = mtlpp::LanguageVersion::Version1_0;
-#endif
+			Result = mtlpp::LanguageVersion::Version2_2;
 			break;
 	}
-	
-	if (!FApplePlatformMisc::IsOSAtLeastVersion(MetalMacOSVersions[Version], MetaliOSVersions[Version], MetaliOSVersions[Version]))
-	{
-		FFormatNamedArguments Args;
-		Args.Add(TEXT("ShaderVersion"), FText::FromString(FString(StandardNames[Version])));
-#if PLATFORM_MAC
-		Args.Add(TEXT("RequiredOS"), FText::FromString(FString::Printf(TEXT("macOS %d.%d.%d"), MetalMacOSVersions[Version][0], MetalMacOSVersions[Version][1], MetalMacOSVersions[Version][2])));
-#else
-		Args.Add(TEXT("RequiredOS"), FText::FromString(FString::Printf(TEXT("%d.%d.%d"), MetaliOSVersions[Version][0], MetaliOSVersions[Version][1], MetaliOSVersions[Version][2])));
-#endif
-		FText LocalizedMsg = FText::Format(NSLOCTEXT("AGXRHI", "ShaderVersionUnsupported", "The current OS version does not support {ShaderVersion} required by the project. You must upgrade to {RequiredOS} to run this project."),Args);
-		
-		FText Title = NSLOCTEXT("AGXRHI", "ShaderVersionUnsupportedTitle", "Shader Version Unsupported");
-		FMessageDialog::Open(EAppMsgType::Ok, LocalizedMsg, &Title);
-		
-		FPlatformMisc::RequestExit(true);
-	}
-	
+
 	return Result;
 }
diff --git a/Engine/Source/Runtime/Apple/AGXRHI/Private/AGXState.cpp b/Engine/Source/Runtime/Apple/AGXRHI/Private/AGXState.cpp
--- a/Engine/Source/Runtime/Apple/AGXRHI/Private/AGXState.cpp
+++ b/Engine/Source/Runtime/Apple/AGXRHI/Private/AGXState.cpp
@@ -266,11 +266,6 @@ static FAGXSampler FindOrCreateSamplerState(const FSamplerStateInitializerRHI& I
 #if PLATFORM_MAC
 		[Desc setBorderColor:Initializer.BorderColor == 0 ? MTLSamplerBorderColorTransparentBlack : MTLSamplerBorderColorOpaqueWhite];
 #endif
-		if (FAGXCommandQueue::SupportsFeature(EAGXFeaturesIABs))
-		{
-			[Desc setSupportArgumentBuffers:YES];
-		}
-		
 		State = FAGXSampler([GMtlDevice newSamplerStateWithDescriptor:Desc], ns::Ownership::Assign);
 		
 		[Desc release];
diff --git a/Engine/Source/Runtime/Apple/AGXRHI/Private/AGXStateCache.cpp b/Engine/Source/Runtime/Apple/AGXRHI/Private/AGXStateCache.cpp
--- a/Engine/Source/Runtime/Apple/AGXRHI/Private/AGXStateCache.cpp
+++ b/Engine/Source/Runtime/Apple/AGXRHI/Private/AGXStateCache.cpp
@@ -7,7 +7,6 @@
 #include "AGXGraphicsPipelineState.h"
 #include "AGXStateCache.h"
 #include "AGXProfiler.h"
-#include "AGXCommandBuffer.h"
 
 #if PLATFORM_MAC
 	#ifndef UINT128_MAX
@@ -1712,7 +1711,7 @@ void FAGXStateCache::SetResourcesFromTables(ShaderType Shader, uint32 ShaderStag
 		const int32 BufferIndex = FMath::FloorLog2(LowestBitMask); // todo: This has a branch on zero, we know it could never be zero...
 		DirtyBits ^= LowestBitMask;
 		FAGXUniformBuffer* Buffer = (FAGXUniformBuffer*)GetBoundUniformBuffers(Frequency)[BufferIndex];
-		if (Buffer && !FAGXCommandQueue::SupportsFeature(EAGXFeaturesIABs))
+		if (Buffer)
 		{
 			check(BufferIndex < Shader->Bindings.ShaderResourceTable.ResourceTableLayoutHashes.Num());
 			check(Buffer->GetLayout().GetHash() == Shader->Bindings.ShaderResourceTable.ResourceTableLayoutHashes[BufferIndex]);
@@ -1871,16 +1870,14 @@ void FAGXStateCache::FlushVisibilityResults(FAGXCommandEncoder& CommandEncoder)
 #if PLATFORM_MAC
 	if(VisibilityResults && VisibilityResults->Buffer && VisibilityResults->Buffer.GetStorageMode() == mtlpp::StorageMode::Managed && VisibilityWritten && CommandEncoder.IsRenderCommandEncoderActive())
 	{
-		TRefCountPtr<FAGXFence> Fence = CommandEncoder.EndEncoding();
+		CommandEncoder.EndEncoding();
 		
         CommandEncoder.BeginBlitCommandEncoding();
-		CommandEncoder.WaitForFence(Fence);
 		
 		mtlpp::BlitCommandEncoder& Encoder = CommandEncoder.GetBlitCommandEncoder();
 
 		// METAL_GPUPROFILE(FAGXProfiler::GetProfiler()->EncodeBlit(CommandEncoder.GetCommandBufferStats(), __FUNCTION__));
 		MTLPP_VALIDATE(mtlpp::BlitCommandEncoder, Encoder, AGXSafeGetRuntimeDebuggingLevel() >= EAGXDebugLevelValidation, Synchronize(VisibilityResults->Buffer));
-		METAL_DEBUG_LAYER(EAGXDebugLevelFastValidation, CommandEncoder.GetBlitCommandEncoderDebugging().Synchronize(VisibilityResults->Buffer));
 		
 		VisibilityWritten = 0;
 	}
@@ -1995,9 +1992,9 @@ void FAGXStateCache::SetRenderPipelineState(FAGXCommandEncoder& CommandEncoder,
 		FAGXShaderPipeline* Pipeline = GetPipelineState();
 		EAGXShaderStages VertexStage = EAGXShaderStages::Vertex;
 		
-		FAGXDebugShaderResourceMask VertexMask = Pipeline->ResourceMask[EAGXShaderVertex];
-		TArray<uint32>& MinVertexBufferSizes = Pipeline->BufferDataSizes[EAGXShaderVertex];
-		const TMap<uint8, uint8>& VertexTexTypes = Pipeline->TextureTypes[EAGXShaderVertex];
+		FAGXDebugShaderResourceMask VertexMask = Pipeline->ResourceMask[SF_Vertex];
+		TArray<uint32>& MinVertexBufferSizes = Pipeline->BufferDataSizes[SF_Vertex];
+		const TMap<uint8, uint8>& VertexTexTypes = Pipeline->TextureTypes[SF_Vertex];
 		while(VertexMask.BufferMask)
 		{
 			uint32 Index = __builtin_ctz(VertexMask.BufferMask);
@@ -2045,9 +2042,9 @@ void FAGXStateCache::SetRenderPipelineState(FAGXCommandEncoder& CommandEncoder,
 			ensure(ShaderSamplers[VertexStage].Samplers[Index]);
 		}
 		
-		FAGXDebugShaderResourceMask FragmentMask = Pipeline->ResourceMask[EAGXShaderFragment];
-		TArray<uint32>& MinFragmentBufferSizes = Pipeline->BufferDataSizes[EAGXShaderFragment];
-		const TMap<uint8, uint8>& FragmentTexTypes = Pipeline->TextureTypes[EAGXShaderFragment];
+		FAGXDebugShaderResourceMask FragmentMask = Pipeline->ResourceMask[SF_Pixel];
+		TArray<uint32>& MinFragmentBufferSizes = Pipeline->BufferDataSizes[SF_Pixel];
+		const TMap<uint8, uint8>& FragmentTexTypes = Pipeline->TextureTypes[SF_Pixel];
 		while(FragmentMask.BufferMask)
 		{
 			uint32 Index = __builtin_ctz(FragmentMask.BufferMask);
@@ -2111,9 +2108,9 @@ void FAGXStateCache::SetComputePipelineState(FAGXCommandEncoder& CommandEncoder)
 		FAGXShaderPipeline* Pipeline = ComputeShader->GetPipeline();
 		check(Pipeline);
 		
-		FAGXDebugShaderResourceMask ComputeMask = Pipeline->ResourceMask[EAGXShaderCompute];
-		TArray<uint32>& MinComputeBufferSizes = Pipeline->BufferDataSizes[EAGXShaderCompute];
-		const TMap<uint8, uint8>& ComputeTexTypes = Pipeline->TextureTypes[EAGXShaderCompute];
+		FAGXDebugShaderResourceMask ComputeMask = Pipeline->ResourceMask[SF_Compute];
+		TArray<uint32>& MinComputeBufferSizes = Pipeline->BufferDataSizes[SF_Compute];
+		const TMap<uint8, uint8>& ComputeTexTypes = Pipeline->TextureTypes[SF_Compute];
 		while(ComputeMask.BufferMask)
 		{
 			uint32 Index = __builtin_ctz(ComputeMask.BufferMask);
@@ -2242,17 +2239,6 @@ void FAGXStateCache::CommitResourceTable(EAGXShaderStages const Frequency, mtlpp
 	}
 }
 
-FAGXBuffer& FAGXStateCache::GetDebugBuffer()
-{
-    if (!DebugBuffer)
-    {
-        // Assume worst case tiling (16x16) and render-target size (4096x4096) on iOS for now
-        uint32 NumTiles = PLATFORM_MAC ? 1 : 65536;
-        DebugBuffer = GetAGXDeviceContext().CreatePooledBuffer(FAGXPooledBufferArgs(NumTiles * sizeof(FAGXDebugInfo), BUF_Dynamic, mtlpp::StorageMode::Shared));
-    }
-    return DebugBuffer;
-}
-
 FTexture2DRHIRef FAGXStateCache::CreateFallbackDepthStencilSurface(uint32 Width, uint32 Height)
 {
 #if PLATFORM_MAC
diff --git a/Engine/Source/Runtime/Apple/AGXRHI/Private/AGXStateCache.h b/Engine/Source/Runtime/Apple/AGXRHI/Private/AGXStateCache.h
--- a/Engine/Source/Runtime/Apple/AGXRHI/Private/AGXStateCache.h
+++ b/Engine/Source/Runtime/Apple/AGXRHI/Private/AGXStateCache.h
@@ -116,7 +116,6 @@ public:
 	
 	bool PrepareToRestart(bool const bCurrentApplied);
 	
-	FAGXBuffer& GetDebugBuffer();
 	FAGXShaderParameterCache& GetShaderParameters(EAGXShaderStages const Stage) { return ShaderParameters[Stage]; }
 	FLinearColor const& GetBlendFactor() const { return BlendFactor; }
 	uint32 GetStencilRef() const { return StencilRef; }
@@ -287,7 +286,6 @@ private:
 	FTextureRHIRef DepthStencilResolve;
 	/** A fallback depth-stencil surface for draw calls that write to depth without a depth-stencil surface bound. */
 	FTexture2DRHIRef FallbackDepthStencilSurface;
-	FAGXBuffer DebugBuffer;
 	mtlpp::RenderPassDescriptor RenderPassDesc;
 	uint32 RasterBits;
     uint8 PipelineBits;
diff --git a/Engine/Source/Runtime/Apple/AGXRHI/Private/AGXTexture.cpp b/Engine/Source/Runtime/Apple/AGXRHI/Private/AGXTexture.cpp
--- a/Engine/Source/Runtime/Apple/AGXRHI/Private/AGXTexture.cpp
+++ b/Engine/Source/Runtime/Apple/AGXRHI/Private/AGXTexture.cpp
@@ -6,7 +6,6 @@
 
 #include "AGXRHIPrivate.h"
 #include "AGXProfiler.h" // for STAT_AGXTexturePageOffTime
-#include "AGXCommandBuffer.h"
 #include "RenderUtils.h"
 #include "Containers/ResourceArray.h"
 #include "Misc/ScopeRWLock.h"
@@ -399,71 +398,12 @@ void FAGXSurface::ReplaceTexture(FAGXContext& Context, FAGXTexture CurrentTextur
 
 void FAGXSurface::MakeAliasable(void)
 {
-	check(!bTextureView);
-	check(ImageSurfaceRef == nullptr);
-	
-	static bool bSupportsHeaps = GetAGXDeviceContext().SupportsFeature(EAGXFeaturesHeaps);
-	if (bSupportsHeaps && Texture.GetStorageMode() == mtlpp::StorageMode::Private && Texture.GetHeap())
-	{
-		if (MSAATexture && (MSAATexture != Texture) && !MSAATexture.IsAliasable())
-		{
-			MSAATexture.MakeAliasable();
-#if STATS || ENABLE_LOW_LEVEL_MEM_TRACKER
-			AGXLLM::LogAliasTexture(MSAATexture);
-#endif
-		}
-		if (!Texture.IsAliasable())
-		{
-			Texture.MakeAliasable();
-#if STATS || ENABLE_LOW_LEVEL_MEM_TRACKER
-			AGXLLM::LogAliasTexture(Texture);
-#endif
-		}
-	}
+	// TODO
 }
 
 void FAGXSurface::MakeUnAliasable(void)
 {
-	check(!bTextureView);
-	check(ImageSurfaceRef == nullptr);
-	
-	static bool bSupportsHeaps = GetAGXDeviceContext().SupportsFeature(EAGXFeaturesHeaps);
-	if (bSupportsHeaps && Texture.GetStorageMode() == mtlpp::StorageMode::Private && Texture.GetHeap() && Texture.IsAliasable())
-	{
-		FAGXTexture OldTexture = Texture;
-		Texture = Reallocate(Texture, mtlpp::TextureUsage::Unknown);
-		AGXSafeReleaseMetalTexture(this, OldTexture);
-		if (MSAATexture && (MSAATexture != OldTexture))
-		{
-			FAGXTexture OldMSAATexture = MSAATexture;
-			MSAATexture = Reallocate(MSAATexture, mtlpp::TextureUsage::Unknown);
-			AGXSafeReleaseMetalTexture(this, OldMSAATexture);
-		}
-		else if(MSAATexture)
-		{
-			MSAATexture = Texture;
-		}
-		
-		for (FAGXShaderResourceView* SRV : SRVs)
-		{
-			if(OldTexture != SRV->TextureView->Texture)
-			{
-				AGXSafeReleaseMetalTexture(this, SRV->TextureView->Texture, true, false);
-			}
-			SRV->TextureView->Texture = nil;
-			
-			SRV->TextureView->MSAATexture = nil;
-			
-			if (SRV->Format == PF_Unknown)
-			{
-				SRV->TextureView->Init(*this, NSMakeRange(SRV->MipLevel, SRV->NumMips));
-			}
-			else
-			{
-				SRV->TextureView->Init(*this, NSMakeRange(SRV->MipLevel, SRV->NumMips), (EPixelFormat)SRV->Format, SRV->bSRGBForceDisable);
-			}
-		}
-	}
+	// TODO
 }
 
 void FAGXSurface::Init(FAGXSurface& Source, NSRange MipRange)
diff --git a/Engine/Source/Runtime/Apple/AGXRHI/Private/AGXTransitionData.cpp b/Engine/Source/Runtime/Apple/AGXRHI/Private/AGXTransitionData.cpp
--- a/Engine/Source/Runtime/Apple/AGXRHI/Private/AGXTransitionData.cpp
+++ b/Engine/Source/Runtime/Apple/AGXRHI/Private/AGXTransitionData.cpp
@@ -12,7 +12,7 @@
 
 #pragma mark - AGX RHI Resource Transition Data Definitions -
 
-#define UE_METAL_TRANSITION_DATA_USES_EXPLICIT_FENCING					0
+// TODO: Put AGXRHI resource transition data definitions here
 
 
 //------------------------------------------------------------------------------
@@ -21,9 +21,9 @@
 
 
 FAGXTransitionData::FAGXTransitionData(ERHIPipeline                         InSrcPipelines,
-										   ERHIPipeline                         InDstPipelines,
-										   ERHITransitionCreateFlags            InCreateFlags,
-										   TArrayView<const FRHITransitionInfo> InInfos)
+									   ERHIPipeline                         InDstPipelines,
+									   ERHITransitionCreateFlags            InCreateFlags,
+									   TArrayView<const FRHITransitionInfo> InInfos)
 {
 	SrcPipelines   = InSrcPipelines;
 	DstPipelines   = InDstPipelines;
@@ -32,48 +32,10 @@ FAGXTransitionData::FAGXTransitionData(ERHIPipeline                         InSr
 	bCrossPipeline = (SrcPipelines != DstPipelines);
 
 	Infos.Append(InInfos.GetData(), InInfos.Num());
-
-	// TODO: Determine whether the AGX RHI needs to create a separate, per-transition fence.
-#if UE_METAL_TRANSITION_DATA_USES_EXPLICIT_FENCING
-	if (bCrossPipeline && !EnumHasAnyFlags(CreateFlags, ERHITransitionCreateFlags::NoFence))
-	{
-		// Get the current context pointer.
-		FAGXContext* Context = GetAGXDeviceContext().GetCurrentContext();
-
-		// Get the current render pass fence.
-		TRefCountPtr<FAGXFence> const& MetalFence = Context->GetCurrentRenderPass().End();
-
-		// Write it again as we may wait on this fence in two different encoders.
-		Context->GetCurrentRenderPass().Update(MetalFence);
-
-		// Write it into the transition data.
-		Fence = MetalFence;
-		Fence->AddRef();
-
-		if (GSupportsEfficientAsyncCompute)
-		{
-			static_cast<FAGXRHICommandContext*>(RHIGetDefaultContext())->RHISubmitCommandsHint();
-		}
-	}
-#endif // UE_METAL_TRANSITION_DATA_USES_EXPLICIT_FENCING
 }
 
 void FAGXTransitionData::BeginResourceTransitions() const
 {
-	// TODO: Determine whether the AGX RHI needs to do anything here.
-#if UE_METAL_TRANSITION_DATA_USES_EXPLICIT_FENCING
-	if (Fence.IsValid())
-	{
-		FAGXContext* Context = GetAGXDeviceContext().GetCurrentContext();
-
-		if (Context->GetCurrentCommandBuffer())
-		{
-			Context->SubmitCommandsHint(EAGXSubmitFlagsNone);
-		}
-
-		Context->GetCurrentRenderPass().Begin(Fence);
-	}
-#endif // UE_METAL_TRANSITION_DATA_USES_EXPLICIT_FENCING
 }
 
 void FAGXTransitionData::EndResourceTransitions() const
diff --git a/Engine/Source/Runtime/Apple/AGXRHI/Private/AGXTransitionData.h b/Engine/Source/Runtime/Apple/AGXRHI/Private/AGXTransitionData.h
--- a/Engine/Source/Runtime/Apple/AGXRHI/Private/AGXTransitionData.h
+++ b/Engine/Source/Runtime/Apple/AGXRHI/Private/AGXTransitionData.h
@@ -24,15 +24,15 @@ class FAGXTransitionData
 {
 public:
 	explicit FAGXTransitionData(ERHIPipeline                         InSrcPipelines,
-								  ERHIPipeline                         InDstPipelines,
-								  ERHITransitionCreateFlags            InCreateFlags,
-								  TArrayView<const FRHITransitionInfo> InInfos);
+								ERHIPipeline                         InDstPipelines,
+								ERHITransitionCreateFlags            InCreateFlags,
+								TArrayView<const FRHITransitionInfo> InInfos);
 
 	// The default destructor is sufficient.
 	~FAGXTransitionData() = default;
 
 	// Disallow default, copy and move constructors.
-	FAGXTransitionData()                             = delete;
+	FAGXTransitionData()                           = delete;
 	FAGXTransitionData(const FAGXTransitionData&)  = delete;
 	FAGXTransitionData(const FAGXTransitionData&&) = delete;
 
@@ -47,6 +47,5 @@ private:
 	ERHIPipeline              DstPipelines   = ERHIPipeline::Num;
 	ERHITransitionCreateFlags CreateFlags    = ERHITransitionCreateFlags::None;
 	bool                      bCrossPipeline = false;
-	FAGXTransitionInfoArray Infos          = {};
-	TRefCountPtr<FAGXFence> Fence          = nullptr;
+	FAGXTransitionInfoArray   Infos          = {};
 };
diff --git a/Engine/Source/Runtime/Apple/AGXRHI/Private/AGXUAV.cpp b/Engine/Source/Runtime/Apple/AGXRHI/Private/AGXUAV.cpp
--- a/Engine/Source/Runtime/Apple/AGXRHI/Private/AGXUAV.cpp
+++ b/Engine/Source/Runtime/Apple/AGXRHI/Private/AGXUAV.cpp
@@ -3,7 +3,6 @@
 
 #include "AGXRHIPrivate.h"
 #include "AGXRHIStagingBuffer.h"
-#include "AGXCommandBuffer.h"
 #include "RenderUtils.h"
 #include "ClearReplacementShaders.h"
 #include "AGXTransitionData.h"
diff --git a/Engine/Source/Runtime/Apple/AGXRHI/Private/AGXVertexBuffer.cpp b/Engine/Source/Runtime/Apple/AGXRHI/Private/AGXVertexBuffer.cpp
--- a/Engine/Source/Runtime/Apple/AGXRHI/Private/AGXVertexBuffer.cpp
+++ b/Engine/Source/Runtime/Apple/AGXRHI/Private/AGXVertexBuffer.cpp
@@ -6,7 +6,6 @@
 
 #include "AGXRHIPrivate.h"
 #include "AGXProfiler.h"
-#include "AGXCommandBuffer.h"
 #include "AGXCommandQueue.h"
 #include "Containers/ResourceArray.h"
 #include "RenderUtils.h"
@@ -108,13 +107,12 @@ void FAGXRHIBuffer::Swap(FAGXRHIBuffer& Other)
 
 static bool CanUsePrivateMemory()
 {
-	return (FAGXCommandQueue::SupportsFeature(EAGXFeaturesEfficientBufferBlits) || FAGXCommandQueue::SupportsFeature(EAGXFeaturesIABs));
+	return (FAGXCommandQueue::SupportsFeature(EAGXFeaturesEfficientBufferBlits));
 }
 
 bool FAGXRHIBuffer::UsePrivateMemory() const
 {
-	return (FAGXCommandQueue::SupportsFeature(EAGXFeaturesEfficientBufferBlits) && EnumHasAnyFlags(Usage, BUF_Dynamic|BUF_Static))
-	|| (FAGXCommandQueue::SupportsFeature(EAGXFeaturesIABs) && EnumHasAnyFlags(Usage, BUF_ShaderResource|BUF_UnorderedAccess));
+	return (FAGXCommandQueue::SupportsFeature(EAGXFeaturesEfficientBufferBlits) && EnumHasAnyFlags(Usage, BUF_Dynamic|BUF_Static));
 }
 
 FAGXRHIBuffer::FAGXRHIBuffer(uint32 InSize, EBufferUsageFlags InUsage, EAGXBufferUsage InAgxUsage, ERHIResourceType InType)
@@ -764,7 +762,7 @@ FBufferRHIRef FAGXDynamicRHI::RHICreateBuffer(uint32 Size, EBufferUsageFlags Usa
 	{
 		check (!Buffer->TransferBuffer);
 
-		if (GAGXBufferZeroFill && !FAGXCommandQueue::SupportsFeature(EAGXFeaturesFences))
+		if (GAGXBufferZeroFill)
 		{
 			for(FAGXRHIBuffer::FAGXBufferAndViews& Backing : Buffer->BufferPool)
 			{
diff --git a/Engine/Source/Runtime/Apple/AGXRHI/Private/AGXViewport.cpp b/Engine/Source/Runtime/Apple/AGXRHI/Private/AGXViewport.cpp
--- a/Engine/Source/Runtime/Apple/AGXRHI/Private/AGXViewport.cpp
+++ b/Engine/Source/Runtime/Apple/AGXRHI/Private/AGXViewport.cpp
@@ -479,29 +479,22 @@ void FAGXViewport::Present(FAGXCommandQueue& CommandQueue, bool bLockToVsync)
 							NSUInteger Width = FMath::Min(Src.GetWidth(), Dst.GetWidth());
 							NSUInteger Height = FMath::Min(Src.GetHeight(), Dst.GetHeight());
 							
-							mtlpp::BlitCommandEncoder Encoder = CurrentCommandBuffer.BlitCommandEncoder();
-							check(Encoder.GetPtr());
-#if MTLPP_CONFIG_VALIDATE && METAL_DEBUG_OPTIONS
-							FAGXBlitCommandEncoderDebugging Debugging;
-							if (AGXSafeGetRuntimeDebuggingLevel() >= EAGXDebugLevelFastValidation)
-							{
-								FAGXCommandBufferDebugging CmdDebug = FAGXCommandBufferDebugging::Get(CurrentCommandBuffer);
-								Debugging = FAGXBlitCommandEncoderDebugging(Encoder, CmdDebug);
-							}
-#endif
-							METAL_GPUPROFILE(Profiler->EncodeBlit(Stats, __FUNCTION__));
-
-							Encoder.Copy(Src, 0, 0, mtlpp::Origin(0, 0, 0), mtlpp::Size(Width, Height, 1), Dst, 0, 0, mtlpp::Origin(0, 0, 0));
-							METAL_DEBUG_LAYER(EAGXDebugLevelFastValidation, Debugging.Copy(Src, 0, 0, mtlpp::Origin(0, 0, 0), mtlpp::Size(Width, Height, 1), Dst, 0, 0, mtlpp::Origin(0, 0, 0)));
+							id<MTLBlitCommandEncoder> Encoder = [CurrentCommandBuffer.GetPtr() blitCommandEncoder];
+							check(Encoder);
 
-							Encoder.EndEncoding();
-							METAL_DEBUG_LAYER(EAGXDebugLevelFastValidation, Debugging.EndEncoder());
-
-							mtlpp::CommandBufferHandler H = [Src, Dst](const mtlpp::CommandBuffer &) {
-								// void
-							};
+							METAL_GPUPROFILE(Profiler->EncodeBlit(Stats, __FUNCTION__));
 
-							CurrentCommandBuffer.AddCompletedHandler(H);
+							[Encoder	copyFromTexture:Src.GetPtr()
+											sourceSlice:0
+											sourceLevel:0
+										   sourceOrigin:MTLOriginMake(0, 0, 0)
+											 sourceSize:MTLSizeMake(Width, Height, 1)
+											  toTexture:Dst.GetPtr()
+									   destinationSlice:0
+									   destinationLevel:0
+									  destinationOrigin:MTLOriginMake(0, 0, 0)];
+
+							[Encoder endEncoding];
 
 							[Drawable release];
 							Drawable = nil;
diff --git a/Engine/Source/Runtime/Apple/AGXRHI/Private/IOS/IOSPlatformAGXConfig.cpp b/Engine/Source/Runtime/Apple/AGXRHI/Private/IOS/IOSPlatformAGXConfig.cpp
new file mode 100644
--- /dev/null
+++ b/Engine/Source/Runtime/Apple/AGXRHI/Private/IOS/IOSPlatformAGXConfig.cpp
@@ -0,0 +1,15 @@
+// Copyright Epic Games, Inc. All rights reserved.
+
+#include "IOSPlatformAGXConfig.h"
+#include "RHI.h"
+#include "AGXRHIPrivate.h"
+
+
+//------------------------------------------------------------------------------
+
+#pragma mark - iOS Platform AGXRHI Config Routines -
+
+void FIOSAGXConfig::PopulateRHIGlobals()
+{
+    // TODO
+}
diff --git a/Engine/Source/Runtime/Apple/AGXRHI/Private/IOS/IOSPlatformAGXConfig.h b/Engine/Source/Runtime/Apple/AGXRHI/Private/IOS/IOSPlatformAGXConfig.h
new file mode 100644
--- /dev/null
+++ b/Engine/Source/Runtime/Apple/AGXRHI/Private/IOS/IOSPlatformAGXConfig.h
@@ -0,0 +1,19 @@
+// Copyright Epic Games, Inc. All rights reserved.
+
+#pragma once
+
+#include "HAL/Platform.h"
+
+#if PLATFORM_TVOS
+#include "../TVOS/TVOSPlatformAGXConfig.h"
+#else
+
+class FIOSAGXConfig
+{
+public:
+	static void PopulateRHIGlobals();
+};
+
+typedef FIOSAGXConfig FPlatformAGXConfig;
+
+#endif // PLATFORM_TVOS
diff --git a/Engine/Source/Runtime/Apple/AGXRHI/Private/Mac/MacPlatformAGXConfig.cpp b/Engine/Source/Runtime/Apple/AGXRHI/Private/Mac/MacPlatformAGXConfig.cpp
new file mode 100644
--- /dev/null
+++ b/Engine/Source/Runtime/Apple/AGXRHI/Private/Mac/MacPlatformAGXConfig.cpp
@@ -0,0 +1,617 @@
+// Copyright Epic Games, Inc. All rights reserved.
+
+#include "MacPlatformAGXConfig.h"
+#include "RHI.h"
+#include "AGXRHIPrivate.h"
+
+
+//------------------------------------------------------------------------------
+
+#pragma mark - Mac Platform AGXRHI Config Definitions -
+
+#define MAC_PLATFORM_AGXRHI_ADAPTER_DRIVER_ON_DENY_LIST                                 AGXRHI_FALSE
+#define MAC_PLATFORM_AGXRHI_DEVICE_IS_AMD_PRE_GCN_ARCHITECTURE                          AGXRHI_FALSE
+#define MAC_PLATFORM_AGXRHI_DEVICE_IS_INTEGRATED                                        AGXRHI_FALSE
+#define MAC_PLATFORM_AGXRHI_ADAPTER_NAME                                                TEXT("Unknown AGX Adapter")
+#define MAC_PLATFORM_AGXRHI_ADAPTER_DRIVER_DATE                                         TEXT("Unknown Metal plugin date")
+#define MAC_PLATFORM_AGXRHI_ADAPTER_USER_DRIVER_VERSION                                 TEXT("Unknown Metal plugin version")
+#define MAC_PLATFORM_AGXRHI_ADAPTER_INTERNAL_DRIVER_VERSION                             TEXT("1.0.1")
+#define MAC_PLATFORM_AGXRHI_DEVICE_ID                                                   2275
+#define MAC_PLATFORM_AGXRHI_DEVICE_REVISION                                             1
+
+#define MAC_PLATFORM_AGXRHI_HARDWARE_HIDDEN_SURFACE_REMOVAL                             AGXRHI_FALSE
+#define MAC_PLATFORM_AGXRHI_IS_INITIALIZED                                              AGXRHI_FALSE
+#define MAC_PLATFORM_AGXRHI_ATTACHMENT_VARIABLE_RATE_SHADING_ENABLED                    AGXRHI_TRUE
+#define MAC_PLATFORM_AGXRHI_FORCE_NO_DELETION_LATENCY_FOR_STREAMING_TEXTURES            AGXRHI_FALSE
+#define MAC_PLATFORM_AGXRHI_IS_HDR_ENABLED                                              AGXRHI_FALSE
+#define MAC_PLATFORM_AGXRHI_LAZY_SHADER_CODE_LOADING                                    AGXRHI_FALSE
+#define MAC_PLATFORM_AGXRHI_NEEDS_EXTRA_DELETION_LATENCY                                AGXRHI_FALSE
+#define MAC_PLATFORM_AGXRHI_NEEDS_UNATLASED_CSM_DEPTHS_WORKAROUND                       AGXRHI_FALSE
+#define MAC_PLATFORM_AGXRHI_REQUIRES_RENDER_TARGET_FOR_PIXEL_SHADER_UAVS                AGXRHI_TRUE
+#define MAC_PLATFORM_AGXRHI_SUPPORTS_ARRAY_INDEX_FROM_ANY_SHADER                        AGXRHI_TRUE
+#define MAC_PLATFORM_AGXRHI_SUPPORTS_ASYNC_TEXTURE_CREATION                             AGXRHI_FALSE
+#define MAC_PLATFORM_AGXRHI_SUPPORTS_ATOMIC_UINT64                                      AGXRHI_FALSE
+#define MAC_PLATFORM_AGXRHI_SUPPORTS_ATTACHMENT_VARIABLE_RATE_SHADING                   AGXRHI_FALSE
+#define MAC_PLATFORM_AGXRHI_SUPPORTS_BACKBUFFER_WITH_CUSTOM_DEPTHSTENCIL                AGXRHI_TRUE
+#define MAC_PLATFORM_AGXRHI_SUPPORTS_BASE_VERTEX_INDEX                                  AGXRHI_TRUE
+#define MAC_PLATFORM_AGXRHI_SUPPORTS_COMPLEX_VARIABLE_RATE_SHADING_COMBINER_OPS         AGXRHI_FALSE
+#define MAC_PLATFORM_AGXRHI_SUPPORTS_CONSERVATIVE_RASTERIZATION                         AGXRHI_FALSE
+#define MAC_PLATFORM_AGXRHI_SUPPORTS_COPY_TO_TEXTURE_MULTIPLE_MIPS                      AGXRHI_TRUE
+#define MAC_PLATFORM_AGXRHI_SUPPORTS_DEPTH_UAV                                          AGXRHI_FALSE
+#define MAC_PLATFORM_AGXRHI_SUPPORTS_DIRECT_GPU_MEMORY_LOCK                             AGXRHI_FALSE
+#define MAC_PLATFORM_AGXRHI_SUPPORTS_DRAW_INDIRECT                                      AGXRHI_TRUE
+#define MAC_PLATFORM_AGXRHI_SUPPORTS_DYNAMIC_RESOLUTION                                 AGXRHI_TRUE
+#define MAC_PLATFORM_AGXRHI_SUPPORTS_EFFICIENT_UPLOAD_ON_RESOURCE_CREATION              AGXRHI_FALSE
+#define MAC_PLATFORM_AGXRHI_SUPPORTS_EXACT_OCCLUSION_QUERIES                            AGXRHI_TRUE
+#define MAC_PLATFORM_AGXRHI_SUPPORTS_EXPLICIT_FMASK                                     AGXRHI_FALSE
+#define MAC_PLATFORM_AGXRHI_SUPPORTS_EXPLICIT_HTILE                                     AGXRHI_FALSE
+#define MAC_PLATFORM_AGXRHI_SUPPORTS_FIRST_INSTANCE                                     AGXRHI_TRUE
+#define MAC_PLATFORM_AGXRHI_SUPPORTS_FRAME_CYCLES_BUBBLES_REMOVAL                       AGXRHI_TRUE
+#define MAC_PLATFORM_AGXRHI_SUPPORTS_GPU_TIMESTAMP_BUBBLES_REMOVAL                      AGXRHI_FALSE
+#define MAC_PLATFORM_AGXRHI_SUPPORTS_HDR_OUTPUT                                         AGXRHI_TRUE
+#define MAC_PLATFORM_AGXRHI_SUPPORTS_INLINE_RAY_TRACING                                 AGXRHI_FALSE
+#define MAC_PLATFORM_AGXRHI_SUPPORTS_LATE_VARIABLE_RATE_SHADING_UPDATE                  AGXRHI_FALSE
+#define MAC_PLATFORM_AGXRHI_SUPPORTS_LAZY_SHADER_CODE_LOADING                           AGXRHI_FALSE
+#define MAC_PLATFORM_AGXRHI_SUPPORTS_MESH_SHADERS_TIER_0                                AGXRHI_FALSE
+#define MAC_PLATFORM_AGXRHI_SUPPORTS_MESH_SHADERS_TIER_1                                AGXRHI_FALSE
+#define MAC_PLATFORM_AGXRHI_SUPPORTS_MSAA_DEPTH_SAMPLE_ACCESS                           AGXRHI_FALSE
+#define MAC_PLATFORM_AGXRHI_SUPPORTS_MULTITHREADED_SHADER_CREATION                      AGXRHI_TRUE
+#define MAC_PLATFORM_AGXRHI_SUPPORTS_MULTITHREADING                                     AGXRHI_TRUE
+#define MAC_PLATFORM_AGXRHI_SUPPORTS_PARALLEL_RHI_EXECUTE                               AGXRHI_TRUE
+#define MAC_PLATFORM_AGXRHI_SUPPORTS_PIPELINE_FILE_CACHE                                AGXRHI_TRUE
+#define MAC_PLATFORM_AGXRHI_SUPPORTS_PIPELINE_STATE_SORT_KEY                            AGXRHI_FALSE
+#define MAC_PLATFORM_AGXRHI_SUPPORTS_PIPELINE_VARIABLE_RATE_SHADING                     AGXRHI_FALSE
+#define MAC_PLATFORM_AGXRHI_SUPPORTS_PIXEL_SHADER_UAVS                                  AGXRHI_TRUE
+#define MAC_PLATFORM_AGXRHI_SUPPORTS_PRIMITIVE_SHADERS                                  AGXRHI_FALSE
+#define MAC_PLATFORM_AGXRHI_SUPPORTS_QUAD_TOPOLOGY                                      AGXRHI_FALSE
+#define MAC_PLATFORM_AGXRHI_SUPPORTS_RAY_TRACING                                        AGXRHI_FALSE
+#define MAC_PLATFORM_AGXRHI_SUPPORTS_RAY_TRACING_AMD_HIT_TOKEN                          AGXRHI_FALSE
+#define MAC_PLATFORM_AGXRHI_SUPPORTS_RAY_TRACING_ASYNC_BUILD_ACCELERATION_STRUCTURE     AGXRHI_FALSE
+#define MAC_PLATFORM_AGXRHI_SUPPORTS_RAY_TRACING_DISPATCH_INDIRECT                      AGXRHI_FALSE
+#define MAC_PLATFORM_AGXRHI_SUPPORTS_RAY_TRACING_PSO_ADDITIONS                          AGXRHI_FALSE
+#define MAC_PLATFORM_AGXRHI_SUPPORTS_RAY_TRACING_SHADERS                                AGXRHI_FALSE
+#define MAC_PLATFORM_AGXRHI_SUPPORTS_RECT_TOPOLOGY                                      AGXRHI_FALSE
+#define MAC_PLATFORM_AGXRHI_SUPPORTS_RESOLVE_CUBEMAP_FACES                              AGXRHI_FALSE
+#define MAC_PLATFORM_AGXRHI_SUPPORTS_RESUMMARIZE_HTILE                                  AGXRHI_FALSE
+#define MAC_PLATFORM_AGXRHI_SUPPORTS_RHI_ON_TASK_THREAD                                 AGXRHI_FALSE
+#define MAC_PLATFORM_AGXRHI_SUPPORTS_RHI_THREAD                                         AGXRHI_TRUE
+#define MAC_PLATFORM_AGXRHI_SUPPORTS_RWTEXTURE_BUFFERS                                  AGXRHI_TRUE
+#define MAC_PLATFORM_AGXRHI_SUPPORTS_SHADER_TIMESTAMP                                   AGXRHI_FALSE
+#define MAC_PLATFORM_AGXRHI_SUPPORTS_STENCIL_REF_FROM_PIXEL_SHADER                      AGXRHI_FALSE
+#define MAC_PLATFORM_AGXRHI_SUPPORTS_TEXTURE_STREAMING                                  AGXRHI_TRUE
+#define MAC_PLATFORM_AGXRHI_SUPPORTS_UAV_FORMAT_ALIASING                                AGXRHI_FALSE
+#define MAC_PLATFORM_AGXRHI_SUPPORTS_UPDATE_FROM_BUFFER_TEXTURE                         AGXRHI_FALSE
+#define MAC_PLATFORM_AGXRHI_SUPPORTS_VARIABLE_RATE_SHADING_ATTACHMENT_ARRAY_TEXTURES    AGXRHI_FALSE
+#define MAC_PLATFORM_AGXRHI_SUPPORTS_WAVE_OPERATIONS                                    AGXRHI_FALSE
+#define MAC_PLATFORM_AGXRHI_THREAD_NEEDS_KICKING                                        AGXRHI_FALSE
+#define MAC_PLATFORM_AGXRHI_VARIABLE_RATE_SHADING_ENABLED                               AGXRHI_TRUE
+#define MAC_PLATFORM_AGXRHI_SUPPORTS_DEPTH_BOUNDS_TEST                                  AGXRHI_FALSE
+#define MAC_PLATFORM_AGXRHI_SUPPORTS_EFFICIENT_ASYNC_COMPUTE                            AGXRHI_TRUE
+#define MAC_PLATFORM_AGXRHI_SUPPORTS_IMAGE_EXTERNAL                                     AGXRHI_FALSE
+#define MAC_PLATFORM_AGXRHI_SUPPORTS_MOBILE_MULTIVIEW                                   AGXRHI_FALSE
+#define MAC_PLATFORM_AGXRHI_SUPPORTS_PARALLEL_OCCLUSION_QUERIES                         AGXRHI_TRUE
+#define MAC_PLATFORM_AGXRHI_SUPPORTS_PARALLEL_RENDERING_TASKS_WITH_SEPARATE_RHI_THREAD  AGXRHI_TRUE
+#define MAC_PLATFORM_AGXRHI_SUPPORTS_PIXEL_LOCAL_STORAGE                                AGXRHI_FALSE
+#define MAC_PLATFORM_AGXRHI_SUPPORTS_QUADBUFFER_STEREO                                  AGXRHI_FALSE
+#define MAC_PLATFORM_AGXRHI_SUPPORTS_RENDER_DEPTH_TARGETABLE_SHADER_RESOURCES           AGXRHI_TRUE
+#define MAC_PLATFORM_AGXRHI_SUPPORTS_RENDER_TARGET_FORMAT_PF_FLOAT_RGBA                 AGXRHI_TRUE
+#define MAC_PLATFORM_AGXRHI_SUPPORTS_RENDER_TARGET_FORMAT_PF_G8                         AGXRHI_FALSE
+#define MAC_PLATFORM_AGXRHI_SUPPORTS_SEPARATE_RENDER_TARGET_BLEND_STATE                 AGXRHI_FALSE // (GMaxRHIFeatureLevel >= ERHIFeatureLevel::SM5)
+#define MAC_PLATFORM_AGXRHI_SUPPORTS_SHADER_DEPTHSTENCIL_FETCH                          AGXRHI_FALSE
+#define MAC_PLATFORM_AGXRHI_SUPPORTS_SHADER_FRAMEBUFFER_FETCH                           AGXRHI_TRUE
+#define MAC_PLATFORM_AGXRHI_SUPPORTS_SHADER_MRT_FRAMEBUFFER_FETCH                       AGXFHI_TRUE // GSupportsShaderFramebufferFetch
+#define MAC_PLATFORM_AGXRHI_SUPPORTS_TEXTURE_3D                                         AGXRHI_TRUE
+#define MAC_PLATFORM_AGXRHI_SUPPORTS_TIMESTAMP_RENDER_QUERIES                           AGXRHI_TRUE
+#define MAC_PLATFORM_AGXRHI_SUPPORTS_TRANSIENT_RESOURCE_ALIASING                        AGXRHI_FALSE
+#define MAC_PLATFORM_AGXRHI_SUPPORTS_VOLUME_TEXTURE_RENDERING                           AGXRHI_TRUE
+#define MAC_PLATFORM_AGXRHI_SUPPORTS_WIDE_MRT                                           AGXRHI_TRUE
+#define MAC_PLATFORM_AGXRHI_USING_NULL_RHI                                              AGXRHI_FALSE
+#define MAC_PLATFORM_AGXRHI_HDR_DISPLAY_OUTPUT_FORMAT                                   PF_Unknown
+#define MAC_PLATFORM_AGXRHI_VARIABLE_RATE_SHADING_IMAGE_FORMAT                          PF_Unknown
+#define MAC_PLATFORM_AGXRHI_MERGEABLE_ACCESS_MASK                                       (ERHIAccess::ReadOnlyMask | ERHIAccess::UAVMask)
+#define MAC_PLATFORM_AGXRHI_MULTI_PIPELINE_MERGEABLE_ACCESS_MASK                        ERHIAccess::Unknown
+#define MAC_PLATFORM_AGXRHI_MAX_FEATURE_LEVEL                                           ERHIFeatureLevel::SM5
+#define MAC_PLATFORM_AGXRHI_MAX_SHADER_PLATFORM                                         SP_METAL_SM5
+#define MAC_PLATFORM_AGXRHI_SHADER_PLATFORM_FOR_FEATURE_LEVEL                           { SP_NumPlatforms }
+#define MAC_PLATFORM_AGXRHI_VARIABLE_RATE_SHADING_IMAGE_DATA_TYPE                       VRSImage_NotSupported
+#define MAC_PLATFORM_AGXRHI_MAX_DISPATCH_THREAD_GROUPS_PER_DIMENSION                    { MAX_uint16, MAX_uint16, MAX_uint16 }
+#define MAC_PLATFORM_AGXRHI_MAX_COMPUTE_DISPATCH_DIMENSION                              MAX_uint16
+#define MAC_PLATFORM_AGXRHI_MAX_CUBE_TEXTURE_DIMENSIONS                                 16384
+#define MAC_PLATFORM_AGXRHI_MAX_SHADOW_DEPTH_BUFFER_SIZE_X                              GMaxTextureDimensions
+#define MAC_PLATFORM_AGXRHI_MAX_SHADOW_DEPTH_BUFFER_SIZE_Y                              GMaxTextureDimensions
+#define MAC_PLATFORM_AGXRHI_MAX_TEXTURE_ARRAY_LAYERS                                    2048
+#define MAC_PLATFORM_AGXRHI_MAX_TEXTURE_DIMENSIONS                                      16384
+#define MAC_PLATFORM_AGXRHI_MAX_TEXTURE_MIP_COUNT                                       0
+#define MAC_PLATFORM_AGXRHI_MAX_TEXTURE_SAMPLERS                                        16
+#define MAC_PLATFORM_AGXRHI_MAX_VOLUME_TEXTURE_DIMENSIONS                               2048
+#define MAC_PLATFORM_AGXRHI_MAX_WORKGROUP_INVOCATIONS                                   1024
+#define MAC_PLATFORM_AGXRHI_POOL_SIZE_VRAM_PERCENTAGE                                   0
+#define MAC_PLATFORM_AGXRHI_MAXIMUM_RECCOMMENDED_OUSTANDING_OCCLUSION_QUERIES           MAX_int32
+#define MAC_PLATFORM_AGXRHI_MAXIMUM_WAVE_SIZE                                           0
+#define MAC_PLATFORM_AGXRHI_MINIMUM_WAVE_SIZE                                           0
+#define MAC_PLATFORM_AGXRHI_PERSISTENT_THREADGROUP_COUNT                                0
+#define MAC_PLATFORM_AGXRHI_VARIABLE_RATE_SHADING_IMAGE_TILE_MAX_HEIGHT                 0
+#define MAC_PLATFORM_AGXRHI_VARIABLE_RATE_SHADING_IMAGE_TILE_MAX_WIDTH                  0
+#define MAC_PLATFORM_AGXRHI_VARIABLE_RATE_SHADING_IMAGE_TILE_MIN_HEIGHT                 0
+#define MAC_PLATFORM_AGXRHI_VARIABLE_RATE_SHADING_IMAGE_TILE_MIN_WIDTH                  0
+#define MAC_PLATFORM_AGXRHI_MAX_BUFFER_DIMENSIONS                                       (1 << 27)
+#define MAC_PLATFORM_AGXRHI_MAX_COMPUTE_SHARED_MEMORY                                   32768
+#define MAC_PLATFORM_AGXRHI_TEXTURE_POOL_SIZE                                           0
+#define MAC_PLATFORM_AGXRHI_RAY_TRACING_ACCELERATION_STRUCTURE_ALIGNMENT                0
+#define MAC_PLATFORM_AGXRHI_RAY_TRACING_INSTANCE_DESCRIPTOR_SIZE                        0
+#define MAC_PLATFORM_AGXRHI_RAY_TRACING_SCRATCH_BUFFER_ALIGNMENT                        0
+#define MAC_PLATFORM_AGXRHI_RAY_TRACING_SHADER_TABLE_ALIGNMENT                          0
+#define MAC_PLATFORM_AGXRHI_VENDOR_ID                                                   0x106B
+#define MAC_PLATFORM_AGXRHI_DEMOTED_LOCAL_MEMORY_SIZE                                   0
+#define MAC_PLATFORM_AGXRHI_PRESENT_COUNTER                                             0
+#define MAC_PLATFORM_AGXRHI_TRANSITION_PRIVATE_DATA_ALIGN_IN_BYTES                      alignof(FAGXTransitionData)
+#define MAC_PLATFORM_AGXRHI_TRANSITION_PRIVATE_DATA_SIZE_IN_BYTES                       sizeof(FAGXTransitionData)
+
+
+//------------------------------------------------------------------------------
+
+#pragma mark - AGXRHI Internal Utility Routines -
+
+/**
+ * @fn AGXUtilCreateCString
+ * 
+ * @brief
+ * This routine extracts a standard C string from a Foundation container.
+ * 
+ * @details
+ * The caller is responsible for releasing resources allocated by this routine
+ * by calling AGXUtilReleaseCString() on the returned pointer.
+ * 
+ * @param InStringRef
+ * A Foundation string container object.
+ * 
+ * @return
+ * A null-terminated C string upon success, or NULL if an error occurred during
+ * the extraction.
+ * 
+ * @see
+ * AGXUtilReleaseCString
+ */
+static const ANSICHAR* AGXUtilCreateCString(CFStringRef InStringRef)
+{
+	CFIndex   BufferLength = CFStringGetLength(InStringRef) + 1;
+	ANSICHAR *Buffer       = new ANSICHAR[BufferLength] { '\0' };
+
+	if (!CFStringGetCString(InStringRef, Buffer, BufferLength, kCFStringEncodingUTF8))
+	{
+		delete [] Buffer;
+		Buffer = nullptr;
+	}
+
+	return Buffer;
+}
+
+/**
+ * @fn AGXUtilReleaseCString
+ * 
+ * @brief
+ * This routine releases resources allocated by AGXUtilCreateCString.
+ * 
+ * @details
+ * Any call to AGXUtilCreateCString must be balanced with a call to this
+ * routine.
+ * 
+ * @param InPointer
+ * A pointer returned from AGXUtilCreateCString. It is valid to pass a NULL
+ * pointer to this routine.
+ * 
+ * @see
+ * AGXUtilCreateCString
+ */
+static void AGXUtilReleaseCString(const ANSICHAR *InPointer)
+{
+	if (InPointer != nullptr)
+	{
+		delete [] InPointer;
+	}
+}
+
+static CFDateRef AGXUtilCopyLastModifiedTimeForFile(CFStringRef Path)
+{
+    CFDateRef FileLastModifiedDate = NULL;
+
+    @autoreleasepool {
+        NSError *Error = nil;
+
+        NSDictionary *Attrs = [[NSFileManager defaultManager] attributesOfItemAtPath:(__bridge NSString*)Path error:&Error];
+        if (Attrs && !Error)
+        {
+            FileLastModifiedDate = (__bridge CFDateRef)[[Attrs fileModificationDate] retain];
+        }
+    } // autoreleasepool
+
+    return FileLastModifiedDate;
+}
+
+
+//------------------------------------------------------------------------------
+
+#pragma mark - Mac Platform AGXRHI Config Support Routines -
+
+/**
+ * @fn MacAGXConfig_PopulateAdapterInfo_GetVendorID
+ * 
+ * @brief
+ * This routine populates the RHI global vendor identifier.
+ * 
+ * @param InServiceDictionary
+ * The dictionary of properties defined by the underlying IOService.
+ * 
+ */
+static void MacAGXConfig_PopulateAdapterInfo_GetVendorID(CFMutableDictionaryRef InServiceDictionary)
+{
+	CFStringRef Key = CFStringCreateWithCString(kCFAllocatorDefault, "vendor-id", kCFStringEncodingUTF8);
+
+	CFTypeRef Value = CFDictionaryGetValue(InServiceDictionary, Key);
+	if ((Value != nullptr) && (CFGetTypeID(Value) == CFDataGetTypeID()))
+	{
+		const UInt32 *U32Value = (const UInt32*)CFDataGetBytePtr((CFDataRef)Value);
+		GRHIVendorId = *U32Value;
+	}
+	else
+	{
+		// The vendor-id cannot be zero or the RHI will assert. Set it to something
+		// non-zero for testing purposes. Note this path will only be taken on non-
+		// Apple GPUs, and will eventually be removed.
+		GRHIVendorId = 1; // Cannot be zero. 
+	}
+
+	CFRelease(Key);
+}
+
+/**
+ * @fn MacAGXConfig_PopulateAdapterInfo_GetAdapterName
+ * 
+ * @brief
+ * This routine populates the RHI global for the adapter name.
+ * 
+ * @param InServiceDictionary
+ * The dictionary of properties defined by the underlying IOService.
+ */
+static void MacAGXConfig_PopulateAdapterInfo_GetAdapterName(CFMutableDictionaryRef InServiceDictionary)
+{
+	CFStringRef Key = CFStringCreateWithCString(kCFAllocatorDefault, "model", kCFStringEncodingUTF8);
+
+	CFTypeRef Value = CFDictionaryGetValue(InServiceDictionary, Key);
+	if ((Value != nullptr) && (CFGetTypeID(Value) == CFStringGetTypeID()))
+	{
+		CFStringRef     ValueAsString = (CFStringRef)Value;
+		const ANSICHAR* Model         = AGXUtilCreateCString(ValueAsString);
+
+		if (Model)
+		{
+			GRHIAdapterName = FString::Printf(TEXT("%s (%s)"), UTF8_TO_TCHAR([[AGXUtilGetDevice() name] UTF8String] /*AGXUtilGetDevice()->name()->utf8String()*/), ANSI_TO_TCHAR(Model));
+		}
+		else
+		{
+			GRHIAdapterName = UTF8_TO_TCHAR([[AGXUtilGetDevice() name] UTF8String] /*AGXUtilGetDevice()->name()->utf8String()*/);
+		}
+
+		AGXUtilReleaseCString(Model);
+	}
+	else
+	{
+		GRHIAdapterName = MAC_PLATFORM_AGXRHI_ADAPTER_NAME;
+	}
+
+	CFRelease(Key);
+}
+
+/**
+ * @fn MacAGXConfig_PopulateAdapterInfo_GetAdapterDriverInfo_GetAdapterDriverDate
+ * 
+ * @brief
+ * This routine is called from MacAGXConfig_PopulateAdapterInfo_GetAdapterDriverInfo
+ * to populate the RHI global for the adapter driver date.
+ * 
+ * @param InBundleRef
+ * The Foundation bundle reference object representing the user-space Metal
+ * plugin.
+ */
+static void MacAGXConfig_PopulateAdapterInfo_GetAdapterDriverInfo_GetAdapterDriverDate(CFBundleRef InBundleRef)
+{
+    CFURLRef            BundleURL     = CFBundleCopyBundleURL(InBundleRef);
+    CFStringRef         BundlePath    = CFURLCopyFileSystemPath(BundleURL, kCFURLPOSIXPathStyle);
+    CFDateRef           BundleDate    = AGXUtilCopyLastModifiedTimeForFile(BundlePath);
+    CFDateFormatterRef  DateFormatter = CFDateFormatterCreateISO8601Formatter(kCFAllocatorDefault, kCFISO8601DateFormatWithDashSeparatorInDate | kCFISO8601DateFormatWithFullDate | kCFISO8601DateFormatWithFullTime);
+    CFStringRef         DateString    = CFDateFormatterCreateStringWithDate(kCFAllocatorDefault, DateFormatter, BundleDate);
+    const ANSICHAR     *DateCString   = AGXUtilCreateCString(DateString);
+    
+    if (DateCString != nullptr)
+    {
+		GRHIAdapterDriverDate = ANSI_TO_TCHAR(DateCString);
+    }
+    else
+    {
+        GRHIAdapterDriverDate = MAC_PLATFORM_AGXRHI_ADAPTER_DRIVER_DATE;
+    }
+    
+    AGXUtilReleaseCString(DateCString);
+
+    CFRelease(DateString);
+    CFRelease(DateFormatter);
+    CFRelease(BundleDate);
+    CFRelease(BundlePath);
+    CFRelease(BundleURL);
+}
+
+/**
+ * @fn MacAGXConfig_PopulateAdapterInfo_GetAdapterDriverInfo_GetAdapterUserDriverVersion
+ * 
+ * @brief
+ * This routine is called from MacAGXConfig_PopulateAdapterInfo_GetAdapterDriverInfo
+ * to populate the RHI global for the adapter user-space driver version.
+ * 
+ * @param InBundleRef
+ * The Foundation bundle reference object representing the user-space Metal
+ * plugin.
+ */
+static void MacAGXConfig_PopulateAdapterInfo_GetAdapterDriverInfo_GetAdapterUserDriverVersion(CFBundleRef InBundleRef)
+{
+    CFStringRef     Version     = (CFStringRef)CFBundleGetValueForInfoDictionaryKey(InBundleRef, CFSTR("CFBundleShortVersionString"));
+    const ANSICHAR *VersionCStr = AGXUtilCreateCString(Version);
+
+    if (VersionCStr != nullptr)
+    {
+        GRHIAdapterUserDriverVersion = ANSI_TO_TCHAR(VersionCStr);
+    }
+    else
+    {
+        GRHIAdapterUserDriverVersion = MAC_PLATFORM_AGXRHI_ADAPTER_USER_DRIVER_VERSION;
+    }
+
+    AGXUtilReleaseCString(VersionCStr);
+}
+
+/**
+ * @fn MacAGXConfig_PopulateAdapterInfo_GetAdapterDriverInfo
+ * 
+ * @brief
+ * This routine populates the RHI global for the adapter driver date and user
+ * driver version.
+ * 
+ * @details
+ * The adapter driver date and user driver version are derived from the Metal
+ * plugin bundle. This routine does the bundle discovery and calls supporting
+ * routines to fill in the respective RHI globals.
+ * 
+ * @param InServiceDictionary
+ * The dictionary of properties defined by the underlying IOService.
+ * 
+ * @see
+ * MacAGXConfig_PopulateAdapterInfo_GetAdapterDriverInfo_GetAdapterDriverDate
+ * MacAGXConfig_PopulateAdapterInfo_GetAdapterDriverInfo_GetAdapterUserDriverVersion
+ */
+static void MacAGXConfig_PopulateAdapterInfo_GetAdapterDriverInfo(CFMutableDictionaryRef InServiceDictionary)
+{
+    CFStringRef Key = CFStringCreateWithCString(kCFAllocatorDefault, "MetalPluginName", kCFStringEncodingUTF8);
+
+    CFTypeRef Value = CFDictionaryGetValue(InServiceDictionary, Key);
+    if ((Value != nullptr) && (CFGetTypeID(Value) == CFStringGetTypeID()))
+    {
+        CFStringRef ValueAsString = (CFStringRef)Value;
+
+        CFStringRef BundleID  = CFStringCreateWithFormat(kCFAllocatorDefault, NULL, CFSTR("com.apple.%@"), ValueAsString);
+        CFBundleRef BundleRef = CFBundleGetBundleWithIdentifier(BundleID);
+
+        MacAGXConfig_PopulateAdapterInfo_GetAdapterDriverInfo_GetAdapterDriverDate(BundleRef);
+        MacAGXConfig_PopulateAdapterInfo_GetAdapterDriverInfo_GetAdapterUserDriverVersion(BundleRef);
+
+        CFRelease(BundleID);
+    }
+
+    CFRelease(Key);
+}
+
+/**
+ * @fn MacAGXConfig_PopulateAdapterInfo
+ * 
+ * @brief
+ * This routine populates the adapter and device related RHI globals.
+ * 
+ * @details
+ * Not all of the RHI globals pertaining to adapter information are relevant
+ * for the Mac. For these, the Mac platform AGXRHI simply fills in known
+ * placeholder values.
+ *   - GRHIAdapterInternalDriverVersion
+ *   - GRHIDeviceId
+ *   - GRHIDeviceRevision
+ */
+static void MacAGXConfig_PopulateAdapterInfo()
+{
+	GRHIAdapterDriverOnDenyList              = MAC_PLATFORM_AGXRHI_ADAPTER_DRIVER_ON_DENY_LIST;
+	GRHIAdapterInternalDriverVersion         = MAC_PLATFORM_AGXRHI_ADAPTER_INTERNAL_DRIVER_VERSION;
+	GRHIDeviceId                             = MAC_PLATFORM_AGXRHI_DEVICE_ID;
+	GRHIDeviceRevision                       = MAC_PLATFORM_AGXRHI_DEVICE_REVISION;
+	GRHIDeviceIsAMDPreGCNArchitecture        = MAC_PLATFORM_AGXRHI_DEVICE_IS_AMD_PRE_GCN_ARCHITECTURE;
+	GRHIDeviceIsIntegrated                   = MAC_PLATFORM_AGXRHI_DEVICE_IS_INTEGRATED;
+
+	CFMutableDictionaryRef Matching          = IORegistryEntryIDMatching([AGXUtilGetDevice() registryID] /*AGXUtilGetDevice()->registryID()*/);
+	io_service_t           Service           = IOServiceGetMatchingService(kIOMasterPortDefault, Matching);
+	CFMutableDictionaryRef ServiceDictionary = nullptr;
+
+	if (IORegistryEntryCreateCFProperties(Service, &ServiceDictionary, kCFAllocatorDefault, kNilOptions) == kIOReturnSuccess)
+	{
+		MacAGXConfig_PopulateAdapterInfo_GetVendorID(ServiceDictionary);
+		MacAGXConfig_PopulateAdapterInfo_GetAdapterName(ServiceDictionary);
+		MacAGXConfig_PopulateAdapterInfo_GetAdapterDriverInfo(ServiceDictionary);
+
+		CFRelease(ServiceDictionary);
+	}
+
+	IOObjectRelease(Service);
+}
+
+static void MacAGXConfig_PopulateFeaturesInfo()
+{
+    //
+}
+
+static void MacAGXConfig_PopulateLimitsInfo()
+{
+    //
+}
+
+static void MacAGXConfig_PopulatePixelFormatsInfo()
+{
+    //
+}
+
+
+//------------------------------------------------------------------------------
+
+#pragma mark - Mac Platform AGXRHI Config Routines -
+
+void FMacAGXConfig::PopulateRHIGlobals()
+{
+    MacAGXConfig_PopulateAdapterInfo();
+    MacAGXConfig_PopulateFeaturesInfo();
+    MacAGXConfig_PopulateLimitsInfo();
+    MacAGXConfig_PopulatePixelFormatsInfo();
+}
+
+/*
+GHardwareHiddenSurfaceRemoval                          = MAC_PLATFORM_AGXRHI_HARDWARE_HIDDEN_SURFACE_REMOVAL;
+GIsRHIInitialized                                      = MAC_PLATFORM_AGXRHI_IS_INITIALIZED;
+GRHIAttachmentVariableRateShadingEnabled               = MAC_PLATFORM_AGXRHI_ATTACHMENT_VARIABLE_RATE_SHADING_ENABLED;
+GRHIForceNoDeletionLatencyForStreamingTextures         = MAC_PLATFORM_AGXRHI_FORCE_NO_DELETION_LATENCY_FOR_STREAMING_TEXTURES;
+GRHIIsHDREnabled                                       = MAC_PLATFORM_AGXRHI_IS_HDR_ENABLED;
+GRHILazyShaderCodeLoading                              = MAC_PLATFORM_AGXRHI_LAZY_SHADER_CODE_LOADING;
+GRHINeedsExtraDeletionLatency                          = MAC_PLATFORM_AGXRHI_NEEDS_EXTRA_DELETION_LATENCY;
+GRHINeedsUnatlasedCSMDepthsWorkaround                  = MAC_PLATFORM_AGXRHI_NEEDS_UNATLASED_CSM_DEPTHS_WORKAROUND;
+GRHIRequiresRenderTargetForPixelShaderUAVs             = MAC_PLATFORM_AGXRHI_REQUIRES_RENDER_TARGET_FOR_PIXEL_SHADER_UAVS;
+GRHISupportsArrayIndexFromAnyShader                    = MAC_PLATFORM_AGXRHI_SUPPORTS_ARRAY_INDEX_FROM_ANY_SHADER;
+GRHISupportsAsyncTextureCreation                       = MAC_PLATFORM_AGXRHI_SUPPORTS_ASYNC_TEXTURE_CREATION;
+GRHISupportsAtomicUInt64                               = MAC_PLATFORM_AGXRHI_SUPPORTS_ATOMIC_UINT64;
+GRHISupportsAttachmentVariableRateShading              = MAC_PLATFORM_AGXRHI_SUPPORTS_ATTACHMENT_VARIABLE_RATE_SHADING;
+GRHISupportsBackBufferWithCustomDepthStencil           = MAC_PLATFORM_AGXRHI_SUPPORTS_BACKBUFFER_WITH_CUSTOM_DEPTHSTENCIL;
+GRHISupportsBaseVertexIndex                            = MAC_PLATFORM_AGXRHI_SUPPORTS_BASE_VERTEX_INDEX;
+GRHISupportsComplexVariableRateShadingCombinerOps      = MAC_PLATFORM_AGXRHI_SUPPORTS_COMPLEX_VARIABLE_RATE_SHADING_COMBINER_OPS;
+GRHISupportsConservativeRasterization                  = MAC_PLATFORM_AGXRHI_SUPPORTS_CONSERVATIVE_RASTERIZATION;
+GRHISupportsCopyToTextureMultipleMips                  = MAC_PLATFORM_AGXRHI_SUPPORTS_COPY_TO_TEXTURE_MULTIPLE_MIPS;
+GRHISupportsDepthUAV                                   = MAC_PLATFORM_AGXRHI_SUPPORTS_DEPTH_UAV;
+GRHISupportsDirectGPUMemoryLock                        = MAC_PLATFORM_AGXRHI_SUPPORTS_DIRECT_GPU_MEMORY_LOCK;
+GRHISupportsDrawIndirect                               = MAC_PLATFORM_AGXRHI_SUPPORTS_DRAW_INDIRECT;
+GRHISupportsDynamicResolution                          = MAC_PLATFORM_AGXRHI_SUPPORTS_DYNAMIC_RESOLUTION;
+GRHISupportsEfficientUploadOnResourceCreation          = MAC_PLATFORM_AGXRHI_SUPPORTS_EFFICIENT_UPLOAD_ON_RESOURCE_CREATION;
+GRHISupportsExactOcclusionQueries                      = MAC_PLATFORM_AGXRHI_SUPPORTS_EXACT_OCCLUSION_QUERIES;
+GRHISupportsExplicitFMask                              = MAC_PLATFORM_AGXRHI_SUPPORTS_EXPLICIT_FMASK;
+GRHISupportsExplicitHTile                              = MAC_PLATFORM_AGXRHI_SUPPORTS_EXPLICIT_HTILE;
+GRHISupportsFirstInstance                              = MAC_PLATFORM_AGXRHI_SUPPORTS_FIRST_INSTANCE;
+GRHISupportsFrameCyclesBubblesRemoval                  = MAC_PLATFORM_AGXRHI_SUPPORTS_FRAME_CYCLES_BUBBLES_REMOVAL;
+GRHISupportsGPUTimestampBubblesRemoval                 = MAC_PLATFORM_AGXRHI_SUPPORTS_GPU_TIMESTAMP_BUBBLES_REMOVAL;
+GRHISupportsHDROutput                                  = MAC_PLATFORM_AGXRHI_SUPPORTS_HDR_OUTPUT;
+GRHISupportsInlineRayTracing                           = MAC_PLATFORM_AGXRHI_SUPPORTS_INLINE_RAY_TRACING;
+GRHISupportsLateVariableRateShadingUpdate              = MAC_PLATFORM_AGXRHI_SUPPORTS_LATE_VARIABLE_RATE_SHADING_UPDATE;
+GRHISupportsLazyShaderCodeLoading                      = MAC_PLATFORM_AGXRHI_SUPPORTS_LAZY_SHADER_CODE_LOADING;
+GRHISupportsMeshShadersTier0                           = MAC_PLATFORM_AGXRHI_SUPPORTS_MESH_SHADERS_TIER_0;
+GRHISupportsMeshShadersTier1                           = MAC_PLATFORM_AGXRHI_SUPPORTS_MESH_SHADERS_TIER_1;
+GRHISupportsMSAADepthSampleAccess                      = MAC_PLATFORM_AGXRHI_SUPPORTS_MSAA_DEPTH_SAMPLE_ACCESS;
+GRHISupportsMultithreadedShaderCreation                = MAC_PLATFORM_AGXRHI_SUPPORTS_MULTITHREADED_SHADER_CREATION;
+GRHISupportsMultithreading                             = MAC_PLATFORM_AGXRHI_SUPPORTS_MULTITHREADING;
+GRHISupportsParallelRHIExecute                         = MAC_PLATFORM_AGXRHI_SUPPORTS_PARALLEL_RHI_EXECUTE;
+GRHISupportsPipelineFileCache                          = MAC_PLATFORM_AGXRHI_SUPPORTS_PIPELINE_FILE_CACHE;
+GRHISupportsPipelineStateSortKey                       = MAC_PLATFORM_AGXRHI_SUPPORTS_PIPELINE_STATE_SORT_KEY;
+GRHISupportsPipelineVariableRateShading                = MAC_PLATFORM_AGXRHI_SUPPORTS_PIPELINE_VARIABLE_RATE_SHADING;
+GRHISupportsPixelShaderUAVs                            = MAC_PLATFORM_AGXRHI_SUPPORTS_PIXEL_SHADER_UAVS;
+GRHISupportsPrimitiveShaders                           = MAC_PLATFORM_AGXRHI_SUPPORTS_PRIMITIVE_SHADERS;
+GRHISupportsQuadTopology                               = MAC_PLATFORM_AGXRHI_SUPPORTS_QUAD_TOPOLOGY;
+GRHISupportsRayTracing                                 = MAC_PLATFORM_AGXRHI_SUPPORTS_RAY_TRACING;
+GRHISupportsRayTracingAMDHitToken                      = MAC_PLATFORM_AGXRHI_SUPPORTS_RAY_TRACING_AMD_HIT_TOKEN;
+GRHISupportsRayTracingAsyncBuildAccelerationStructure  = MAC_PLATFORM_AGXRHI_SUPPORTS_RAY_TRACING_ASYNC_BUILD_ACCELERATION_STRUCTURE;
+GRHISupportsRayTracingDispatchIndirect                 = MAC_PLATFORM_AGXRHI_SUPPORTS_RAY_TRACING_DISPATCH_INDIRECT;
+GRHISupportsRayTracingPSOAdditions                     = MAC_PLATFORM_AGXRHI_SUPPORTS_RAY_TRACING_PSO_ADDITIONS;
+GRHISupportsRayTracingShaders                          = MAC_PLATFORM_AGXRHI_SUPPORTS_RAY_TRACING_SHADERS;
+GRHISupportsRectTopology                               = MAC_PLATFORM_AGXRHI_SUPPORTS_RECT_TOPOLOGY;
+GRHISupportsResolveCubemapFaces                        = MAC_PLATFORM_AGXRHI_SUPPORTS_RESOLVE_CUBEMAP_FACES;
+GRHISupportsResummarizeHTile                           = MAC_PLATFORM_AGXRHI_SUPPORTS_RESUMMARIZE_HTILE;
+GRHISupportsRHIOnTaskThread                            = MAC_PLATFORM_AGXRHI_SUPPORTS_RHI_ON_TASK_THREAD;
+GRHISupportsRHIThread                                  = MAC_PLATFORM_AGXRHI_SUPPORTS_RHI_THREAD;
+GRHISupportsRWTextureBuffers                           = MAC_PLATFORM_AGXRHI_SUPPORTS_RWTEXTURE_BUFFERS;
+GRHISupportsShaderTimestamp                            = MAC_PLATFORM_AGXRHI_SUPPORTS_SHADER_TIMESTAMP;
+GRHISupportsStencilRefFromPixelShader                  = MAC_PLATFORM_AGXRHI_SUPPORTS_STENCIL_REF_FROM_PIXEL_SHADER;
+GRHISupportsTextureStreaming                           = MAC_PLATFORM_AGXRHI_SUPPORTS_TEXTURE_STREAMING;
+GRHISupportsUAVFormatAliasing                          = MAC_PLATFORM_AGXRHI_SUPPORTS_UAV_FORMAT_ALIASING;
+GRHISupportsUpdateFromBufferTexture                    = MAC_PLATFORM_AGXRHI_SUPPORTS_UPDATE_FROM_BUFFER_TEXTURE;
+GRHISupportsVariableRateShadingAttachmentArrayTextures = MAC_PLATFORM_AGXRHI_SUPPORTS_VARIABLE_RATE_SHADING_ATTACHMENT_ARRAY_TEXTURES;
+GRHISupportsWaveOperations                             = MAC_PLATFORM_AGXRHI_SUPPORTS_WAVE_OPERATIONS;
+GRHIThreadNeedsKicking                                 = MAC_PLATFORM_AGXRHI_THREAD_NEEDS_KICKING;
+GRHIVariableRateShadingEnabled                         = MAC_PLATFORM_AGXRHI_VARIABLE_RATE_SHADING_ENABLED;
+GSupportsDepthBoundsTest                               = MAC_PLATFORM_AGXRHI_SUPPORTS_DEPTH_BOUNDS_TEST;
+GSupportsEfficientAsyncCompute                         = MAC_PLATFORM_AGXRHI_SUPPORTS_EFFICIENT_ASYNC_COMPUTE;
+GSupportsImageExternal                                 = MAC_PLATFORM_AGXRHI_SUPPORTS_IMAGE_EXTERNAL;
+GSupportsMobileMultiView                               = MAC_PLATFORM_AGXRHI_SUPPORTS_MOBILE_MULTIVIEW;
+GSupportsParallelOcclusionQueries                      = MAC_PLATFORM_AGXRHI_SUPPORTS_PARALLEL_OCCLUSION_QUERIES;
+GSupportsParallelRenderingTasksWithSeparateRHIThread   = MAC_PLATFORM_AGXRHI_SUPPORTS_PARALLEL_RENDERING_TASKS_WITH_SEPARATE_RHI_THREAD;
+GSupportsPixelLocalStorage                             = MAC_PLATFORM_AGXRHI_SUPPORTS_PIXEL_LOCAL_STORAGE;
+GSupportsQuadBufferStereo                              = MAC_PLATFORM_AGXRHI_SUPPORTS_QUADBUFFER_STEREO;
+GSupportsRenderDepthTargetableShaderResources          = MAC_PLATFORM_AGXRHI_SUPPORTS_RENDER_DEPTH_TARGETABLE_SHADER_RESOURCES;
+GSupportsRenderTargetFormat_PF_FloatRGBA               = MAC_PLATFORM_AGXRHI_SUPPORTS_RENDERTARGET_FORMAT_PF_FLOAT_RGBA;
+GSupportsRenderTargetFormat_PF_G8                      = MAC_PLATFORM_AGXRHI_SUPPORTS_RENDERTARGET_FORMAT_PF_G_8;
+GSupportsSeparateRenderTargetBlendState                = MAC_PLATFORM_AGXRHI_SUPPORTS_SEPARATE_RENDERTARGET_BLEND_STATE;
+GSupportsShaderDepthStencilFetch                       = MAC_PLATFORM_AGXRHI_SUPPORTS_SHADER_DEPTHSTENCIL_FETCH;
+GSupportsShaderFramebufferFetch                        = MAC_PLATFORM_AGXRHI_SUPPORTS_SHADER_FRAMEBUFFER_FETCH;
+GSupportsShaderMRTFramebufferFetch                     = MAC_PLATFORM_AGXRHI_SUPPORTS_SHADER_MRT_FRAMEBUFFER_FETCH;
+GSupportsTexture3D                                     = MAC_PLATFORM_AGXRHI_SUPPORTS_TEXTURE_3D;
+GSupportsTimestampRenderQueries                        = MAC_PLATFORM_AGXRHI_SUPPORTS_TIMESTAMP_RENDER_QUERIES;
+GSupportsTransientResourceAliasing                     = MAC_PLATFORM_AGXRHI_SUPPORTS_TRANSIENT_RESOURCE_ALIASING;
+GSupportsVolumeTextureRendering                        = MAC_PLATFORM_AGXRHI_SUPPORTS_VOLUME_TEXTURE_RENDERING;
+GSupportsWideMRT                                       = MAC_PLATFORM_AGXRHI_SUPPORTS_WIDE_MRT;
+GUsingNullRHI                                          = MAC_PLATFORM_AGXRHI_USING_NULL_RHI;
+GRHIHDRDisplayOutputFormat                             = MAC_PLATFORM_AGXRHI_HDR_DISPLAY_OUTPUT_FORMAT;
+GRHIVariableRateShadingImageFormat                     = MAC_PLATFORM_AGXRHI_VARIABLE_RATE_SHADING_IMAGE_FORMAT;
+GRHIMergeableAccessMask                                = MAC_PLATFORM_AGXRHI_MERGEABLE_ACCESS_MASK;
+GRHIMultiPipelineMergeableAccessMask                   = MAC_PLATFORM_AGXRHI_MULTI_PIPELINE_MERGEABLE_ACCESS_MASK;
+GMaxRHIFeatureLevel                                    = MAC_PLATFORM_AGXRHI_MAX_FEATURE_LEVEL;
+GMaxRHIShaderPlatform                                  = MAC_PLATFORM_AGXRHI_MAX_SHADER_PLATFORM;
+GShaderPlatformForFeatureLevel                         = ERHIFeatureLevel::Num]	MAC_PLATFORM_AGXRHI_SHADER_PLATFORM_FOR_FEATURE_LEVEL;
+GRHIVariableRateShadingImageDataType                   = MAC_PLATFORM_AGXRHI_VARIABLE_RATE_SHADING_IMAGE_DATA_TYPE;
+GRHIMaxDispatchThreadGroupsPerDimension                = MAC_PLATFORM_AGXRHI_MAX_DISPATCH_THREAD_GROUPS_PER_DIMENSION;
+GPixelFormats                                          = PF_MAX]	MAC_PLATFORM_AGXRHI_PIXEL_FORMATS;
+GRHIDefaultMSAASampleOffsets                           = 1+2+4+8+16]	MAC_PLATFORM_AGXRHI_DEFAULT_MSAA_SAMPLE_OFFSETS;
+GVertexElementTypeSupport                              = MAC_PLATFORM_AGXRHI_VERTEX_ELEMENT_TYPE_SUPPORT;
+GCurrentNumDrawCallsRHI                                = MAX_NUM_GPUS]	MAC_PLATFORM_AGXRHI_CURRENT_NUM_DRAW_CALLS;
+GCurrentNumPrimitivesDrawnRHI                          = MAX_NUM_GPUS]	MAC_PLATFORM_AGXRHI_CURRENT_NUM_PRIMITIVES_DRAWN;
+GCurrentRendertargetMemorySize                         = MAC_PLATFORM_AGXRHI_CURRENT_RENDERTARGET_MEMORY_SIZE;
+GCurrentTextureMemorySize                              = MAC_PLATFORM_AGXRHI_CURRENT_TEXTURE_MEMORY_SIZE;
+GDrawUPIndexCheckCount                                 = MAC_PLATFORM_AGXRHI_DRAW_UP_INDEX_CHECK_COUNT;
+GDrawUPVertexCheckCount                                = MAC_PLATFORM_AGXRHI_DRAW_UP_VERTEX_CHECK_COUNT;
+GMaxComputeDispatchDimension                           = MAC_PLATFORM_AGXRHI_MAX_COMPUTE_DISPATCH_DIMENSION;
+GMaxCubeTextureDimensions                              = MAC_PLATFORM_AGXRHI_MAX_CUBE_TEXTURE_DIMENSIONS;
+GMaxShadowDepthBufferSizeX                             = MAC_PLATFORM_AGXRHI_MAX_SHADOW_DEPTH_BUFFER_SIZE_X;
+GMaxShadowDepthBufferSizeY                             = MAC_PLATFORM_AGXRHI_MAX_SHADOW_DEPTH_BUFFER_SIZE_Y;
+GMaxTextureArrayLayers                                 = MAC_PLATFORM_AGXRHI_MAX_TEXTURE_ARRAY_LAYERS;
+GMaxTextureDimensions                                  = MAC_PLATFORM_AGXRHI_MAX_TEXTURE_DIMENSIONS;
+GMaxTextureMipCount                                    = MAC_PLATFORM_AGXRHI_MAX_TEXTURE_MIP_COUNT;
+GMaxTextureSamplers                                    = MAC_PLATFORM_AGXRHI_MAX_TEXTURE_SAMPLERS;
+GMaxVolumeTextureDimensions                            = MAC_PLATFORM_AGXRHI_MAX_VOLUME_TEXTURE_DIMENSIONS;
+GMaxWorkGroupInvocations                               = MAC_PLATFORM_AGXRHI_MAX_WORKGROUP_INVOCATIONS;
+GNumDrawCallsRHI                                       = MAX_NUM_GPUS]	MAC_PLATFORM_AGXRHI_NUM_DRAW_CALLS;
+GNumPrimitivesDrawnRHI                                 = MAX_NUM_GPUS]	MAC_PLATFORM_AGXRHI_NUM_PRIMITIVES_DRAWN;
+GPoolSizeVRAMPercentage                                = MAC_PLATFORM_AGXRHI_POOL_SIZE_VRAM_PERCENTAGE;
+GRHIMaximumReccommendedOustandingOcclusionQueries      = MAC_PLATFORM_AGXRHI_MAXIMUM_RECCOMMENDED_OUSTANDING_OCCLUSION_QUERIES;
+GRHIMaximumWaveSize                                    = MAC_PLATFORM_AGXRHI_MAXIMUM_WAVE_SIZE;
+GRHIMinimumWaveSize                                    = MAC_PLATFORM_AGXRHI_MINIMUM_WAVE_SIZE;
+GRHIPersistentThreadGroupCount                         = MAC_PLATFORM_AGXRHI_PERSISTENT_THREAD_GROUP_COUNT;
+GRHIVariableRateShadingImageTileMaxHeight              = MAC_PLATFORM_AGXRHI_VARIABLE_RATE_SHADING_IMAGE_TILE_MAX_HEIGHT;
+GRHIVariableRateShadingImageTileMaxWidth               = MAC_PLATFORM_AGXRHI_VARIABLE_RATE_SHADING_IMAGE_TILE_MAX_WIDTH;
+GRHIVariableRateShadingImageTileMinHeight              = MAC_PLATFORM_AGXRHI_VARIABLE_RATE_SHADING_IMAGE_TILE_MIN_HEIGHT;
+GRHIVariableRateShadingImageTileMinWidth               = MAC_PLATFORM_AGXRHI_VARIABLE_RATE_SHADING_IMAGE_TILE_MIN_WIDTH;
+GCurrentNumDrawCallsRHIPtr                             = MAC_PLATFORM_AGXRHI_CURRENT_NUM_DRAW_CALLS_PTR;
+GMaxBufferDimensions                                   = MAC_PLATFORM_AGXRHI_MAX_BUFFER_DIMENSIONS;
+GMaxComputeSharedMemory                                = MAC_PLATFORM_AGXRHI_MAX_COMPUTE_SHARED_MEMORY;
+GTexturePoolSize                                       = MAC_PLATFORM_AGXRHI_TEXTURE_POOL_SIZE;
+GRHIRayTracingAccelerationStructureAlignment           = MAC_PLATFORM_AGXRHI_RAY_TRACING_ACCELERATION_STRUCTURE_ALIGNMENT;
+GRHIRayTracingInstanceDescriptorSize                   = MAC_PLATFORM_AGXRHI_RAY_TRACING_INSTANCE_DESCRIPTOR_SIZE;
+GRHIRayTracingScratchBufferAlignment                   = MAC_PLATFORM_AGXRHI_RAY_TRACING_SCRATCH_BUFFER_ALIGNMENT;
+GRHIRayTracingShaderTableAlignment                     = MAC_PLATFORM_AGXRHI_RAY_TRACING_SHADER_TABLE_ALIGNMENT;
+GDemotedLocalMemorySize                                = MAC_PLATFORM_AGXRHI_DEMOTED_LOCAL_MEMORY_SIZE;
+GRHIPresentCounter                                     = MAC_PLATFORM_AGXRHI_PRESENT_COUNTER;
+GRHITransitionPrivateData_AlignInBytes                 = MAC_PLATFORM_AGXRHI_TRANSITION_PRIVATE_DATA_ALIGN_IN_BYTES;
+GRHITransitionPrivateData_SizeInBytes                  = MAC_PLATFORM_AGXRHI_TRANSITION_PRIVATE_DATA_SIZE_IN_BYTES;
+ */
diff --git a/Engine/Source/Runtime/Apple/AGXRHI/Private/Mac/MacPlatformAGXConfig.h b/Engine/Source/Runtime/Apple/AGXRHI/Private/Mac/MacPlatformAGXConfig.h
new file mode 100644
--- /dev/null
+++ b/Engine/Source/Runtime/Apple/AGXRHI/Private/Mac/MacPlatformAGXConfig.h
@@ -0,0 +1,11 @@
+// Copyright Epic Games, Inc. All rights reserved.
+
+#pragma once
+
+class FMacAGXConfig
+{
+public:
+	static void PopulateRHIGlobals();
+};
+
+typedef FMacAGXConfig FPlatformAGXConfig;
diff --git a/Engine/Source/Runtime/Apple/AGXRHI/Private/Shaders/Types/AGXComputeShader.cpp b/Engine/Source/Runtime/Apple/AGXRHI/Private/Shaders/Types/AGXComputeShader.cpp
--- a/Engine/Source/Runtime/Apple/AGXRHI/Private/Shaders/Types/AGXComputeShader.cpp
+++ b/Engine/Source/Runtime/Apple/AGXRHI/Private/Shaders/Types/AGXComputeShader.cpp
@@ -55,7 +55,7 @@ FAGXShaderPipeline* FAGXComputeShader::GetPipeline()
 		{
 			ns::AutoReleased<ns::Array<mtlpp::PipelineBufferDescriptor>> PipelineBuffers = Descriptor.GetBuffers();
 
-			uint32 ImmutableBuffers = Bindings.ConstantBuffers | Bindings.ArgumentBuffers;
+			uint32 ImmutableBuffers = Bindings.ConstantBuffers;
 			while(ImmutableBuffers)
 			{
 				uint32 Index = __builtin_ctz(ImmutableBuffers);
diff --git a/Engine/Source/Runtime/Apple/AGXRHI/Private/Shaders/Types/Templates/AGXBaseShader.h b/Engine/Source/Runtime/Apple/AGXRHI/Private/Shaders/Types/Templates/AGXBaseShader.h
--- a/Engine/Source/Runtime/Apple/AGXRHI/Private/Shaders/Types/Templates/AGXBaseShader.h
+++ b/Engine/Source/Runtime/Apple/AGXRHI/Private/Shaders/Types/Templates/AGXBaseShader.h
@@ -20,7 +20,7 @@
 #pragma mark - AGX RHI Base Shader Class Support Routines
 
 
-extern mtlpp::LanguageVersion AGXValidateVersion(uint8 Version);
+extern mtlpp::LanguageVersion AGXValidateVersion(uint32 Version);
 
 
 //------------------------------------------------------------------------------
@@ -76,12 +76,6 @@ public:
 	// List of memory copies from RHIUniformBuffer to packed uniforms
 	TArray<CrossCompiler::FUniformBufferCopyInfo> UniformBuffersCopyInfo;
 
-	/* Argument encoders for shader IABs */
-	TMap<uint32, mtlpp::ArgumentEncoder> ArgumentEncoders;
-
-	/* Tier1 Argument buffer bitmasks */
-	TMap<uint32, TBitArray<>> ArgumentBitmasks;
-
 	/* Uniform buffer static slots */
 	TArray<FUniformBufferStaticSlot> StaticSlots;
 
@@ -302,34 +296,18 @@ void TAGXBaseShader<BaseResourceType, ShaderType>::Init(TArrayView<const uint8>
 			mtlpp::LanguageVersion MetalVersion;
 			switch(Header.Version)
 			{
-				case 6:
-				case 5:
-				case 4:
-					MetalVersion = mtlpp::LanguageVersion::Version2_1;
-					break;
-				case 3:
-					MetalVersion = mtlpp::LanguageVersion::Version2_0;
-					break;
-				case 2:
-					MetalVersion = mtlpp::LanguageVersion::Version1_2;
+				case 7:
+					MetalVersion = mtlpp::LanguageVersion::Version2_4;
 					break;
-				case 1:
-					MetalVersion = mtlpp::LanguageVersion::Version1_1;
+				case 6:
+					MetalVersion = mtlpp::LanguageVersion::Version2_3;
 					break;
-				case 0:
-#if PLATFORM_MAC
-					MetalVersion = mtlpp::LanguageVersion::Version1_1;
-#else
-					MetalVersion = mtlpp::LanguageVersion::Version1_0;
-#endif
+				case 5:
+					MetalVersion = mtlpp::LanguageVersion::Version2_2;
 					break;
 				default:
 					UE_LOG(LogRHI, Fatal, TEXT("Failed to create shader with unknown version %d: %s"), Header.Version, *FString(NewShaderString));
-#if PLATFORM_MAC
-					MetalVersion = mtlpp::LanguageVersion::Version1_1;
-#else
-					MetalVersion = mtlpp::LanguageVersion::Version1_0;
-#endif
+					MetalVersion = mtlpp::LanguageVersion::Version2_2;
 					break;
 				}
 			CompileOptions.SetLanguageVersion(MetalVersion);
@@ -488,30 +466,6 @@ mtlpp::Function TAGXBaseShader<BaseResourceType, ShaderType>::GetCompiledFunctio
 		}
 	}
 
-	if (FAGXCommandQueue::SupportsFeature(EAGXFeaturesIABs) && Bindings.ArgumentBuffers && ArgumentEncoders.Num() == 0)
-	{
-		uint32 ArgumentBuffers = Bindings.ArgumentBuffers;
-		while(ArgumentBuffers)
-		{
-			uint32 Index = __builtin_ctz(ArgumentBuffers);
-			ArgumentBuffers &= ~(1 << Index);
-
-			mtlpp::ArgumentEncoder ArgumentEncoder = Function.NewArgumentEncoderWithBufferIndex(Index);
-			ArgumentEncoders.Add(Index, ArgumentEncoder);
-
-			TBitArray<> Resources;
-			for (uint8 Id : Bindings.ArgumentBufferMasks[Index])
-			{
-				if (Id >= Resources.Num())
-				{
-					Resources.Add(false, (Id + 1) - Resources.Num());
-				}
-				Resources[Id] = true;
-			}
-			ArgumentBitmasks.Add(Index, Resources);
-		}
-	}
-
 	check(Func);
 	return Func;
 }
diff --git a/Engine/Source/Runtime/Apple/AGXRHI/Private/TVOS/TVOSPlatformAGXConfig.cpp b/Engine/Source/Runtime/Apple/AGXRHI/Private/TVOS/TVOSPlatformAGXConfig.cpp
new file mode 100644
--- /dev/null
+++ b/Engine/Source/Runtime/Apple/AGXRHI/Private/TVOS/TVOSPlatformAGXConfig.cpp
@@ -0,0 +1,15 @@
+// Copyright Epic Games, Inc. All rights reserved.
+
+#include "TVOSPlatformAGXConfig.h"
+#include "RHI.h"
+#include "AGXRHIPrivate.h"
+
+
+//------------------------------------------------------------------------------
+
+#pragma mark - tvOS Platform AGXRHI Config Routines -
+
+void FTVOSAGXConfig::PopulateRHIGlobals()
+{
+    // TODO
+}
diff --git a/Engine/Source/Runtime/Apple/AGXRHI/Private/TVOS/TVOSPlatformAGXConfig.h b/Engine/Source/Runtime/Apple/AGXRHI/Private/TVOS/TVOSPlatformAGXConfig.h
new file mode 100644
--- /dev/null
+++ b/Engine/Source/Runtime/Apple/AGXRHI/Private/TVOS/TVOSPlatformAGXConfig.h
@@ -0,0 +1,11 @@
+// Copyright Epic Games, Inc. All rights reserved.
+
+#pragma once
+
+class FTVOSAGXConfig
+{
+public:
+	static void PopulateRHIGlobals();
+};
+
+typedef FTVOSAGXConfig FPlatformAGXConfig;
