Fixup copyright notice casing.

#rnx

diff --git a/Engine/Source/Runtime/Apple/AGXRHI/Private/IOS/IOSPlatformAGXConfig.cpp b/Engine/Source/Runtime/Apple/AGXRHI/Private/IOS/IOSPlatformAGXConfig.cpp
--- a/Engine/Source/Runtime/Apple/AGXRHI/Private/IOS/IOSPlatformAGXConfig.cpp
+++ b/Engine/Source/Runtime/Apple/AGXRHI/Private/IOS/IOSPlatformAGXConfig.cpp
@@ -1,4 +1,4 @@
-// Copyright Epic Games, Inc. All rights reserved.
+// Copyright Epic Games, Inc. All Rights Reserved.
 
 #include "IOSPlatformAGXConfig.h"
 #include "RHI.h"
diff --git a/Engine/Source/Runtime/Apple/AGXRHI/Private/IOS/IOSPlatformAGXConfig.h b/Engine/Source/Runtime/Apple/AGXRHI/Private/IOS/IOSPlatformAGXConfig.h
--- a/Engine/Source/Runtime/Apple/AGXRHI/Private/IOS/IOSPlatformAGXConfig.h
+++ b/Engine/Source/Runtime/Apple/AGXRHI/Private/IOS/IOSPlatformAGXConfig.h
@@ -1,4 +1,4 @@
-// Copyright Epic Games, Inc. All rights reserved.
+// Copyright Epic Games, Inc. All Rights Reserved.
 
 #pragma once
 
diff --git a/Engine/Source/Runtime/Apple/AGXRHI/Private/Mac/MacPlatformAGXConfig.cpp b/Engine/Source/Runtime/Apple/AGXRHI/Private/Mac/MacPlatformAGXConfig.cpp
--- a/Engine/Source/Runtime/Apple/AGXRHI/Private/Mac/MacPlatformAGXConfig.cpp
+++ b/Engine/Source/Runtime/Apple/AGXRHI/Private/Mac/MacPlatformAGXConfig.cpp
@@ -1,4 +1,4 @@
-// Copyright Epic Games, Inc. All rights reserved.
+// Copyright Epic Games, Inc. All Rights Reserved.
 
 #include "MacPlatformAGXConfig.h"
 #include "RHI.h"
diff --git a/Engine/Source/Runtime/Apple/AGXRHI/Private/Mac/MacPlatformAGXConfig.h b/Engine/Source/Runtime/Apple/AGXRHI/Private/Mac/MacPlatformAGXConfig.h
--- a/Engine/Source/Runtime/Apple/AGXRHI/Private/Mac/MacPlatformAGXConfig.h
+++ b/Engine/Source/Runtime/Apple/AGXRHI/Private/Mac/MacPlatformAGXConfig.h
@@ -1,4 +1,4 @@
-// Copyright Epic Games, Inc. All rights reserved.
+// Copyright Epic Games, Inc. All Rights Reserved.
 
 #pragma once
 
diff --git a/Engine/Source/Runtime/Apple/AGXRHI/Private/TVOS/TVOSPlatformAGXConfig.cpp b/Engine/Source/Runtime/Apple/AGXRHI/Private/TVOS/TVOSPlatformAGXConfig.cpp
--- a/Engine/Source/Runtime/Apple/AGXRHI/Private/TVOS/TVOSPlatformAGXConfig.cpp
+++ b/Engine/Source/Runtime/Apple/AGXRHI/Private/TVOS/TVOSPlatformAGXConfig.cpp
@@ -1,4 +1,4 @@
-// Copyright Epic Games, Inc. All rights reserved.
+// Copyright Epic Games, Inc. All Rights Reserved.
 
 #include "TVOSPlatformAGXConfig.h"
 #include "RHI.h"
diff --git a/Engine/Source/Runtime/Apple/AGXRHI/Private/TVOS/TVOSPlatformAGXConfig.h b/Engine/Source/Runtime/Apple/AGXRHI/Private/TVOS/TVOSPlatformAGXConfig.h
--- a/Engine/Source/Runtime/Apple/AGXRHI/Private/TVOS/TVOSPlatformAGXConfig.h
+++ b/Engine/Source/Runtime/Apple/AGXRHI/Private/TVOS/TVOSPlatformAGXConfig.h
@@ -1,4 +1,4 @@
-// Copyright Epic Games, Inc. All rights reserved.
+// Copyright Epic Games, Inc. All Rights Reserved.
 
 #pragma once
 
