Undo //UE5/Main/Engine/Source/Runtime/Apple/AGXRHI/Private/... changelist 18589140
Causing CIS compile errors for MAC and Android platforms.

diff --git a/Engine/Source/Runtime/Apple/AGXRHI/Private/AGXBlitCommandEncoder.cpp b/Engine/Source/Runtime/Apple/AGXRHI/Private/AGXBlitCommandEncoder.cpp
new file mode 100644
--- /dev/null
+++ b/Engine/Source/Runtime/Apple/AGXRHI/Private/AGXBlitCommandEncoder.cpp
@@ -0,0 +1,280 @@
+// Copyright Epic Games, Inc. All Rights Reserved.
+
+/*=============================================================================
+	AGXBlitCommandEncoder.cpp: AGX command encoder wrapper.
+=============================================================================*/
+
+#include "AGXRHIPrivate.h"
+
+#include "AGXBlitCommandEncoder.h"
+#include "AGXCommandBuffer.h"
+#include "AGXFence.h"
+
+#if MTLPP_CONFIG_VALIDATE && METAL_DEBUG_OPTIONS
+extern int32 GAGXRuntimeDebugLevel;
+
+@interface FAGXDebugBlitCommandEncoder : FAGXDebugCommandEncoder
+{
+	@public
+	id<MTLBlitCommandEncoder> Inner;
+	FAGXCommandBufferDebugging Buffer;
+}
+/** Initialise the wrapper with the provided command-buffer. */
+-(id)initWithEncoder:(id<MTLBlitCommandEncoder>)Encoder andCommandBuffer:(FAGXCommandBufferDebugging const&)Buffer;
+
+@end
+
+@implementation FAGXDebugBlitCommandEncoder
+
+APPLE_PLATFORM_OBJECT_ALLOC_OVERRIDES(FAGXDebugBlitCommandEncoder)
+
+-(id)initWithEncoder:(id<MTLBlitCommandEncoder>)Encoder andCommandBuffer:(FAGXCommandBufferDebugging const&)SourceBuffer
+{
+	id Self = [super init];
+	if (Self)
+	{
+        Inner = Encoder;
+		Buffer = SourceBuffer;
+	}
+	return Self;
+}
+
+-(void)dealloc
+{
+	[super dealloc];
+}
+
+@end
+
+FAGXBlitCommandEncoderDebugging::FAGXBlitCommandEncoderDebugging()
+{
+	
+}
+FAGXBlitCommandEncoderDebugging::FAGXBlitCommandEncoderDebugging(mtlpp::BlitCommandEncoder& Encoder, FAGXCommandBufferDebugging& Buffer)
+: FAGXCommandEncoderDebugging((FAGXDebugCommandEncoder*)[[[FAGXDebugBlitCommandEncoder alloc] initWithEncoder:Encoder.GetPtr() andCommandBuffer:Buffer] autorelease])
+{
+	Buffer.BeginBlitCommandEncoder([NSString stringWithFormat:@"Blit: %@", Encoder.GetLabel().GetPtr()]);
+	Encoder.SetAssociatedObject((void const*)&FAGXBlitCommandEncoderDebugging::Get, (FAGXCommandEncoderDebugging const&)*this);
+}
+FAGXBlitCommandEncoderDebugging::FAGXBlitCommandEncoderDebugging(FAGXDebugCommandEncoder* handle)
+: FAGXCommandEncoderDebugging((FAGXDebugCommandEncoder*)handle)
+{
+	
+}
+
+FAGXBlitCommandEncoderDebugging FAGXBlitCommandEncoderDebugging::Get(mtlpp::BlitCommandEncoder& Encoder)
+{
+	return Encoder.GetAssociatedObject<FAGXBlitCommandEncoderDebugging>((void const*)&FAGXBlitCommandEncoderDebugging::Get);
+}
+
+void FAGXBlitCommandEncoderDebugging::InsertDebugSignpost(ns::String const& Label)
+{
+	((FAGXDebugBlitCommandEncoder*)m_ptr)->Buffer.InsertDebugSignpost(Label);
+}
+void FAGXBlitCommandEncoderDebugging::PushDebugGroup(ns::String const& Group)
+{
+	((FAGXDebugBlitCommandEncoder*)m_ptr)->Buffer.PushDebugGroup(Group);
+}
+void FAGXBlitCommandEncoderDebugging::PopDebugGroup()
+{
+	((FAGXDebugBlitCommandEncoder*)m_ptr)->Buffer.PopDebugGroup();
+}
+
+void FAGXBlitCommandEncoderDebugging::EndEncoder()
+{
+	((FAGXDebugBlitCommandEncoder*)m_ptr)->Buffer.EndCommandEncoder();
+}
+
+#if PLATFORM_MAC
+void FAGXBlitCommandEncoderDebugging::Synchronize(mtlpp::Resource const& resource)
+{
+#if METAL_DEBUG_OPTIONS
+	switch(((FAGXDebugBlitCommandEncoder*)m_ptr)->Buffer.GetPtr()->DebugLevel)
+	{
+		case EAGXDebugLevelConditionalSubmit:
+		case EAGXDebugLevelWaitForComplete:
+		case EAGXDebugLevelLogOperations:
+		{
+			((FAGXDebugBlitCommandEncoder*)m_ptr)->Buffer.Blit([NSString stringWithFormat:@"%s", __PRETTY_FUNCTION__]);
+		}
+		default:
+		{
+			break;
+		}
+	}
+#endif
+}
+
+void FAGXBlitCommandEncoderDebugging::Synchronize(FAGXTexture const& texture, NSUInteger slice, NSUInteger level)
+{
+#if METAL_DEBUG_OPTIONS
+	switch(((FAGXDebugBlitCommandEncoder*)m_ptr)->Buffer.GetPtr()->DebugLevel)
+	{
+		case EAGXDebugLevelConditionalSubmit:
+		case EAGXDebugLevelWaitForComplete:
+		case EAGXDebugLevelLogOperations:
+		{
+			((FAGXDebugBlitCommandEncoder*)m_ptr)->Buffer.Blit([NSString stringWithFormat:@"%s", __PRETTY_FUNCTION__]);
+		}
+		default:
+		{
+			break;
+		}
+	}
+#endif
+}
+#endif
+
+void FAGXBlitCommandEncoderDebugging::Copy(FAGXTexture const& sourceTexture, NSUInteger sourceSlice, NSUInteger sourceLevel, mtlpp::Origin const& sourceOrigin, mtlpp::Size const& sourceSize, FAGXTexture const& destinationTexture, NSUInteger destinationSlice, NSUInteger destinationLevel, mtlpp::Origin const& destinationOrigin)
+{
+#if METAL_DEBUG_OPTIONS
+	switch(((FAGXDebugBlitCommandEncoder*)m_ptr)->Buffer.GetPtr()->DebugLevel)
+	{
+		case EAGXDebugLevelConditionalSubmit:
+		case EAGXDebugLevelWaitForComplete:
+		case EAGXDebugLevelLogOperations:
+		{
+			((FAGXDebugBlitCommandEncoder*)m_ptr)->Buffer.Blit([NSString stringWithFormat:@"%s", __PRETTY_FUNCTION__]);
+		}
+		default:
+		{
+			break;
+		}
+	}
+#endif
+}
+
+void FAGXBlitCommandEncoderDebugging::Copy(FAGXBuffer const& sourceBuffer, NSUInteger sourceOffset, NSUInteger sourceBytesPerRow, NSUInteger sourceBytesPerImage, mtlpp::Size const& sourceSize, FAGXTexture const& destinationTexture, NSUInteger destinationSlice, NSUInteger destinationLevel, mtlpp::Origin const& destinationOrigin)
+{
+#if METAL_DEBUG_OPTIONS
+	switch(((FAGXDebugBlitCommandEncoder*)m_ptr)->Buffer.GetPtr()->DebugLevel)
+	{
+		case EAGXDebugLevelConditionalSubmit:
+		case EAGXDebugLevelWaitForComplete:
+		case EAGXDebugLevelLogOperations:
+		{
+			((FAGXDebugBlitCommandEncoder*)m_ptr)->Buffer.Blit([NSString stringWithFormat:@"%s", __PRETTY_FUNCTION__]);
+		}
+		default:
+		{
+			break;
+		}
+	}
+#endif
+}
+
+void FAGXBlitCommandEncoderDebugging::Copy(FAGXBuffer const& sourceBuffer, NSUInteger sourceOffset, NSUInteger sourceBytesPerRow, NSUInteger sourceBytesPerImage, mtlpp::Size const& sourceSize, FAGXTexture const& destinationTexture, NSUInteger destinationSlice, NSUInteger destinationLevel, mtlpp::Origin const& destinationOrigin, mtlpp::BlitOption options)
+{
+#if METAL_DEBUG_OPTIONS
+	switch(((FAGXDebugBlitCommandEncoder*)m_ptr)->Buffer.GetPtr()->DebugLevel)
+	{
+		case EAGXDebugLevelConditionalSubmit:
+		case EAGXDebugLevelWaitForComplete:
+		case EAGXDebugLevelLogOperations:
+		{
+			((FAGXDebugBlitCommandEncoder*)m_ptr)->Buffer.Blit([NSString stringWithFormat:@"%s", __PRETTY_FUNCTION__]);
+		}
+		default:
+		{
+			break;
+		}
+	}
+#endif
+}
+
+void FAGXBlitCommandEncoderDebugging::Copy(FAGXTexture const& sourceTexture, NSUInteger sourceSlice, NSUInteger sourceLevel, mtlpp::Origin const& sourceOrigin, mtlpp::Size const& sourceSize, FAGXBuffer const& destinationBuffer, NSUInteger destinationOffset, NSUInteger destinationBytesPerRow, NSUInteger destinationBytesPerImage)
+{
+#if METAL_DEBUG_OPTIONS
+	switch(((FAGXDebugBlitCommandEncoder*)m_ptr)->Buffer.GetPtr()->DebugLevel)
+	{
+		case EAGXDebugLevelConditionalSubmit:
+		case EAGXDebugLevelWaitForComplete:
+		case EAGXDebugLevelLogOperations:
+		{
+			((FAGXDebugBlitCommandEncoder*)m_ptr)->Buffer.Blit([NSString stringWithFormat:@"%s", __PRETTY_FUNCTION__]);
+		}
+		default:
+		{
+			break;
+		}
+	}
+#endif
+}
+
+void FAGXBlitCommandEncoderDebugging::Copy(FAGXTexture const& sourceTexture, NSUInteger sourceSlice, NSUInteger sourceLevel, mtlpp::Origin const& sourceOrigin, mtlpp::Size const& sourceSize, FAGXBuffer const& destinationBuffer, NSUInteger destinationOffset, NSUInteger destinationBytesPerRow, NSUInteger destinationBytesPerImage, mtlpp::BlitOption options)
+{
+#if METAL_DEBUG_OPTIONS
+	switch(((FAGXDebugBlitCommandEncoder*)m_ptr)->Buffer.GetPtr()->DebugLevel)
+	{
+		case EAGXDebugLevelConditionalSubmit:
+		case EAGXDebugLevelWaitForComplete:
+		case EAGXDebugLevelLogOperations:
+		{
+			((FAGXDebugBlitCommandEncoder*)m_ptr)->Buffer.Blit([NSString stringWithFormat:@"%s", __PRETTY_FUNCTION__]);
+		}
+		default:
+		{
+			break;
+		}
+	}
+#endif
+}
+
+void FAGXBlitCommandEncoderDebugging::GenerateMipmaps(FAGXTexture const& texture)
+{
+#if METAL_DEBUG_OPTIONS
+	switch(((FAGXDebugBlitCommandEncoder*)m_ptr)->Buffer.GetPtr()->DebugLevel)
+	{
+		case EAGXDebugLevelConditionalSubmit:
+		case EAGXDebugLevelWaitForComplete:
+		case EAGXDebugLevelLogOperations:
+		{
+			((FAGXDebugBlitCommandEncoder*)m_ptr)->Buffer.Blit([NSString stringWithFormat:@"%s", __PRETTY_FUNCTION__]);
+		}
+		default:
+		{
+			break;
+		}
+	}
+#endif
+}
+
+void FAGXBlitCommandEncoderDebugging::Fill(FAGXBuffer const& buffer, ns::Range const& range, uint8_t value)
+{
+#if METAL_DEBUG_OPTIONS
+	switch(((FAGXDebugBlitCommandEncoder*)m_ptr)->Buffer.GetPtr()->DebugLevel)
+	{
+		case EAGXDebugLevelConditionalSubmit:
+		case EAGXDebugLevelWaitForComplete:
+		case EAGXDebugLevelLogOperations:
+		{
+			((FAGXDebugBlitCommandEncoder*)m_ptr)->Buffer.Blit([NSString stringWithFormat:@"%s", __PRETTY_FUNCTION__]);
+		}
+		default:
+		{
+			break;
+		}
+	}
+#endif
+}
+
+void FAGXBlitCommandEncoderDebugging::Copy(FAGXBuffer const& sourceBuffer, NSUInteger sourceOffset, FAGXBuffer const& destinationBuffer, NSUInteger destinationOffset, NSUInteger size)
+{
+#if METAL_DEBUG_OPTIONS
+	switch(((FAGXDebugBlitCommandEncoder*)m_ptr)->Buffer.GetPtr()->DebugLevel)
+	{
+		case EAGXDebugLevelConditionalSubmit:
+		case EAGXDebugLevelWaitForComplete:
+		case EAGXDebugLevelLogOperations:
+		{
+			((FAGXDebugBlitCommandEncoder*)m_ptr)->Buffer.Blit([NSString stringWithFormat:@"%s", __PRETTY_FUNCTION__]);
+		}
+		default:
+		{
+			break;
+		}
+	}
+#endif
+}
+
+#endif
diff --git a/Engine/Source/Runtime/Apple/AGXRHI/Private/AGXBlitCommandEncoder.h b/Engine/Source/Runtime/Apple/AGXRHI/Private/AGXBlitCommandEncoder.h
new file mode 100644
--- /dev/null
+++ b/Engine/Source/Runtime/Apple/AGXRHI/Private/AGXBlitCommandEncoder.h
@@ -0,0 +1,53 @@
+// Copyright Epic Games, Inc. All Rights Reserved.
+
+#pragma once
+
+#include <Metal/Metal.h>
+#include "AGXDebugCommandEncoder.h"
+
+#if MTLPP_CONFIG_VALIDATE && METAL_DEBUG_OPTIONS
+NS_ASSUME_NONNULL_BEGIN
+
+@class FAGXDebugCommandBuffer;
+@class FAGXDebugBlitCommandEncoder;
+
+class FAGXBlitCommandEncoderDebugging : public FAGXCommandEncoderDebugging
+{
+public:
+	FAGXBlitCommandEncoderDebugging();
+	FAGXBlitCommandEncoderDebugging(mtlpp::BlitCommandEncoder& Encoder, FAGXCommandBufferDebugging& Buffer);
+	FAGXBlitCommandEncoderDebugging(FAGXDebugCommandEncoder* handle);
+	
+	static FAGXBlitCommandEncoderDebugging Get(mtlpp::BlitCommandEncoder& Buffer);
+	
+	void InsertDebugSignpost(ns::String const& Label);
+	void PushDebugGroup(ns::String const& Group);
+	void PopDebugGroup();
+	
+	void EndEncoder();
+	
+#if PLATFORM_MAC
+	void Synchronize(mtlpp::Resource const& resource);
+	
+	void Synchronize(FAGXTexture const& texture, NSUInteger slice, NSUInteger level);
+#endif
+	
+	void Copy(FAGXTexture const& sourceTexture, NSUInteger sourceSlice, NSUInteger sourceLevel, mtlpp::Origin const& sourceOrigin, mtlpp::Size const& sourceSize, FAGXTexture const& destinationTexture, NSUInteger destinationSlice, NSUInteger destinationLevel, mtlpp::Origin const& destinationOrigin);
+	
+	void Copy(FAGXBuffer const& sourceBuffer, NSUInteger sourceOffset, NSUInteger sourceBytesPerRow, NSUInteger sourceBytesPerImage, mtlpp::Size const& sourceSize, FAGXTexture const& destinationTexture, NSUInteger destinationSlice, NSUInteger destinationLevel, mtlpp::Origin const& destinationOrigin);
+	
+	void Copy(FAGXBuffer const& sourceBuffer, NSUInteger sourceOffset, NSUInteger sourceBytesPerRow, NSUInteger sourceBytesPerImage, mtlpp::Size const& sourceSize, FAGXTexture const& destinationTexture, NSUInteger destinationSlice, NSUInteger destinationLevel, mtlpp::Origin const& destinationOrigin, mtlpp::BlitOption options);
+	
+	void Copy(FAGXTexture const& sourceTexture, NSUInteger sourceSlice, NSUInteger sourceLevel, mtlpp::Origin const& sourceOrigin, mtlpp::Size const& sourceSize, FAGXBuffer const& destinationBuffer, NSUInteger destinationOffset, NSUInteger destinationBytesPerRow, NSUInteger destinationBytesPerImage);
+	
+	void Copy(FAGXTexture const& sourceTexture, NSUInteger sourceSlice, NSUInteger sourceLevel, mtlpp::Origin const& sourceOrigin, mtlpp::Size const& sourceSize, FAGXBuffer const& destinationBuffer, NSUInteger destinationOffset, NSUInteger destinationBytesPerRow, NSUInteger destinationBytesPerImage, mtlpp::BlitOption options);
+	
+	void GenerateMipmaps(FAGXTexture const& texture);
+	
+	void Fill(FAGXBuffer const& buffer, ns::Range const& range, uint8_t value);
+	
+	void Copy(FAGXBuffer const& sourceBuffer, NSUInteger sourceOffset, FAGXBuffer const& destinationBuffer, NSUInteger destinationOffset, NSUInteger size);
+};
+
+NS_ASSUME_NONNULL_END
+#endif
diff --git a/Engine/Source/Runtime/Apple/AGXRHI/Private/AGXBuffer.cpp b/Engine/Source/Runtime/Apple/AGXRHI/Private/AGXBuffer.cpp
--- a/Engine/Source/Runtime/Apple/AGXRHI/Private/AGXBuffer.cpp
+++ b/Engine/Source/Runtime/Apple/AGXRHI/Private/AGXBuffer.cpp
@@ -5,6 +5,7 @@
 #include "Templates/AlignmentTemplates.h"
 #include "AGXLLM.h"
 #include <objc/runtime.h>
+#include "AGXCommandBuffer.h"
 #include "AGXProfiler.h"
 #include "AGXRenderPass.h"
 
@@ -171,6 +172,7 @@ FAGXSubBufferHeap::FAGXSubBufferHeap(NSUInteger Size, NSUInteger Alignment, mtlp
 , UsedSize(0)
 {
 	Options = (mtlpp::ResourceOptions)FAGXCommandQueue::GetCompatibleResourceOptions(Options);
+	static bool bSupportsHeaps = GetAGXDeviceContext().SupportsFeature(EAGXFeaturesHeaps);
 	NSUInteger FullSize = Align(Size, Alignment);
 	METAL_GPUPROFILE(FAGXScopedCPUStats CPUStat(FString::Printf(TEXT("AllocBuffer: %llu, %llu"), FullSize, Options)));
 	
@@ -179,22 +181,43 @@ FAGXSubBufferHeap::FAGXSubBufferHeap(NSUInteger Size, NSUInteger Alignment, mtlp
 	check(Storage != mtlpp::StorageMode::Managed /* Managed memory cannot be safely suballocated! When you overwrite existing data the GPU buffer is immediately disposed of! */);
 #endif
 
-	ParentBuffer = MTLPP_VALIDATE(mtlpp::Device, GMtlppDevice, AGXSafeGetRuntimeDebuggingLevel() >= EAGXDebugLevelValidation, NewBuffer(FullSize, Options));
-	check(ParentBuffer.GetPtr());
-	check(ParentBuffer.GetLength() >= FullSize);
+	if (bSupportsHeaps && Storage == mtlpp::StorageMode::Private)
+	{
+		mtlpp::HeapDescriptor Desc;
+		Desc.SetSize(FullSize);
+		Desc.SetStorageMode(Storage);
+		ParentHeap = GMtlppDevice.NewHeap(Desc);
+		check(ParentHeap.GetPtr());
 #if STATS || ENABLE_LOW_LEVEL_MEM_TRACKER
-	AGXLLM::LogAllocBuffer(ParentBuffer);
+		AGXLLM::LogAllocHeap(ParentHeap);
 #endif
-	FreeRanges.Add(ns::Range(0, FullSize));
-
+	}
+	else
+	{
+		ParentBuffer = MTLPP_VALIDATE(mtlpp::Device, GMtlppDevice, AGXSafeGetRuntimeDebuggingLevel() >= EAGXDebugLevelValidation, NewBuffer(FullSize, Options));
+		check(ParentBuffer.GetPtr());
+		check(ParentBuffer.GetLength() >= FullSize);
+#if STATS || ENABLE_LOW_LEVEL_MEM_TRACKER
+		AGXLLM::LogAllocBuffer(ParentBuffer);
+#endif
+		FreeRanges.Add(ns::Range(0, FullSize));
+	}
 	INC_MEMORY_STAT_BY(STAT_AGXBufferUnusedMemory, FullSize);
 	INC_MEMORY_STAT_BY(STAT_AGXHeapBufferUnusedMemory, FullSize);
 }
 
 FAGXSubBufferHeap::~FAGXSubBufferHeap()
 {
-	DEC_MEMORY_STAT_BY(STAT_AGXBufferUnusedMemory, ParentBuffer.GetLength());
-	DEC_MEMORY_STAT_BY(STAT_AGXHeapBufferUnusedMemory, ParentBuffer.GetLength());
+	if (ParentHeap)
+	{
+		DEC_MEMORY_STAT_BY(STAT_AGXBufferUnusedMemory, ParentHeap.GetSize());
+		DEC_MEMORY_STAT_BY(STAT_AGXHeapBufferUnusedMemory, ParentHeap.GetSize());
+	}
+	else
+	{
+		DEC_MEMORY_STAT_BY(STAT_AGXBufferUnusedMemory, ParentBuffer.GetLength());
+		DEC_MEMORY_STAT_BY(STAT_AGXHeapBufferUnusedMemory, ParentBuffer.GetLength());
+	}
 }
 
 void FAGXSubBufferHeap::SetOwner(ns::Range const& Range, FAGXRHIBuffer* Owner, bool bIsSwap)
@@ -228,6 +251,13 @@ void FAGXSubBufferHeap::FreeRange(ns::Range const& Range)
 			}
 		}
 	}
+	if (ParentHeap)
+	{
+		INC_MEMORY_STAT_BY(STAT_AGXBufferUnusedMemory, Range.Length);
+		INC_MEMORY_STAT_BY(STAT_AGXHeapBufferUnusedMemory, Range.Length);
+		DEC_MEMORY_STAT_BY(STAT_AGXHeapBufferMemory, Range.Length);
+	}
+	else
 	{
 #if METAL_DEBUG_OPTIONS
 		if (GIsRHIInitialized)
@@ -294,27 +324,62 @@ void FAGXSubBufferHeap::FreeRange(ns::Range const& Range)
 
 ns::String FAGXSubBufferHeap::GetLabel() const
 {
-	return ParentBuffer.GetLabel();
+	if (ParentHeap)
+	{
+		return ParentHeap.GetLabel();
+	}
+	else
+	{
+		return ParentBuffer.GetLabel();
+	}
 }
 
 mtlpp::StorageMode FAGXSubBufferHeap::GetStorageMode() const
 {
-	return ParentBuffer.GetStorageMode();
+	if (ParentHeap)
+	{
+		return ParentHeap.GetStorageMode();
+	}
+	else
+	{
+		return ParentBuffer.GetStorageMode();
+	}
 }
 
 mtlpp::CpuCacheMode FAGXSubBufferHeap::GetCpuCacheMode() const
 {
-	return ParentBuffer.GetCpuCacheMode();
+	if (ParentHeap)
+	{
+		return ParentHeap.GetCpuCacheMode();
+	}
+	else
+	{
+		return ParentBuffer.GetCpuCacheMode();
+	}
 }
 
 NSUInteger FAGXSubBufferHeap::GetSize() const
 {
-	return ParentBuffer.GetLength();
+	if (ParentHeap)
+	{
+		return ParentHeap.GetSize();
+	}
+	else
+	{
+		return ParentBuffer.GetLength();
+	}
 }
 
 NSUInteger FAGXSubBufferHeap::GetUsedSize() const
 {
-	return UsedSize;
+	if (ParentHeap)
+	{
+		return ParentHeap.GetUsedSize();
+	}
+	else
+	{
+		return UsedSize;
+	}
 }
 
 int64 FAGXSubBufferHeap::NumCurrentAllocations() const
@@ -324,24 +389,50 @@ int64 FAGXSubBufferHeap::NumCurrentAllocations() const
 
 void FAGXSubBufferHeap::SetLabel(const ns::String& label)
 {
-	ParentBuffer.SetLabel(label);
+	if (ParentHeap)
+	{
+		ParentHeap.SetLabel(label);
+	}
+	else
+	{
+		ParentBuffer.SetLabel(label);
+	}
 }
 
 NSUInteger FAGXSubBufferHeap::MaxAvailableSize() const
 {
-	if (UsedSize < GetSize())
+	if (ParentHeap)
 	{
-		return FreeRanges.Last().Length;
+		return ParentHeap.MaxAvailableSizeWithAlignment(MinAlign);
 	}
 	else
 	{
-		return 0;
+		if (UsedSize < GetSize())
+		{
+			return FreeRanges.Last().Length;
+		}
+		else
+		{
+			return 0;
+		}
 	}
 }
 
 bool FAGXSubBufferHeap::CanAllocateSize(NSUInteger Size) const
 {
-	return Size <= MaxAvailableSize();
+	if (ParentHeap)
+	{
+		NSUInteger Storage = (NSUInteger(GetStorageMode()) << mtlpp::ResourceStorageModeShift);
+		NSUInteger Cache = (NSUInteger(GetCpuCacheMode()) << mtlpp::ResourceCpuCacheModeShift);
+		mtlpp::ResourceOptions Opt = mtlpp::ResourceOptions(Storage | Cache);
+		
+		NSUInteger Align = GMtlppDevice.HeapBufferSizeAndAlign(Size, Opt).Align;
+		return Size <= ParentHeap.MaxAvailableSizeWithAlignment(Align);
+	}
+	else
+	{
+		return Size <= MaxAvailableSize();
+	}
 }
 
 FAGXBuffer FAGXSubBufferHeap::NewBuffer(NSUInteger length)
@@ -349,6 +440,21 @@ FAGXBuffer FAGXSubBufferHeap::NewBuffer(NSUInteger length)
 	NSUInteger Size = Align(length, MinAlign);
 	FAGXBuffer Result;
 	
+	if (ParentHeap)
+	{
+		NSUInteger Storage = (NSUInteger(GetStorageMode()) << mtlpp::ResourceStorageModeShift);
+		NSUInteger Cache = (NSUInteger(GetCpuCacheMode()) << mtlpp::ResourceCpuCacheModeShift);
+		mtlpp::ResourceOptions Opt = mtlpp::ResourceOptions(Storage | Cache);
+		
+		Result = FAGXBuffer(ParentHeap.NewBuffer(Size, Opt), this);
+#if STATS || ENABLE_LOW_LEVEL_MEM_TRACKER
+		AGXLLM::LogAllocBuffer(Result);
+#endif
+		DEC_MEMORY_STAT_BY(STAT_AGXBufferUnusedMemory, Result.GetLength());
+		DEC_MEMORY_STAT_BY(STAT_AGXHeapBufferUnusedMemory, Result.GetLength());
+		INC_MEMORY_STAT_BY(STAT_AGXHeapBufferMemory, Result.GetLength());
+	}
+	else
 	{
 		check(ParentBuffer && ParentBuffer.GetPtr());
 	
@@ -406,7 +512,14 @@ FAGXBuffer FAGXSubBufferHeap::NewBuffer(NSUInteger length)
 
 mtlpp::PurgeableState FAGXSubBufferHeap::SetPurgeableState(mtlpp::PurgeableState state)
 {
-	return ParentBuffer.SetPurgeableState(state);
+	if (ParentHeap)
+	{
+		return ParentHeap.SetPurgeableState(state);
+	}
+	else
+	{
+		return ParentBuffer.SetPurgeableState(state);
+	}
 }
 
 #pragma mark --
@@ -545,84 +658,163 @@ FAGXSubBufferMagazine::FAGXSubBufferMagazine(NSUInteger Size, NSUInteger ChunkSi
 , UsedSize(0)
 {
 	Options = (mtlpp::ResourceOptions)FAGXCommandQueue::GetCompatibleResourceOptions(Options);
-	mtlpp::StorageMode Storage = (mtlpp::StorageMode)((Options & mtlpp::ResourceStorageModeMask) >> mtlpp::ResourceStorageModeShift);
-	NSUInteger FullSize = Align(Size, MinAlign);
+    static bool bSupportsHeaps = GetAGXDeviceContext().SupportsFeature(EAGXFeaturesHeaps);
+    mtlpp::StorageMode Storage = (mtlpp::StorageMode)((Options & mtlpp::ResourceStorageModeMask) >> mtlpp::ResourceStorageModeShift);
+    if (PLATFORM_IOS && bSupportsHeaps && Storage == mtlpp::StorageMode::Private)
+    {
+        MinAlign = GMtlppDevice.HeapBufferSizeAndAlign(BlockSize, Options).Align;
+    }
+    
+    NSUInteger FullSize = Align(Size, MinAlign);
 	METAL_GPUPROFILE(FAGXScopedCPUStats CPUStat(FString::Printf(TEXT("AllocBuffer: %llu, %llu"), FullSize, Options)));
 	
 #if PLATFORM_MAC
 	check(Storage != mtlpp::StorageMode::Managed /* Managed memory cannot be safely suballocated! When you overwrite existing data the GPU buffer is immediately disposed of! */);
 #endif
 
-	ParentBuffer = MTLPP_VALIDATE(mtlpp::Device, GMtlppDevice, AGXSafeGetRuntimeDebuggingLevel() >= EAGXDebugLevelValidation, NewBuffer(FullSize, Options));
-	check(ParentBuffer.GetPtr());
-	check(ParentBuffer.GetLength() >= FullSize);
-	METAL_FATAL_ASSERT(ParentBuffer, TEXT("Failed to create heap of size %u and resource options %u"), Size, (uint32)Options);
+	if (bSupportsHeaps && Storage == mtlpp::StorageMode::Private)
+	{
+		mtlpp::HeapDescriptor Desc;
+		Desc.SetSize(FullSize);
+		Desc.SetStorageMode(Storage);
+		ParentHeap = GMtlppDevice.NewHeap(Desc);
+		check(ParentHeap.GetPtr());
+		METAL_FATAL_ASSERT(ParentHeap, TEXT("Failed to create heap of size %u and resource options %u"), Size, (uint32)Options);
 #if STATS || ENABLE_LOW_LEVEL_MEM_TRACKER
-	AGXLLM::LogAllocBuffer(ParentBuffer);
+		AGXLLM::LogAllocHeap(ParentHeap);
 #endif
-	
-	INC_MEMORY_STAT_BY(STAT_AGXBufferUnusedMemory, FullSize);
-	INC_MEMORY_STAT_BY(STAT_AGXMagazineBufferUnusedMemory, FullSize);
-	uint32 BlockCount = FullSize / ChunkSize;
-	Blocks.AddZeroed(BlockCount);
+	}
+	else
+	{
+		ParentBuffer = MTLPP_VALIDATE(mtlpp::Device, GMtlppDevice, AGXSafeGetRuntimeDebuggingLevel() >= EAGXDebugLevelValidation, NewBuffer(FullSize, Options));
+		check(ParentBuffer.GetPtr());
+		check(ParentBuffer.GetLength() >= FullSize);
+		METAL_FATAL_ASSERT(ParentBuffer, TEXT("Failed to create heap of size %u and resource options %u"), Size, (uint32)Options);
+#if STATS || ENABLE_LOW_LEVEL_MEM_TRACKER
+		AGXLLM::LogAllocBuffer(ParentBuffer);
+#endif
+		
+		INC_MEMORY_STAT_BY(STAT_AGXBufferUnusedMemory, FullSize);
+		INC_MEMORY_STAT_BY(STAT_AGXMagazineBufferUnusedMemory, FullSize);
+        uint32 BlockCount = FullSize / ChunkSize;
+        Blocks.AddZeroed(BlockCount);
+	}
 }
 
 FAGXSubBufferMagazine::~FAGXSubBufferMagazine()
 {
-	DEC_MEMORY_STAT_BY(STAT_AGXBufferUnusedMemory, ParentBuffer.GetLength());
-	DEC_MEMORY_STAT_BY(STAT_AGXMagazineBufferUnusedMemory, ParentBuffer.GetLength());
+	if (ParentHeap)
+	{
+		DEC_MEMORY_STAT_BY(STAT_AGXBufferUnusedMemory, ParentHeap.GetSize());
+		DEC_MEMORY_STAT_BY(STAT_AGXMagazineBufferUnusedMemory, ParentHeap.GetSize());
+	}
+	else
+	{
+		DEC_MEMORY_STAT_BY(STAT_AGXBufferUnusedMemory, ParentBuffer.GetLength());
+		DEC_MEMORY_STAT_BY(STAT_AGXMagazineBufferUnusedMemory, ParentBuffer.GetLength());
+	}
 }
 
 void FAGXSubBufferMagazine::FreeRange(ns::Range const& Range)
 {
 	FPlatformAtomics::InterlockedDecrement(&OutstandingAllocs);
-
-#if METAL_DEBUG_OPTIONS
-	if (GIsRHIInitialized)
+	if (ParentHeap)
 	{
-		MTLPP_VALIDATE_ONLY(mtlpp::Buffer, ParentBuffer, AGXSafeGetRuntimeDebuggingLevel() >= EAGXDebugLevelValidation, ReleaseRange(Range));
-		FAGXBuffer Buf(ParentBuffer.NewBuffer(Range), false);
-		GetAGXDeviceContext().ValidateIsInactiveBuffer(Buf);
+		INC_MEMORY_STAT_BY(STAT_AGXBufferUnusedMemory, Range.Length);
+		INC_MEMORY_STAT_BY(STAT_AGXMagazineBufferUnusedMemory, Range.Length);
+		DEC_MEMORY_STAT_BY(STAT_AGXMagazineBufferMemory, Range.Length);
 	}
+	else
+	{
+#if METAL_DEBUG_OPTIONS
+		if (GIsRHIInitialized)
+		{
+			MTLPP_VALIDATE_ONLY(mtlpp::Buffer, ParentBuffer, AGXSafeGetRuntimeDebuggingLevel() >= EAGXDebugLevelValidation, ReleaseRange(Range));
+			FAGXBuffer Buf(ParentBuffer.NewBuffer(Range), false);
+			GetAGXDeviceContext().ValidateIsInactiveBuffer(Buf);
+		}
 #endif
-
-	uint32 BlockIndex = Range.Location / Range.Length;
-	FPlatformAtomics::AtomicStore(&Blocks[BlockIndex], 0);
-	FPlatformAtomics::InterlockedAdd(&UsedSize, -((int64)Range.Length));
 	
-	INC_MEMORY_STAT_BY(STAT_AGXBufferUnusedMemory, Range.Length);
-	INC_MEMORY_STAT_BY(STAT_AGXMagazineBufferUnusedMemory, Range.Length);
-	DEC_MEMORY_STAT_BY(STAT_AGXMagazineBufferMemory, Range.Length);
+		uint32 BlockIndex = Range.Location / Range.Length;
+		FPlatformAtomics::AtomicStore(&Blocks[BlockIndex], 0);
+		FPlatformAtomics::InterlockedAdd(&UsedSize, -((int64)Range.Length));
+		
+		INC_MEMORY_STAT_BY(STAT_AGXBufferUnusedMemory, Range.Length);
+		INC_MEMORY_STAT_BY(STAT_AGXMagazineBufferUnusedMemory, Range.Length);
+		DEC_MEMORY_STAT_BY(STAT_AGXMagazineBufferMemory, Range.Length);
+	}
 }
 
 ns::String FAGXSubBufferMagazine::GetLabel() const
 {
-	return ParentBuffer.GetLabel();
+	if (ParentHeap)
+	{
+		return ParentHeap.GetLabel();
+	}
+	else
+	{
+		return ParentBuffer.GetLabel();
+	}
 }
 
 mtlpp::StorageMode FAGXSubBufferMagazine::GetStorageMode() const
 {
-	return ParentBuffer.GetStorageMode();
+	if (ParentHeap)
+	{
+		return ParentHeap.GetStorageMode();
+	}
+	else
+	{
+		return ParentBuffer.GetStorageMode();
+	}
 }
 
 mtlpp::CpuCacheMode FAGXSubBufferMagazine::GetCpuCacheMode() const
 {
-	return ParentBuffer.GetCpuCacheMode();
+	if (ParentHeap)
+	{
+		return ParentHeap.GetCpuCacheMode();
+	}
+	else
+	{
+		return ParentBuffer.GetCpuCacheMode();
+	}
 }
 
 NSUInteger FAGXSubBufferMagazine::GetSize() const
 {
-	return ParentBuffer.GetLength();
+	if (ParentHeap)
+	{
+		return ParentHeap.GetSize();
+	}
+	else
+	{
+		return ParentBuffer.GetLength();
+	}
 }
 
 NSUInteger FAGXSubBufferMagazine::GetUsedSize() const
 {
-	return (NSUInteger)FPlatformAtomics::AtomicRead(&UsedSize);
+	if (ParentHeap)
+	{
+		return ParentHeap.GetUsedSize();
+	}
+	else
+	{
+		return (NSUInteger)FPlatformAtomics::AtomicRead(&UsedSize);
+	}
 }
 
 NSUInteger FAGXSubBufferMagazine::GetFreeSize() const
 {
-	return GetSize() - GetUsedSize();
+	if (ParentHeap)
+	{
+		return ParentHeap.MaxAvailableSizeWithAlignment(MinAlign);
+	}
+	else
+	{
+		return GetSize() - GetUsedSize();
+	}
 }
 
 int64 FAGXSubBufferMagazine::NumCurrentAllocations() const
@@ -637,7 +829,14 @@ bool FAGXSubBufferMagazine::CanAllocateSize(NSUInteger Size) const
 
 void FAGXSubBufferMagazine::SetLabel(const ns::String& label)
 {
-	ParentBuffer.SetLabel(label);
+	if (ParentHeap)
+	{
+		ParentHeap.SetLabel(label);
+	}
+	else
+	{
+		ParentBuffer.SetLabel(label);
+	}
 }
 
 FAGXBuffer FAGXSubBufferMagazine::NewBuffer()
@@ -645,19 +844,36 @@ FAGXBuffer FAGXSubBufferMagazine::NewBuffer()
 	NSUInteger Size = BlockSize;
 	FAGXBuffer Result;
 
-	check(ParentBuffer && ParentBuffer.GetPtr());
-	
-	for (uint32 i = 0; i < Blocks.Num(); i++)
+	if (ParentHeap)
+	{
+		NSUInteger Storage = (NSUInteger(GetStorageMode()) << mtlpp::ResourceStorageModeShift);
+		NSUInteger Cache = (NSUInteger(GetCpuCacheMode()) << mtlpp::ResourceCpuCacheModeShift);
+		mtlpp::ResourceOptions Opt = mtlpp::ResourceOptions(Storage | Cache);
+		
+		Result = FAGXBuffer(ParentHeap.NewBuffer(Size, Opt), this);
+#if STATS || ENABLE_LOW_LEVEL_MEM_TRACKER
+		AGXLLM::LogAllocBuffer(Result);
+#endif
+		DEC_MEMORY_STAT_BY(STAT_AGXBufferUnusedMemory, Result.GetLength());
+		DEC_MEMORY_STAT_BY(STAT_AGXMagazineBufferUnusedMemory, Result.GetLength());
+		INC_MEMORY_STAT_BY(STAT_AGXMagazineBufferMemory, Result.GetLength());
+	}
+	else
 	{
-		if (FPlatformAtomics::InterlockedCompareExchange(&Blocks[i], 1, 0) == 0)
+		check(ParentBuffer && ParentBuffer.GetPtr());
+		
+		for (uint32 i = 0; i < Blocks.Num(); i++)
 		{
-			ns::Range Range(i * BlockSize, BlockSize);
-			FPlatformAtomics::InterlockedAdd(&UsedSize, ((int64)Range.Length));
-			DEC_MEMORY_STAT_BY(STAT_AGXBufferUnusedMemory, Range.Length);
-			DEC_MEMORY_STAT_BY(STAT_AGXMagazineBufferUnusedMemory, Range.Length);
-			INC_MEMORY_STAT_BY(STAT_AGXMagazineBufferMemory, Range.Length);
-			Result = FAGXBuffer(MTLPP_VALIDATE(mtlpp::Buffer, ParentBuffer, AGXSafeGetRuntimeDebuggingLevel() >= EAGXDebugLevelValidation, NewBuffer(Range)), this);
-			break;
+			if (FPlatformAtomics::InterlockedCompareExchange(&Blocks[i], 1, 0) == 0)
+			{
+				ns::Range Range(i * BlockSize, BlockSize);
+				FPlatformAtomics::InterlockedAdd(&UsedSize, ((int64)Range.Length));
+				DEC_MEMORY_STAT_BY(STAT_AGXBufferUnusedMemory, Range.Length);
+				DEC_MEMORY_STAT_BY(STAT_AGXMagazineBufferUnusedMemory, Range.Length);
+				INC_MEMORY_STAT_BY(STAT_AGXMagazineBufferMemory, Range.Length);
+				Result = FAGXBuffer(MTLPP_VALIDATE(mtlpp::Buffer, ParentBuffer, AGXSafeGetRuntimeDebuggingLevel() >= EAGXDebugLevelValidation, NewBuffer(Range)), this);
+				break;
+			}
 		}
 	}
 
@@ -668,7 +884,14 @@ FAGXBuffer FAGXSubBufferMagazine::NewBuffer()
 
 mtlpp::PurgeableState FAGXSubBufferMagazine::SetPurgeableState(mtlpp::PurgeableState state)
 {
-	return ParentBuffer.SetPurgeableState(state);
+	if (ParentHeap)
+	{
+		return ParentHeap.SetPurgeableState(state);
+	}
+	else
+	{
+		return ParentBuffer.SetPurgeableState(state);
+	}
 }
 
 FAGXRingBufferRef::FAGXRingBufferRef(FAGXBuffer Buf)
@@ -1165,13 +1388,63 @@ FAGXResourceHeap::TextureHeapSize FAGXResourceHeap::TextureSizeToIndex(uint32 Si
 	return (TextureHeapSize)Lower;
 }
 
+mtlpp::Heap FAGXResourceHeap::GetTextureHeap(mtlpp::TextureDescriptor Desc, mtlpp::SizeAndAlign Size)
+{
+	mtlpp::Heap Result;
+	static bool bTextureHeaps = FParse::Param(FCommandLine::Get(),TEXT("agxtextureheaps"));
+	if (FAGXCommandQueue::SupportsFeature(EAGXFeaturesHeaps) && bTextureHeaps && Size.Size <= HeapTextureHeapSizes[MaxTextureSize])
+	{
+		FAGXResourceHeap::TextureHeapSize HeapIndex = TextureSizeToIndex(Size.Size);
+
+		EAGXHeapTextureUsage UsageIndex = EAGXHeapTextureUsageNum;
+		mtlpp::StorageMode StorageMode = Desc.GetStorageMode();
+		mtlpp::CpuCacheMode CPUMode = Desc.GetCpuCacheMode();
+		if ((Desc.GetUsage() & mtlpp::TextureUsage::RenderTarget) && StorageMode == mtlpp::StorageMode::Private && CPUMode == mtlpp::CpuCacheMode::DefaultCache)
+		{
+			UsageIndex = PLATFORM_MAC ? EAGXHeapTextureUsageNum : EAGXHeapTextureUsageRenderTarget;
+		}
+		else if (StorageMode == mtlpp::StorageMode::Private && CPUMode == mtlpp::CpuCacheMode::WriteCombined)
+		{
+			UsageIndex = EAGXHeapTextureUsageResource;
+		}
+		
+		if (UsageIndex < EAGXHeapTextureUsageNum)
+		{
+			for (mtlpp::Heap& Heap : TextureHeaps[UsageIndex][HeapIndex])
+			{
+				if (Heap.MaxAvailableSizeWithAlignment(Size.Align) >= Size.Size)
+				{
+					Result = Heap;
+					break;
+				}
+			}
+			if (!Result)
+			{
+				MTLHeapDescriptor* HeapDesc = [[MTLHeapDescriptor alloc] init];
+				[HeapDesc setSize:HeapTextureHeapSizes[HeapIndex]];
+				[HeapDesc setStorageMode:(MTLStorageMode)Desc.GetStorageMode()];
+				[HeapDesc setCpuCacheMode:(MTLCPUCacheMode)Desc.GetCpuCacheMode()];
+				Result = mtlpp::Heap([GMtlDevice newHeapWithDescriptor:HeapDesc], nullptr, ns::Ownership::Assign);
+				[HeapDesc release];
+#if STATS || ENABLE_LOW_LEVEL_MEM_TRACKER
+				AGXLLM::LogAllocHeap(Result);
+#endif
+				TextureHeaps[UsageIndex][HeapIndex].Add(Result);
+			}
+			check(Result);
+		}
+	}
+	return Result;
+}
+
 FAGXBuffer FAGXResourceHeap::CreateBuffer(uint32 Size, uint32 Alignment, EBufferUsageFlags Flags, mtlpp::ResourceOptions Options, bool bForceUnique)
 {
 	LLM_SCOPE_METAL(ELLMTagAGX::Buffers);
 	LLM_PLATFORM_SCOPE_METAL(ELLMTagAGX::Buffers);
 	
+	static bool bSupportsHeaps = GetAGXDeviceContext().SupportsFeature(EAGXFeaturesHeaps);
 	static bool bSupportsBufferSubAllocation = FAGXCommandQueue::SupportsFeature(EAGXFeaturesBufferSubAllocation);
-	bForceUnique |= (!bSupportsBufferSubAllocation);
+	bForceUnique |= (!bSupportsBufferSubAllocation && !bSupportsHeaps);
 	
 	uint32 Usage = EnumHasAnyFlags(Flags, BUF_Static) ? UsageStatic : UsageDynamic;
 	
@@ -1383,7 +1656,18 @@ FAGXTexture FAGXResourceHeap::CreateTexture(mtlpp::TextureDescriptor Desc, FAGXS
 	LLM_SCOPE_METAL(ELLMTagAGX::Textures);
 	LLM_PLATFORM_SCOPE_METAL(ELLMTagAGX::Textures);
 	
-	if (Desc.GetUsage() & mtlpp::TextureUsage::RenderTarget)
+	mtlpp::SizeAndAlign Res = GMtlppDevice.HeapTextureSizeAndAlign(Desc);
+	mtlpp::Heap Heap = GetTextureHeap(Desc, Res);
+	if (Heap)
+	{
+		METAL_GPUPROFILE(FAGXScopedCPUStats CPUStat(FString::Printf(TEXT("AllocTexture: %s"), TEXT("")/**FString([Desc.GetPtr() description])*/)));
+		FAGXTexture Texture = Heap.NewTexture(Desc);
+#if STATS || ENABLE_LOW_LEVEL_MEM_TRACKER
+		AGXLLM::LogAllocTexture(Desc, Texture);
+#endif
+		return Texture;
+	}
+	else if (Desc.GetUsage() & mtlpp::TextureUsage::RenderTarget)
 	{
 		LLM_PLATFORM_SCOPE_METAL(ELLMTagAGX::RenderTargets);
 		return TargetPool.CreateTexture(Desc);
@@ -1396,7 +1680,7 @@ FAGXTexture FAGXResourceHeap::CreateTexture(mtlpp::TextureDescriptor Desc, FAGXS
 
 void FAGXResourceHeap::ReleaseTexture(FAGXSurface* Surface, FAGXTexture& Texture)
 {
-	if (Texture && !Texture.GetBuffer() && !Texture.GetParentTexture())
+	if (Texture && !Texture.GetBuffer() && !Texture.GetParentTexture() && !Texture.GetHeap())
 	{
         if (Texture.GetUsage() & mtlpp::TextureUsage::RenderTarget)
         {
diff --git a/Engine/Source/Runtime/Apple/AGXRHI/Private/AGXBuffer.h b/Engine/Source/Runtime/Apple/AGXRHI/Private/AGXBuffer.h
--- a/Engine/Source/Runtime/Apple/AGXRHI/Private/AGXBuffer.h
+++ b/Engine/Source/Runtime/Apple/AGXRHI/Private/AGXBuffer.h
@@ -65,6 +65,7 @@ private:
 	NSUInteger MinAlign;
 	NSUInteger UsedSize;
 	mtlpp::Buffer ParentBuffer;
+	mutable mtlpp::Heap ParentHeap;
 	TArray<ns::Range> FreeRanges;
     TArray<Allocation> AllocRanges;
 };
@@ -124,6 +125,7 @@ private:
 	int64 volatile OutstandingAllocs;
 	int64 volatile UsedSize;
 	mtlpp::Buffer ParentBuffer;
+	mutable mtlpp::Heap ParentHeap;
 	TArray<int8> Blocks;
 };
 
@@ -426,6 +428,8 @@ private:
 	uint32 GetHeapIndex(uint32 Size);
 	TextureHeapSize TextureSizeToIndex(uint32 Size);
 	
+	mtlpp::Heap GetTextureHeap(mtlpp::TextureDescriptor Desc, mtlpp::SizeAndAlign Size);
+	
 private:
 	static uint32 MagazineSizes[NumMagazineSizes];
 	static uint32 HeapSizes[NumHeapSizes];
@@ -454,4 +458,6 @@ private:
 	/** We can reuse texture allocations as well, to minimize their performance impact */
 	FAGXTexturePool TexturePool;
 	FAGXTexturePool TargetPool;
+	
+	TArray<mtlpp::Heap> TextureHeaps[EAGXHeapTextureUsageNum][NumTextureHeapSizes];
 };
diff --git a/Engine/Source/Runtime/Apple/AGXRHI/Private/AGXCommandBuffer.cpp b/Engine/Source/Runtime/Apple/AGXRHI/Private/AGXCommandBuffer.cpp
new file mode 100644
--- /dev/null
+++ b/Engine/Source/Runtime/Apple/AGXRHI/Private/AGXCommandBuffer.cpp
@@ -0,0 +1,274 @@
+// Copyright Epic Games, Inc. All Rights Reserved.
+
+/*=============================================================================
+	AGXCommandBuffer.cpp: AGX RHI command buffer wrapper.
+=============================================================================*/
+
+#include "AGXRHIPrivate.h"
+
+#include "AGXCommandBuffer.h"
+#include "AGXRenderCommandEncoder.h"
+#include "AGXBlitCommandEncoder.h"
+#include "AGXComputeCommandEncoder.h"
+#include <objc/runtime.h>
+
+NSString* GAGXDebugCommandTypeNames[EAGXDebugCommandTypeInvalid] = {
+	@"RenderEncoder",
+	@"ComputeEncoder",
+	@"BlitEncoder",
+    @"EndEncoder",
+    @"Pipeline",
+	@"Draw",
+	@"Dispatch",
+	@"Blit",
+	@"Signpost",
+	@"PushGroup",
+	@"PopGroup"
+};
+
+extern int32 GAGXRuntimeDebugLevel;
+
+uint32 AGXSafeGetRuntimeDebuggingLevel()
+{
+	return GIsRHIInitialized ? GetAGXDeviceContext().GetCommandQueue().GetRuntimeDebuggingLevel() : GAGXRuntimeDebugLevel;
+}
+
+#if MTLPP_CONFIG_VALIDATE && METAL_DEBUG_OPTIONS
+
+@implementation FAGXDebugCommandBuffer
+
+APPLE_PLATFORM_OBJECT_ALLOC_OVERRIDES(FAGXDebugCommandBuffer)
+
+-(id)initWithCommandBuffer:(id<MTLCommandBuffer>)Buffer
+{
+	id Self = [super init];
+	if (Self)
+	{
+        DebugLevel = (EAGXDebugLevel)GAGXRuntimeDebugLevel;
+		InnerBuffer = Buffer;
+		DebugGroup = [NSMutableArray new];
+		ActiveEncoder = nil;
+		DebugInfoBuffer = nil;
+		if (DebugLevel >= EAGXDebugLevelValidation)
+		{
+			DebugInfoBuffer = [Buffer.device newBufferWithLength:BufferOffsetAlignment options:0];
+		}
+	}
+	return Self;
+}
+
+-(void)dealloc
+{
+	for (FAGXDebugCommand* Command : DebugCommands)
+	{
+		[Command->Label release];
+		[Command->PassDesc release];
+		delete Command;
+	}
+	
+	[DebugGroup release];
+	[DebugInfoBuffer release];
+	DebugInfoBuffer = nil;
+	
+	[super dealloc];
+}
+
+@end
+
+FAGXCommandBufferDebugging FAGXCommandBufferDebugging::Get(mtlpp::CommandBuffer& Buffer)
+{
+	return Buffer.GetAssociatedObject<FAGXCommandBufferDebugging>((void const*)&FAGXCommandBufferDebugging::Get);
+}
+
+FAGXCommandBufferDebugging::FAGXCommandBufferDebugging()
+: ns::Object<FAGXDebugCommandBuffer*>(nullptr)
+{
+	
+}
+FAGXCommandBufferDebugging::FAGXCommandBufferDebugging(mtlpp::CommandBuffer& Buffer)
+: ns::Object<FAGXDebugCommandBuffer*>([[FAGXDebugCommandBuffer alloc] initWithCommandBuffer:Buffer.GetPtr()], ns::Ownership::Assign)
+{
+	Buffer.SetAssociatedObject((void const*)&FAGXCommandBufferDebugging::Get, *this);
+}
+FAGXCommandBufferDebugging::FAGXCommandBufferDebugging(FAGXDebugCommandBuffer* handle)
+: ns::Object<FAGXDebugCommandBuffer*>(handle)
+{
+	
+}
+
+ns::AutoReleased<ns::String> FAGXCommandBufferDebugging::GetDescription()
+{
+	NSMutableString* String = [[NSMutableString new] autorelease];
+	NSString* Label = m_ptr->InnerBuffer.label ? m_ptr->InnerBuffer.label : @"Unknown";
+	[String appendFormat:@"Command Buffer %p %@:", m_ptr->InnerBuffer, Label];
+	return ns::AutoReleased<ns::String>(String);
+}
+
+ns::AutoReleased<ns::String> FAGXCommandBufferDebugging::GetDebugDescription()
+{
+	NSMutableString* String = [[NSMutableString new] autorelease];
+	NSString* Label = m_ptr->InnerBuffer.label ? m_ptr->InnerBuffer.label : @"Unknown";
+	[String appendFormat:@"Command Buffer %p %@:", m_ptr->InnerBuffer, Label];
+	
+	uint32 Index = 0;
+	if (m_ptr->DebugInfoBuffer)
+	{
+		Index = *((uint32*)m_ptr->DebugInfoBuffer.contents);
+	}
+	
+	uint32 Count = 1;
+	for (FAGXDebugCommand* Command : m_ptr->DebugCommands)
+	{
+		if (Index == Count++)
+		{
+			[String appendFormat:@"\n\t--> %@: %@", GAGXDebugCommandTypeNames[Command->Type], Command->Label];
+		}
+		else
+		{
+			[String appendFormat:@"\n\t%@: %@", GAGXDebugCommandTypeNames[Command->Type], Command->Label];
+		}
+	}
+	
+	return ns::AutoReleased<ns::String>(String);
+}
+
+void FAGXCommandBufferDebugging::BeginRenderCommandEncoder(ns::String const& Label, mtlpp::RenderPassDescriptor const& Desc)
+{
+	if (m_ptr->DebugLevel >= EAGXDebugLevelValidation)
+	{
+		if (m_ptr->DebugLevel >= EAGXDebugLevelLogOperations)
+		{
+			check(!m_ptr->ActiveEncoder);
+			m_ptr->ActiveEncoder = [Label.GetPtr() retain];
+			FAGXDebugCommand* Command = new FAGXDebugCommand;
+			Command->Type = EAGXDebugCommandTypeRenderEncoder;
+			Command->Label = [Label.GetPtr() retain];
+			Command->PassDesc = [Desc.GetPtr() retain];
+			m_ptr->DebugCommands.Add(Command);
+		}
+	}
+}
+void FAGXCommandBufferDebugging::BeginComputeCommandEncoder(ns::String const& Label)
+{
+	if (m_ptr->DebugLevel >= EAGXDebugLevelLogOperations)
+	{
+		check(!m_ptr->ActiveEncoder);
+		m_ptr->ActiveEncoder = [Label.GetPtr() retain];
+		FAGXDebugCommand* Command = new FAGXDebugCommand;
+		Command->Type = EAGXDebugCommandTypeComputeEncoder;
+		Command->Label = [m_ptr->ActiveEncoder retain];
+		Command->PassDesc = nil;
+		m_ptr->DebugCommands.Add(Command);
+	}
+}
+void FAGXCommandBufferDebugging::BeginBlitCommandEncoder(ns::String const& Label)
+{
+	if (m_ptr->DebugLevel >= EAGXDebugLevelLogOperations)
+	{
+		check(!m_ptr->ActiveEncoder);
+		m_ptr->ActiveEncoder = [Label.GetPtr() retain];
+		FAGXDebugCommand* Command = new FAGXDebugCommand;
+		Command->Type = EAGXDebugCommandTypeBlitEncoder;
+		Command->Label = [m_ptr->ActiveEncoder retain];
+		Command->PassDesc = nil;
+		m_ptr->DebugCommands.Add(Command);
+	}
+}
+void FAGXCommandBufferDebugging::EndCommandEncoder()
+{
+	if (m_ptr->DebugLevel >= EAGXDebugLevelLogOperations)
+	{
+		check(m_ptr->ActiveEncoder);
+		FAGXDebugCommand* Command = new FAGXDebugCommand;
+		Command->Type = EAGXDebugCommandTypeEndEncoder;
+		Command->Label = m_ptr->ActiveEncoder;
+		Command->PassDesc = nil;
+		m_ptr->DebugCommands.Add(Command);
+		m_ptr->ActiveEncoder = nil;
+	}
+}
+
+void FAGXCommandBufferDebugging::SetPipeline(ns::String const& Desc)
+{
+	if (m_ptr->DebugLevel >= EAGXDebugLevelLogOperations)
+	{
+		FAGXDebugCommand* Command = new FAGXDebugCommand;
+		Command->Type = EAGXDebugCommandTypePipeline;
+		Command->Label = [Desc.GetPtr() retain];
+		Command->PassDesc = nil;
+		m_ptr->DebugCommands.Add(Command);
+	}
+}
+void FAGXCommandBufferDebugging::Draw(ns::String const& Desc)
+{
+	if (m_ptr->DebugLevel >= EAGXDebugLevelLogOperations)
+	{
+		FAGXDebugCommand* Command = new FAGXDebugCommand;
+		Command->Type = EAGXDebugCommandTypeDraw;
+		Command->Label = [Desc.GetPtr() retain];
+		Command->PassDesc = nil;
+		m_ptr->DebugCommands.Add(Command);
+	}
+}
+void FAGXCommandBufferDebugging::Dispatch(ns::String const& Desc)
+{
+	if (m_ptr->DebugLevel >= EAGXDebugLevelLogOperations)
+	{
+		FAGXDebugCommand* Command = new FAGXDebugCommand;
+		Command->Type = EAGXDebugCommandTypeDispatch;
+		Command->Label = [Desc.GetPtr() retain];
+		Command->PassDesc = nil;
+		m_ptr->DebugCommands.Add(Command);
+	}
+}
+void FAGXCommandBufferDebugging::Blit(ns::String const& Desc)
+{
+	if (m_ptr->DebugLevel >= EAGXDebugLevelLogOperations)
+	{
+		FAGXDebugCommand* Command = new FAGXDebugCommand;
+		Command->Type = EAGXDebugCommandTypeBlit;
+		Command->Label = [Desc.GetPtr() retain];
+		Command->PassDesc = nil;
+		m_ptr->DebugCommands.Add(Command);
+	}
+}
+
+void FAGXCommandBufferDebugging::InsertDebugSignpost(ns::String const& Label)
+{
+	if (m_ptr->DebugLevel >= EAGXDebugLevelLogOperations)
+	{
+		FAGXDebugCommand* Command = new FAGXDebugCommand;
+		Command->Type = EAGXDebugCommandTypeSignpost;
+		Command->Label = [Label.GetPtr() retain];
+		Command->PassDesc = nil;
+		m_ptr->DebugCommands.Add(Command);
+	}
+}
+void FAGXCommandBufferDebugging::PushDebugGroup(ns::String const& Group)
+{
+	if (m_ptr->DebugLevel >= EAGXDebugLevelLogOperations)
+	{
+		[m_ptr->DebugGroup addObject:Group.GetPtr()];
+		FAGXDebugCommand* Command = new FAGXDebugCommand;
+		Command->Type = EAGXDebugCommandTypePushGroup;
+		Command->Label = [Group.GetPtr() retain];
+		Command->PassDesc = nil;
+		m_ptr->DebugCommands.Add(Command);
+	}
+}
+void FAGXCommandBufferDebugging::PopDebugGroup()
+{
+	if (m_ptr->DebugLevel >= EAGXDebugLevelLogOperations)
+	{
+		if (m_ptr->DebugGroup.lastObject)
+		{
+			FAGXDebugCommand* Command = new FAGXDebugCommand;
+			Command->Type = EAGXDebugCommandTypePopGroup;
+			Command->Label = [m_ptr->DebugGroup.lastObject retain];
+			Command->PassDesc = nil;
+			m_ptr->DebugCommands.Add(Command);
+		}
+	}
+}
+
+#endif
diff --git a/Engine/Source/Runtime/Apple/AGXRHI/Private/AGXCommandBuffer.h b/Engine/Source/Runtime/Apple/AGXRHI/Private/AGXCommandBuffer.h
new file mode 100644
--- /dev/null
+++ b/Engine/Source/Runtime/Apple/AGXRHI/Private/AGXCommandBuffer.h
@@ -0,0 +1,104 @@
+// Copyright Epic Games, Inc. All Rights Reserved.
+
+#pragma once
+
+#include <Metal/Metal.h>
+#include "AGXResources.h"
+#include "MetalShaderResources.h"
+
+/**
+ * EAGXDebugCommandType: Types of command recorded in our debug command-buffer wrapper.
+ */
+enum EAGXDebugCommandType
+{
+	EAGXDebugCommandTypeRenderEncoder,
+	EAGXDebugCommandTypeComputeEncoder,
+	EAGXDebugCommandTypeBlitEncoder,
+	EAGXDebugCommandTypeEndEncoder,
+    EAGXDebugCommandTypePipeline,
+	EAGXDebugCommandTypeDraw,
+	EAGXDebugCommandTypeDispatch,
+	EAGXDebugCommandTypeBlit,
+	EAGXDebugCommandTypeSignpost,
+	EAGXDebugCommandTypePushGroup,
+	EAGXDebugCommandTypePopGroup,
+	EAGXDebugCommandTypeInvalid
+};
+
+/**
+ * EAGXDebugLevel: Level of AGX RHI debug features to be enabled.
+ */
+enum EAGXDebugLevel
+{
+	EAGXDebugLevelOff,
+	EAGXDebugLevelFastValidation,
+	EAGXDebugLevelResetOnBind,
+	EAGXDebugLevelConditionalSubmit,
+	EAGXDebugLevelValidation,
+	EAGXDebugLevelLogOperations,
+	EAGXDebugLevelWaitForComplete,
+};
+
+NS_ASSUME_NONNULL_BEGIN
+/**
+ * FAGXDebugCommand: The data recorded for each command in the debug command-buffer wrapper.
+ */
+struct FAGXDebugCommand
+{
+	NSString* Label;
+	EAGXDebugCommandType Type;
+	MTLRenderPassDescriptor* PassDesc;
+};
+
+#if MTLPP_CONFIG_VALIDATE && METAL_DEBUG_OPTIONS
+
+/**
+ * FAGXDebugCommandBuffer: Wrapper around id<MTLCommandBuffer> that records information about commands.
+ * This allows reporting of substantially more information in debug modes which can be especially helpful
+ * when debugging GPU command-buffer failures.
+ */
+@interface FAGXDebugCommandBuffer : FApplePlatformObject
+{
+@public
+	NSMutableArray<NSString*>* DebugGroup;
+	NSString* ActiveEncoder;
+	id<MTLCommandBuffer> InnerBuffer;
+	TArray<FAGXDebugCommand*> DebugCommands;
+	EAGXDebugLevel DebugLevel;
+	id<MTLBuffer> DebugInfoBuffer;
+};
+
+/** Initialise the wrapper with the provided command-buffer. */
+-(id)initWithCommandBuffer:(id<MTLCommandBuffer>)Buffer;
+
+@end
+
+class FAGXCommandBufferDebugging : public ns::Object<FAGXDebugCommandBuffer*>
+{
+public:
+	FAGXCommandBufferDebugging();
+	FAGXCommandBufferDebugging(mtlpp::CommandBuffer& Buffer);
+	FAGXCommandBufferDebugging(FAGXDebugCommandBuffer* handle);
+	
+	static FAGXCommandBufferDebugging Get(mtlpp::CommandBuffer& Buffer);
+	ns::AutoReleased<ns::String> GetDescription();
+	ns::AutoReleased<ns::String> GetDebugDescription();
+	
+	void BeginRenderCommandEncoder(ns::String const& Label, mtlpp::RenderPassDescriptor const& Desc);
+	void BeginComputeCommandEncoder(ns::String const& Label);
+	void BeginBlitCommandEncoder(ns::String const& Label);
+	void EndCommandEncoder();
+	
+	void SetPipeline(ns::String const& Desc);
+	void Draw(ns::String const& Desc);
+	void Dispatch(ns::String const& Desc);
+	void Blit(ns::String const& Desc);
+	
+	void InsertDebugSignpost(ns::String const& Label);
+	void PushDebugGroup(ns::String const& Group);
+	void PopDebugGroup();
+};
+
+#endif
+
+NS_ASSUME_NONNULL_END
diff --git a/Engine/Source/Runtime/Apple/AGXRHI/Private/AGXCommandEncoder.cpp b/Engine/Source/Runtime/Apple/AGXRHI/Private/AGXCommandEncoder.cpp
--- a/Engine/Source/Runtime/Apple/AGXRHI/Private/AGXCommandEncoder.cpp
+++ b/Engine/Source/Runtime/Apple/AGXRHI/Private/AGXCommandEncoder.cpp
@@ -9,6 +9,9 @@
 #include "AGXGraphicsPipelineState.h"
 #include "AGXCommandBufferFence.h"
 #include "AGXCommandEncoder.h"
+#include "AGXCommandBuffer.h"
+#include "AGXComputeCommandEncoder.h"
+#include "AGXRenderCommandEncoder.h"
 #include "AGXProfiler.h"
 #include "MetalShaderResources.h"
 
@@ -18,6 +21,194 @@ const uint32 EncoderRingBufferSize = 1024 * 1024;
 extern int32 GAGXBufferScribble;
 #endif
 
+static TCHAR const* const GAGXCommandDataTypeName[] = {
+	TEXT("DrawPrimitive"),
+	TEXT("DrawPrimitiveIndexed"),
+	TEXT("DrawPrimitivePatch"),
+	TEXT("DrawPrimitiveIndirect"),
+	TEXT("DrawPrimitiveIndexedIndirect"),
+	TEXT("Dispatch"),
+	TEXT("DispatchIndirect"),
+};
+
+
+FString FAGXCommandData::ToString() const
+{
+	FString Result;
+	if ((uint32)CommandType < (uint32)FAGXCommandData::Type::Num)
+	{
+		Result = GAGXCommandDataTypeName[(uint32)CommandType];
+		switch(CommandType)
+		{
+			case FAGXCommandData::Type::DrawPrimitive:
+				Result += FString::Printf(TEXT(" BaseInstance: %u InstanceCount: %u VertexCount: %u VertexStart: %u"), Draw.BaseInstance, Draw.InstanceCount, Draw.VertexCount, Draw.VertexStart);
+				break;
+			case FAGXCommandData::Type::DrawPrimitiveIndexed:
+				Result += FString::Printf(TEXT(" BaseInstance: %u BaseVertex: %u IndexCount: %u IndexStart: %u InstanceCount: %u"), DrawIndexed.BaseInstance, DrawIndexed.BaseVertex, DrawIndexed.IndexCount, DrawIndexed.IndexStart, DrawIndexed.InstanceCount);
+				break;
+			case FAGXCommandData::Type::DrawPrimitivePatch:
+				Result += FString::Printf(TEXT(" BaseInstance: %u InstanceCount: %u PatchCount: %u PatchStart: %u"), DrawPatch.BaseInstance, DrawPatch.InstanceCount, DrawPatch.PatchCount, DrawPatch.PatchStart);
+				break;
+			case FAGXCommandData::Type::Dispatch:
+				Result += FString::Printf(TEXT(" X: %u Y: %u Z: %u"), (uint32)Dispatch.threadgroupsPerGrid[0], (uint32)Dispatch.threadgroupsPerGrid[1], (uint32)Dispatch.threadgroupsPerGrid[2]);
+				break;
+			case FAGXCommandData::Type::DispatchIndirect:
+				Result += FString::Printf(TEXT(" Buffer: %p Offset: %u"), (void*)DispatchIndirect.ArgumentBuffer, (uint32)DispatchIndirect.ArgumentOffset);
+				break;
+			case FAGXCommandData::Type::DrawPrimitiveIndirect:
+			case FAGXCommandData::Type::DrawPrimitiveIndexedIndirect:
+			case FAGXCommandData::Type::Num:
+			default:
+				break;
+		}
+	}
+	return Result;
+};
+
+struct FAGXCommandContextDebug
+{
+	TArray<FAGXCommandDebug> Commands;
+	TSet<TRefCountPtr<FAGXGraphicsPipelineState>> PSOs;
+	TSet<TRefCountPtr<FAGXComputeShader>> ComputeShaders;
+	FAGXBuffer DebugBuffer;
+};
+
+@interface FAGXCommandBufferDebug : FApplePlatformObject
+{
+	@public
+	TArray<FAGXCommandContextDebug> Contexts;
+	uint32 Index;
+}
+@end
+@implementation FAGXCommandBufferDebug
+APPLE_PLATFORM_OBJECT_ALLOC_OVERRIDES(FAGXCommandBufferDebug)
+- (instancetype)init
+{
+	id Self = [super init];
+	if (Self)
+	{
+		Index = ~0u;
+	}
+	return Self;
+}
+- (void)dealloc
+{
+	Contexts.Empty();
+	[super dealloc];
+}
+@end
+
+char const* FAGXCommandBufferMarkers::kTableAssociationKey = "FAGXCommandBufferMarkers::kTableAssociationKey";
+
+FAGXCommandBufferMarkers::FAGXCommandBufferMarkers(void)
+: ns::Object<FAGXCommandBufferDebug*, ns::CallingConvention::ObjectiveC>(nil)
+{
+	
+}
+
+FAGXCommandBufferMarkers::FAGXCommandBufferMarkers(mtlpp::CommandBuffer& CmdBuf)
+: ns::Object<FAGXCommandBufferDebug*, ns::CallingConvention::ObjectiveC>([FAGXCommandBufferDebug new], ns::Ownership::Assign)
+{
+	CmdBuf.SetAssociatedObject<FAGXCommandBufferMarkers>(FAGXCommandBufferMarkers::kTableAssociationKey, *this);
+	m_ptr->Contexts.SetNum(1);
+}
+
+
+FAGXCommandBufferMarkers::FAGXCommandBufferMarkers(FAGXCommandBufferDebug* CmdBuf)
+: ns::Object<FAGXCommandBufferDebug*, ns::CallingConvention::ObjectiveC>(CmdBuf)
+{
+	
+}
+
+void FAGXCommandBufferMarkers::AllocateContexts(uint32 NumContexts)
+{
+	if (m_ptr && m_ptr->Contexts.Num() < NumContexts)
+	{
+		m_ptr->Contexts.SetNum(NumContexts);
+	}
+}
+
+uint32 FAGXCommandBufferMarkers::AddCommand(uint32 CmdBufIndex, uint32 Encoder, uint32 ContextIndex, FAGXBuffer& DebugBuffer, FAGXGraphicsPipelineState* PSO, FAGXComputeShader* ComputeShader, FAGXCommandData& Data)
+{
+	uint32 Num = 0;
+	if (m_ptr)
+	{
+		if (m_ptr->Index == ~0u)
+		{
+			m_ptr->Index = CmdBufIndex;
+		}
+		
+		FAGXCommandContextDebug& Context = m_ptr->Contexts[ContextIndex];
+		if (Context.DebugBuffer != DebugBuffer)
+		{
+			Context.DebugBuffer = DebugBuffer;
+		}
+		
+		if (PSO)
+			Context.PSOs.Add(PSO);
+		if (ComputeShader)
+			Context.ComputeShaders.Add(ComputeShader);
+		
+		Num = Context.Commands.Num();
+		FAGXCommandDebug Command;
+        Command.CmdBufIndex = CmdBufIndex;
+		Command.Encoder = Encoder;
+		Command.Index = Num;
+		Command.PSO = PSO;
+		Command.ComputeShader = ComputeShader;
+		Command.Data = Data;
+		Context.Commands.Add(Command);
+	}
+	return Num;
+}
+
+TArray<FAGXCommandDebug>* FAGXCommandBufferMarkers::GetCommands(uint32 ContextIndex)
+{
+	TArray<FAGXCommandDebug>* Result = nullptr;
+	if (m_ptr)
+	{
+		FAGXCommandContextDebug& Context = m_ptr->Contexts[ContextIndex];
+		Result = &Context.Commands;
+	}
+	return Result;
+}
+
+ns::AutoReleased<FAGXBuffer> FAGXCommandBufferMarkers::GetDebugBuffer(uint32 ContextIndex)
+{
+	ns::AutoReleased<FAGXBuffer> Buffer;
+	if (m_ptr)
+	{
+		FAGXCommandContextDebug& Context = m_ptr->Contexts[ContextIndex];
+		Buffer = Context.DebugBuffer;
+	}
+	return Buffer;
+}
+
+uint32 FAGXCommandBufferMarkers::NumContexts() const
+{
+	uint32 Num = 0;
+	if (m_ptr)
+	{
+		Num = m_ptr->Contexts.Num();
+	}
+	return Num;
+}
+
+uint32 FAGXCommandBufferMarkers::GetIndex() const
+{
+	uint32 Num = 0;
+	if (m_ptr)
+	{
+		Num = m_ptr->Index;
+	}
+	return Num;
+}
+
+FAGXCommandBufferMarkers FAGXCommandBufferMarkers::Get(mtlpp::CommandBuffer const& CmdBuf)
+{
+	return CmdBuf.GetAssociatedObject<FAGXCommandBufferMarkers>(FAGXCommandBufferMarkers::kTableAssociationKey);
+}
+
 #pragma mark - Public C++ Boilerplate -
 
 FAGXCommandEncoder::FAGXCommandEncoder(FAGXCommandList& CmdList, EAGXCommandEncoderType InType)
@@ -25,10 +216,16 @@ FAGXCommandEncoder::FAGXCommandEncoder(FAGXCommandList& CmdList, EAGXCommandEnco
 , bSupportsMetalFeaturesSetBytes(CmdList.GetCommandQueue().SupportsFeature(EAGXFeaturesSetBytes))
 , RingBuffer(EncoderRingBufferSize, BufferOffsetAlignment, FAGXCommandQueue::GetCompatibleResourceOptions((mtlpp::ResourceOptions)(mtlpp::ResourceOptions::HazardTrackingModeUntracked | BUFFER_RESOURCE_STORAGE_MANAGED)))
 , RenderPassDesc(nil)
+, EncoderFence(nil)
 #if ENABLE_METAL_GPUPROFILE
 , CommandBufferStats(nullptr)
 #endif
+#if METAL_DEBUG_OPTIONS
+, WaitCount(0)
+, UpdateCount(0)
+#endif
 , DebugGroups([NSMutableArray new])
+, FenceStage(mtlpp::RenderStages::Fragment)
 , EncoderNum(0)
 , CmdBufIndex(0)
 , Type(InType)
@@ -182,6 +379,12 @@ void FAGXCommandEncoder::StartCommandBuffer(void)
 	{
 		CmdBufIndex++;
 		CommandBuffer = CommandList.GetCommandQueue().CreateCommandBuffer();
+		METAL_DEBUG_LAYER(EAGXDebugLevelFastValidation, CommandBufferDebug = FAGXCommandBufferDebugging::Get(CommandBuffer));
+		
+		if (GAGXCommandBufferDebuggingEnabled)
+		{
+			CommandBufferMarkers = FAGXCommandBufferMarkers(CommandBuffer);
+		}
 		
 		if ([DebugGroups count] > 0)
 		{
@@ -336,6 +539,11 @@ mtlpp::BlitCommandEncoder& FAGXCommandEncoder::GetBlitCommandEncoder(void)
 	return BlitCommandEncoder;
 }
 
+TRefCountPtr<FAGXFence> const& FAGXCommandEncoder::GetEncoderFence(void) const
+{
+	return EncoderFence;
+}
+	
 #pragma mark - Public Command Encoder Mutators -
 
 void FAGXCommandEncoder::BeginParallelRenderCommandEncoding(uint32 NumChildren)
@@ -345,10 +553,14 @@ void FAGXCommandEncoder::BeginParallelRenderCommandEncoding(uint32 NumChildren)
 	check(CommandBuffer);
 	check(IsRenderCommandEncoderActive() == false && IsComputeCommandEncoderActive() == false && IsBlitCommandEncoderActive() == false);
 	
+	FenceResources.Append(TransitionedResources);
+	
 	ParallelRenderCommandEncoder = MTLPP_VALIDATE(mtlpp::CommandBuffer, CommandBuffer, AGXSafeGetRuntimeDebuggingLevel() >= EAGXDebugLevelValidation, ParallelRenderCommandEncoder(RenderPassDesc));
+	METAL_DEBUG_LAYER(EAGXDebugLevelFastValidation, ParallelEncoderDebug = FAGXParallelRenderCommandEncoderDebugging(ParallelRenderCommandEncoder, RenderPassDesc, CommandBufferDebug));
 	
 	EncoderNum++;
 	
+	check(!EncoderFence);
 	NSString* Label = nil;
 	
 	if(GetEmitDrawEvents())
@@ -361,6 +573,7 @@ void FAGXCommandEncoder::BeginParallelRenderCommandEncoding(uint32 NumChildren)
 			for (NSString* Group in DebugGroups)
 			{
 				ParallelRenderCommandEncoder.PushDebugGroup(Group);
+				METAL_DEBUG_LAYER(EAGXDebugLevelFastValidation, RenderEncoderDebug.PushDebugGroup(Group));
 			}
 		}
 	}
@@ -369,6 +582,7 @@ void FAGXCommandEncoder::BeginParallelRenderCommandEncoding(uint32 NumChildren)
 	{
 		mtlpp::RenderCommandEncoder CommandEncoder = MTLPP_VALIDATE(mtlpp::ParallelRenderCommandEncoder, ParallelRenderCommandEncoder, AGXSafeGetRuntimeDebuggingLevel() >= EAGXDebugLevelValidation, GetRenderCommandEncoder());
 		ChildRenderCommandEncoders.Add(CommandEncoder);
+		METAL_DEBUG_LAYER(EAGXDebugLevelFastValidation, ParallelEncoderDebug.GetRenderCommandEncoderDebugger(CommandEncoder));
 	}
 }
 
@@ -378,16 +592,21 @@ void FAGXCommandEncoder::BeginRenderCommandEncoding(void)
 	check(CommandList.IsParallel() || CommandBuffer);
 	check(IsRenderCommandEncoderActive() == false && IsComputeCommandEncoderActive() == false && IsBlitCommandEncoderActive() == false);
 	
+	FenceResources.Append(TransitionedResources);
+	
 	if (!CommandList.IsParallel() || Type == EAGXCommandEncoderPrologue)
 	{
 		RenderCommandEncoder = MTLPP_VALIDATE(mtlpp::CommandBuffer, CommandBuffer, AGXSafeGetRuntimeDebuggingLevel() >= EAGXDebugLevelValidation, RenderCommandEncoder(RenderPassDesc));
+		METAL_DEBUG_LAYER(EAGXDebugLevelFastValidation, RenderEncoderDebug = FAGXRenderCommandEncoderDebugging(RenderCommandEncoder, RenderPassDesc, CommandBufferDebug));
 		EncoderNum++;	
 	}
 	else
 	{
 		RenderCommandEncoder = GetAGXDeviceContext().GetParallelRenderCommandEncoder(CommandList.GetParallelIndex(), ParallelRenderCommandEncoder, CommandBuffer);
+		METAL_DEBUG_LAYER(EAGXDebugLevelFastValidation, RenderEncoderDebug = FAGXRenderCommandEncoderDebugging::Get(RenderCommandEncoder));
 	}
 	
+	check(!EncoderFence);
 	NSString* Label = nil;
 	
 	if(GetEmitDrawEvents())
@@ -400,9 +619,15 @@ void FAGXCommandEncoder::BeginRenderCommandEncoding(void)
 			for (NSString* Group in DebugGroups)
 			{
 				RenderCommandEncoder.PushDebugGroup(Group);
+				METAL_DEBUG_LAYER(EAGXDebugLevelFastValidation, RenderEncoderDebug.PushDebugGroup(Group));
 			}
 		}
 	}
+	
+	if (CommandList.IsImmediate())
+	{
+		EncoderFence = CommandList.GetCommandQueue().CreateFence(Label);
+	}
 }
 
 void FAGXCommandEncoder::BeginComputeCommandEncoding(mtlpp::DispatchType DispatchType)
@@ -410,6 +635,9 @@ void FAGXCommandEncoder::BeginComputeCommandEncoding(mtlpp::DispatchType Dispatc
 	check(CommandBuffer);
 	check(IsRenderCommandEncoderActive() == false && IsComputeCommandEncoderActive() == false && IsBlitCommandEncoderActive() == false);
 	
+	FenceResources.Append(TransitionedResources);
+	TransitionedResources.Empty();
+	
 	if (DispatchType == mtlpp::DispatchType::Serial)
 	{
 		ComputeCommandEncoder = MTLPP_VALIDATE(mtlpp::CommandBuffer, CommandBuffer, AGXSafeGetRuntimeDebuggingLevel() >= EAGXDebugLevelValidation, ComputeCommandEncoder());
@@ -418,9 +646,11 @@ void FAGXCommandEncoder::BeginComputeCommandEncoding(mtlpp::DispatchType Dispatc
 	{
 		ComputeCommandEncoder = MTLPP_VALIDATE(mtlpp::CommandBuffer, CommandBuffer, AGXSafeGetRuntimeDebuggingLevel() >= EAGXDebugLevelValidation, ComputeCommandEncoder(DispatchType));
 	}
+	METAL_DEBUG_LAYER(EAGXDebugLevelFastValidation, ComputeEncoderDebug = FAGXComputeCommandEncoderDebugging(ComputeCommandEncoder, CommandBufferDebug));
 
 	EncoderNum++;
 	
+	check(!EncoderFence);
 	NSString* Label = nil;
 	
 	if(GetEmitDrawEvents())
@@ -433,9 +663,12 @@ void FAGXCommandEncoder::BeginComputeCommandEncoding(mtlpp::DispatchType Dispatc
 			for (NSString* Group in DebugGroups)
 			{
 				ComputeCommandEncoder.PushDebugGroup(Group);
+				METAL_DEBUG_LAYER(EAGXDebugLevelFastValidation, ComputeEncoderDebug.PushDebugGroup(Group));
 			}
 		}
 	}
+	
+	EncoderFence = CommandList.GetCommandQueue().CreateFence(Label);
 }
 
 void FAGXCommandEncoder::BeginBlitCommandEncoding(void)
@@ -443,10 +676,15 @@ void FAGXCommandEncoder::BeginBlitCommandEncoding(void)
 	check(CommandBuffer);
 	check(IsRenderCommandEncoderActive() == false && IsComputeCommandEncoderActive() == false && IsBlitCommandEncoderActive() == false);
 	
+	FenceResources.Append(TransitionedResources);
+	TransitionedResources.Empty();
+	
 	BlitCommandEncoder = MTLPP_VALIDATE(mtlpp::CommandBuffer, CommandBuffer, AGXSafeGetRuntimeDebuggingLevel() >= EAGXDebugLevelValidation, BlitCommandEncoder());
+	METAL_DEBUG_LAYER(EAGXDebugLevelFastValidation, BlitEncoderDebug = FAGXBlitCommandEncoderDebugging(BlitCommandEncoder, CommandBufferDebug));
 	
 	EncoderNum++;
 	
+	check(!EncoderFence);
 	NSString* Label = nil;
 	
 	if(GetEmitDrawEvents())
@@ -459,19 +697,25 @@ void FAGXCommandEncoder::BeginBlitCommandEncoding(void)
 			for (NSString* Group in DebugGroups)
 			{
 				BlitCommandEncoder.PushDebugGroup(Group);
+				METAL_DEBUG_LAYER(EAGXDebugLevelFastValidation, BlitEncoderDebug.PushDebugGroup(Group));
 			}
 		}
 	}
+	
+	EncoderFence = CommandList.GetCommandQueue().CreateFence(Label);
 }
 
-void FAGXCommandEncoder::EndEncoding(void)
+TRefCountPtr<FAGXFence> FAGXCommandEncoder::EndEncoding(void)
 {
+	static bool bSupportsFences = CommandList.GetCommandQueue().SupportsFeature(EAGXFeaturesFences);
+	TRefCountPtr<FAGXFence> Fence = nullptr;
 	@autoreleasepool
 	{
 		if(IsRenderCommandEncoderActive())
 		{
 			if (RenderCommandEncoder)
 			{
+				check(!bSupportsFences || EncoderFence || !CommandList.IsImmediate());
 				if (ParallelRenderCommandEncoder.GetPtr() == nil)
 				{
 					check(RenderPassDesc);
@@ -500,15 +744,52 @@ void FAGXCommandEncoder::EndEncoding(void)
 					}
 				}
 				
+				for (FAGXFence* FragFence : FragmentFences)
+				{
+					if (FragFence->NeedsWait(mtlpp::RenderStages::Fragment))
+					{
+						mtlpp::Fence FragmentFence = FragFence->Get(mtlpp::RenderStages::Fragment);
+						mtlpp::Fence FragInnerFence = METAL_DEBUG_OPTION(CommandList.GetCommandQueue().GetRuntimeDebuggingLevel() >= EAGXDebugLevelValidation ? mtlpp::Fence(((FAGXDebugFence*)FragmentFence.GetPtr()).Inner) :) FragmentFence;
+						
+						RenderCommandEncoder.WaitForFence(FragInnerFence, FenceStage);
+						METAL_DEBUG_LAYER(EAGXDebugLevelFastValidation, RenderEncoderDebug.AddWaitFence(FragmentFence));
+						FragFence->Wait(mtlpp::RenderStages::Fragment);
+					}
+				}
+				FragmentFences.Empty();
+				
+				if (FenceStage == mtlpp::RenderStages::Vertex)
+				{
+					FenceResources.Empty();
+					FenceStage = mtlpp::RenderStages::Fragment;
+				}
+				
+				if (EncoderFence && EncoderFence->NeedsWrite(mtlpp::RenderStages::Fragment))
+				{
+					Fence = EncoderFence;
+				}
+				UpdateFence(EncoderFence);
+				
+#if METAL_DEBUG_OPTIONS
+				if (bSupportsFences && AGXSafeGetRuntimeDebuggingLevel() >= EAGXDebugLevelFastValidation && (!WaitCount || !UpdateCount))
+				{
+					UE_LOG(LogAGX, Error, TEXT("%s has incorrect fence waits (%u) vs. updates (%u)."), *FString(RenderCommandEncoder.GetLabel()), WaitCount, UpdateCount);
+					
+				}
+				WaitCount = 0;
+				UpdateCount = 0;
+#endif
 				RenderCommandEncoder.EndEncoding();
+				METAL_DEBUG_LAYER(EAGXDebugLevelFastValidation, RenderEncoderDebug.EndEncoder());
 				RenderCommandEncoder = nil;
+				EncoderFence = nullptr;
 			}
 			
 			if (ParallelRenderCommandEncoder && IsParallel())
 			{
 				RingBuffer.Commit(CommandBuffer);
 				
-#if METAL_DEBUG_OPTIONS
+			#if METAL_DEBUG_OPTIONS
 				if(CommandList.GetCommandQueue().GetRuntimeDebuggingLevel() >= EAGXDebugLevelValidation)
 				{
 					for (FAGXBuffer const& Buffer : ActiveBuffers)
@@ -525,7 +806,7 @@ void FAGXCommandEncoder::EndEncoding(void)
 											}
 										});
 				}
-#endif // METAL_DEBUG_OPTIONS
+			#endif
 
 				BufferBindingHistory.Empty();
 				TextureBindingHistory.Empty();
@@ -567,6 +848,7 @@ void FAGXCommandEncoder::EndEncoding(void)
 				}
 
 				ParallelRenderCommandEncoder.EndEncoding();
+				METAL_DEBUG_LAYER(EAGXDebugLevelFastValidation, ParallelEncoderDebug.EndEncoder());
 				ParallelRenderCommandEncoder = nil;
 
 				ChildRenderCommandEncoders.Empty();
@@ -574,13 +856,83 @@ void FAGXCommandEncoder::EndEncoding(void)
 		}
 		else if(IsComputeCommandEncoderActive())
 		{
+			check(!bSupportsFences || EncoderFence);
+			
+			for (FAGXFence* FragFence : FragmentFences)
+			{
+				if (FragFence->NeedsWait(mtlpp::RenderStages::Fragment))
+				{
+					mtlpp::Fence FragmentFence = FragFence->Get(mtlpp::RenderStages::Fragment);
+					mtlpp::Fence FragInnerFence = METAL_DEBUG_OPTION(CommandList.GetCommandQueue().GetRuntimeDebuggingLevel() >= EAGXDebugLevelValidation ? mtlpp::Fence(((FAGXDebugFence*)FragmentFence.GetPtr()).Inner) :) FragmentFence;
+					
+					ComputeCommandEncoder.WaitForFence(FragInnerFence);
+					METAL_DEBUG_LAYER(EAGXDebugLevelFastValidation, ComputeEncoderDebug.AddWaitFence(FragmentFence));
+					FragFence->Wait(mtlpp::RenderStages::Fragment);
+				}
+			}
+			FragmentFences.Empty();
+			FenceResources.Empty();
+			FenceStage = mtlpp::RenderStages::Fragment;
+			
+			if (EncoderFence && EncoderFence->NeedsWrite(mtlpp::RenderStages::Fragment))
+			{
+				Fence = EncoderFence;
+			}
+			UpdateFence(EncoderFence);
+			
+#if METAL_DEBUG_OPTIONS
+			if (bSupportsFences && AGXSafeGetRuntimeDebuggingLevel() >= EAGXDebugLevelFastValidation && (!WaitCount || !UpdateCount))
+			{
+				UE_LOG(LogAGX, Error, TEXT("%s has incorrect fence waits (%u) vs. updates (%u)."), *FString(ComputeCommandEncoder.GetLabel()), WaitCount, UpdateCount);
+				
+			}
+			WaitCount = 0;
+			UpdateCount = 0;
+#endif
 			ComputeCommandEncoder.EndEncoding();
+			METAL_DEBUG_LAYER(EAGXDebugLevelFastValidation, ComputeEncoderDebug.EndEncoder());
 			ComputeCommandEncoder = nil;
+			EncoderFence = nullptr;
 		}
 		else if(IsBlitCommandEncoderActive())
 		{
+			// check(!bSupportsFences || EncoderFence);
+			
+			for (FAGXFence* FragFence : FragmentFences)
+			{
+				if (FragFence->NeedsWait(mtlpp::RenderStages::Fragment))
+				{
+					mtlpp::Fence FragmentFence = FragFence->Get(mtlpp::RenderStages::Fragment);
+					mtlpp::Fence FragInnerFence = METAL_DEBUG_OPTION(CommandList.GetCommandQueue().GetRuntimeDebuggingLevel() >= EAGXDebugLevelValidation ? mtlpp::Fence(((FAGXDebugFence*)FragmentFence.GetPtr()).Inner) :) FragmentFence;
+					
+					BlitCommandEncoder.WaitForFence(FragInnerFence);
+					METAL_DEBUG_LAYER(EAGXDebugLevelFastValidation, BlitEncoderDebug.AddWaitFence(FragmentFence));
+					FragFence->Wait(mtlpp::RenderStages::Fragment);
+				}
+			}
+			FragmentFences.Empty();
+			FenceResources.Empty();
+			FenceStage = mtlpp::RenderStages::Fragment;
+			
+			if (EncoderFence && EncoderFence->NeedsWrite(mtlpp::RenderStages::Fragment))
+			{
+				Fence = EncoderFence;
+			}
+			UpdateFence(EncoderFence);
+			
+#if METAL_DEBUG_OPTIONS
+			if (bSupportsFences && AGXSafeGetRuntimeDebuggingLevel() >= EAGXDebugLevelFastValidation && (!WaitCount || !UpdateCount))
+			{
+				UE_LOG(LogAGX, Error, TEXT("%s has incorrect fence waits (%u) vs. updates (%u)."), *FString(BlitCommandEncoder.GetLabel()), WaitCount, UpdateCount);
+				
+			}
+			WaitCount = 0;
+			UpdateCount = 0;
+#endif
 			BlitCommandEncoder.EndEncoding();
+			METAL_DEBUG_LAYER(EAGXDebugLevelFastValidation, BlitEncoderDebug.EndEncoder());
 			BlitCommandEncoder = nil;
+			EncoderFence = nullptr;
 		}
 	}
 	
@@ -596,6 +948,7 @@ void FAGXCommandEncoder::EndEncoding(void)
 		FMemory::Memzero(ShaderBuffers[Frequency].Usage);
 		ShaderBuffers[Frequency].Bound = 0;
 	}
+    return Fence;
 }
 
 void FAGXCommandEncoder::InsertCommandBufferFence(FAGXCommandBufferFence& Fence, mtlpp::CommandBufferHandler Handler)
@@ -619,6 +972,176 @@ void FAGXCommandEncoder::AddCompletionHandler(mtlpp::CommandBufferHandler Handle
 	Block_release(HeapHandler);
 }
 
+void FAGXCommandEncoder::UpdateFence(FAGXFence* Fence)
+{
+	check(IsRenderCommandEncoderActive() || IsComputeCommandEncoderActive() || IsBlitCommandEncoderActive());
+	static bool bSupportsFences = CommandList.GetCommandQueue().SupportsFeature(EAGXFeaturesFences);
+	if ((bSupportsFences METAL_DEBUG_OPTION(|| CommandList.GetCommandQueue().GetRuntimeDebuggingLevel() >= EAGXDebugLevelValidation)) && Fence)
+	{
+		mtlpp::Fence VertexFence = Fence->Get(mtlpp::RenderStages::Vertex);
+		mtlpp::Fence InnerFence = METAL_DEBUG_OPTION(CommandList.GetCommandQueue().GetRuntimeDebuggingLevel() >= EAGXDebugLevelValidation ? mtlpp::Fence(((FAGXDebugFence*)VertexFence.GetPtr()).Inner) :) VertexFence;
+		if (RenderCommandEncoder)
+		{
+			mtlpp::Fence FragmentFence = Fence->Get(mtlpp::RenderStages::Fragment);
+			mtlpp::Fence FragInnerFence = METAL_DEBUG_OPTION(CommandList.GetCommandQueue().GetRuntimeDebuggingLevel() >= EAGXDebugLevelValidation ? mtlpp::Fence(((FAGXDebugFence*)FragmentFence.GetPtr()).Inner) :) FragmentFence;
+			
+			if (Fence->NeedsWrite(mtlpp::RenderStages::Vertex))
+			{
+				RenderCommandEncoder.UpdateFence(InnerFence, (mtlpp::RenderStages)(mtlpp::RenderStages::Vertex));
+				METAL_DEBUG_LAYER(EAGXDebugLevelFastValidation, RenderEncoderDebug.AddUpdateFence(VertexFence));
+				Fence->Write(mtlpp::RenderStages::Vertex);
+				METAL_DEBUG_LAYER(EAGXDebugLevelFastValidation, UpdateCount++);
+			}
+
+			if (Fence->NeedsWrite(mtlpp::RenderStages::Fragment))
+			{
+				RenderCommandEncoder.UpdateFence(FragInnerFence, (mtlpp::RenderStages)(mtlpp::RenderStages::Fragment));
+				METAL_DEBUG_LAYER(EAGXDebugLevelFastValidation, RenderEncoderDebug.AddUpdateFence(FragmentFence));
+				Fence->Write(mtlpp::RenderStages::Fragment);
+				METAL_DEBUG_LAYER(EAGXDebugLevelFastValidation, UpdateCount++);
+			}
+		}
+		else if (ComputeCommandEncoder && Fence->NeedsWrite(mtlpp::RenderStages::Vertex))
+		{
+			ComputeCommandEncoder.UpdateFence(InnerFence);
+			METAL_DEBUG_LAYER(EAGXDebugLevelFastValidation, ComputeEncoderDebug.AddUpdateFence(VertexFence));
+			Fence->Write(mtlpp::RenderStages::Vertex);
+			METAL_DEBUG_LAYER(EAGXDebugLevelFastValidation, UpdateCount++);
+		}
+		else if (BlitCommandEncoder && Fence->NeedsWrite(mtlpp::RenderStages::Vertex))
+		{
+			BlitCommandEncoder.UpdateFence(InnerFence);
+			METAL_DEBUG_LAYER(EAGXDebugLevelFastValidation, BlitEncoderDebug.AddUpdateFence(VertexFence));
+			Fence->Write(mtlpp::RenderStages::Vertex);
+			METAL_DEBUG_LAYER(EAGXDebugLevelFastValidation, UpdateCount++);
+		}
+	}
+}
+
+void FAGXCommandEncoder::WaitForFence(FAGXFence* Fence)
+{
+	check(IsRenderCommandEncoderActive() || IsComputeCommandEncoderActive() || IsBlitCommandEncoderActive());
+	static bool bSupportsFences = CommandList.GetCommandQueue().SupportsFeature(EAGXFeaturesFences);
+	if ((bSupportsFences METAL_DEBUG_OPTION(|| CommandList.GetCommandQueue().GetRuntimeDebuggingLevel() >= EAGXDebugLevelValidation)) && Fence)
+	{
+		if (Fence->NeedsWait(mtlpp::RenderStages::Vertex))
+		{
+			METAL_DEBUG_LAYER(EAGXDebugLevelFastValidation, WaitCount++);
+			
+			mtlpp::Fence VertexFence = Fence->Get(mtlpp::RenderStages::Vertex);
+			mtlpp::Fence InnerFence = METAL_DEBUG_OPTION(CommandList.GetCommandQueue().GetRuntimeDebuggingLevel() >= EAGXDebugLevelValidation ? mtlpp::Fence(((FAGXDebugFence*)VertexFence.GetPtr()).Inner) :) VertexFence;
+			if (RenderCommandEncoder)
+			{
+				RenderCommandEncoder.WaitForFence(InnerFence, (mtlpp::RenderStages)(mtlpp::RenderStages::Vertex));
+				METAL_DEBUG_LAYER(EAGXDebugLevelFastValidation, RenderEncoderDebug.AddWaitFence(VertexFence));
+				Fence->Wait(mtlpp::RenderStages::Vertex);
+			}
+			else if (ComputeCommandEncoder)
+			{
+				ComputeCommandEncoder.WaitForFence(InnerFence);
+				METAL_DEBUG_LAYER(EAGXDebugLevelFastValidation, ComputeEncoderDebug.AddWaitFence(VertexFence));
+				Fence->Wait(mtlpp::RenderStages::Vertex);
+			}
+			else if (BlitCommandEncoder)
+			{
+				BlitCommandEncoder.WaitForFence(InnerFence);
+				METAL_DEBUG_LAYER(EAGXDebugLevelFastValidation, BlitEncoderDebug.AddWaitFence(VertexFence));
+				Fence->Wait(mtlpp::RenderStages::Vertex);
+			}
+		}
+		if (Fence->NeedsWait(mtlpp::RenderStages::Fragment))
+		{
+			if (FenceStage == mtlpp::RenderStages::Vertex || BlitCommandEncoder)
+			{
+				mtlpp::Fence FragmentFence = Fence->Get(mtlpp::RenderStages::Fragment);
+				mtlpp::Fence FragInnerFence = METAL_DEBUG_OPTION(CommandList.GetCommandQueue().GetRuntimeDebuggingLevel() >= EAGXDebugLevelValidation ? mtlpp::Fence(((FAGXDebugFence*)FragmentFence.GetPtr()).Inner) :) FragmentFence;
+				if (RenderCommandEncoder)
+				{
+					RenderCommandEncoder.WaitForFence(FragInnerFence, (mtlpp::RenderStages)(mtlpp::RenderStages::Vertex));
+					METAL_DEBUG_LAYER(EAGXDebugLevelFastValidation, RenderEncoderDebug.AddWaitFence(FragmentFence));
+					Fence->Wait(mtlpp::RenderStages::Fragment);
+				}
+				else if (ComputeCommandEncoder)
+				{
+					ComputeCommandEncoder.WaitForFence(FragInnerFence);
+					METAL_DEBUG_LAYER(EAGXDebugLevelFastValidation, ComputeEncoderDebug.AddWaitFence(FragmentFence));
+					Fence->Wait(mtlpp::RenderStages::Fragment);
+				}
+				else if (BlitCommandEncoder)
+				{
+					BlitCommandEncoder.WaitForFence(FragInnerFence);
+					METAL_DEBUG_LAYER(EAGXDebugLevelFastValidation, BlitEncoderDebug.AddWaitFence(FragmentFence));
+					Fence->Wait(mtlpp::RenderStages::Fragment);
+				}
+				METAL_DEBUG_LAYER(EAGXDebugLevelFastValidation, WaitCount++);
+			}
+			else
+			{
+				METAL_DEBUG_LAYER(EAGXDebugLevelFastValidation, WaitCount++);
+				FragmentFences.Add(Fence);
+			}
+		}
+	}
+}
+
+void FAGXCommandEncoder::WaitAndUpdateFence(FAGXFence* Fence)
+{
+	check(IsRenderCommandEncoderActive() || IsComputeCommandEncoderActive() || IsBlitCommandEncoderActive());
+	static bool bSupportsFences = CommandList.GetCommandQueue().SupportsFeature(EAGXFeaturesFences);
+	if ((bSupportsFences METAL_DEBUG_OPTION(|| CommandList.GetCommandQueue().GetRuntimeDebuggingLevel() >= EAGXDebugLevelValidation)) && Fence)
+	{
+		METAL_DEBUG_LAYER(EAGXDebugLevelFastValidation, WaitCount++);
+		METAL_DEBUG_LAYER(EAGXDebugLevelFastValidation, UpdateCount++);
+		
+		mtlpp::Fence VertexFence = Fence->Get(mtlpp::RenderStages::Vertex);
+		mtlpp::Fence InnerFence = METAL_DEBUG_OPTION(CommandList.GetCommandQueue().GetRuntimeDebuggingLevel() >= EAGXDebugLevelValidation ? mtlpp::Fence(((FAGXDebugFence*)VertexFence.GetPtr()).Inner) :) VertexFence;
+		if (RenderCommandEncoder)
+		{
+			mtlpp::Fence FragmentFence = Fence->Get(mtlpp::RenderStages::Fragment);
+			mtlpp::Fence FragInnerFence = METAL_DEBUG_OPTION(CommandList.GetCommandQueue().GetRuntimeDebuggingLevel() >= EAGXDebugLevelValidation ? mtlpp::Fence(((FAGXDebugFence*)FragmentFence.GetPtr()).Inner) :) FragmentFence;
+			
+			METAL_DEBUG_LAYER(EAGXDebugLevelFastValidation, WaitCount++);
+			METAL_DEBUG_LAYER(EAGXDebugLevelFastValidation, UpdateCount++);
+			
+			RenderCommandEncoder.WaitForFence(InnerFence, (mtlpp::RenderStages)(mtlpp::RenderStages::Vertex));
+			METAL_DEBUG_LAYER(EAGXDebugLevelFastValidation, RenderEncoderDebug.AddWaitFence(VertexFence));
+			Fence->Wait(mtlpp::RenderStages::Vertex);
+			
+			RenderCommandEncoder.WaitForFence(FragInnerFence, (mtlpp::RenderStages)(mtlpp::RenderStages::Fragment));
+			METAL_DEBUG_LAYER(EAGXDebugLevelFastValidation, RenderEncoderDebug.AddWaitFence(FragmentFence));
+			Fence->Wait(mtlpp::RenderStages::Fragment);
+			
+			RenderCommandEncoder.UpdateFence(InnerFence, (mtlpp::RenderStages)(mtlpp::RenderStages::Fragment));
+			METAL_DEBUG_LAYER(EAGXDebugLevelFastValidation, RenderEncoderDebug.AddUpdateFence(VertexFence));
+			Fence->Write(mtlpp::RenderStages::Vertex);
+			
+			RenderCommandEncoder.UpdateFence(FragInnerFence, (mtlpp::RenderStages)(mtlpp::RenderStages::Vertex));
+			METAL_DEBUG_LAYER(EAGXDebugLevelFastValidation, RenderEncoderDebug.AddUpdateFence(FragmentFence));
+			Fence->Write(mtlpp::RenderStages::Fragment);
+		}
+		else if (ComputeCommandEncoder)
+		{
+			ComputeCommandEncoder.WaitForFence(InnerFence);
+			METAL_DEBUG_LAYER(EAGXDebugLevelFastValidation, ComputeEncoderDebug.AddWaitFence(VertexFence));
+			Fence->Wait(mtlpp::RenderStages::Vertex);
+			
+			ComputeCommandEncoder.UpdateFence(InnerFence);
+			METAL_DEBUG_LAYER(EAGXDebugLevelFastValidation, ComputeEncoderDebug.AddUpdateFence(VertexFence));
+			Fence->Write(mtlpp::RenderStages::Vertex);
+		}
+		else if (BlitCommandEncoder)
+		{
+			BlitCommandEncoder.WaitForFence(InnerFence);
+			METAL_DEBUG_LAYER(EAGXDebugLevelFastValidation, BlitEncoderDebug.AddWaitFence(VertexFence));
+			Fence->Wait(mtlpp::RenderStages::Vertex);
+			
+			BlitCommandEncoder.UpdateFence(InnerFence);
+			METAL_DEBUG_LAYER(EAGXDebugLevelFastValidation, BlitEncoderDebug.AddUpdateFence(VertexFence));
+			Fence->Write(mtlpp::RenderStages::Vertex);
+		}
+	}
+}
+
 #pragma mark - Public Debug Support -
 
 void FAGXCommandEncoder::InsertDebugSignpost(ns::String const& String)
@@ -628,6 +1151,7 @@ void FAGXCommandEncoder::InsertDebugSignpost(ns::String const& String)
 		if (RenderCommandEncoder)
 		{
 			RenderCommandEncoder.InsertDebugSignpost(String);
+			METAL_DEBUG_LAYER(EAGXDebugLevelFastValidation, RenderEncoderDebug.InsertDebugSignpost(String));
 		}
 		else if (ParallelRenderCommandEncoder && !IsParallel())
 		{
@@ -637,10 +1161,12 @@ void FAGXCommandEncoder::InsertDebugSignpost(ns::String const& String)
 		else if (ComputeCommandEncoder)
 		{
 			ComputeCommandEncoder.InsertDebugSignpost(String);
+			METAL_DEBUG_LAYER(EAGXDebugLevelFastValidation, ComputeEncoderDebug.InsertDebugSignpost(String));
 		}
 		else if (BlitCommandEncoder)
 		{
 			BlitCommandEncoder.InsertDebugSignpost(String);
+			METAL_DEBUG_LAYER(EAGXDebugLevelFastValidation, BlitEncoderDebug.InsertDebugSignpost(String));
 		}
 	}
 }
@@ -653,6 +1179,7 @@ void FAGXCommandEncoder::PushDebugGroup(ns::String const& String)
 		if (RenderCommandEncoder)
 		{
 			RenderCommandEncoder.PushDebugGroup(String);
+			METAL_DEBUG_LAYER(EAGXDebugLevelFastValidation, RenderEncoderDebug.PushDebugGroup(String));
 		}
 		else if (ParallelRenderCommandEncoder && !IsParallel())
 		{
@@ -662,10 +1189,12 @@ void FAGXCommandEncoder::PushDebugGroup(ns::String const& String)
 		else if (ComputeCommandEncoder)
 		{
 			ComputeCommandEncoder.PushDebugGroup(String);
+			METAL_DEBUG_LAYER(EAGXDebugLevelFastValidation, ComputeEncoderDebug.PushDebugGroup(String));
 		}
 		else if (BlitCommandEncoder)
 		{
 			BlitCommandEncoder.PushDebugGroup(String);
+			METAL_DEBUG_LAYER(EAGXDebugLevelFastValidation, BlitEncoderDebug.PushDebugGroup(String));
 		}
 	}
 }
@@ -678,6 +1207,7 @@ void FAGXCommandEncoder::PopDebugGroup(void)
 		if (RenderCommandEncoder)
 		{
 			RenderCommandEncoder.PopDebugGroup();
+			METAL_DEBUG_LAYER(EAGXDebugLevelFastValidation, RenderEncoderDebug.PopDebugGroup());
 		}
 		else if (ParallelRenderCommandEncoder && !IsParallel())
 		{
@@ -687,14 +1217,21 @@ void FAGXCommandEncoder::PopDebugGroup(void)
 		else if (ComputeCommandEncoder)
 		{
 			ComputeCommandEncoder.PopDebugGroup();
+			METAL_DEBUG_LAYER(EAGXDebugLevelFastValidation, ComputeEncoderDebug.PopDebugGroup());
 		}
 		else if (BlitCommandEncoder)
 		{
 			BlitCommandEncoder.PopDebugGroup();
+			METAL_DEBUG_LAYER(EAGXDebugLevelFastValidation, BlitEncoderDebug.PopDebugGroup());
 		}
 	}
 }
 
+FAGXCommandBufferMarkers& FAGXCommandEncoder::GetMarkers(void)
+{
+	return CommandBufferMarkers;
+}
+
 #if ENABLE_METAL_GPUPROFILE
 FAGXCommandBufferStats* FAGXCommandEncoder::GetCommandBufferStats(void)
 {
@@ -755,6 +1292,7 @@ void FAGXCommandEncoder::SetRenderPipelineState(FAGXShaderPipeline* PipelineStat
 {
 	check (RenderCommandEncoder);
 	{
+		METAL_DEBUG_LAYER(EAGXDebugLevelFastValidation, RenderEncoderDebug.SetPipeline(PipelineState));
 		RenderCommandEncoder.SetRenderPipelineState(PipelineState->RenderPipelineState);
 	}
 }
@@ -838,6 +1376,7 @@ void FAGXCommandEncoder::SetDepthStencilState(mtlpp::DepthStencilState const& In
     check (RenderCommandEncoder);
 	{
 		RenderCommandEncoder.SetDepthStencilState(InDepthStencilState);
+		METAL_DEBUG_LAYER(EAGXDebugLevelFastValidation, RenderEncoderDebug.SetDepthStencilState(InDepthStencilState));
 	}
 }
 
@@ -934,14 +1473,17 @@ void FAGXCommandEncoder::SetShaderBytes(mtlpp::FunctionType const FunctionType,
 			{
 				case mtlpp::FunctionType::Vertex:
 					check(RenderCommandEncoder);
+					METAL_DEBUG_LAYER(EAGXDebugLevelFastValidation, RenderEncoderDebug.SetBytes(EAGXShaderVertex, Bytes, Length, Index));
 					RenderCommandEncoder.SetVertexData(Bytes, Length, Index);
 					break;
 				case mtlpp::FunctionType::Fragment:
 					check(RenderCommandEncoder);
+					METAL_DEBUG_LAYER(EAGXDebugLevelFastValidation, RenderEncoderDebug.SetBytes(EAGXShaderFragment, Bytes, Length, Index));
 					RenderCommandEncoder.SetFragmentData(Bytes, Length, Index);
 					break;
 				case mtlpp::FunctionType::Kernel:
 					check(ComputeCommandEncoder);
+					METAL_DEBUG_LAYER(EAGXDebugLevelFastValidation, ComputeEncoderDebug.SetBytes(Bytes, Length, Index));
 					ComputeCommandEncoder.SetBytes(Bytes, Length, Index);
 					break;
 				default:
@@ -988,14 +1530,17 @@ void FAGXCommandEncoder::SetShaderBufferOffset(mtlpp::FunctionType FunctionType,
 		case mtlpp::FunctionType::Vertex:
 			check (RenderCommandEncoder);
 			RenderCommandEncoder.SetVertexBufferOffset(Offset + ShaderBuffers[uint32(FunctionType)].Buffers[index].GetOffset(), index);
+			METAL_DEBUG_LAYER(EAGXDebugLevelFastValidation, RenderEncoderDebug.SetBufferOffset(EAGXShaderVertex, Offset + ShaderBuffers[uint32(FunctionType)].Buffers[index].GetOffset(), index));
 			break;
 		case mtlpp::FunctionType::Fragment:
 			check(RenderCommandEncoder);
 			RenderCommandEncoder.SetFragmentBufferOffset(Offset + ShaderBuffers[uint32(FunctionType)].Buffers[index].GetOffset(), index);
+			METAL_DEBUG_LAYER(EAGXDebugLevelFastValidation, RenderEncoderDebug.SetBufferOffset(EAGXShaderFragment, Offset + ShaderBuffers[uint32(FunctionType)].Buffers[index].GetOffset(), index));
 			break;
 		case mtlpp::FunctionType::Kernel:
 			check (ComputeCommandEncoder);
 			ComputeCommandEncoder.SetBufferOffset(Offset + ShaderBuffers[uint32(FunctionType)].Buffers[index].GetOffset(), index);
+			METAL_DEBUG_LAYER(EAGXDebugLevelFastValidation, ComputeEncoderDebug.SetBufferOffset(Offset + ShaderBuffers[uint32(FunctionType)].Buffers[index].GetOffset(), index));
 			break;
 		default:
 			check(false);
@@ -1011,15 +1556,21 @@ void FAGXCommandEncoder::SetShaderTexture(mtlpp::FunctionType FunctionType, FAGX
 		case mtlpp::FunctionType::Vertex:
 			check (RenderCommandEncoder);
 			FenceResource(Texture);
+			// MTLPP_VALIDATE(mtlpp::RenderCommandEncoder, RenderCommandEncoder, AGXSafeGetRuntimeDebuggingLevel() >= EAGXDebugLevelValidation, UseResource(Texture, mtlpp::ResourceUsage::Read));
+			METAL_DEBUG_LAYER(EAGXDebugLevelFastValidation, RenderEncoderDebug.SetTexture(EAGXShaderVertex, Texture, index));
 			RenderCommandEncoder.SetVertexTexture(Texture, index);
 			break;
 		case mtlpp::FunctionType::Fragment:
 			check(RenderCommandEncoder);
+			// MTLPP_VALIDATE(mtlpp::RenderCommandEncoder, RenderCommandEncoder, AGXSafeGetRuntimeDebuggingLevel() >= EAGXDebugLevelValidation, UseResource(Texture, mtlpp::ResourceUsage::Read));
+			METAL_DEBUG_LAYER(EAGXDebugLevelFastValidation, RenderEncoderDebug.SetTexture(EAGXShaderFragment, Texture, index));
 			RenderCommandEncoder.SetFragmentTexture(Texture, index);
 			break;
 		case mtlpp::FunctionType::Kernel:
 			check (ComputeCommandEncoder);
 			FenceResource(Texture);
+			// MTLPP_VALIDATE(mtlpp::ComputeCommandEncoder, ComputeCommandEncoder, AGXSafeGetRuntimeDebuggingLevel() >= EAGXDebugLevelValidation, UseResource(Texture, mtlpp::ResourceUsage::Read));
+			METAL_DEBUG_LAYER(EAGXDebugLevelFastValidation, ComputeEncoderDebug.SetTexture(Texture, index));
 			ComputeCommandEncoder.SetTexture(Texture, index);
 			break;
 		default:
@@ -1051,14 +1602,17 @@ void FAGXCommandEncoder::SetShaderSamplerState(mtlpp::FunctionType FunctionType,
 	{
 		case mtlpp::FunctionType::Vertex:
        		check (RenderCommandEncoder);
+			METAL_DEBUG_LAYER(EAGXDebugLevelFastValidation, RenderEncoderDebug.SetSamplerState(EAGXShaderVertex, Sampler, index));
 			RenderCommandEncoder.SetVertexSamplerState(Sampler, index);
 			break;
 		case mtlpp::FunctionType::Fragment:
 			check (RenderCommandEncoder);
+			METAL_DEBUG_LAYER(EAGXDebugLevelFastValidation, RenderEncoderDebug.SetSamplerState(EAGXShaderFragment, Sampler, index));
 			RenderCommandEncoder.SetFragmentSamplerState(Sampler, index);
 			break;
 		case mtlpp::FunctionType::Kernel:
 			check (ComputeCommandEncoder);
+			METAL_DEBUG_LAYER(EAGXDebugLevelFastValidation, ComputeEncoderDebug.SetSamplerState(Sampler, index));
 			ComputeCommandEncoder.SetSamplerState(Sampler, index);
 			break;
 		default:
@@ -1075,8 +1629,23 @@ void FAGXCommandEncoder::SetShaderSideTable(mtlpp::FunctionType const FunctionTy
 	}
 }
 
+void FAGXCommandEncoder::UseIndirectArgumentResource(FAGXTexture const& Texture, mtlpp::ResourceUsage const Usage)
+{
+	FenceResource(Texture);
+	UseResource(Texture, Usage);
+	TextureBindingHistory.Add(ns::AutoReleased<FAGXTexture>(Texture));
+}
+
+void FAGXCommandEncoder::UseIndirectArgumentResource(FAGXBuffer const& Buffer, mtlpp::ResourceUsage const Usage)
+{
+	FenceResource(Buffer);
+	UseResource(Buffer, Usage);
+	BufferBindingHistory.Add(ns::AutoReleased<FAGXBuffer>(Buffer));
+}
+
 void FAGXCommandEncoder::TransitionResources(mtlpp::Resource const& Resource)
 {
+	TransitionedResources.Add(Resource.GetPtr());
 }
 
 #pragma mark - Public Compute State Mutators -
@@ -1085,6 +1654,7 @@ void FAGXCommandEncoder::SetComputePipelineState(FAGXShaderPipeline* State)
 {
 	check (ComputeCommandEncoder);
 	{
+		METAL_DEBUG_LAYER(EAGXDebugLevelFastValidation, ComputeEncoderDebug.SetPipeline(State));
 		ComputeCommandEncoder.SetComputePipelineState(State->ComputePipelineState);
 	}
 }
@@ -1112,10 +1682,110 @@ bool FAGXCommandEncoder::HasBufferBindingHistory(FAGXBuffer const& Buffer) const
 
 void FAGXCommandEncoder::FenceResource(mtlpp::Texture const& Resource)
 {
+	mtlpp::Resource::Type Res = Resource.GetPtr();
+	ns::AutoReleased<mtlpp::Texture> Parent = Resource.GetParentTexture();
+	ns::AutoReleased<mtlpp::Buffer> Buffer = Resource.GetBuffer();
+	if (Parent)
+	{
+		Res = Parent.GetPtr();
+	}
+	else if (Buffer)
+	{
+		Res = Buffer.GetPtr();
+	}
+	if (FenceStage == mtlpp::RenderStages::Vertex || FenceResources.Contains(Res))
+	{
+		FenceStage = mtlpp::RenderStages::Vertex;
+		
+		for (FAGXFence* FragFence : FragmentFences)
+		{
+			if (FragFence->NeedsWait(mtlpp::RenderStages::Fragment))
+			{
+				mtlpp::Fence FragmentFence = FragFence->Get(mtlpp::RenderStages::Fragment);
+				mtlpp::Fence FragInnerFence = METAL_DEBUG_OPTION(CommandList.GetCommandQueue().GetRuntimeDebuggingLevel() >= EAGXDebugLevelValidation ? mtlpp::Fence(((FAGXDebugFence*)FragmentFence.GetPtr()).Inner) :) FragmentFence;
+				
+				if (RenderCommandEncoder)
+				{
+					RenderCommandEncoder.WaitForFence(FragInnerFence, (mtlpp::RenderStages)(mtlpp::RenderStages::Vertex));
+					METAL_DEBUG_LAYER(EAGXDebugLevelFastValidation, RenderEncoderDebug.AddWaitFence(FragmentFence));
+					FragFence->Wait(mtlpp::RenderStages::Fragment);
+				}
+				else if (ComputeCommandEncoder)
+				{
+					ComputeCommandEncoder.WaitForFence(FragInnerFence);
+					METAL_DEBUG_LAYER(EAGXDebugLevelFastValidation, ComputeEncoderDebug.AddWaitFence(FragmentFence));
+					FragFence->Wait(mtlpp::RenderStages::Fragment);
+				}
+				else if (BlitCommandEncoder)
+				{
+					BlitCommandEncoder.WaitForFence(FragInnerFence);
+					METAL_DEBUG_LAYER(EAGXDebugLevelFastValidation, BlitEncoderDebug.AddWaitFence(FragmentFence));
+					FragFence->Wait(mtlpp::RenderStages::Fragment);
+				}
+				METAL_DEBUG_LAYER(EAGXDebugLevelFastValidation, WaitCount++);
+			}
+		}
+		FragmentFences.Empty();
+	}
 }
 
 void FAGXCommandEncoder::FenceResource(mtlpp::Buffer const& Resource)
 {
+	mtlpp::Resource::Type Res = Resource.GetPtr();
+	if (FenceStage == mtlpp::RenderStages::Vertex || FenceResources.Contains(Res))
+	{
+		FenceStage = mtlpp::RenderStages::Vertex;
+		
+		for (FAGXFence* FragFence : FragmentFences)
+		{
+			if (FragFence->NeedsWait(mtlpp::RenderStages::Fragment))
+			{
+				mtlpp::Fence FragmentFence = FragFence->Get(mtlpp::RenderStages::Fragment);
+				mtlpp::Fence FragInnerFence = METAL_DEBUG_OPTION(CommandList.GetCommandQueue().GetRuntimeDebuggingLevel() >= EAGXDebugLevelValidation ? mtlpp::Fence(((FAGXDebugFence*)FragmentFence.GetPtr()).Inner) :) FragmentFence;
+				
+				if (RenderCommandEncoder)
+				{
+					RenderCommandEncoder.WaitForFence(FragInnerFence, (mtlpp::RenderStages)(mtlpp::RenderStages::Vertex));
+					METAL_DEBUG_LAYER(EAGXDebugLevelFastValidation, RenderEncoderDebug.AddWaitFence(FragmentFence));
+					FragFence->Wait(mtlpp::RenderStages::Fragment);
+				}
+				else if (ComputeCommandEncoder)
+				{
+					ComputeCommandEncoder.WaitForFence(FragInnerFence);
+					METAL_DEBUG_LAYER(EAGXDebugLevelFastValidation, ComputeEncoderDebug.AddWaitFence(FragmentFence));
+					FragFence->Wait(mtlpp::RenderStages::Fragment);
+				}
+				else if (BlitCommandEncoder)
+				{
+					BlitCommandEncoder.WaitForFence(FragInnerFence);
+					METAL_DEBUG_LAYER(EAGXDebugLevelFastValidation, BlitEncoderDebug.AddWaitFence(FragmentFence));
+					FragFence->Wait(mtlpp::RenderStages::Fragment);
+				}
+			}
+		}
+		FragmentFences.Empty();
+	}
+}
+
+void FAGXCommandEncoder::UseResource(mtlpp::Resource const& Resource, mtlpp::ResourceUsage const Usage)
+{
+	static bool UseResourceAvailable = FAGXCommandQueue::SupportsFeature(EAGXFeaturesIABs);
+	if (UseResourceAvailable || AGXSafeGetRuntimeDebuggingLevel() >= EAGXDebugLevelValidation)
+	{
+		mtlpp::ResourceUsage Current = ResourceUsage.FindRef(Resource.GetPtr());
+		if (Current != Usage)
+		{
+			ResourceUsage.Add(Resource.GetPtr(), Usage);
+			if (RenderCommandEncoder)
+			{
+				MTLPP_VALIDATE(mtlpp::RenderCommandEncoder, RenderCommandEncoder, AGXSafeGetRuntimeDebuggingLevel() >= EAGXDebugLevelValidation, UseResource(Resource, Usage));
+			}
+			else if (ComputeCommandEncoder)
+			{
+				MTLPP_VALIDATE(mtlpp::ComputeCommandEncoder, ComputeCommandEncoder, AGXSafeGetRuntimeDebuggingLevel() >= EAGXDebugLevelValidation, UseResource(Resource, Usage));
+			}
+		}
+	}
 }
 
 void FAGXCommandEncoder::SetShaderBufferInternal(mtlpp::FunctionType Function, uint32 Index)
@@ -1150,17 +1820,23 @@ void FAGXCommandEncoder::SetShaderBufferInternal(mtlpp::FunctionType Function, u
 				ShaderBuffers[uint32(Function)].Bound |= (1 << Index);
 				check(RenderCommandEncoder);
 				FenceResource(Buffer);
+				// MTLPP_VALIDATE(mtlpp::RenderCommandEncoder, RenderCommandEncoder, AGXSafeGetRuntimeDebuggingLevel() >= EAGXDebugLevelValidation, UseResource(Buffer, mtlpp::ResourceUsage::Read));
+				METAL_DEBUG_LAYER(EAGXDebugLevelFastValidation, RenderEncoderDebug.SetBuffer(EAGXShaderVertex, Buffer, Offset, Index));
 				RenderCommandEncoder.SetVertexBuffer(Buffer, Offset, Index);
 				break;
 			case mtlpp::FunctionType::Fragment:
 				ShaderBuffers[uint32(Function)].Bound |= (1 << Index);
 				check(RenderCommandEncoder);
+				// MTLPP_VALIDATE(mtlpp::RenderCommandEncoder, RenderCommandEncoder, AGXSafeGetRuntimeDebuggingLevel() >= EAGXDebugLevelValidation, UseResource(Buffer, mtlpp::ResourceUsage::Read));
+				METAL_DEBUG_LAYER(EAGXDebugLevelFastValidation, RenderEncoderDebug.SetBuffer(EAGXShaderFragment, Buffer, Offset, Index));
 				RenderCommandEncoder.SetFragmentBuffer(Buffer, Offset, Index);
 				break;
 			case mtlpp::FunctionType::Kernel:
 				ShaderBuffers[uint32(Function)].Bound |= (1 << Index);
 				check(ComputeCommandEncoder);
 				FenceResource(Buffer);
+				// MTLPP_VALIDATE(mtlpp::ComputeCommandEncoder, ComputeCommandEncoder, AGXSafeGetRuntimeDebuggingLevel() >= EAGXDebugLevelValidation, UseResource(Buffer, mtlpp::ResourceUsage::Read));
+				METAL_DEBUG_LAYER(EAGXDebugLevelFastValidation, ComputeEncoderDebug.SetBuffer(Buffer, Offset, Index));
 				ComputeCommandEncoder.SetBuffer(Buffer, Offset, Index);
 				break;
 			default:
@@ -1186,16 +1862,19 @@ void FAGXCommandEncoder::SetShaderBufferInternal(mtlpp::FunctionType Function, u
 			case mtlpp::FunctionType::Vertex:
 				ShaderBuffers[uint32(Function)].Bound |= (1 << Index);
 				check(RenderCommandEncoder);
+				METAL_DEBUG_LAYER(EAGXDebugLevelFastValidation, RenderEncoderDebug.SetBytes(EAGXShaderVertex, Bytes, Len, Index));
 				RenderCommandEncoder.SetVertexData(Bytes, Len, Index);
 				break;
 			case mtlpp::FunctionType::Fragment:
 				ShaderBuffers[uint32(Function)].Bound |= (1 << Index);
 				check(RenderCommandEncoder);
+				METAL_DEBUG_LAYER(EAGXDebugLevelFastValidation, RenderEncoderDebug.SetBytes(EAGXShaderFragment, Bytes, Len, Index));
 				RenderCommandEncoder.SetFragmentData(Bytes, Len, Index);
 				break;
 			case mtlpp::FunctionType::Kernel:
 				ShaderBuffers[uint32(Function)].Bound |= (1 << Index);
 				check(ComputeCommandEncoder);
+				METAL_DEBUG_LAYER(EAGXDebugLevelFastValidation, ComputeEncoderDebug.SetBytes(Bytes, Len, Index));
 				ComputeCommandEncoder.SetBytes(Bytes, Len, Index);
 				break;
 			default:
