Undo //UE5/Main/Engine/Source/Runtime/Apple/AGXRHI/Private/... changelist 18589140

diff --git a/Engine/Source/Runtime/Apple/AGXRHI/Private/AGXCommandEncoder.h b/Engine/Source/Runtime/Apple/AGXRHI/Private/AGXCommandEncoder.h
--- a/Engine/Source/Runtime/Apple/AGXRHI/Private/AGXCommandEncoder.h
+++ b/Engine/Source/Runtime/Apple/AGXRHI/Private/AGXCommandEncoder.h
@@ -4,6 +4,11 @@
 
 #include <Metal/Metal.h>
 #include "AGXBuffer.h"
+#include "AGXDebugCommandEncoder.h"
+#include "AGXBlitCommandEncoder.h"
+#include "AGXComputeCommandEncoder.h"
+#include "AGXRenderCommandEncoder.h"
+#include "AGXFence.h"
 #include "AGXProfiler.h"
 #include "command_buffer.hpp"
 
@@ -35,6 +40,67 @@ enum EAGXSubmitFlags
 	EAGXSubmitFlagsLastCommandBuffer = 1 << 6,
 };
 
+@class FAGXCommandBufferDebug;
+
+struct FAGXCommandData
+{
+	enum class Type
+	{
+		DrawPrimitive,
+		DrawPrimitiveIndexed,
+		DrawPrimitivePatch,
+		DrawPrimitiveIndirect,
+		DrawPrimitiveIndexedIndirect,
+		Dispatch,
+		DispatchIndirect,
+		Num,
+	};
+	struct DispatchIndirectArgs
+	{
+		id ArgumentBuffer;
+		NSUInteger ArgumentOffset;
+	};
+	Type CommandType;
+	union
+	{
+		mtlpp::DrawPrimitivesIndirectArguments Draw;
+		mtlpp::DrawIndexedPrimitivesIndirectArguments DrawIndexed;
+		mtlpp::DrawPatchIndirectArguments DrawPatch;
+		MTLDispatchThreadgroupsIndirectArguments Dispatch;
+		DispatchIndirectArgs DispatchIndirect;
+	};
+	FString ToString() const;
+};
+
+struct FAGXCommandDebug
+{
+    uint32 CmdBufIndex;
+	uint32 Encoder;
+	uint32 Index;
+	FAGXGraphicsPipelineState* PSO;
+	FAGXComputeShader* ComputeShader;
+	FAGXCommandData Data;
+};
+
+class FAGXCommandBufferMarkers : public ns::Object<FAGXCommandBufferDebug*, ns::CallingConvention::ObjectiveC>
+{
+public:
+	FAGXCommandBufferMarkers(void);
+	FAGXCommandBufferMarkers(mtlpp::CommandBuffer& CmdBuf);
+	FAGXCommandBufferMarkers(FAGXCommandBufferDebug* CmdBuf);
+	
+	void AllocateContexts(uint32 NumContexts);
+	uint32 AddCommand(uint32 CmdBufIndex, uint32 Encoder, uint32 ContextIndex, FAGXBuffer& DebugBuffer, FAGXGraphicsPipelineState* PSO, FAGXComputeShader* ComputeShader, FAGXCommandData& Data);
+	TArray<FAGXCommandDebug>* GetCommands(uint32 Context);
+	ns::AutoReleased<FAGXBuffer> GetDebugBuffer(uint32 ContextIndex);
+	uint32 NumContexts() const;
+	uint32 GetIndex() const;
+	
+	static FAGXCommandBufferMarkers Get(mtlpp::CommandBuffer const& CmdBuf);
+	
+	static char const* kTableAssociationKey;
+};
+
 /**
  * EAGXCommandEncoderType:
  *   EAGXCommandEncoderCurrent: The primary encoder that is used for draw calls & dispatches
@@ -142,6 +208,9 @@ public:
 	/** @returns The active blit command encoder or nil if there isn't one. */
 	mtlpp::BlitCommandEncoder& GetBlitCommandEncoder(void);
 	
+	/** @returns The MTLFence for the current encoder or nil if there isn't one. */
+	TRefCountPtr<FAGXFence> const& GetEncoderFence(void) const;
+	
 	/** @returns The number of encoded passes in the command buffer. */
 	uint32 NumEncodedPasses(void) const { return EncoderNum; }
 	
@@ -165,7 +234,7 @@ public:
 	void BeginBlitCommandEncoding(void);
 	
 	/** Declare that all command generation from this encoder is complete, and detach from the MTLCommandBuffer if there is an encoder active or does nothing if there isn't. */
-	void EndEncoding(void);
+	TRefCountPtr<FAGXFence> EndEncoding(void);
 	
 	/** Initialises a fence for the current command-buffer optionally adding a command-buffer completion handler to the command-buffer */
 	void InsertCommandBufferFence(FAGXCommandBufferFence& Fence, mtlpp::CommandBufferHandler Handler);
@@ -173,6 +242,18 @@ public:
 	/** Adds a command-buffer completion handler to the command-buffer */
 	void AddCompletionHandler(mtlpp::CommandBufferHandler Handler);
 	
+	/** Update the event to capture all GPU work so far enqueued by this encoder. */
+	void UpdateFence(FAGXFence* Fence);
+	
+	/** Prevent further GPU work until the event is reached. */
+	void WaitForFence(FAGXFence* Fence);
+	
+	/**
+	 * Prevent further GPU work until the event is reached and then kick the fence again so future encoders may wait on it.
+	 * Unlike the functions above this is not guarded and the caller must ensure proper matching waits & updates.
+	 */
+	void WaitAndUpdateFence(FAGXFence* Fence);
+
 #pragma mark - Public Debug Support -
 	
 	/*
@@ -190,11 +271,31 @@ public:
 	/* Pop the latest named string off of the stack. */
 	void PopDebugGroup(void);
 	
+	/** Returns the object that records debug command markers into the command-buffer. */
+	FAGXCommandBufferMarkers& GetMarkers(void);
+	
 #if ENABLE_METAL_GPUPROFILE
 	/* Get the command-buffer stats object. */
 	FAGXCommandBufferStats* GetCommandBufferStats(void);
 #endif
 
+#if MTLPP_CONFIG_VALIDATE && METAL_DEBUG_OPTIONS
+	/** @returns The active render command encoder or nil if there isn't one. */
+	FAGXCommandBufferDebugging& GetCommandBufferDebugging(void) { return CommandBufferDebug; } 
+	
+	/** @returns The active render command encoder or nil if there isn't one. */
+	FAGXRenderCommandEncoderDebugging& GetRenderCommandEncoderDebugging(void) { return RenderEncoderDebug; } 
+	
+	/** @returns The active compute command encoder or nil if there isn't one. */
+	FAGXComputeCommandEncoderDebugging& GetComputeCommandEncoderDebugging(void) { return ComputeEncoderDebug; }
+	
+	/** @returns The active blit command encoder or nil if there isn't one. */
+	FAGXBlitCommandEncoderDebugging& GetBlitCommandEncoderDebugging(void) { return BlitEncoderDebug; }
+	
+	/** @returns The active blit command encoder or nil if there isn't one. */
+	FAGXParallelRenderCommandEncoderDebugging& GetParallelRenderCommandEncoderDebugging(void) { return ParallelEncoderDebug; }
+#endif
+	
 #pragma mark - Public Render State Mutators -
 
 	/**
@@ -342,6 +443,20 @@ public:
 	 */
 	void SetShaderSideTable(mtlpp::FunctionType const FunctionType, NSUInteger const Index);
 	
+    /*
+     * Inform the driver that an indirect argument resource will be used and what the usage mode will be.
+     * @param Texture The texture that will be used as an indirect argument.
+     * @param Usage The usage mode for the texture.
+     */
+	void UseIndirectArgumentResource(FAGXTexture const& Texture, mtlpp::ResourceUsage const Usage);
+    
+    /*
+     * Inform the driver that an indirect argument resource will be used and what the usage mode will be.
+     * @param Buffer The buffer that will be used as an indirect argument.
+     * @param Usage The usage mode for the texture.
+     */
+	void UseIndirectArgumentResource(FAGXBuffer const& Buffer, mtlpp::ResourceUsage const Usage);
+
 	/*
 	 * Transition resource so that we can barrier fragment->vertex stages.
 	 * @param Resource The resource we are going to make readable having been written.
@@ -388,6 +503,8 @@ private:
 	
 	void FenceResource(mtlpp::Texture const& Resource);
 	void FenceResource(mtlpp::Buffer const& Resource);
+
+	void UseResource(mtlpp::Resource const& Resource, mtlpp::ResourceUsage const Usage);
 	
 #pragma mark - Private Type Declarations -
 private:
@@ -456,10 +573,22 @@ public:
 	mtlpp::ComputeCommandEncoder ComputeCommandEncoder;
 	mtlpp::BlitCommandEncoder BlitCommandEncoder;
 	TArray<mtlpp::RenderCommandEncoder> ChildRenderCommandEncoders;
+	FAGXCommandBufferMarkers CommandBufferMarkers;
+	
+	METAL_DEBUG_ONLY(FAGXCommandBufferDebugging CommandBufferDebug);
+	METAL_DEBUG_ONLY(FAGXRenderCommandEncoderDebugging RenderEncoderDebug);
+	METAL_DEBUG_ONLY(FAGXComputeCommandEncoderDebugging ComputeEncoderDebug);
+	METAL_DEBUG_ONLY(FAGXBlitCommandEncoderDebugging BlitEncoderDebug);
+	METAL_DEBUG_ONLY(FAGXParallelRenderCommandEncoderDebugging ParallelEncoderDebug);
 	
+	TRefCountPtr<FAGXFence> EncoderFence;
 #if ENABLE_METAL_GPUPROFILE
 	FAGXCommandBufferStats* CommandBufferStats;
 #endif
+#if METAL_DEBUG_OPTIONS
+	uint8 WaitCount;
+	uint8 UpdateCount;
+#endif
 
 	TArray<ns::Object<mtlpp::CommandBufferHandler>> CompletionHandlers;
 	NSMutableArray* DebugGroups;
@@ -471,6 +600,12 @@ public:
 	
 	TMap<mtlpp::Resource::Type, mtlpp::ResourceUsage> ResourceUsage;
 	
+	TSet<mtlpp::Resource::Type> TransitionedResources;
+	TSet<mtlpp::Resource::Type> FenceResources;
+	
+	TSet<TRefCountPtr<FAGXFence>> FragmentFences;
+	
+	mtlpp::RenderStages FenceStage;
 	uint32 EncoderNum;
 	uint32 CmdBufIndex;
 	EAGXCommandEncoderType Type;
diff --git a/Engine/Source/Runtime/Apple/AGXRHI/Private/AGXCommandList.cpp b/Engine/Source/Runtime/Apple/AGXRHI/Private/AGXCommandList.cpp
--- a/Engine/Source/Runtime/Apple/AGXRHI/Private/AGXCommandList.cpp
+++ b/Engine/Source/Runtime/Apple/AGXRHI/Private/AGXCommandList.cpp
@@ -10,6 +10,7 @@
 #include "AGXCommandList.h"
 #include "AGXCommandQueue.h"
 #include "AGXProfiler.h"
+#include "AGXCommandBuffer.h"
 #include "AGXCommandEncoder.h"
 #include "ns.hpp"
 
@@ -59,6 +60,115 @@ static void ReportMetalCommandBufferFailure(mtlpp::CommandBuffer const& Complete
     if (bDoCheck && !GIsSuspended && !GIsRenderingThreadSuspended)
 #endif
     {
+		FAGXCommandBufferMarkers Markers = FAGXCommandBufferMarkers::Get(CompletedBuffer);
+		if (Markers)
+		{
+			uint32 CommandBufferIndex = Markers.GetIndex();
+			FAGXDebugInfo BrokenDraw;
+			BrokenDraw.EncoderIndex = UINT32_MAX;
+			BrokenDraw.CommandIndex = UINT32_MAX;
+			BrokenDraw.CmdBuffIndex = UINT32_MAX;
+			BrokenDraw.ContextIndex = UINT32_MAX;
+            BrokenDraw.CommandBuffer = 0;
+			uint32 BrokenContext = 0;
+			bool bFoundBrokenDraw = false;
+			for (uint32 i = 0; (CommandBufferIndex != ~0u) && !bFoundBrokenDraw && i < Markers.NumContexts(); i++)
+			{
+				ns::AutoReleased<FAGXBuffer> DebugBuffer = Markers.GetDebugBuffer(i);
+				TArray<FAGXCommandDebug>* Commands = Markers.GetCommands(i);
+				if (DebugBuffer && Commands && Commands->Num())
+				{
+					uint32 DebugLength = DebugBuffer.GetLength();
+					uint32 DebugCount = DebugLength / sizeof(FAGXDebugInfo);
+					check(DebugCount >= 1);
+					
+					FAGXDebugInfo* DebugArray = (FAGXDebugInfo*)DebugBuffer.GetContents();
+					FAGXDebugInfo CurrentDraw = DebugArray[0];
+					
+					// On TBDR we find the disjoint where one tile has progressed further than another
+					// We find the earliest failure in the probably vain hope that this is the tile that actually failed
+					// There's actually no guarantee of that, it might be one of the later tiles that exploded
+					for (uint32 j = 0; j < DebugCount; j++)
+					{
+						if (CommandBufferIndex == DebugArray[j].CmdBuffIndex && DebugArray[j].EncoderIndex < CurrentDraw.EncoderIndex && DebugArray[j].CommandIndex < CurrentDraw.CommandIndex)
+						{
+							CurrentDraw = DebugArray[j];
+							break;
+						}
+					}
+
+					// Find the first command to fail - which depends on the order of encoders - parallel contexts make this more complicated
+					if (CurrentDraw.CmdBuffIndex == BrokenDraw.CmdBuffIndex && CurrentDraw.CommandIndex < Commands->Num() && CurrentDraw.EncoderIndex < BrokenDraw.EncoderIndex)
+					{
+						BrokenDraw = CurrentDraw;
+                        BrokenContext = i;
+                        bFoundBrokenDraw = true;
+					}
+                    else if(Commands->Num() > 1)
+                    {
+						BrokenDraw.EncoderIndex = 0;
+						BrokenDraw.CommandIndex = 0;
+						BrokenDraw.CmdBuffIndex = CommandBufferIndex;
+						BrokenDraw.ContextIndex = 0;
+                        BrokenContext = i;
+						BrokenDraw.CommandBuffer = (uintptr_t)CompletedBuffer.GetPtr();
+                        BrokenDraw = DebugArray[1];
+                        bFoundBrokenDraw = true;
+                    }
+				}
+			}
+			if (bFoundBrokenDraw)
+			{
+				id<MTLCommandBuffer> Ptr = reinterpret_cast<id<MTLCommandBuffer>>(BrokenDraw.CommandBuffer);
+				UE_LOG(LogAGX, Error, TEXT("GPU last wrote Command Buffer: %u (%llx) Encoder Index: %d Context Index: %d Draw Index: %d PSO: VS: %u_%u, PS: %u_%u."), BrokenDraw.CmdBuffIndex, BrokenDraw.CommandBuffer, BrokenDraw.EncoderIndex, BrokenDraw.ContextIndex, BrokenDraw.CommandIndex, BrokenDraw.PSOSignature[0], BrokenDraw.PSOSignature[1], BrokenDraw.PSOSignature[2], BrokenDraw.PSOSignature[3]);
+			
+				ns::AutoReleased<FAGXBuffer> DebugBuffer = Markers.GetDebugBuffer(BrokenContext);
+				TArray<FAGXCommandDebug>* Commands = Markers.GetCommands(BrokenContext);
+				if (Commands->Num())
+				{
+					UE_LOG(LogAGX, Error, TEXT("Failed executing following commands:"));
+					uint32 StartIdx = BrokenDraw.CommandIndex;
+					for (uint32 CmdIdx = StartIdx; CmdIdx < Commands->Num(); CmdIdx++)
+					{
+						FAGXCommandDebug& Command = (*Commands)[CmdIdx];
+						FString VSHash = Command.PSO->VertexShader->GetHash().ToString();
+						uint32 VSSig[2] = { Command.PSO->VertexShader->SourceLen, Command.PSO->VertexShader->SourceCRC };
+						
+						FString PSHash;
+						uint32 PSSig[2] = { 0, 0 };
+						if (Command.PSO->PixelShader.IsValid())
+						{
+							PSHash = Command.PSO->PixelShader->GetHash().ToString();
+							PSSig[0] = Command.PSO->PixelShader->SourceLen;
+							PSSig[1] = Command.PSO->PixelShader->SourceCRC;
+						}
+						
+						UE_LOG(LogAGX, Error, TEXT("Command Buffer: %d (%p) Encoder: %d Command: %d: %s PSO: VS: %s (%u_%u), PS: %s (%u_%u)"), Command.CmdBufIndex, CompletedBuffer.GetPtr(), Command.Encoder, Command.Index, *Command.Data.ToString(), *VSHash, VSSig[0], VSSig[1], *PSHash, PSSig[0], PSSig[1]);
+					}
+				}
+                if (DebugBuffer)
+                {
+                    uint32 DebugLength = DebugBuffer.GetLength();
+                    uint32 DebugCount = DebugLength / sizeof(FAGXDebugInfo);
+                    check(DebugCount >= 1);
+                    
+                    FAGXDebugInfo* DebugArray = (FAGXDebugInfo*)DebugBuffer.GetContents();
+                    for (uint32 i = 0; i < DebugCount; i++)
+                    {
+                        FAGXDebugInfo& Command = DebugArray[i];
+						
+						// Stop when nothing has been written into the buffer
+						if (Command.CmdBuffIndex == 0 && Command.CommandBuffer == 0 && Command.EncoderIndex == 0 && Command.CommandIndex == 0 && Command.PSOSignature[0] == 0 && Command.PSOSignature[1] == 0 && Command.PSOSignature[2] == 0 && Command.PSOSignature[3] == 0)
+						{
+							break;
+						}
+						
+                        UE_LOG(LogAGX, Error, TEXT("Command Buffer: %d (%p) Debug Buffer: %p Tile: %u Context: %d Encoder: %d Command: %d PSO: VS: %u_%u, PS: %u_%u"), Command.CmdBuffIndex, Command.CommandBuffer, DebugBuffer.GetPtr(), i, Command.ContextIndex, Command.EncoderIndex, Command.CommandIndex, Command.PSOSignature[0], Command.PSOSignature[1], Command.PSOSignature[2], Command.PSOSignature[3]);
+                    }
+                }
+			}
+		}
+		
 #if PLATFORM_IOS
         UE_LOG(LogAGX, Warning, TEXT("Command Buffer %s Failed with %s Error! Error Domain: %s Code: %d Description %s %s %s"), *LabelString, ErrorType, *DomainString, Code, *ErrorString, *FailureString, *RecoveryString);
         FIOSPlatformMisc::GPUAssert();
diff --git a/Engine/Source/Runtime/Apple/AGXRHI/Private/AGXCommandQueue.cpp b/Engine/Source/Runtime/Apple/AGXRHI/Private/AGXCommandQueue.cpp
--- a/Engine/Source/Runtime/Apple/AGXRHI/Private/AGXCommandQueue.cpp
+++ b/Engine/Source/Runtime/Apple/AGXRHI/Private/AGXCommandQueue.cpp
@@ -7,6 +7,7 @@
 #include "AGXRHIPrivate.h"
 
 #include "AGXCommandQueue.h"
+#include "AGXCommandBuffer.h"
 #include "AGXCommandList.h"
 #include "AGXProfiler.h"
 #include "Misc/ConfigCacheIni.h"
@@ -16,6 +17,7 @@
 NSUInteger FAGXCommandQueue::PermittedOptions = 0;
 uint64 FAGXCommandQueue::Features = 0;
 extern mtlpp::VertexFormat GAGXFColorVertexFormat;
+bool GAGXCommandBufferDebuggingEnabled = 0;
 
 #pragma mark - Public C++ Boilerplate -
 
@@ -24,17 +26,23 @@ FAGXCommandQueue::FAGXCommandQueue(uint32 const MaxNumCommandBuffers /* = 0 */)
 , RuntimeDebuggingLevel(EAGXDebugLevelOff)
 {
 	int32 MaxShaderVersion = 0;
-	int32 DefaultMaxShaderVersion = 5; // MSL v2.2
-	int32 MinShaderVersion = 5; // MSL v2.2
-
+	int32 IndirectArgumentTier = 0;
 #if PLATFORM_MAC
-	const TCHAR* const Settings = TEXT("/Script/MacTargetPlatform.MacTargetSettings");
+	int32 DefaultMaxShaderVersion = 5; // MSL v2.2
+	int32 MinShaderVersion = 4; // MSL v2.1
+    const TCHAR* const Settings = TEXT("/Script/MacTargetPlatform.MacTargetSettings");
 #else
-	const TCHAR* const Settings = TEXT("/Script/IOSRuntimeSettings.IOSRuntimeSettings");
+	int32 DefaultMaxShaderVersion = 2;
+	int32 MinShaderVersion = 2;
+    const TCHAR* const Settings = TEXT("/Script/IOSRuntimeSettings.IOSRuntimeSettings");
 #endif
     if(!GConfig->GetInt(Settings, TEXT("MaxShaderLanguageVersion"), MaxShaderVersion, GEngineIni))
+    {
+        MaxShaderVersion = DefaultMaxShaderVersion;
+    }
+	if(!GConfig->GetInt(Settings, TEXT("IndirectArgumentTier"), IndirectArgumentTier, GEngineIni))
 	{
-		MaxShaderVersion = DefaultMaxShaderVersion;
+		IndirectArgumentTier = 0;
 	}
 	MaxShaderVersion = FMath::Max(MinShaderVersion, MaxShaderVersion);
 	AGXValidateVersion(MaxShaderVersion);
@@ -49,74 +57,151 @@ FAGXCommandQueue::FAGXCommandQueue(uint32 const MaxNumCommandBuffers /* = 0 */)
 	}
 	check(CommandQueue);
 #if PLATFORM_IOS
-	NSOperatingSystemVersion Vers = [[NSProcessInfo processInfo]operatingSystemVersion];
-	Features = EAGXFeaturesSetBufferOffset | EAGXFeaturesSetBytes;
-
-#if PLATFORM_TVOS
-	Features &= ~(EAGXFeaturesSetBytes);
-
-	if ([GMtlDevice supportsFeatureSet:MTLFeatureSet_tvOS_GPUFamily2_v1])
+	NSOperatingSystemVersion Vers = [[NSProcessInfo processInfo] operatingSystemVersion];
+	if(Vers.majorVersion >= 9)
 	{
-		Features |= EAGXFeaturesCountingQueries | EAGXFeaturesBaseVertexInstance | EAGXFeaturesIndirectBuffer | EAGXFeaturesMSAADepthResolve | EAGXFeaturesMSAAStoreAndResolve;
-	}
+		Features = EAGXFeaturesSetBufferOffset | EAGXFeaturesSetBytes;
 
-	Features |= EAGXFeaturesPrivateBufferSubAllocation | EAGXFeaturesGPUCaptureManager | EAGXFeaturesBufferSubAllocation | EAGXFeaturesParallelRenderEncoders | EAGXFeaturesPipelineBufferMutability | EAGXFeaturesMaxThreadsPerThreadgroup | EAGXFeaturesTextureBuffers;
-	GAGXFColorVertexFormat = mtlpp::VertexFormat::UChar4Normalized_BGRA;
+#if PLATFORM_TVOS
+        Features &= ~(EAGXFeaturesSetBytes);
+		
+		if ([GMtlDevice supportsFeatureSet:MTLFeatureSet_tvOS_GPUFamily2_v1])
+		{
+			Features |= EAGXFeaturesCountingQueries | EAGXFeaturesBaseVertexInstance | EAGXFeaturesIndirectBuffer | EAGXFeaturesMSAADepthResolve | EAGXFeaturesMSAAStoreAndResolve;
+		}
+		
+		if(Vers.majorVersion > 10)
+		{
+			Features |= EAGXFeaturesPrivateBufferSubAllocation;
+			
+			if(Vers.majorVersion >= 11)
+			{
+				Features |= EAGXFeaturesGPUCaptureManager | EAGXFeaturesBufferSubAllocation | EAGXFeaturesParallelRenderEncoders | EAGXFeaturesPipelineBufferMutability;
+				
+				if (MaxShaderVersion >= 3)
+				{
+					GAGXFColorVertexFormat = mtlpp::VertexFormat::UChar4Normalized_BGRA;
+				}
 
+				if (Vers.majorVersion >= 12)
+				{
+					Features |= EAGXFeaturesMaxThreadsPerThreadgroup;
+					
+					if (FParse::Param(FCommandLine::Get(),TEXT("metalfence")))
+					{
+					Features |= EAGXFeaturesFences;
+					}
+					
+					if (FParse::Param(FCommandLine::Get(),TEXT("metalheap")))
+					{
+					Features |= EAGXFeaturesHeaps;
+					}
+					
+					if (MaxShaderVersion >= 4)
+					{
+						Features |= EAGXFeaturesTextureBuffers;
+					}
+				}
+			}
+		}
 #else
-	if ([GMtlDevice supportsFeatureSet:MTLFeatureSet_iOS_GPUFamily3_v1])
-	{
-		Features |= EAGXFeaturesCountingQueries | EAGXFeaturesBaseVertexInstance | EAGXFeaturesIndirectBuffer | EAGXFeaturesMSAADepthResolve;
-	}
+		if ([GMtlDevice supportsFeatureSet:MTLFeatureSet_iOS_GPUFamily3_v1])
+		{
+			Features |= EAGXFeaturesCountingQueries | EAGXFeaturesBaseVertexInstance | EAGXFeaturesIndirectBuffer | EAGXFeaturesMSAADepthResolve;
+		}
 		
-	if ([GMtlDevice supportsFeatureSet:MTLFeatureSet_iOS_GPUFamily3_v2])
-	{
-		Features |= EAGXFeaturesMSAAStoreAndResolve;
-	}
-
-	// Turning the below option on will allocate more buffer memory which isn't generally desirable on iOS
-	// Features |= EAGXFeaturesEfficientBufferBlits;
-
-	// These options are fine however as thye just change how we allocate small buffers
-	Features |= EAGXFeaturesBufferSubAllocation;
-	Features |= EAGXFeaturesPrivateBufferSubAllocation;
-	GAGXFColorVertexFormat = mtlpp::VertexFormat::UChar4Normalized_BGRA;
-
-	Features |= EAGXFeaturesPresentMinDuration
-		| EAGXFeaturesGPUCaptureManager
-		| EAGXFeaturesBufferSubAllocation
-		| EAGXFeaturesParallelRenderEncoders
-		| EAGXFeaturesPipelineBufferMutability;
-
-	Features |= EAGXFeaturesMaxThreadsPerThreadgroup;
-	Features |= EAGXFeaturesTextureBuffers;
-
-	if ([GMtlDevice supportsFeatureSet : MTLFeatureSet_iOS_GPUFamily4_v1])
-	{
-		Features |= EAGXFeaturesTileShaders;
+		if ([GMtlDevice supportsFeatureSet:MTLFeatureSet_iOS_GPUFamily3_v2] || [GMtlDevice supportsFeatureSet:MTLFeatureSet_iOS_GPUFamily2_v3] || [GMtlDevice supportsFeatureSet:MTLFeatureSet_iOS_GPUFamily1_v3])
+		{
+			if (FParse::Param(FCommandLine::Get(),TEXT("metalfence")))
+			{
+				Features |= EAGXFeaturesFences;
+			}
+			
+			if (FParse::Param(FCommandLine::Get(),TEXT("metalheap")))
+			{
+				Features |= EAGXFeaturesHeaps;
+			}
+		}
+		
+		if ([GMtlDevice supportsFeatureSet:MTLFeatureSet_iOS_GPUFamily3_v2])
+		{
+			Features |= EAGXFeaturesMSAAStoreAndResolve;
+		}
+		
+		if(Vers.majorVersion > 10 || (Vers.majorVersion == 10 && Vers.minorVersion >= 3))
+        {
+			// Turning the below option on will allocate more buffer memory which isn't generally desirable on iOS
+			// Features |= EAGXFeaturesEfficientBufferBlits;
+			
+			// These options are fine however as thye just change how we allocate small buffers
+            Features |= EAGXFeaturesBufferSubAllocation;
+			Features |= EAGXFeaturesPrivateBufferSubAllocation;
+			
+			if(Vers.majorVersion >= 11)
+			{
+				if (MaxShaderVersion >= 3)
+				{
+					GAGXFColorVertexFormat = mtlpp::VertexFormat::UChar4Normalized_BGRA;
+				}
+				
+				Features |= EAGXFeaturesPresentMinDuration | EAGXFeaturesGPUCaptureManager | EAGXFeaturesBufferSubAllocation | EAGXFeaturesParallelRenderEncoders | EAGXFeaturesPipelineBufferMutability;
+                
+				// Turn on Texture Buffers! These are faster on the GPU as we don't need to do out-of-bounds tests but require Metal 2.1 and macOS 10.14
+				if (Vers.majorVersion >= 12)
+				{
+					Features |= EAGXFeaturesMaxThreadsPerThreadgroup;
+                    if (!FParse::Param(FCommandLine::Get(),TEXT("nometalfence")))
+                    {
+                        Features |= EAGXFeaturesFences;
+                    }
+                    
+                    if (!FParse::Param(FCommandLine::Get(),TEXT("nometalheap")))
+                    {
+                        Features |= EAGXFeaturesHeaps;
+                    }
+					
+					if (MaxShaderVersion >= 4)
+					{
+						Features |= EAGXFeaturesTextureBuffers;
+					}
+                    
+                    if ([GMtlDevice supportsFeatureSet:MTLFeatureSet_iOS_GPUFamily4_v1])
+                    {
+                        Features |= EAGXFeaturesTileShaders;
                         
-	}
+                        // The below implies tile shaders which are necessary to order the draw calls and generate a buffer that shows what PSOs/draws ran on each tile.
+                        IConsoleVariable* GPUCrashDebuggingCVar = IConsoleManager::Get().FindConsoleVariable(TEXT("r.GPUCrashDebugging"));
+                        GAGXCommandBufferDebuggingEnabled = (GPUCrashDebuggingCVar && GPUCrashDebuggingCVar->GetInt() != 0) || FParse::Param(FCommandLine::Get(),TEXT("metalgpudebug"));
+                    }
                     
-	if ([GMtlDevice supportsFeatureSet:MTLFeatureSet_iOS_GPUFamily5_v1])
+					if ([GMtlDevice supportsFeatureSet:MTLFeatureSet_iOS_GPUFamily5_v1])
+					{
+						Features |= EAGXFeaturesLayeredRendering;
+					}
+				}
+			}
+        }
+#endif
+	}
+	else if(Vers.majorVersion == 8 && Vers.minorVersion >= 3)
 	{
-		Features |= EAGXFeaturesLayeredRendering;
+		Features = EAGXFeaturesSetBufferOffset;
 	}
-#endif
 #else // Assume that Mac & other platforms all support these from the start. They can diverge later.
-	const bool bIsNVIDIA = [[GMtlDevice name]rangeOfString:@"Nvidia" options:NSCaseInsensitiveSearch] .location != NSNotFound;
+	const bool bIsNVIDIA = [[GMtlDevice name] rangeOfString:@"Nvidia" options:NSCaseInsensitiveSearch].location != NSNotFound;
 	Features = EAGXFeaturesCountingQueries | EAGXFeaturesBaseVertexInstance | EAGXFeaturesIndirectBuffer | EAGXFeaturesLayeredRendering | EAGXFeaturesCubemapArrays;
 	if (!bIsNVIDIA)
 	{
 		Features |= EAGXFeaturesSetBufferOffset;
 	}
 	if ([GMtlDevice supportsFeatureSet:MTLFeatureSet_macOS_GPUFamily1_v2])
-	{
+    {
         Features |= (   EAGXFeaturesSetBytes
-			| EAGXFeaturesMSAADepthResolve
-			| EAGXFeaturesMSAAStoreAndResolve
-			| EAGXFeaturesEfficientBufferBlits
-			| EAGXFeaturesBufferSubAllocation
-			| EAGXFeaturesPrivateBufferSubAllocation
+					  | EAGXFeaturesMSAADepthResolve
+					  | EAGXFeaturesMSAAStoreAndResolve
+					  | EAGXFeaturesEfficientBufferBlits
+					  | EAGXFeaturesBufferSubAllocation
+					  | EAGXFeaturesPrivateBufferSubAllocation
 					  | EAGXFeaturesMaxThreadsPerThreadgroup );
 		
 		GAGXFColorVertexFormat = mtlpp::VertexFormat::UChar4Normalized_BGRA;
@@ -126,18 +211,68 @@ FAGXCommandQueue::FAGXCommandQueue(uint32 const MaxNumCommandBuffers /* = 0 */)
 			Features |= EAGXFeaturesParallelRenderEncoders;
 		}
 
-		Features |= EAGXFeaturesTextureBuffers;
+		{
+			if (MaxShaderVersion >= 4)
+			{
+				Features |= EAGXFeaturesTextureBuffers;
+            }
+            if (IndirectArgumentTier >= 1)
+            {
+                Features |= EAGXFeaturesIABs;
+				
+				if (IndirectArgumentTier >= 2)
+				{
+					Features |= EAGXFeaturesTier2IABs;
+				}
+            }
+            
+            IConsoleVariable* GPUCrashDebuggingCVar = IConsoleManager::Get().FindConsoleVariable(TEXT("r.GPUCrashDebugging"));
+            GAGXCommandBufferDebuggingEnabled = (GPUCrashDebuggingCVar && GPUCrashDebuggingCVar->GetInt() != 0) || FParse::Param(FCommandLine::Get(),TEXT("metalgpudebug"));
+            
+            // The editor spawns so many viewports and preview icons that we can run out of hardware fences!
+			// Need to figure out a way to safely flush the rendering and reuse the fences when that happens.
+#if WITH_EDITORONLY_DATA
+			if (!GIsEditor)
+#endif
+			{
+				if (FParse::Param(FCommandLine::Get(),TEXT("metalfence")))
+				{
+					Features |= EAGXFeaturesFences;
+				}
+				
+				// There are still too many driver bugs to use MTLHeap on macOS - nothing works without causing random, undebuggable GPU hangs that completely deadlock the Mac and don't generate any validation errors or command-buffer failures
+				if (FParse::Param(FCommandLine::Get(),TEXT("forcemetalheap")))
+				{
+					Features |= EAGXFeaturesHeaps;
+				}
+			}
+		}
     }
     else if ([[GMtlDevice name] rangeOfString:@"Nvidia" options:NSCaseInsensitiveSearch].location != NSNotFound)
-	{
+    {
 		// Using set*Bytes fixes bugs on Nvidia for 10.11 so we should use it...
-		Features |= EAGXFeaturesSetBytes;
-	}
+    	Features |= EAGXFeaturesSetBytes;
+    }
     
     if ([GMtlDevice supportsFeatureSet:MTLFeatureSet_macOS_GPUFamily1_v3])
-	{
-		Features |= EAGXFeaturesMultipleViewports | EAGXFeaturesPipelineBufferMutability | EAGXFeaturesGPUCaptureManager;
-	}
+    {
+        Features |= EAGXFeaturesMultipleViewports | EAGXFeaturesPipelineBufferMutability | EAGXFeaturesGPUCaptureManager;
+		
+		if (FParse::Param(FCommandLine::Get(),TEXT("metalfence")))
+		{
+			Features |= EAGXFeaturesFences;
+		}
+		
+		if (FParse::Param(FCommandLine::Get(),TEXT("metalheap")))
+		{
+			Features |= EAGXFeaturesHeaps;
+		}
+		
+		if (FParse::Param(FCommandLine::Get(),TEXT("metaliabs")))
+		{
+			Features |= EAGXFeaturesIABs;
+		}
+    }
 #endif
 	
 #if !UE_BUILD_SHIPPING
@@ -157,13 +292,20 @@ FAGXCommandQueue::FAGXCommandQueue(uint32 const MaxNumCommandBuffers /* = 0 */)
 	PermittedOptions = 0;
 	PermittedOptions |= mtlpp::ResourceOptions::CpuCacheModeDefaultCache;
 	PermittedOptions |= mtlpp::ResourceOptions::CpuCacheModeWriteCombined;
-	PermittedOptions |= mtlpp::ResourceOptions::StorageModeShared;
-	PermittedOptions |= mtlpp::ResourceOptions::StorageModePrivate;
+	{
+		PermittedOptions |= mtlpp::ResourceOptions::StorageModeShared;
+		PermittedOptions |= mtlpp::ResourceOptions::StorageModePrivate;
 #if PLATFORM_MAC
-	PermittedOptions |= mtlpp::ResourceOptions::StorageModeManaged;
+		PermittedOptions |= mtlpp::ResourceOptions::StorageModeManaged;
 #else
-	PermittedOptions |= mtlpp::ResourceOptions::StorageModeMemoryless;
+		PermittedOptions |= mtlpp::ResourceOptions::StorageModeMemoryless;
 #endif
+		// You can't use HazardUntracked under the validation layer due to bugs in the layer when trying to create linear-textures/texture-buffers
+		if ((Features & EAGXFeaturesFences) && !(Features & EAGXFeaturesValidation))
+		{
+			PermittedOptions |= mtlpp::ResourceOptions::HazardTrackingModeUntracked;
+		}
+	}
 }
 
 FAGXCommandQueue::~FAGXCommandQueue(void)
@@ -198,6 +340,7 @@ mtlpp::CommandBuffer FAGXCommandQueue::CreateCommandBuffer(void)
 		
 		if (RuntimeDebuggingLevel > EAGXDebugLevelOff)
 		{			
+			METAL_DEBUG_ONLY(FAGXCommandBufferDebugging AddDebugging(CmdBuffer));
 			MTLPP_VALIDATION(mtlpp::CommandBufferValidationTable ValidatedCommandBuffer(CmdBuffer));
 		}
 	}
@@ -242,6 +385,40 @@ void FAGXCommandQueue::SubmitCommandBuffers(TArray<mtlpp::CommandBuffer> BufferL
 	}
 }
 
+FAGXFence* FAGXCommandQueue::CreateFence(ns::String const& Label) const
+{
+	if ((Features & EAGXFeaturesFences) != 0)
+	{
+		FAGXFence* InternalFence = FAGXFencePool::Get().AllocateFence();
+		for (uint32 i = mtlpp::RenderStages::Vertex; InternalFence && i <= mtlpp::RenderStages::Fragment; i++)
+		{
+			mtlpp::Fence InnerFence = InternalFence->Get((mtlpp::RenderStages)i);
+			NSString* String = nil;
+			if (GetEmitDrawEvents())
+			{
+				String = [NSString stringWithFormat:@"%u %p: %@", i, InnerFence.GetPtr(), Label.GetPtr()];
+			}
+	#if METAL_DEBUG_OPTIONS
+			if (RuntimeDebuggingLevel >= EAGXDebugLevelValidation)
+			{
+				FAGXDebugFence* Fence = (FAGXDebugFence*)InnerFence.GetPtr();
+				Fence.label = String;
+			}
+			else
+	#endif
+			if(InnerFence && String)
+			{
+				InnerFence.SetLabel(String);
+			}
+		}
+		return InternalFence;
+	}
+	else
+	{
+		return nullptr;
+	}
+}
+
 void FAGXCommandQueue::GetCommittedCommandBufferFences(TArray<mtlpp::CommandBufferFence>& Fences)
 {
 	TArray<mtlpp::CommandBufferFence*> Temp;
@@ -272,7 +449,7 @@ mtlpp::ResourceOptions FAGXCommandQueue::GetCompatibleResourceOptions(mtlpp::Res
 void FAGXCommandQueue::InsertDebugCaptureBoundary(void)
 {
 	PRAGMA_DISABLE_DEPRECATION_WARNINGS
-		[CommandQueue insertDebugCaptureBoundary];
+	[CommandQueue insertDebugCaptureBoundary];
 	PRAGMA_ENABLE_DEPRECATION_WARNINGS
 }
 
diff --git a/Engine/Source/Runtime/Apple/AGXRHI/Private/AGXCommandQueue.h b/Engine/Source/Runtime/Apple/AGXRHI/Private/AGXCommandQueue.h
--- a/Engine/Source/Runtime/Apple/AGXRHI/Private/AGXCommandQueue.h
+++ b/Engine/Source/Runtime/Apple/AGXRHI/Private/AGXCommandQueue.h
@@ -36,9 +36,9 @@ typedef NS_OPTIONS(uint64, EAGXFeatures)
 	/** Supports detailed statistics */
 	EAGXFeaturesStatistics = 1 << 8,
 	/** Supports the explicit MTLHeap APIs */
-	EAGXFeaturesHeaps_REMOVED = 1 << 9,
+	EAGXFeaturesHeaps = 1 << 9,
 	/** Supports the explicit MTLFence APIs */
-	EAGXFeaturesFences_REMOVED = 1 << 10,
+	EAGXFeaturesFences = 1 << 10,
 	/** Supports MSAA Depth Resolves */
 	EAGXFeaturesMSAADepthResolve = 1 << 11,
 	/** Supports Store & Resolve in a single store action */
@@ -66,7 +66,7 @@ typedef NS_OPTIONS(uint64, EAGXFeatures)
 	/** Supports parallel render encoders */
 	EAGXFeaturesParallelRenderEncoders = 1llu << 23llu,
 	/** Supports indirect argument buffers */
-	EAGXFeaturesIABs_REMOVED = 1llu << 24llu,
+	EAGXFeaturesIABs = 1llu << 24llu,
 	/** Supports specifying the mutability of buffers bound to PSOs */
     EAGXFeaturesPipelineBufferMutability = 1llu << 25llu,
     /** Supports tile shaders */
@@ -74,7 +74,7 @@ typedef NS_OPTIONS(uint64, EAGXFeatures)
 	/** Unused Reserved Bit */
 	EAGXFeaturesUnusedReservedBit27 = 1llu << 27llu, // was EAGXFeaturesSeparateTessellation
 	/** Supports indirect argument buffers Tier 2 */
-	EAGXFeaturesTier2IABs_REMOVED = 1llu << 28llu,
+	EAGXFeaturesTier2IABs = 1llu << 28llu,
 };
 
 /**
@@ -118,6 +118,9 @@ public:
 	 */
 	void SubmitCommandBuffers(TArray<mtlpp::CommandBuffer> BufferList, uint32 Index, uint32 Count);
 
+	/** @returns Creates a new MTLFence or nil if this is unsupported */
+	FAGXFence* CreateFence(ns::String const& Label) const;
+	
 	/** @params Fences An array of command-buffer fences for the committed command-buffers */
 	void GetCommittedCommandBufferFences(TArray<mtlpp::CommandBufferFence>& Fences);
 	
diff --git a/Engine/Source/Runtime/Apple/AGXRHI/Private/AGXCommands.cpp b/Engine/Source/Runtime/Apple/AGXRHI/Private/AGXCommands.cpp
--- a/Engine/Source/Runtime/Apple/AGXRHI/Private/AGXCommands.cpp
+++ b/Engine/Source/Runtime/Apple/AGXRHI/Private/AGXCommands.cpp
@@ -15,6 +15,7 @@
 #include "ShaderParameterUtils.h"
 #include "SceneUtils.h"
 #include "AGXProfiler.h"
+#include "AGXCommandBuffer.h"
 #include "StaticBoundShaderState.h"
 #include "EngineGlobals.h"
 #include "PipelineStateCache.h"
diff --git a/Engine/Source/Runtime/Apple/AGXRHI/Private/AGXComputeCommandEncoder.cpp b/Engine/Source/Runtime/Apple/AGXRHI/Private/AGXComputeCommandEncoder.cpp
new file mode 100644
--- /dev/null
+++ b/Engine/Source/Runtime/Apple/AGXRHI/Private/AGXComputeCommandEncoder.cpp
@@ -0,0 +1,454 @@
+// Copyright Epic Games, Inc. All Rights Reserved.
+
+/*=============================================================================
+	AGXComputeCommandEncoder.cpp: AGX command encoder wrapper.
+=============================================================================*/
+
+#include "AGXRHIPrivate.h"
+
+#include "AGXComputeCommandEncoder.h"
+#include "AGXCommandBuffer.h"
+#include "AGXFence.h"
+#include "AGXPipeline.h"
+
+#if MTLPP_CONFIG_VALIDATE && METAL_DEBUG_OPTIONS
+NS_ASSUME_NONNULL_BEGIN
+
+#if METAL_DEBUG_OPTIONS
+static NSString* GAGXDebugComputeShader = @"#include <metal_stdlib>\n"
+@"using namespace metal;\n"
+@"kernel void WriteCommandIndexCS(constant uint* Input [[ buffer(0) ]], device atomic_uint* Output [[ buffer(1) ]])\n"
+@"{\n"
+@"	atomic_store_explicit(Output, Input[0], memory_order_relaxed);\n"
+@"}\n";
+
+static id <MTLComputePipelineState> GetDebugComputeShaderState(id<MTLDevice> Device)
+{
+	static id<MTLComputePipelineState> State = nil;
+	if (!State)
+	{
+		id<MTLLibrary> Lib = [Device newLibraryWithSource:GAGXDebugComputeShader options:nil error:nullptr];
+		check(Lib);
+		id<MTLFunction> Func = [Lib newFunctionWithName:@"WriteCommandIndexCS"];
+		check(Func);
+		State = [Device newComputePipelineStateWithFunction:Func error:nil];
+		[Func release];
+		[Lib release];
+	}
+	check(State);
+	return State;
+}
+#endif
+
+@interface FAGXDebugComputeCommandEncoder : FAGXDebugCommandEncoder
+{
+@public
+	id<MTLComputeCommandEncoder> Inner;
+	FAGXCommandBufferDebugging Buffer;
+	FAGXShaderPipeline* Pipeline;
+
+#pragma mark - Private Member Variables -
+#if METAL_DEBUG_OPTIONS
+	FAGXDebugShaderResourceMask ResourceMask;
+    FAGXDebugBufferBindings ShaderBuffers;
+    FAGXDebugTextureBindings ShaderTextures;
+    FAGXDebugSamplerBindings ShaderSamplers;
+#endif
+}
+/** Initialise the wrapper with the provided command-buffer. */
+-(id)initWithEncoder:(id<MTLComputeCommandEncoder>)Encoder andCommandBuffer:(FAGXCommandBufferDebugging const&)Buffer;
+@end
+
+@implementation FAGXDebugComputeCommandEncoder
+
+APPLE_PLATFORM_OBJECT_ALLOC_OVERRIDES(FAGXDebugComputeCommandEncoder)
+
+-(id)initWithEncoder:(id<MTLComputeCommandEncoder>)Encoder andCommandBuffer:(FAGXCommandBufferDebugging const&)SourceBuffer
+{
+	id Self = [super init];
+	if (Self)
+	{
+        Inner = Encoder;
+		Buffer = SourceBuffer;
+        Pipeline = nil;
+	}
+	return Self;
+}
+
+-(void)dealloc
+{
+	[Pipeline release];
+	[super dealloc];
+}
+
+@end
+
+#if METAL_DEBUG_OPTIONS
+void FAGXComputeCommandEncoderDebugging::InsertDebugDispatch()
+{
+//	switch (Buffer->DebugLevel)
+//	{
+//		case EAGXDebugLevelConditionalSubmit:
+//		case EAGXDebugLevelWaitForComplete:
+//		case EAGXDebugLevelLogOperations:
+//		case EAGXDebugLevelValidation:
+//		{
+//			uint32 const Index = Buffer->DebugCommands.Num();
+//			[Inner setBytes:&Index length:sizeof(Index) atIndex:0];
+//			[Inner setBuffer:Buffer->DebugInfoBuffer offset:0 atIndex:1];
+//			[Inner setComputePipelineState:GetDebugComputeShaderState(Inner.device)];
+//			
+//			[Inner dispatchThreadgroups:mtlpp::Size const&Make(1, 1, 1) threadsPerThreadgroup:mtlpp::Size const&Make(1, 1, 1)];
+//
+//			if (Pipeline && Pipeline->ComputePipelineState)
+//			{
+//				[Inner setComputePipelineState:Pipeline->ComputePipelineState];
+//			}
+//			
+//			if (ShaderBuffers.Buffers[0])
+//			{
+//				[Inner setBuffer:ShaderBuffers.Buffers[0] offset:ShaderBuffers.Offsets[0] atIndex:0];
+//			}
+//			else if (ShaderBuffers.Bytes[0])
+//			{
+//				[Inner setBytes:ShaderBuffers.Bytes[0] length:ShaderBuffers.Offsets[0] atIndex:0];
+//			}
+//			
+//			if (ShaderBuffers.Buffers[1])
+//			{
+//				[Inner setBuffer:ShaderBuffers.Buffers[1] offset:ShaderBuffers.Offsets[1] atIndex:1];
+//			}
+//			else if (ShaderBuffers.Bytes[1])
+//			{
+//				[Inner setBytes:ShaderBuffers.Bytes[1] length:ShaderBuffers.Offsets[1] atIndex:1];
+//			}
+//		}
+//		default:
+//		{
+//			break;
+//		}
+//	}
+}
+#endif
+
+FAGXComputeCommandEncoderDebugging::FAGXComputeCommandEncoderDebugging()
+{
+	
+}
+FAGXComputeCommandEncoderDebugging::FAGXComputeCommandEncoderDebugging(mtlpp::ComputeCommandEncoder& Encoder, FAGXCommandBufferDebugging& Buffer)
+: FAGXCommandEncoderDebugging((FAGXDebugCommandEncoder*)[[[FAGXDebugComputeCommandEncoder alloc] initWithEncoder:Encoder.GetPtr() andCommandBuffer:Buffer] autorelease])
+{
+	Buffer.BeginComputeCommandEncoder([NSString stringWithFormat:@"Compute: %@", Encoder.GetLabel().GetPtr()]);
+	Encoder.SetAssociatedObject((void const*)&FAGXComputeCommandEncoderDebugging::Get, (FAGXCommandEncoderDebugging const&)*this);
+}
+FAGXComputeCommandEncoderDebugging::FAGXComputeCommandEncoderDebugging(FAGXDebugCommandEncoder* handle)
+: FAGXCommandEncoderDebugging((FAGXDebugCommandEncoder*)handle)
+{
+	
+}
+
+FAGXComputeCommandEncoderDebugging FAGXComputeCommandEncoderDebugging::Get(mtlpp::ComputeCommandEncoder& Encoder)
+{
+	return Encoder.GetAssociatedObject<FAGXComputeCommandEncoderDebugging>((void const*)&FAGXComputeCommandEncoderDebugging::Get);
+}
+
+void FAGXComputeCommandEncoderDebugging::InsertDebugSignpost(ns::String const& Label)
+{
+	((FAGXDebugComputeCommandEncoder*)m_ptr)->Buffer.InsertDebugSignpost(Label);
+}
+void FAGXComputeCommandEncoderDebugging::PushDebugGroup(ns::String const& Group)
+{
+    ((FAGXDebugComputeCommandEncoder*)m_ptr)->Buffer.PushDebugGroup(Group);
+}
+void FAGXComputeCommandEncoderDebugging::PopDebugGroup()
+{
+    ((FAGXDebugComputeCommandEncoder*)m_ptr)->Buffer.PopDebugGroup();
+#if METAL_DEBUG_OPTIONS
+	InsertDebugDispatch();
+#endif
+}
+
+void FAGXComputeCommandEncoderDebugging::EndEncoder()
+{
+	((FAGXDebugComputeCommandEncoder*)m_ptr)->Buffer.EndCommandEncoder();
+}
+
+void FAGXComputeCommandEncoderDebugging::DispatchThreadgroups(mtlpp::Size const& threadgroupsPerGrid, mtlpp::Size const& threadsPerThreadgroup)
+{
+#if METAL_DEBUG_OPTIONS
+	switch(((FAGXDebugComputeCommandEncoder*)m_ptr)->Buffer.GetPtr()->DebugLevel)
+	{
+		case EAGXDebugLevelConditionalSubmit:
+		case EAGXDebugLevelWaitForComplete:
+		case EAGXDebugLevelLogOperations:
+		{
+			((FAGXDebugComputeCommandEncoder*)m_ptr)->Buffer.Dispatch([NSString stringWithFormat:@"%s", __PRETTY_FUNCTION__]);
+		}
+		case EAGXDebugLevelValidation:
+		{
+			Validate();
+		}
+		default:
+		{
+			break;
+		}
+	}
+#endif
+}
+
+void FAGXComputeCommandEncoderDebugging::SetPipeline(FAGXShaderPipeline* Pipeline)
+{
+#if METAL_DEBUG_OPTIONS
+	((FAGXDebugComputeCommandEncoder*)m_ptr)->Pipeline = [Pipeline retain];
+	switch (((FAGXDebugComputeCommandEncoder*)m_ptr)->Buffer.GetPtr()->DebugLevel)
+	{
+		case EAGXDebugLevelConditionalSubmit:
+		case EAGXDebugLevelWaitForComplete:
+		case EAGXDebugLevelLogOperations:
+		{
+			((FAGXDebugComputeCommandEncoder*)m_ptr)->Buffer.SetPipeline(Pipeline->ComputePipelineState.GetLabel());
+		}
+		default:
+		{
+			break;
+		}
+	}
+#endif
+}
+
+void FAGXComputeCommandEncoderDebugging::SetBytes(const void * bytes, NSUInteger length, NSUInteger index)
+{
+#if METAL_DEBUG_OPTIONS
+	switch (((FAGXDebugComputeCommandEncoder*)m_ptr)->Buffer.GetPtr()->DebugLevel)
+	{
+		case EAGXDebugLevelConditionalSubmit:
+		case EAGXDebugLevelWaitForComplete:
+		case EAGXDebugLevelLogOperations:
+		case EAGXDebugLevelValidation:
+		{
+			((FAGXDebugComputeCommandEncoder*)m_ptr)->ShaderBuffers.Buffers[index] = nil;
+			((FAGXDebugComputeCommandEncoder*)m_ptr)->ShaderBuffers.Bytes[index] = bytes;
+			((FAGXDebugComputeCommandEncoder*)m_ptr)->ShaderBuffers.Offsets[index] = length;
+			((FAGXDebugComputeCommandEncoder*)m_ptr)->ResourceMask.BufferMask = bytes ? (((FAGXDebugComputeCommandEncoder*)m_ptr)->ResourceMask.BufferMask | (1 << (index))) : (((FAGXDebugComputeCommandEncoder*)m_ptr)->ResourceMask.BufferMask & ~(1 << (index)));
+		}
+		default:
+		{
+			break;
+		}
+	}
+#endif
+}
+void FAGXComputeCommandEncoderDebugging::SetBuffer( FAGXBuffer const& buffer, NSUInteger offset, NSUInteger index)
+{
+#if METAL_DEBUG_OPTIONS
+	switch (((FAGXDebugComputeCommandEncoder*)m_ptr)->Buffer.GetPtr()->DebugLevel)
+	{
+		case EAGXDebugLevelConditionalSubmit:
+		case EAGXDebugLevelWaitForComplete:
+		case EAGXDebugLevelLogOperations:
+		case EAGXDebugLevelValidation:
+		{
+			((FAGXDebugComputeCommandEncoder*)m_ptr)->ShaderBuffers.Buffers[index] = buffer;
+			((FAGXDebugComputeCommandEncoder*)m_ptr)->ShaderBuffers.Bytes[index] = nil;
+			((FAGXDebugComputeCommandEncoder*)m_ptr)->ShaderBuffers.Offsets[index] = offset;
+			((FAGXDebugComputeCommandEncoder*)m_ptr)->ResourceMask.BufferMask = buffer ? (((FAGXDebugComputeCommandEncoder*)m_ptr)->ResourceMask.BufferMask | (1 << (index))) : (((FAGXDebugComputeCommandEncoder*)m_ptr)->ResourceMask.BufferMask & ~(1 << (index)));
+		}
+		default:
+		{
+			break;
+		}
+	}
+#endif
+}
+void FAGXComputeCommandEncoderDebugging::SetBufferOffset(NSUInteger offset, NSUInteger index)
+{
+#if METAL_DEBUG_OPTIONS
+	switch (((FAGXDebugComputeCommandEncoder*)m_ptr)->Buffer.GetPtr()->DebugLevel)
+	{
+		case EAGXDebugLevelConditionalSubmit:
+		case EAGXDebugLevelWaitForComplete:
+		case EAGXDebugLevelLogOperations:
+		case EAGXDebugLevelValidation:
+		{
+			((FAGXDebugComputeCommandEncoder*)m_ptr)->ShaderBuffers.Offsets[index] = offset;
+			check(((FAGXDebugComputeCommandEncoder*)m_ptr)->ResourceMask.BufferMask & (1 << (index)));
+		}
+		default:
+		{
+			break;
+		}
+	}
+#endif
+}
+
+void FAGXComputeCommandEncoderDebugging::SetTexture( FAGXTexture const& texture, NSUInteger index)
+{
+#if METAL_DEBUG_OPTIONS
+	switch (((FAGXDebugComputeCommandEncoder*)m_ptr)->Buffer.GetPtr()->DebugLevel)
+	{
+		case EAGXDebugLevelConditionalSubmit:
+		case EAGXDebugLevelWaitForComplete:
+		case EAGXDebugLevelLogOperations:
+		case EAGXDebugLevelValidation:
+		{
+			((FAGXDebugComputeCommandEncoder*)m_ptr)->ShaderTextures.Textures[index] = texture;
+			((FAGXDebugComputeCommandEncoder*)m_ptr)->ResourceMask.TextureMask = texture ? (((FAGXDebugComputeCommandEncoder*)m_ptr)->ResourceMask.TextureMask | (FAGXTextureMask(1) << (index))) : (((FAGXDebugComputeCommandEncoder*)m_ptr)->ResourceMask.TextureMask & ~(FAGXTextureMask(1) << (index)));
+		}
+		default:
+		{
+			break;
+		}
+	}
+#endif
+}
+
+void FAGXComputeCommandEncoderDebugging::SetSamplerState( mtlpp::SamplerState const& sampler, NSUInteger index)
+{
+#if METAL_DEBUG_OPTIONS
+	switch (((FAGXDebugComputeCommandEncoder*)m_ptr)->Buffer.GetPtr()->DebugLevel)
+	{
+		case EAGXDebugLevelConditionalSubmit:
+		case EAGXDebugLevelWaitForComplete:
+		case EAGXDebugLevelLogOperations:
+		case EAGXDebugLevelValidation:
+		{
+			((FAGXDebugComputeCommandEncoder*)m_ptr)->ShaderSamplers.Samplers[index] = sampler;
+			((FAGXDebugComputeCommandEncoder*)m_ptr)->ResourceMask.SamplerMask = sampler ? (((FAGXDebugComputeCommandEncoder*)m_ptr)->ResourceMask.SamplerMask | (1 << (index))) : (((FAGXDebugComputeCommandEncoder*)m_ptr)->ResourceMask.SamplerMask & ~(1 << (index)));
+		}
+		default:
+		{
+			break;
+		}
+	}
+#endif
+}
+
+void FAGXComputeCommandEncoderDebugging::SetSamplerState( mtlpp::SamplerState const& sampler, float lodMinClamp, float lodMaxClamp, NSUInteger index)
+{
+#if METAL_DEBUG_OPTIONS
+	switch (((FAGXDebugComputeCommandEncoder*)m_ptr)->Buffer.GetPtr()->DebugLevel)
+	{
+		case EAGXDebugLevelConditionalSubmit:
+		case EAGXDebugLevelWaitForComplete:
+		case EAGXDebugLevelLogOperations:
+		case EAGXDebugLevelValidation:
+		{
+			((FAGXDebugComputeCommandEncoder*)m_ptr)->ShaderSamplers.Samplers[index] = sampler;
+			((FAGXDebugComputeCommandEncoder*)m_ptr)->ResourceMask.SamplerMask = sampler ? (((FAGXDebugComputeCommandEncoder*)m_ptr)->ResourceMask.SamplerMask | (1 << (index))) : (((FAGXDebugComputeCommandEncoder*)m_ptr)->ResourceMask.SamplerMask & ~(1 << (index)));
+		}
+		default:
+		{
+			break;
+		}
+	}
+#endif
+}
+
+void FAGXComputeCommandEncoderDebugging::DispatchThreadgroupsWithIndirectBuffer(FAGXBuffer const& indirectBuffer, NSUInteger indirectBufferOffset, mtlpp::Size const& threadsPerThreadgroup)
+{
+#if METAL_DEBUG_OPTIONS
+	switch(((FAGXDebugComputeCommandEncoder*)m_ptr)->Buffer.GetPtr()->DebugLevel)
+	{
+		case EAGXDebugLevelConditionalSubmit:
+		case EAGXDebugLevelWaitForComplete:
+		case EAGXDebugLevelLogOperations:
+		{
+			((FAGXDebugComputeCommandEncoder*)m_ptr)->Buffer.Dispatch([NSString stringWithFormat:@"%s", __PRETTY_FUNCTION__]);
+		}
+		case EAGXDebugLevelValidation:
+		{
+			Validate();
+		}
+		default:
+		{
+			break;
+		}
+	}
+#endif
+}
+
+void FAGXComputeCommandEncoderDebugging::Validate()
+{
+	bool bOK = true;
+	
+#if METAL_DEBUG_OPTIONS
+	switch(((FAGXDebugComputeCommandEncoder*)m_ptr)->Buffer.GetPtr()->DebugLevel)
+	{
+		case EAGXDebugLevelConditionalSubmit:
+		case EAGXDebugLevelWaitForComplete:
+		case EAGXDebugLevelLogOperations:
+		case EAGXDebugLevelValidation:
+		{
+			check(((FAGXDebugComputeCommandEncoder*)m_ptr)->Pipeline);
+			
+			MTLComputePipelineReflection* Reflection = ((FAGXDebugComputeCommandEncoder*)m_ptr)->Pipeline->ComputePipelineReflection;
+			check(Reflection);
+	
+			NSArray<MTLArgument*>* Arguments = Reflection.arguments;
+			for (uint32 i = 0; i < Arguments.count; i++)
+			{
+				MTLArgument* Arg = [Arguments objectAtIndex:i];
+				check(Arg);
+				switch(Arg.type)
+				{
+					case MTLArgumentTypeBuffer:
+					{
+						checkf(Arg.index < ML_MaxBuffers, TEXT("Metal buffer index exceeded!"));
+						if ((((FAGXDebugComputeCommandEncoder*)m_ptr)->ShaderBuffers.Buffers[Arg.index] == nil && ((FAGXDebugComputeCommandEncoder*)m_ptr)->ShaderBuffers.Bytes[Arg.index] == nil))
+						{
+							UE_LOG(LogAGX, Warning, TEXT("Unbound buffer at Metal index %u which will crash the driver: %s"), (uint32)Arg.index, *FString([Arg description]));
+							bOK = false;
+						}
+						break;
+					}
+					case MTLArgumentTypeThreadgroupMemory:
+					{
+						break;
+					}
+					case MTLArgumentTypeTexture:
+					{
+						checkf(Arg.index < ML_MaxTextures, TEXT("Metal texture index exceeded!"));
+						if (((FAGXDebugComputeCommandEncoder*)m_ptr)->ShaderTextures.Textures[Arg.index] == nil)
+						{
+							UE_LOG(LogAGX, Warning, TEXT("Unbound texture at Metal index %u  which will crash the driver: %s"), (uint32)Arg.index, *FString([Arg description]));
+							bOK = false;
+						}
+						else if (((FAGXDebugComputeCommandEncoder*)m_ptr)->ShaderTextures.Textures[Arg.index].textureType != Arg.textureType)
+						{
+							UE_LOG(LogAGX, Warning, TEXT("Incorrect texture type bound at Metal index %u which will crash the driver: %s\n%s"), (uint32)Arg.index, *FString([Arg description]), *FString([((FAGXDebugComputeCommandEncoder*)m_ptr)->ShaderTextures.Textures[Arg.index] description]));
+							bOK = false;
+						}
+						break;
+					}
+					case MTLArgumentTypeSampler:
+					{
+						checkf(Arg.index < ML_MaxSamplers, TEXT("Metal sampler index exceeded!"));
+						if (((FAGXDebugComputeCommandEncoder*)m_ptr)->ShaderSamplers.Samplers[Arg.index] == nil)
+						{
+							UE_LOG(LogAGX, Warning, TEXT("Unbound sampler at Metal index %u which will crash the driver: %s"), (uint32)Arg.index, *FString([Arg description]));
+							bOK = false;
+						}
+						break;
+					}
+					default:
+						check(false);
+						break;
+				}
+			}
+			break;
+		}
+		default:
+		{
+			break;
+		}
+	}
+	
+    if (!bOK)
+    {
+        UE_LOG(LogAGX, Error, TEXT("Metal Validation failures for compute shader:\n%s"), (((FAGXDebugComputeCommandEncoder*)m_ptr)->Pipeline && ((FAGXDebugComputeCommandEncoder*)m_ptr)->Pipeline->ComputeSource) ? *FString(((FAGXDebugComputeCommandEncoder*)m_ptr)->Pipeline->ComputeSource) : TEXT("nil"));
+    }
+#endif
+}
+
+NS_ASSUME_NONNULL_END
+#endif
diff --git a/Engine/Source/Runtime/Apple/AGXRHI/Private/AGXComputeCommandEncoder.h b/Engine/Source/Runtime/Apple/AGXRHI/Private/AGXComputeCommandEncoder.h
new file mode 100644
--- /dev/null
+++ b/Engine/Source/Runtime/Apple/AGXRHI/Private/AGXComputeCommandEncoder.h
@@ -0,0 +1,42 @@
+// Copyright Epic Games, Inc. All Rights Reserved.
+
+#pragma once
+
+#include "AGXDebugCommandEncoder.h"
+
+#if MTLPP_CONFIG_VALIDATE && METAL_DEBUG_OPTIONS
+NS_ASSUME_NONNULL_BEGIN
+
+@class FAGXDebugComputeCommandEncoder;
+@class FAGXShaderPipeline;
+
+
+class FAGXComputeCommandEncoderDebugging : public FAGXCommandEncoderDebugging
+{
+	void InsertDebugDispatch();
+	void Validate();
+public:
+	FAGXComputeCommandEncoderDebugging();
+	FAGXComputeCommandEncoderDebugging(mtlpp::ComputeCommandEncoder& Encoder, FAGXCommandBufferDebugging& Buffer);
+	FAGXComputeCommandEncoderDebugging(FAGXDebugCommandEncoder* handle);
+	
+	static FAGXComputeCommandEncoderDebugging Get(mtlpp::ComputeCommandEncoder& Buffer);
+	
+	void InsertDebugSignpost(ns::String const& Label);
+	void PushDebugGroup(ns::String const& Group);
+	void PopDebugGroup();
+	void EndEncoder();
+	void DispatchThreadgroups(mtlpp::Size const& threadgroupsPerGrid, mtlpp::Size const& threadsPerThreadgroup);
+	void SetPipeline(FAGXShaderPipeline* Pipeline);
+	void SetBytes(const void * bytes, NSUInteger length, NSUInteger index);
+	void SetBuffer( FAGXBuffer const& buffer, NSUInteger offset, NSUInteger index);
+	void SetBufferOffset(NSUInteger offset, NSUInteger index);
+	void SetTexture( FAGXTexture const& texture, NSUInteger index);
+	void SetSamplerState( mtlpp::SamplerState const& sampler, NSUInteger index);
+	void SetSamplerState( mtlpp::SamplerState const& sampler, float lodMinClamp, float lodMaxClamp, NSUInteger index);
+	
+	void DispatchThreadgroupsWithIndirectBuffer(FAGXBuffer const& indirectBuffer, NSUInteger indirectBufferOffset, mtlpp::Size const& threadsPerThreadgroup);
+};
+
+NS_ASSUME_NONNULL_END
+#endif
diff --git a/Engine/Source/Runtime/Apple/AGXRHI/Private/AGXContext.cpp b/Engine/Source/Runtime/Apple/AGXRHI/Private/AGXContext.cpp
--- a/Engine/Source/Runtime/Apple/AGXRHI/Private/AGXContext.cpp
+++ b/Engine/Source/Runtime/Apple/AGXRHI/Private/AGXContext.cpp
@@ -15,6 +15,7 @@
 
 #include "AGXContext.h"
 #include "AGXProfiler.h"
+#include "AGXCommandBuffer.h"
 
 #include "AGXFrameAllocator.h"
 
@@ -162,11 +163,6 @@ mtlpp::Device GMtlppDevice;
 #pragma mark - AGX Device Context Support Routines
 
 
-uint32 AGXSafeGetRuntimeDebuggingLevel()
-{
-	return GIsRHIInitialized ? GetAGXDeviceContext().GetCommandQueue().GetRuntimeDebuggingLevel() : GAGXRuntimeDebugLevel;
-}
-
 #if PLATFORM_MAC
 static id<NSObject> GAGXDeviceObserver = nil;
 
@@ -383,6 +379,11 @@ FAGXDeviceContext* FAGXDeviceContext::CreateDeviceContext()
 	FAGXCommandQueue* Queue = new FAGXCommandQueue(GAGXCommandQueueSize);
 	check(Queue);
 	
+	if (FAGXCommandQueue::SupportsFeature(EAGXFeaturesFences))
+	{
+		FAGXFencePool::Get().Init();
+	}
+	
 	return new FAGXDeviceContext(DeviceIndex, Queue);
 }
 
@@ -517,7 +518,7 @@ void FAGXDeviceContext::ClearFreeList()
 				{
 					ScribbleBuffer(Buffer);
 				}
-				if (GAGXResourcePurgeOnDelete && !Buffer.GetParentBuffer())
+				if (GAGXResourcePurgeOnDelete && !Buffer.GetHeap() && !Buffer.GetParentBuffer())
 				{
 					Buffer.SetPurgeableState(mtlpp::PurgeableState::Empty);
 				}
@@ -529,7 +530,7 @@ void FAGXDeviceContext::ClearFreeList()
                 if (!Texture.GetBuffer() && !Texture.GetParentTexture())
 				{
 #if METAL_DEBUG_OPTIONS
-					if (GAGXResourcePurgeOnDelete)
+					if (GAGXResourcePurgeOnDelete && !Texture.GetHeap())
 					{
 						Texture.SetPurgeableState(mtlpp::PurgeableState::Empty);
 					}
@@ -537,6 +538,10 @@ void FAGXDeviceContext::ClearFreeList()
 					Heap.ReleaseTexture(nullptr, Texture);
 				}
 			}
+			for ( FAGXFence* Fence : Pair->FenceFreeList )
+			{
+				FAGXFencePool::Get().ReleaseFence(Fence);
+			}
 			delete Pair;
 			DelayedFreeLists.RemoveAt(Index, 1, false);
 		}
@@ -643,7 +648,7 @@ bool FAGXDeviceContext::FAGXDelayedFreeList::IsComplete() const
 	return bFinished;
 }
 
-void FAGXDeviceContext::FlushFreeList()
+void FAGXDeviceContext::FlushFreeList(bool const bFlushFences)
 {
 	FAGXDelayedFreeList* NewList = new FAGXDelayedFreeList;
 	
@@ -655,6 +660,25 @@ void FAGXDeviceContext::FlushFreeList()
 	NewList->UsedBuffers = MoveTemp(UsedBuffers);
 	NewList->UsedTextures = MoveTemp(UsedTextures);
 	NewList->ObjectFreeList = ObjectFreeList;
+	if (bFlushFences)
+	{
+		TArray<FAGXFence*> Fences;
+		FenceFreeList.PopAll(Fences);
+		for (FAGXFence* Fence : Fences)
+		{
+			if(!UsedFences.Contains(Fence))
+			{
+				UsedFences.Add(Fence);
+			}
+		}
+		NewList->FenceFreeList = MoveTemp(UsedFences);
+	}
+#if METAL_DEBUG_OPTIONS
+	if (FrameFences.Num())
+	{
+		FrameFences.Empty();
+	}
+#endif
 	ObjectFreeList.Empty(ObjectFreeList.Num());
 	FreeListMutex.Unlock();
 	
@@ -759,11 +783,28 @@ void FAGXDeviceContext::ReleaseTexture(FAGXTexture& Texture)
 	}
 }
 
+void FAGXDeviceContext::ReleaseFence(FAGXFence* Fence)
+{
+#if METAL_DEBUG_OPTIONS
+	if(GetCommandList().GetCommandQueue().GetRuntimeDebuggingLevel() >= EAGXDebugLevelValidation)
+	{
+		FScopeLock Lock(&FreeListMutex);
+		FrameFences.Add(Fence);
+	}
+#endif
+	
+	if (GIsAGXInitialized) // @todo zebra: there seems to be some race condition at exit when the framerate is very low
+	{
+		check(Fence);
+		FenceFreeList.Push(Fence);
+	}
+}
+
 FAGXTexture FAGXDeviceContext::CreateTexture(FAGXSurface* Surface, mtlpp::TextureDescriptor Descriptor)
 {
 	FAGXTexture Tex = Heap.CreateTexture(Descriptor, Surface);
 #if METAL_DEBUG_OPTIONS
-	if (GAGXResourcePurgeOnDelete)
+	if (GAGXResourcePurgeOnDelete && !Tex.GetHeap())
 	{
 		Tex.SetPurgeableState(mtlpp::PurgeableState::NonVolatile);
 	}
@@ -788,7 +829,7 @@ FAGXBuffer FAGXDeviceContext::CreatePooledBuffer(FAGXPooledBufferArgs const& Arg
     FAGXBuffer Buffer = Heap.CreateBuffer(Args.Size, RequestedBufferOffsetAlignment, Args.Flags, FAGXCommandQueue::GetCompatibleResourceOptions((mtlpp::ResourceOptions)(CpuResourceOption | mtlpp::ResourceOptions::HazardTrackingModeUntracked | ((NSUInteger)Args.Storage << mtlpp::ResourceStorageModeShift))));
 	check(Buffer && Buffer.GetPtr());
 #if METAL_DEBUG_OPTIONS
-	if (GAGXResourcePurgeOnDelete)
+	if (GAGXResourcePurgeOnDelete && !Buffer.GetHeap())
 	{
 		Buffer.SetPurgeableState(mtlpp::PurgeableState::NonVolatile);
 	}
@@ -813,15 +854,18 @@ void FAGXDeviceContext::ReleaseBuffer(FAGXBuffer& Buffer)
 
 struct FAGXRHICommandUpdateFence final : public FRHICommand<FAGXRHICommandUpdateFence>
 {
+	TRefCountPtr<FAGXFence> Fence;
 	uint32 Num;
 	
-	FORCEINLINE_DEBUGGABLE FAGXRHICommandUpdateFence(uint32 InNum)
-	: Num(InNum)
+	FORCEINLINE_DEBUGGABLE FAGXRHICommandUpdateFence(TRefCountPtr<FAGXFence>& InFence, uint32 InNum)
+	: Fence(InFence)
+	, Num(InNum)
 	{
 	}
 	
 	void Execute(FRHICommandListBase& CmdList)
 	{
+		GetAGXDeviceContext().SetParallelPassFences(nil, Fence);
 		GetAGXDeviceContext().FinishFrame(true);
 		GetAGXDeviceContext().BeginParallelRenderCommandEncoding(Num);
 	}
@@ -842,6 +886,11 @@ FAGXRHICommandContext* FAGXDeviceContext::AcquireContext(int32 NewIndex, int32 N
 	}
 	check(Context);
 	
+	if(ParallelFences.Num() < NewNum)
+	{
+		ParallelFences.AddDefaulted(NewNum - ParallelFences.Num());
+	}
+	
 	NSString* StartLabel = nil;
 	NSString* EndLabel = nil;
 #if METAL_DEBUG_OPTIONS
@@ -849,16 +898,23 @@ FAGXRHICommandContext* FAGXDeviceContext::AcquireContext(int32 NewIndex, int32 N
 	EndLabel = [NSString stringWithFormat:@"End Parallel Context Index %d Num %d", NewIndex, NewNum];
 #endif
 	
+	TRefCountPtr<FAGXFence> StartFence = (NewIndex == 0 ? CommandList.GetCommandQueue().CreateFence(StartLabel) : ParallelFences[NewIndex - 1].GetReference());
+	TRefCountPtr<FAGXFence> EndFence = CommandList.GetCommandQueue().CreateFence(EndLabel);
+	ParallelFences[NewIndex] = EndFence;
+	
+	// Give the context the fences so that we can properly order the parallel contexts.
+	Context->GetInternalContext().SetParallelPassFences(StartFence, EndFence);
+	
 	if (NewIndex == 0)
 	{
 		if (FRHICommandListExecutor::GetImmediateCommandList().Bypass() || !IsRunningRHIInSeparateThread())
 		{
-			FAGXRHICommandUpdateFence UpdateCommand(NewNum);
+			FAGXRHICommandUpdateFence UpdateCommand(StartFence, NewNum);
 			UpdateCommand.Execute(FRHICommandListExecutor::GetImmediateCommandList());
 		}
 		else
 		{
-			new (FRHICommandListExecutor::GetImmediateCommandList().AllocCommand<FAGXRHICommandUpdateFence>()) FAGXRHICommandUpdateFence(NewNum);
+			new (FRHICommandListExecutor::GetImmediateCommandList().AllocCommand<FAGXRHICommandUpdateFence>()) FAGXRHICommandUpdateFence(StartFence, NewNum);
 			FRHICommandListExecutor::GetImmediateCommandList().RHIThreadFence(true);
 			FRHICommandListExecutor::GetImmediateCommandList().ImmediateFlush(EImmediateFlushType::DispatchToRHIThread);
 		}
@@ -962,12 +1018,14 @@ bool FAGXDeviceContext::ValidateIsInactiveBuffer(FAGXBuffer const& Buffer)
 uint32 FAGXContext::CurrentContextTLSSlot = FPlatformTLS::AllocTlsSlot();
 #endif
 
-FAGXContext::FAGXContext(FAGXCommandQueue& Queue, bool bIsImmediate)
+FAGXContext::FAGXContext(FAGXCommandQueue& Queue, bool const bIsImmediate)
 : CommandQueue(Queue)
 , CommandList(Queue, bIsImmediate)
 , StateCache(bIsImmediate)
 , RenderPass(CommandList, StateCache)
 , QueryBuffer(new FAGXQueryBufferPool(this))
+, StartFence(nil)
+, EndFence(nil)
 , NumParallelContextsInPass(0)
 {
 	// create a semaphore for multi-buffering the command buffer
@@ -1031,7 +1089,24 @@ void FAGXContext::MakeCurrent(FAGXContext* Context)
 }
 #endif
 
-void FAGXContext::InitFrame(bool bImmediateContext, uint32 Index, uint32 Num)
+void FAGXContext::SetParallelPassFences(FAGXFence* Start, FAGXFence* End)
+{
+	check(!StartFence && !EndFence);
+	StartFence = Start;
+	EndFence = End;
+}
+
+TRefCountPtr<FAGXFence> const& FAGXContext::GetParallelPassStartFence(void) const
+{
+	return StartFence;
+}
+
+TRefCountPtr<FAGXFence> const& FAGXContext::GetParallelPassEndFence(void) const
+{
+	return EndFence;
+}
+
+void FAGXContext::InitFrame(bool const bImmediateContext, uint32 Index, uint32 Num)
 {
 #if ENABLE_METAL_GPUPROFILE
 	FPlatformTLS::SetTlsValue(CurrentContextTLSSlot, this);
@@ -1054,14 +1129,23 @@ void FAGXContext::InitFrame(bool bImmediateContext, uint32 Index, uint32 Num)
 	RenderPass.ShrinkRingBuffers();
 	
 	// Begin the render pass frame.
-	RenderPass.Begin();
+	RenderPass.Begin(StartFence);
+	
+	// Unset the start fence, the render-pass owns it and we can consider it encoded now!
+	StartFence = nullptr;
 	
 	// make sure first SetRenderTarget goes through
 	StateCache.InvalidateRenderTargets();
 }
 
-void FAGXContext::FinishFrame(bool bImmediateContext)
+void FAGXContext::FinishFrame(bool const bImmediateContext)
 {
+	// Ensure that we update the end fence for parallel contexts.
+	RenderPass.Update(EndFence);
+	
+	// Unset the end fence, the render-pass owns it and we can consider it encoded now!
+	EndFence = nullptr;
+	
 	// End the render pass
 	RenderPass.End();
 	
@@ -1318,7 +1402,7 @@ bool FAGXContext::PrepareToDraw(uint32 PrimitiveType)
 	return true;
 }
 
-void FAGXContext::SetRenderPassInfo(const FRHIRenderPassInfo& RenderTargetsInfo, bool bRestart)
+void FAGXContext::SetRenderPassInfo(const FRHIRenderPassInfo& RenderTargetsInfo, bool const bRestart)
 {
 	if (CommandList.IsParallel())
 	{
@@ -1527,7 +1611,7 @@ void FAGXContext::AsyncGenerateMipmapsForTexture(FAGXTexture const& Texture)
 	RenderPass.AsyncGenerateMipmapsForTexture(Texture);
 }
 
-void FAGXContext::SubmitAsyncCommands(mtlpp::CommandBufferHandler ScheduledHandler, mtlpp::CommandBufferHandler CompletionHandler, bool bWait)
+void FAGXContext::SubmitAsyncCommands(mtlpp::CommandBufferHandler ScheduledHandler, mtlpp::CommandBufferHandler CompletionHandler, bool const bWait)
 {
 	RenderPass.AddAsyncCommandBufferHandlers(ScheduledHandler, CompletionHandler);
 	if (bWait)
@@ -1617,7 +1701,8 @@ void FAGXDeviceContext::SetParallelRenderPassDescriptor(FRHIRenderPassInfo const
 
 	if (!RenderPass.IsWithinParallelPass())
 	{
-		RenderPass.Begin();
+		RenderPass.Begin(EndFence);
+		EndFence = nullptr;
 		StateCache.InvalidateRenderTargets();
 		SetRenderPassInfo(TargetInfo, false);
 	}
@@ -1639,7 +1724,8 @@ void FAGXDeviceContext::EndParallelRenderCommandEncoding(void)
 	if (FPlatformAtomics::InterlockedDecrement(&ActiveParallelContexts) == 0)
 	{
 		RenderPass.EndRenderPass();
-		RenderPass.Begin(true);
+		RenderPass.Begin(StartFence, true);
+		StartFence = nullptr;
 		FPlatformAtomics::AtomicStore(&NumParallelContextsInPass, 0);
 	}
 }
@@ -1689,6 +1775,12 @@ public:
 		{
 			check(Index == NewIndex && Num == NewNum);
 			
+			if (Index == (Num - 1))
+			{
+				TRefCountPtr<FAGXFence> Fence = CmdContext->GetInternalContext().GetParallelPassEndFence();
+				GetAGXDeviceContext().SetParallelPassFences(Fence, nil);
+			}
+
 			CmdContext->GetInternalContext().FinishFrame(false);
 			GetAGXDeviceContext().EndParallelRenderCommandEncoding();
 
diff --git a/Engine/Source/Runtime/Apple/AGXRHI/Private/AGXContext.h b/Engine/Source/Runtime/Apple/AGXRHI/Private/AGXContext.h
--- a/Engine/Source/Runtime/Apple/AGXRHI/Private/AGXContext.h
+++ b/Engine/Source/Runtime/Apple/AGXRHI/Private/AGXContext.h
@@ -29,7 +29,7 @@ class FAGXContext
 {
 	friend class FAGXCommandContextContainer;
 public:
-	FAGXContext(FAGXCommandQueue& Queue, bool bIsImmediate);
+	FAGXContext(FAGXCommandQueue& Queue, bool const bIsImmediate);
 	virtual ~FAGXContext();
 	
 	FAGXCommandQueue& GetCommandQueue();
@@ -51,7 +51,7 @@ public:
 	/**
 	 * Set the color, depth and stencil render targets, and then make the new command buffer/encoder
 	 */
-	void SetRenderPassInfo(const FRHIRenderPassInfo& RenderTargetsInfo, bool bRestart = false);
+	void SetRenderPassInfo(const FRHIRenderPassInfo& RenderTargetsInfo, bool const bRestart = false);
 	
 	/**
 	 * Allocate from a dynamic ring buffer - by default align to the allowed alignment for offset field when setting buffers
@@ -99,7 +99,7 @@ public:
 	
     void AsyncGenerateMipmapsForTexture(FAGXTexture const& Texture);
     
-	void SubmitAsyncCommands(mtlpp::CommandBufferHandler ScheduledHandler, mtlpp::CommandBufferHandler CompletionHandler, bool bWait);
+	void SubmitAsyncCommands(mtlpp::CommandBufferHandler ScheduledHandler, mtlpp::CommandBufferHandler CompletionHandler, bool const bWait);
 	
 	void SynchronizeTexture(FAGXTexture const& Texture, uint32 Slice, uint32 Level);
 	
@@ -118,8 +118,12 @@ public:
 	static FAGXContext* GetCurrentContext();
 #endif
 	
-	void InitFrame(bool bImmediateContext, uint32 Index, uint32 Num);
-	void FinishFrame(bool bImmediateContext);
+	void SetParallelPassFences(FAGXFence* Start, FAGXFence* End);
+	TRefCountPtr<FAGXFence> const& GetParallelPassStartFence(void) const;
+	TRefCountPtr<FAGXFence> const& GetParallelPassEndFence(void) const;
+	
+	void InitFrame(bool const bImmediateContext, uint32 Index, uint32 Num);
+	void FinishFrame(bool const bImmediateContext);
 
 	// Track Write->Read transitions for TBDR Fragment->Verex fencing
 	void TransitionResource(FRHIUnorderedAccessView* InResource);
@@ -147,6 +151,12 @@ protected:
 	/** A pool of buffers for writing visibility query results. */
 	TSharedPtr<FAGXQueryBufferPool, ESPMode::ThreadSafe> QueryBuffer;
 	
+	/** Initial fence to wait on for parallel contexts */
+	TRefCountPtr<FAGXFence> StartFence;
+	
+	/** Fence to update at the end for parallel contexts */
+	TRefCountPtr<FAGXFence> EndFence;
+	
 #if ENABLE_METAL_GPUPROFILE
 	/** the slot to store a per-thread context ref */
 	static uint32 CurrentContextTLSSlot;
@@ -187,9 +197,10 @@ public:
 	void ReleaseObject(id Object);
 	void ReleaseTexture(FAGXSurface* Surface, FAGXTexture& Texture);
 	void ReleaseTexture(FAGXTexture& Texture);
+	void ReleaseFence(FAGXFence* Fence);
 	
 	void BeginFrame();
-	void FlushFreeList();
+	void FlushFreeList(bool const bFlushFences = true);
 	void ClearFreeList();
 	void DrainHeap();
 	void EndFrame();
@@ -260,6 +271,8 @@ private:
 	/** Free lists for releasing objects only once it is safe to do so */
 	TSet<FAGXBuffer> UsedBuffers;
 	TSet<FAGXTexture> UsedTextures;
+	TSet<FAGXFence*> UsedFences;
+	TLockFreePointerListLIFO<FAGXFence> FenceFreeList;
 	TSet<id> ObjectFreeList;
 	struct FAGXDelayedFreeList
 	{
@@ -267,6 +280,7 @@ private:
 		TArray<mtlpp::CommandBufferFence> Fences;
 		TSet<FAGXBuffer> UsedBuffers;
 		TSet<FAGXTexture> UsedTextures;
+		TSet<FAGXFence*> FenceFreeList;
 		TSet<id> ObjectFreeList;
 #if METAL_DEBUG_OPTIONS
 		int32 DeferCount;
@@ -281,6 +295,9 @@ private:
 	TMap<FAGXRHIBuffer*, FAGXFrameAllocator::AllocationEntry> OutstandingLocks;
 	
 #if METAL_DEBUG_OPTIONS
+	/** The list of fences for the current frame */
+	TArray<FAGXFence*> FrameFences;
+    
     FCriticalSection ActiveBuffersMutex;
     
     /** These are the active buffers that cannot be CPU modified */
@@ -290,6 +307,9 @@ private:
 	/** Free-list of contexts for parallel encoding */
 	TLockFreePointerListLIFO<FAGXRHICommandContext> ParallelContexts;
 	
+	/** Fences for parallel execution */
+	TArray<TRefCountPtr<FAGXFence>> ParallelFences;
+	
 	/** Critical section for FreeList */
 	FCriticalSection FreeListMutex;
 	
diff --git a/Engine/Source/Runtime/Apple/AGXRHI/Private/AGXDebugCommandEncoder.cpp b/Engine/Source/Runtime/Apple/AGXRHI/Private/AGXDebugCommandEncoder.cpp
new file mode 100644
--- /dev/null
+++ b/Engine/Source/Runtime/Apple/AGXRHI/Private/AGXDebugCommandEncoder.cpp
@@ -0,0 +1,64 @@
+// Copyright Epic Games, Inc. All Rights Reserved.
+
+#include "AGXRHIPrivate.h"
+
+#include "AGXDebugCommandEncoder.h"
+#include "AGXCommandBuffer.h"
+#include "AGXFence.h"
+
+#if MTLPP_CONFIG_VALIDATE && METAL_DEBUG_OPTIONS
+extern int32 GAGXRuntimeDebugLevel;
+
+@implementation FAGXDebugCommandEncoder
+
+APPLE_PLATFORM_OBJECT_ALLOC_OVERRIDES(FAGXDebugCommandEncoder)
+
+-(instancetype)init
+{
+	id Self = [super init];
+	if (Self)
+	{
+		UpdatedFences = [[NSHashTable weakObjectsHashTable] retain];
+		WaitingFences = [[NSHashTable weakObjectsHashTable] retain];
+	}
+	return Self;
+}
+
+-(void)dealloc
+{
+	[UpdatedFences release];
+	[WaitingFences release];
+	[super dealloc];
+}
+@end
+
+FAGXCommandEncoderDebugging::FAGXCommandEncoderDebugging()
+{
+	
+}
+
+FAGXCommandEncoderDebugging::FAGXCommandEncoderDebugging(FAGXDebugCommandEncoder* handle)
+: ns::Object<FAGXDebugCommandEncoder*>(handle)
+{
+	
+}
+
+void FAGXCommandEncoderDebugging::AddUpdateFence(id Fence)
+{
+	if ((EAGXDebugLevel)GAGXRuntimeDebugLevel >= EAGXDebugLevelValidation && Fence)
+	{
+		[m_ptr->UpdatedFences addObject:(FAGXDebugFence*)Fence];
+		[(FAGXDebugFence*)Fence updatingEncoder:m_ptr];
+	}
+}
+
+void FAGXCommandEncoderDebugging::AddWaitFence(id Fence)
+{
+	if ((EAGXDebugLevel)GAGXRuntimeDebugLevel >= EAGXDebugLevelValidation && Fence)
+	{
+		[m_ptr->WaitingFences addObject:(FAGXDebugFence*)Fence];
+		[(FAGXDebugFence*)Fence waitingEncoder:m_ptr];
+	}
+}
+#endif
+
diff --git a/Engine/Source/Runtime/Apple/AGXRHI/Private/AGXDebugCommandEncoder.h b/Engine/Source/Runtime/Apple/AGXRHI/Private/AGXDebugCommandEncoder.h
new file mode 100644
--- /dev/null
+++ b/Engine/Source/Runtime/Apple/AGXRHI/Private/AGXDebugCommandEncoder.h
@@ -0,0 +1,109 @@
+// Copyright Epic Games, Inc. All Rights Reserved.
+
+#pragma once
+
+#include "AGXCommandBuffer.h"
+
+// For some reason when including this file while building the editor Clang 9 ignores this pragma from MacPlatformCompilerPreSetup.h,
+// resulting in errors in FAGXDebugBufferBindings, FAGXDebugTextureBindings and FAGXDebugSamplerBindings. Readding it here works around this problem.
+#if (__clang_major__ >= 9)
+#pragma clang diagnostic ignored "-Wnullability-inferred-on-nested-type"
+#endif
+
+NS_ASSUME_NONNULL_BEGIN
+
+/**
+ * The sampler, buffer and texture resource limits as defined here:
+ * https://developer.apple.com/library/ios/documentation/Miscellaneous/Conceptual/MetalProgrammingGuide/Render-Ctx/Render-Ctx.html
+ */
+#if PLATFORM_IOS
+#define METAL_MAX_BUFFERS 31
+#define METAL_MAX_TEXTURES 31
+typedef uint32 FAGXTextureMask;
+#elif PLATFORM_MAC
+#define METAL_MAX_BUFFERS 31
+#define METAL_MAX_TEXTURES 128
+typedef __uint128_t FAGXTextureMask;
+#else
+#error "Unsupported Platform!"
+#endif
+typedef uint32 FAGXBufferMask;
+typedef uint16 FAGXSamplerMask;
+
+enum EAGXLimits
+{
+	ML_MaxSamplers = 16, /** Maximum number of samplers */
+	ML_MaxBuffers = METAL_MAX_BUFFERS, /** Maximum number of buffers */
+	ML_MaxTextures = METAL_MAX_TEXTURES, /** Maximum number of textures - there are more textures available on Mac than iOS */
+	ML_MaxViewports = 16 /** Technically this may be different at runtime, but this is the likely absolute upper-bound */
+};
+
+enum EAGXShaderFrequency
+{
+    EAGXShaderVertex = 0,
+    EAGXShaderFragment = 1,
+    EAGXShaderCompute = 2,
+	EAGXShaderStream = 3,
+    EAGXShaderRenderNum = 2,
+	EAGXShaderStagesNum = 4
+};
+
+/** A structure for quick mask-testing of shader-stage resource bindings */
+struct FAGXDebugShaderResourceMask
+{
+	FAGXTextureMask TextureMask;
+	FAGXBufferMask BufferMask;
+	FAGXSamplerMask SamplerMask;
+};
+
+/** A structure of arrays for the current buffer binding settings. */
+struct FAGXDebugBufferBindings
+{
+    /** The bound buffers or nil. */
+    id<MTLBuffer> _Nullable Buffers[ML_MaxBuffers];
+    /** Optional bytes buffer used instead of an id<MTLBuffer> */
+    void const* _Nullable Bytes[ML_MaxBuffers];
+    /** The bound buffer offsets or 0. */
+    NSUInteger Offsets[ML_MaxBuffers];
+};
+
+/** A structure of arrays for the current texture binding settings. */
+struct FAGXDebugTextureBindings
+{
+    /** The bound textures or nil. */
+    id<MTLTexture> _Nullable Textures[ML_MaxTextures];
+};
+
+/** A structure of arrays for the current sampler binding settings. */
+struct FAGXDebugSamplerBindings
+{
+    /** The bound sampler states or nil. */
+    id<MTLSamplerState> _Nullable Samplers[ML_MaxSamplers];
+};
+
+#if MTLPP_CONFIG_VALIDATE && METAL_DEBUG_OPTIONS
+class FAGXCommandBufferDebugging;
+@class FAGXDebugCommandBuffer;
+@class FAGXDebugFence;
+
+@interface FAGXDebugCommandEncoder : FApplePlatformObject
+{
+@public
+	NSHashTable<FAGXDebugFence*>* UpdatedFences;
+	NSHashTable<FAGXDebugFence*>* WaitingFences;
+}
+-(instancetype)init;
+@end
+
+class FAGXCommandEncoderDebugging : public ns::Object<FAGXDebugCommandEncoder*>
+{
+public:
+	FAGXCommandEncoderDebugging();
+	FAGXCommandEncoderDebugging(FAGXDebugCommandEncoder* handle);
+	
+	void AddUpdateFence(id Fence);
+	void AddWaitFence(id Fence);
+};
+
+#endif
+NS_ASSUME_NONNULL_END
diff --git a/Engine/Source/Runtime/Apple/AGXRHI/Private/AGXFence.cpp b/Engine/Source/Runtime/Apple/AGXRHI/Private/AGXFence.cpp
new file mode 100644
--- /dev/null
+++ b/Engine/Source/Runtime/Apple/AGXRHI/Private/AGXFence.cpp
@@ -0,0 +1,228 @@
+// Copyright Epic Games, Inc. All Rights Reserved.
+
+#include "AGXRHIPrivate.h"
+
+#include "AGXFence.h"
+#include "AGXCommandBuffer.h"
+#include "AGXCommandQueue.h"
+#include "AGXContext.h"
+#include "AGXProfiler.h"
+
+@implementation FAGXDebugFence
+@synthesize Inner;
+
+APPLE_PLATFORM_OBJECT_ALLOC_OVERRIDES(FAGXDebugFence)
+
+-(instancetype)init
+{
+	id Self = [super init];
+	if (Self)
+	{
+		Label = nil;
+	}
+	return Self;
+}
+
+-(void)dealloc
+{
+	[self validate];
+	FAGXDebugCommandEncoder* Encoder = nil;
+	while ((Encoder = UpdatingEncoders.Pop()))
+	{
+		[Encoder release];
+	}
+	Encoder = nil;
+	while ((Encoder = WaitingEncoders.Pop()))
+	{
+		[Encoder release];
+	}
+	[Label release];
+	[super dealloc];
+}
+
+-(id <MTLDevice>) device
+{
+	if (Inner)
+	{
+		return Inner.device;
+	}
+	else
+	{
+		return nil;
+	}
+}
+
+-(NSString *_Nullable)label
+{
+	return Label;
+}
+
+-(void)setLabel:(NSString *_Nullable)Text
+{
+	[Text retain];
+	[Label release];
+	Label = Text;
+	if(Inner)
+	{
+		Inner.label = Text;
+	}
+}
+
+-(void)validate
+{
+	UE_CLOG(UpdatingEncoders.IsEmpty() != WaitingEncoders.IsEmpty(), LogAGX, Fatal, TEXT("Fence with unmatched updates/waits destructed - there's a gap in fence (%p) %s"), self, Label ? *FString(Label) : TEXT("Null"));
+}
+
+-(void)updatingEncoder:(FAGXDebugCommandEncoder*)Encoder
+{
+	check(Encoder);
+	UpdatingEncoders.Push([Encoder retain]);
+}
+
+-(void)waitingEncoder:(FAGXDebugCommandEncoder*)Encoder
+{
+	check(Encoder);
+	WaitingEncoders.Push([Encoder retain]);
+}
+
+-(TLockFreePointerListLIFO<FAGXDebugCommandEncoder>*)updatingEncoders
+{
+	return &UpdatingEncoders;
+}
+
+-(TLockFreePointerListLIFO<FAGXDebugCommandEncoder>*)waitingEncoders
+{
+	return &WaitingEncoders;
+}
+@end
+
+#if METAL_DEBUG_OPTIONS
+extern int32 GAGXRuntimeDebugLevel;
+#endif
+
+uint32 FAGXFence::Release() const
+{
+	uint32 Refs = uint32(FPlatformAtomics::InterlockedDecrement(&NumRefs));
+	if(Refs == 0)
+	{
+#if METAL_DEBUG_OPTIONS // When using validation we need to use fences only once per-frame in order to make it tractable
+		if (GAGXRuntimeDebugLevel >= EAGXDebugLevelValidation)
+		{
+			AGXSafeReleaseMetalFence(const_cast<FAGXFence*>(this));
+		}
+		else // However in a final game, we need to reuse fences aggressively so that we don't run out when loading into projects
+#endif
+		{
+			FAGXFencePool::Get().ReleaseFence(const_cast<FAGXFence*>(this));
+		}
+	}
+	return Refs;
+}
+
+#if METAL_DEBUG_OPTIONS
+void FAGXFence::Validate(void) const
+{
+	if (GetAGXDeviceContext().GetCommandQueue().GetRuntimeDebuggingLevel() >= EAGXDebugLevelValidation && Get(mtlpp::RenderStages::Vertex))
+	{
+		[(FAGXDebugFence*)Get(mtlpp::RenderStages::Vertex).GetPtr() validate];
+		if (Get(mtlpp::RenderStages::Fragment))
+		{
+			[(FAGXDebugFence*)Get(mtlpp::RenderStages::Fragment).GetPtr() validate];
+		}
+	}
+}
+#endif
+
+void FAGXFencePool::Init()
+{
+	for (int32 i = 0; i < FAGXFencePool::NumFences; i++)
+	{
+#if METAL_DEBUG_OPTIONS
+		if (GAGXRuntimeDebugLevel >= EAGXDebugLevelValidation)
+		{
+			FAGXDebugFence* VertexFence = [[FAGXDebugFence new] autorelease];
+			VertexFence.Inner = GMtlppDevice.NewFence();
+			FAGXDebugFence* FragmentFence = [[FAGXDebugFence new] autorelease];
+			FragmentFence.Inner = GMtlppDevice.NewFence();
+			FAGXFence* F = new FAGXFence;
+			F->Set(mtlpp::RenderStages::Vertex, VertexFence);
+			F->Set(mtlpp::RenderStages::Fragment, FragmentFence);
+			Fences.Add(F);
+			Lifo.Push(F);
+		}
+		else
+#endif
+		{
+			FAGXFence* F = new FAGXFence;
+			F->Set(mtlpp::RenderStages::Vertex, GMtlppDevice.NewFence());
+			F->Set(mtlpp::RenderStages::Fragment, GMtlppDevice.NewFence());
+#if METAL_DEBUG_OPTIONS
+			if (GAGXRuntimeDebugLevel >= EAGXDebugLevelValidation)
+			{
+				Fences.Add(F);
+			}
+#endif
+			Lifo.Push(F);
+		}
+	}
+	Count = FAGXFencePool::NumFences;
+    Allocated = 0;
+}
+
+FAGXFence* FAGXFencePool::AllocateFence()
+{
+	FAGXFence* Fence = Lifo.Pop();
+	if (Fence)
+	{
+        INC_DWORD_STAT(STAT_AGXFenceCount);
+        FPlatformAtomics::InterlockedDecrement(&Count);
+        FPlatformAtomics::InterlockedIncrement(&Allocated);
+#if METAL_DEBUG_OPTIONS
+		if (GAGXRuntimeDebugLevel >= EAGXDebugLevelValidation)
+		{
+			FScopeLock Lock(&Mutex);
+			check(Fences.Contains(Fence));
+			Fences.Remove(Fence);
+		}
+#endif
+	}
+	check(Fence);
+	Fence->Reset();
+	return Fence;
+}
+
+void FAGXFence::ValidateUsage(FAGXFence* InFence)
+{
+	if (InFence)
+	{
+		if (InFence->NumWrites(mtlpp::RenderStages::Vertex) != InFence->NumWaits(mtlpp::RenderStages::Vertex))
+		{
+			UE_LOG(LogAGX, Warning, TEXT("%p (%s) writes %d waits %d"), InFence, *FString(InFence->Get(mtlpp::RenderStages::Vertex).GetLabel()), (uint32)InFence->NumWrites(mtlpp::RenderStages::Vertex), (uint32)InFence->NumWaits(mtlpp::RenderStages::Vertex));
+		}
+		if (InFence->NumWrites(mtlpp::RenderStages::Fragment) != InFence->NumWaits(mtlpp::RenderStages::Fragment))
+		{
+			UE_LOG(LogAGX, Warning, TEXT("%p (%s) writes %d waits %d"), InFence, *FString(InFence->Get(mtlpp::RenderStages::Fragment).GetLabel()), (uint32)InFence->NumWrites(mtlpp::RenderStages::Fragment), (uint32)InFence->NumWaits(mtlpp::RenderStages::Fragment));
+		}
+	}
+}
+
+void FAGXFencePool::ReleaseFence(FAGXFence* const InFence)
+{
+	if (InFence)
+	{
+        DEC_DWORD_STAT(STAT_AGXFenceCount);
+        FPlatformAtomics::InterlockedDecrement(&Allocated);
+#if METAL_DEBUG_OPTIONS
+		if (GAGXRuntimeDebugLevel >= EAGXDebugLevelValidation)
+		{
+			FScopeLock Lock(&Mutex);
+			FAGXFence::ValidateUsage(InFence);
+			check(!Fences.Contains(InFence));
+			Fences.Add(InFence);
+		}
+#endif
+		FPlatformAtomics::InterlockedIncrement(&Count);
+		check(Count <= FAGXFencePool::NumFences);
+		Lifo.Push(InFence);
+	}
+}
diff --git a/Engine/Source/Runtime/Apple/AGXRHI/Private/AGXFence.h b/Engine/Source/Runtime/Apple/AGXRHI/Private/AGXFence.h
new file mode 100644
--- /dev/null
+++ b/Engine/Source/Runtime/Apple/AGXRHI/Private/AGXFence.h
@@ -0,0 +1,207 @@
+// Copyright Epic Games, Inc. All Rights Reserved.
+
+
+#pragma once
+
+#import <Metal/Metal.h>
+
+#pragma clang diagnostic ignored "-Wnullability-completeness"
+
+@class FAGXDebugCommandEncoder;
+
+@interface FAGXDebugFence : FApplePlatformObject<MTLFence>
+{
+	TLockFreePointerListLIFO<FAGXDebugCommandEncoder> UpdatingEncoders;
+	TLockFreePointerListLIFO<FAGXDebugCommandEncoder> WaitingEncoders;
+	NSString* Label;
+}
+@property (retain) id<MTLFence> Inner;
+-(void)updatingEncoder:(FAGXDebugCommandEncoder*)Encoder;
+-(void)waitingEncoder:(FAGXDebugCommandEncoder*)Encoder;
+-(TLockFreePointerListLIFO<FAGXDebugCommandEncoder>*)updatingEncoders;
+-(TLockFreePointerListLIFO<FAGXDebugCommandEncoder>*)waitingEncoders;
+-(void)validate;
+@end
+
+@protocol MTLDeviceExtensions <MTLDevice>
+/*!
+ @method newFence
+ @abstract Create a new MTLFence object
+ */
+- (id <MTLFence>)newFence;
+@end
+
+@protocol MTLBlitCommandEncoderExtensions <MTLBlitCommandEncoder>
+/*!
+ @abstract Update the event to capture all GPU work so far enqueued by this encoder. */
+-(void) updateFence:(id <MTLFence>)fence;
+/*!
+ @abstract Prevent further GPU work until the event is reached. */
+-(void) waitForFence:(id <MTLFence>)fence;
+@end
+@protocol MTLComputeCommandEncoderExtensions <MTLComputeCommandEncoder>
+/*!
+ @abstract Update the event to capture all GPU work so far enqueued by this encoder. */
+-(void) updateFence:(id <MTLFence>)fence;
+/*!
+ @abstract Prevent further GPU work until the event is reached. */
+-(void) waitForFence:(id <MTLFence>)fence;
+@end
+@protocol MTLRenderCommandEncoderExtensions <MTLRenderCommandEncoder>
+/*!
+ @abstract Update the event to capture all GPU work so far enqueued by this encoder for the given
+ stages.
+ @discussion Unlike <st>updateFence:</st>, this method will update the event when the given stage(s) complete, allowing for commands to overlap in execution.
+ */
+-(void) updateFence:(id <MTLFence>)fence afterStages:(MTLRenderStages)stages;
+/*!
+ @abstract Prevent further GPU work until the event is reached for the given stages.
+ @discussion Unlike <st>waitForFence:</st>, this method will only block commands assoicated with the given stage(s), allowing for commands to overlap in execution.
+ */
+-(void) waitForFence:(id <MTLFence>)fence beforeStages:(MTLRenderStages)stages;
+@end
+
+class FAGXFence
+{
+	enum
+	{
+		NumFenceStages = 2
+	};
+public:
+	FAGXFence()
+	: NumRefs(0)
+	{
+		for (uint32 i = 0; i < NumFenceStages; i++)
+		{
+			Writes[i] = 0;
+			Waits[i] = 0;
+		}
+	}
+	
+	explicit FAGXFence(FAGXFence const& Other)
+	: NumRefs(0)
+	{
+		operator=(Other);
+	}
+	
+	~FAGXFence()
+	{
+		check(!NumRefs);
+	}
+	
+	uint32 AddRef() const
+	{
+		return uint32(FPlatformAtomics::InterlockedIncrement(&NumRefs));
+	}
+	uint32 Release() const;
+	uint32 GetRefCount() const
+	{
+		return uint32(FPlatformAtomics::AtomicRead(&NumRefs));
+	}
+	
+	FAGXFence& operator=(FAGXFence const& Other)
+	{
+		if (&Other != this)
+		{
+			for (uint32 i = 0; i < NumFenceStages; i++)
+			{
+				Fences[i] = Other.Fences[i];
+			}
+		}
+		return *this;
+	}
+	
+#if METAL_DEBUG_OPTIONS
+	void Validate(void) const;
+#endif
+	
+	void Reset(void)
+	{
+		for (uint32 i = 0; i < NumFenceStages; i++)
+		{
+			Writes[i] = 0;
+			Waits[i] = 0;
+		}
+	}
+	
+	void Write(mtlpp::RenderStages Stage)
+	{
+		Writes[(uint32)Stage - 1u]++;
+	}
+	
+	void Wait(mtlpp::RenderStages Stage)
+	{
+		Waits[(uint32)Stage - 1u]++;
+	}
+	
+	int8 NumWrites(mtlpp::RenderStages Stage) const
+	{
+		return Writes[(uint32)Stage - 1u];
+	}
+	
+	int8 NumWaits(mtlpp::RenderStages Stage) const
+	{
+		return Waits[(uint32)Stage - 1u];
+	}
+	
+	bool NeedsWrite(mtlpp::RenderStages Stage) const
+	{
+		return Writes[(uint32)Stage - 1u] == 0 || (Waits[(uint32)Stage - 1u] > Writes[(uint32)Stage - 1u]);
+	}
+	
+	bool NeedsWait(mtlpp::RenderStages Stage) const
+	{
+		return Waits[(uint32)Stage - 1u] == 0 || (Writes[(uint32)Stage - 1u] > Waits[(uint32)Stage - 1u]);
+	}
+	
+	mtlpp::Fence Get(mtlpp::RenderStages Stage) const
+	{
+		return Fences[(uint32)Stage - 1u];
+	}
+	
+	void Set(mtlpp::RenderStages Stage, mtlpp::Fence Fence)
+	{
+		Fences[(uint32)Stage - 1u] = Fence;
+	}
+	
+	static void ValidateUsage(FAGXFence* Fence);
+	
+private:
+	mtlpp::Fence Fences[NumFenceStages];
+	int8 Writes[NumFenceStages];
+	int8 Waits[NumFenceStages];
+	mutable int32 NumRefs;	
+};
+
+class FAGXFencePool
+{
+	enum
+	{
+		NumFences = 2048
+	};
+public:
+	FAGXFencePool() {}
+	
+	static FAGXFencePool& Get()
+	{
+		static FAGXFencePool sSelf;
+		return sSelf;
+	}
+	
+	void Init();
+	
+	FAGXFence* AllocateFence();
+	void ReleaseFence(FAGXFence* const InFence);
+	
+	int32 Max() const { return Count; }
+	int32 Num() const { return Allocated; }
+	
+private:
+	int32 Count;
+	int32 Allocated;
+#if METAL_DEBUG_OPTIONS
+	TSet<FAGXFence*> Fences;
+	FCriticalSection Mutex;
+#endif
+	TLockFreePointerListFIFO<FAGXFence, PLATFORM_CACHE_LINE_SIZE> Lifo;
+};
diff --git a/Engine/Source/Runtime/Apple/AGXRHI/Private/AGXLLM.cpp b/Engine/Source/Runtime/Apple/AGXRHI/Private/AGXLLM.cpp
--- a/Engine/Source/Runtime/Apple/AGXRHI/Private/AGXLLM.cpp
+++ b/Engine/Source/Runtime/Apple/AGXRHI/Private/AGXLLM.cpp
@@ -274,6 +274,33 @@ void AGXLLM::LogAllocBuffer(mtlpp::Buffer const& Buffer)
 	}
 }
 
+void AGXLLM::LogAllocHeap(mtlpp::Heap const& Heap)
+{
+	void* Ptr = (void*)Heap.GetPtr();
+	uint64 Size = Heap.GetSize();
+	
+	INC_MEMORY_STAT_BY(STAT_AGXHeapMemory, Size);
+	INC_DWORD_STAT(STAT_AGXHeapCount);
+	
+	LLM(FLowLevelMemTracker::Get().OnLowLevelAlloc(ELLMTracker::Platform, Ptr, Size, ELLMTag::Untagged, ELLMAllocType::System));
+	// Assign a dealloc handler to untrack the memory - but don't track the dispatch block!
+	{
+		LLM_SCOPED_PAUSE_TRACKING(ELLMAllocType::System);
+		
+		objc_setAssociatedObject(Heap.GetPtr(), (void*)&AGXLLM::LogAllocHeap,
+								 [[[FAGXDeallocHandler alloc] initWithBlock:^{
+			LLM_SCOPE_METAL(ELLMTagAGX::Heaps);
+			LLM_PLATFORM_SCOPE_METAL(ELLMTagAGX::Heaps);
+			
+			LLM(FLowLevelMemTracker::Get().OnLowLevelFree(ELLMTracker::Platform, Ptr, ELLMAllocType::System));
+			
+			DEC_MEMORY_STAT_BY(STAT_AGXHeapMemory, Size);
+			DEC_DWORD_STAT(STAT_AGXHeapCount);
+		}] autorelease],
+		OBJC_ASSOCIATION_RETAIN);
+	}
+}
+
 void AGXLLM::LogAliasTexture(mtlpp::Texture const& Texture)
 {
 	objc_setAssociatedObject(Texture.GetPtr(), (void*)&AGXLLM::LogAllocTexture, nil, OBJC_ASSOCIATION_RETAIN);
diff --git a/Engine/Source/Runtime/Apple/AGXRHI/Private/AGXPipeline.cpp b/Engine/Source/Runtime/Apple/AGXRHI/Private/AGXPipeline.cpp
--- a/Engine/Source/Runtime/Apple/AGXRHI/Private/AGXPipeline.cpp
+++ b/Engine/Source/Runtime/Apple/AGXRHI/Private/AGXPipeline.cpp
@@ -13,6 +13,7 @@
 #include "MetalShaderResources.h"
 #include "AGXProfiler.h"
 #include "AGXCommandQueue.h"
+#include "AGXCommandBuffer.h"
 #include "RenderUtils.h"
 #include "Misc/ScopeRWLock.h"
 #include "HAL/PThreadEvent.h"
@@ -35,6 +36,122 @@ static uint32 BlendBitOffsets[] = { Offset_BlendState0, Offset_BlendState1, Offs
 static uint32 RTBitOffsets[] = { Offset_RenderTargetFormat0, Offset_RenderTargetFormat1, Offset_RenderTargetFormat2, Offset_RenderTargetFormat3, Offset_RenderTargetFormat4, Offset_RenderTargetFormat5, Offset_RenderTargetFormat6, Offset_RenderTargetFormat7 };
 static_assert(Offset_RasterEnd < 64 && Offset_End < 128, "Offset_RasterEnd must be < 64 && Offset_End < 128");
 
+static float RoundUpNearestEven(const float f)
+{
+	const float ret = FMath::CeilToFloat(f);
+	const float isOdd = (float)(((int)ret) & 1);
+	return ret + isOdd;
+}
+
+// A tile-based or vertex-based debug shader for trying to emulate Aftermath style failure reporting
+static NSString* GAGXDebugShader = @"#include <metal_stdlib>\n"
+"#include <metal_compute>\n"
+"\n"
+"using namespace metal;\n"
+"\n"
+"struct DebugInfo\n"
+"{\n"
+"   uint CmdBuffIndex;\n"
+"	uint EncoderIndex;\n"
+"   uint ContextIndex;\n"
+"   uint CommandIndex;\n"
+"   uint CommandBuffer[2];\n"
+"	uint PSOSignature[4];\n"
+"};\n"
+"\n"
+#if !PLATFORM_MAC
+"// Executes once per-tile\n"
+"kernel void Main_Debug(constant DebugInfo *debugTable [[ buffer(0) ]], device DebugInfo* debugBuffer [[ buffer(1) ]], uint2 threadgroup_position_in_grid [[ threadgroup_position_in_grid ]], uint2 threadgroups_per_grid [[ threadgroups_per_grid ]])\n"
+"{\n"
+"	// Write Pass, Draw indices\n"
+"	// Write Vertex+Fragment PSO sig (in form VertexLen, VertexCRC, FragLen, FragCRC)\n"
+"   uint tile_index = threadgroup_position_in_grid.x + (threadgroup_position_in_grid.y * threadgroups_per_grid.x);"
+"	debugBuffer[tile_index] = debugTable[0];\n"
+"}";
+#else
+"// Executes once as a point draw call\n"
+"vertex void Main_Debug(constant DebugInfo *debugTable [[ buffer(0) ]], device DebugInfo* debugBuffer [[ buffer(1) ]])\n"
+"{\n"
+"	// Write Pass, Draw indices\n"
+"	// Write Vertex+Fragment PSO sig (in form VertexLen, VertexCRC, FragLen, FragCRC)\n"
+"	debugBuffer[0] = debugTable[0];\n"
+"}";
+#endif
+
+// A compute debug shader for trying to emulate Aftermath style failure reporting
+static NSString* GAGXDebugMarkerComputeShader = @"#include <metal_stdlib>\n"
+"#include <metal_compute>\n"
+"\n"
+"using namespace metal;\n"
+"\n"
+"struct DebugInfo\n"
+"{\n"
+"   uint CmdBuffIndex;\n"
+"	uint EncoderIndex;\n"
+"   uint ContextIndex;\n"
+"   uint CommandIndex;\n"
+"   uint CommandBuffer[2];\n"
+"	uint PSOSignature[4];\n"
+"};\n"
+"\n"
+"// Executes once\n"
+"kernel void Main_Debug(constant DebugInfo *debugTable [[ buffer(0) ]], device DebugInfo* debugBuffer [[ buffer(1) ]])\n"
+"{\n"
+"	// Write Pass, Draw indices\n"
+"	// Write Vertex+Fragment PSO sig (in form VertexLen, VertexCRC, FragLen, FragCRC)\n"
+"	debugBuffer[0] = debugTable[0];\n"
+"}";
+
+struct FAGXHelperFunctions
+{
+    mtlpp::Library DebugShadersLib;
+    mtlpp::Function DebugFunc;
+	
+	mtlpp::Library DebugComputeShadersLib;
+	mtlpp::Function DebugComputeFunc;
+	mtlpp::ComputePipelineState DebugComputeState;
+	
+    FAGXHelperFunctions()
+    {
+#if !PLATFORM_TVOS
+        if (GAGXCommandBufferDebuggingEnabled)
+        {
+            mtlpp::CompileOptions CompileOptions;
+            ns::AutoReleasedError Error;
+
+			DebugShadersLib = GMtlppDevice.NewLibrary(GAGXDebugShader, CompileOptions, &Error);
+            DebugFunc = DebugShadersLib.NewFunction(@"Main_Debug");
+			
+			DebugComputeShadersLib = GMtlppDevice.NewLibrary(GAGXDebugMarkerComputeShader, CompileOptions, &Error);
+			DebugComputeFunc = DebugComputeShadersLib.NewFunction(@"Main_Debug");
+			
+			DebugComputeState = GMtlppDevice.NewComputePipelineState(DebugComputeFunc, &Error);
+        }
+#endif
+    }
+    
+    static FAGXHelperFunctions& Get()
+    {
+        static FAGXHelperFunctions sSelf;
+        return sSelf;
+    }
+    
+    mtlpp::Function GetDebugFunction()
+    {
+        return DebugFunc;
+    }
+	
+	mtlpp::ComputePipelineState GetDebugComputeState()
+	{
+		return DebugComputeState;
+	}
+};
+
+mtlpp::ComputePipelineState AGXGetMetalDebugComputeState()
+{
+	return FAGXHelperFunctions::Get().GetDebugComputeState();
+}
+
 struct FAGXGraphicsPipelineKey
 {
 	FAGXRenderPipelineHash RenderPipelineHash;
@@ -326,8 +443,10 @@ APPLE_PLATFORM_OBJECT_ALLOC_OVERRIDES(FAGXShaderPipeline)
 	{
 		RenderPipelineReflection = mtlpp::RenderPipelineReflection(nil);
 		ComputePipelineReflection = mtlpp::ComputePipelineReflection(nil);
+		StreamPipelineReflection = mtlpp::RenderPipelineReflection(nil);
 #if METAL_DEBUG_OPTIONS
 		RenderDesc = mtlpp::RenderPipelineDescriptor(nil);
+		StreamDesc = mtlpp::RenderPipelineDescriptor(nil);
 		ComputeDesc = mtlpp::ComputePipelineDescriptor(nil);
 #endif
 	}
@@ -338,8 +457,8 @@ APPLE_PLATFORM_OBJECT_ALLOC_OVERRIDES(FAGXShaderPipeline)
 {
 	if (RenderPipelineReflection)
 	{
-		[self initResourceMask:SF_Vertex];
-		[self initResourceMask:SF_Pixel];
+		[self initResourceMask:EAGXShaderVertex];
+		[self initResourceMask:EAGXShaderFragment];
 		
 		if (AGXSafeGetRuntimeDebuggingLevel() < EAGXDebugLevelValidation)
 		{
@@ -348,20 +467,29 @@ APPLE_PLATFORM_OBJECT_ALLOC_OVERRIDES(FAGXShaderPipeline)
 	}
 	if (ComputePipelineReflection)
 	{
-		[self initResourceMask:SF_Compute];
+		[self initResourceMask:EAGXShaderCompute];
 		
 		if (AGXSafeGetRuntimeDebuggingLevel() < EAGXDebugLevelValidation)
 		{
 			ComputePipelineReflection = mtlpp::ComputePipelineReflection(nil);
 		}
 	}
+	if (StreamPipelineReflection)
+	{
+		[self initResourceMask:EAGXShaderStream];
+		
+		if (AGXSafeGetRuntimeDebuggingLevel() < EAGXDebugLevelValidation)
+		{
+			StreamPipelineReflection = mtlpp::RenderPipelineReflection(nil);
+		}
+	}
 }
-- (void)initResourceMask:(EShaderFrequency)Frequency
+- (void)initResourceMask:(EAGXShaderFrequency)Frequency
 {
 	NSArray<MTLArgument*>* Arguments = nil;
 	switch(Frequency)
 	{
-		case SF_Vertex:
+		case EAGXShaderVertex:
 		{
 			MTLRenderPipelineReflection* Reflection = RenderPipelineReflection;
 			check(Reflection);
@@ -369,7 +497,7 @@ APPLE_PLATFORM_OBJECT_ALLOC_OVERRIDES(FAGXShaderPipeline)
 			Arguments = Reflection.vertexArguments;
 			break;
 		}
-		case SF_Pixel:
+		case EAGXShaderFragment:
 		{
 			MTLRenderPipelineReflection* Reflection = RenderPipelineReflection;
 			check(Reflection);
@@ -377,7 +505,7 @@ APPLE_PLATFORM_OBJECT_ALLOC_OVERRIDES(FAGXShaderPipeline)
 			Arguments = Reflection.fragmentArguments;
 			break;
 		}
-		case SF_Compute:
+		case EAGXShaderCompute:
 		{
 			MTLComputePipelineReflection* Reflection = ComputePipelineReflection;
 			check(Reflection);
@@ -385,6 +513,14 @@ APPLE_PLATFORM_OBJECT_ALLOC_OVERRIDES(FAGXShaderPipeline)
 			Arguments = Reflection.arguments;
 			break;
 		}
+		case EAGXShaderStream:
+		{
+			MTLRenderPipelineReflection* Reflection = StreamPipelineReflection;
+			check(Reflection);
+			
+			Arguments = Reflection.vertexArguments;
+			break;
+		}
 		default:
 			check(false);
 			break;
@@ -659,7 +795,7 @@ static bool ConfigureRenderPipelineDescriptor(mtlpp::RenderPipelineDescriptor& R
 		FMetalShaderBindings& VertexBindings = VertexShader->Bindings;
 		int8 VertexSideTable = VertexShader->SideTableBinding;
 		{
-			uint32 ImmutableBuffers = VertexBindings.ConstantBuffers;
+			uint32 ImmutableBuffers = VertexBindings.ConstantBuffers | VertexBindings.ArgumentBuffers;
 			while(ImmutableBuffers)
 			{
 				uint32 Index = __builtin_ctz(ImmutableBuffers);
@@ -681,7 +817,7 @@ static bool ConfigureRenderPipelineDescriptor(mtlpp::RenderPipelineDescriptor& R
 		if (PixelShader)
 		{
 			ns::AutoReleased<ns::Array<mtlpp::PipelineBufferDescriptor>> FragmentPipelineBuffers = RenderPipelineDesc.GetFragmentBuffers();
-			uint32 ImmutableBuffers = PixelShader->Bindings.ConstantBuffers;
+			uint32 ImmutableBuffers = PixelShader->Bindings.ConstantBuffers | PixelShader->Bindings.ArgumentBuffers;
 			while(ImmutableBuffers)
 			{
 				uint32 Index = __builtin_ctz(ImmutableBuffers);
@@ -792,6 +928,27 @@ static FAGXShaderPipeline* CreateMTLRenderPipeline(bool const bSync, FAGXGraphic
         Pipeline->FragmentSource = PixelShader ? PixelShader->GetSourceCode() : nil;
 #endif
 		
+#if !PLATFORM_TVOS
+		if (GAGXCommandBufferDebuggingEnabled)
+		{
+#if PLATFORM_MAC
+			DebugPipelineDesc.SetVertexFunction(FAGXHelperFunctions::Get().GetDebugFunction());
+			DebugPipelineDesc.SetRasterizationEnabled(false);
+#else
+			DebugPipelineDesc.SetTileFunction(FAGXHelperFunctions::Get().GetDebugFunction());
+			DebugPipelineDesc.SetRasterSampleCount(RenderPipelineDesc.GetSampleCount());
+			DebugPipelineDesc.SetThreadgroupSizeMatchesTileSize(false);
+#endif
+#if ENABLE_METAL_GPUPROFILE
+			DebugPipelineDesc.SetLabel(@"Main_Debug");
+#endif
+
+			ns::AutoReleasedError RenderError;
+			METAL_GPUPROFILE(FAGXScopedCPUStats CPUStat(FString::Printf(TEXT("NewDebugPipeline: %s"), TEXT("")/**FString([RenderPipelineDesc.GetPtr() description])*/)));
+			Pipeline->DebugPipelineState = GMtlppDevice.NewRenderPipelineState(DebugPipelineDesc, mtlpp::PipelineOption::NoPipelineOption, Reflection, nullptr);
+		}
+#endif
+        
 #if METAL_DEBUG_OPTIONS
         if (GFrameCounter > 3)
         {
diff --git a/Engine/Source/Runtime/Apple/AGXRHI/Private/AGXPipeline.h b/Engine/Source/Runtime/Apple/AGXRHI/Private/AGXPipeline.h
--- a/Engine/Source/Runtime/Apple/AGXRHI/Private/AGXPipeline.h
+++ b/Engine/Source/Runtime/Apple/AGXRHI/Private/AGXPipeline.h
@@ -3,6 +3,10 @@
 #pragma once
 
 #include "AGXRHIPrivate.h"
+#if METAL_DEBUG_OPTIONS
+#include "AGXDebugCommandEncoder.h"
+#endif
+
 #include "ShaderPipelineCache.h"
 
 enum EAGXPipelineHashBits
@@ -67,20 +71,24 @@ private:
 @public
 	mtlpp::RenderPipelineState RenderPipelineState;
 	mtlpp::ComputePipelineState ComputePipelineState;
-	TArray<uint32> BufferDataSizes[SF_NumStandardFrequencies];
-	TMap<uint8, uint8> TextureTypes[SF_NumStandardFrequencies];
-	FAGXDebugShaderResourceMask ResourceMask[SF_NumStandardFrequencies];
+	mtlpp::RenderPipelineState StreamPipelineState;
+	mtlpp::RenderPipelineState DebugPipelineState;
+	TArray<uint32> BufferDataSizes[EAGXShaderStagesNum];
+	TMap<uint8, uint8> TextureTypes[EAGXShaderStagesNum];
+	FAGXDebugShaderResourceMask ResourceMask[EAGXShaderStagesNum];
 	mtlpp::RenderPipelineReflection RenderPipelineReflection;
+	mtlpp::RenderPipelineReflection StreamPipelineReflection;
 	mtlpp::ComputePipelineReflection ComputePipelineReflection;
 #if METAL_DEBUG_OPTIONS
 	ns::String VertexSource;
 	ns::String FragmentSource;
 	ns::String ComputeSource;
 	mtlpp::RenderPipelineDescriptor RenderDesc;
+	mtlpp::RenderPipelineDescriptor StreamDesc;
 	mtlpp::ComputePipelineDescriptor ComputeDesc;
 #endif
 }
 - (instancetype)init;
 - (void)initResourceMask;
-- (void)initResourceMask:(EShaderFrequency)Frequency;
+- (void)initResourceMask:(EAGXShaderFrequency)Frequency;
 @end
diff --git a/Engine/Source/Runtime/Apple/AGXRHI/Private/AGXProfiler.cpp b/Engine/Source/Runtime/Apple/AGXRHI/Private/AGXProfiler.cpp
--- a/Engine/Source/Runtime/Apple/AGXRHI/Private/AGXProfiler.cpp
+++ b/Engine/Source/Runtime/Apple/AGXRHI/Private/AGXProfiler.cpp
@@ -4,6 +4,7 @@
 #include "AGXProfiler.h"
 #include "EngineGlobals.h"
 #include "StaticBoundShaderState.h"
+#include "AGXCommandBuffer.h"
 #include "HAL/FileManager.h"
 
 DEFINE_STAT(STAT_AGXUniformMemAlloc);
