Undo //UE5/Main/Engine/Source/Runtime/Apple/AGXRHI/Private/... changelist 18589140

diff --git a/Engine/Source/Runtime/Apple/AGXRHI/Private/AGXRHI.cpp b/Engine/Source/Runtime/Apple/AGXRHI/Private/AGXRHI.cpp
--- a/Engine/Source/Runtime/Apple/AGXRHI/Private/AGXRHI.cpp
+++ b/Engine/Source/Runtime/Apple/AGXRHI/Private/AGXRHI.cpp
@@ -23,8 +23,6 @@
 #include "Engine/RendererSettings.h"
 #include "AGXTransitionData.h"
 
-#include COMPILED_PLATFORM_HEADER(PlatformAGXConfig.h)
-
 DEFINE_LOG_CATEGORY(LogAGX)
 
 bool GIsAGXInitialized = false;
@@ -51,10 +49,6 @@ static TAutoConsoleVariable<int32> CVarIntelUseRHIThread(
 													ECVF_Default | ECVF_RenderThreadSafe
 													);
 
-id<MTLDevice> AGXUtilGetDevice()
-{
-	return GMtlDevice;
-}
 
 static void ValidateTargetedRHIFeatureLevelExists(EShaderPlatform Platform)
 {
@@ -187,6 +181,10 @@ FAGXDynamicRHI::FAGXDynamicRHI(ERHIFeatureLevel::Type RequestedFeatureLevel)
 	// Metal always needs a render target to render with fragment shaders!
 	GRHIRequiresRenderTargetForPixelShaderUAVs = true;
 
+	//@todo-rco: Query name from API
+	GRHIAdapterName = TEXT("AGX");
+	GRHIVendorId = 1; // non-zero to avoid asserts
+
 	bool const bRequestedFeatureLevel = (RequestedFeatureLevel != ERHIFeatureLevel::Num);
 	bool bSupportsPointLights = false;
 	bool bSupportsRHIThread = false;
@@ -281,8 +279,7 @@ FAGXDynamicRHI::FAGXDynamicRHI(ERHIFeatureLevel::Type RequestedFeatureLevel)
 	GShaderPlatformForFeatureLevel[ERHIFeatureLevel::SM5] = (GMaxRHIFeatureLevel >= ERHIFeatureLevel::SM5) ? GMaxRHIShaderPlatform : SP_NumPlatforms;
 
 #else // PLATFORM_IOS
-	FPlatformAGXConfig::PopulateRHIGlobals();
-
+                
 	uint32 DeviceIndex = ((FAGXDeviceContext*)ImmediateContext.Context)->GetDeviceIndex();
 	
 	TArray<FMacPlatformMisc::FGPUDescriptor> const& GPUs = FPlatformMisc::GetGPUDescriptors();
@@ -292,6 +289,8 @@ FAGXDynamicRHI::FAGXDynamicRHI(ERHIFeatureLevel::Type RequestedFeatureLevel)
 	bool bSupportsD24S8 = false;
 	bool bSupportsD16 = false;
 	
+	GRHIAdapterName = FString([GMtlDevice name]);
+	
 	// However they don't all support other features depending on the version of the OS.
 	bool bSupportsTiledReflections = false;
 	bool bSupportsDistanceFields = false;
@@ -342,6 +341,7 @@ FAGXDynamicRHI::FAGXDynamicRHI(ERHIFeatureLevel::Type RequestedFeatureLevel)
 	else if(GRHIAdapterName.Contains("Apple"))
 	{
 		bSupportsPointLights = true;
+		GRHIVendorId = 0x106B;
 		bSupportsTiledReflections = true;
 		bSupportsDistanceFields = true;
 		bSupportsRHIThread = true;
@@ -561,8 +561,6 @@ FAGXDynamicRHI::FAGXDynamicRHI(ERHIFeatureLevel::Type RequestedFeatureLevel)
 	GSupportsWideMRT = true;
 	GSupportsSeparateRenderTargetBlendState = (GMaxRHIFeatureLevel >= ERHIFeatureLevel::SM5);
 
-	// TODO: GSupportsTransientResourceAliasing requires support for MTLHeap and MTLFence
-
 	GRHISupportsPipelineFileCache = true;
 
 #if PLATFORM_MAC
@@ -664,55 +662,42 @@ FAGXDynamicRHI::FAGXDynamicRHI(ERHIFeatureLevel::Type RequestedFeatureLevel)
 	GAGXBufferFormats[PF_R5G6B5_UNORM         ] = { mtlpp::PixelFormat::Invalid, (uint8)EMetalBufferFormat::R5G6B5Unorm };
 	GAGXBufferFormats[PF_B5G5R5A1_UNORM       ] = { mtlpp::PixelFormat::Invalid, (uint8)EMetalBufferFormat::B5G5R5A1Unorm };
 	GAGXBufferFormats[PF_R8G8B8A8             ] = { mtlpp::PixelFormat::RGBA8Unorm, (uint8)EMetalBufferFormat::RGBA8Unorm };
-	GAGXBufferFormats[PF_A8R8G8B8             ] = { mtlpp::PixelFormat::RGBA8Unorm, (uint8)EMetalBufferFormat::RGBA8Unorm }; // mtlpp::PixelFormat::BGRA8Unorm/EMetalBufferFormat::BGRA8Unorm,  < We don't support this as a vertex-format so we have code to swizzle in the shader
-	GAGXBufferFormats[PF_BC4                  ] = { mtlpp::PixelFormat::Invalid, (uint8)EMetalBufferFormat::Unknown };
+	GAGXBufferFormats[PF_A8R8G8B8				] = { mtlpp::PixelFormat::RGBA8Unorm, (uint8)EMetalBufferFormat::RGBA8Unorm }; // mtlpp::PixelFormat::BGRA8Unorm/EMetalBufferFormat::BGRA8Unorm,  < We don't support this as a vertex-format so we have code to swizzle in the shader
+	GAGXBufferFormats[PF_BC4					] = { mtlpp::PixelFormat::Invalid, (uint8)EMetalBufferFormat::Unknown };
 	GAGXBufferFormats[PF_R8G8                 ] = { mtlpp::PixelFormat::RG8Unorm, (uint8)EMetalBufferFormat::RG8Unorm };
-	GAGXBufferFormats[PF_ATC_RGB              ] = { mtlpp::PixelFormat::Invalid, (uint8)EMetalBufferFormat::Unknown };
-	GAGXBufferFormats[PF_ATC_RGBA_E           ] = { mtlpp::PixelFormat::Invalid, (uint8)EMetalBufferFormat::Unknown };
-	GAGXBufferFormats[PF_ATC_RGBA_I           ] = { mtlpp::PixelFormat::Invalid, (uint8)EMetalBufferFormat::Unknown };
-	GAGXBufferFormats[PF_X24_G8               ] = { mtlpp::PixelFormat::Invalid, (uint8)EMetalBufferFormat::Unknown };
-	GAGXBufferFormats[PF_ETC1                 ] = { mtlpp::PixelFormat::Invalid, (uint8)EMetalBufferFormat::Unknown };
-	GAGXBufferFormats[PF_ETC2_RGB             ] = { mtlpp::PixelFormat::Invalid, (uint8)EMetalBufferFormat::Unknown };
-	GAGXBufferFormats[PF_ETC2_RGBA            ] = { mtlpp::PixelFormat::Invalid, (uint8)EMetalBufferFormat::Unknown };
-	GAGXBufferFormats[PF_R32G32B32A32_UINT    ] = { mtlpp::PixelFormat::RGBA32Uint, (uint8)EMetalBufferFormat::RGBA32Uint };
-	GAGXBufferFormats[PF_R16G16_UINT          ] = { mtlpp::PixelFormat::RG16Uint, (uint8)EMetalBufferFormat::RG16Uint };
-	GAGXBufferFormats[PF_R32G32_UINT          ] = { mtlpp::PixelFormat::RG32Uint, (uint8)EMetalBufferFormat::RG32Uint };
+	GAGXBufferFormats[PF_ATC_RGB				] = { mtlpp::PixelFormat::Invalid, (uint8)EMetalBufferFormat::Unknown };
+	GAGXBufferFormats[PF_ATC_RGBA_E			] = { mtlpp::PixelFormat::Invalid, (uint8)EMetalBufferFormat::Unknown };
+	GAGXBufferFormats[PF_ATC_RGBA_I			] = { mtlpp::PixelFormat::Invalid, (uint8)EMetalBufferFormat::Unknown };
+	GAGXBufferFormats[PF_X24_G8				] = { mtlpp::PixelFormat::Invalid, (uint8)EMetalBufferFormat::Unknown };
+	GAGXBufferFormats[PF_ETC1					] = { mtlpp::PixelFormat::Invalid, (uint8)EMetalBufferFormat::Unknown };
+	GAGXBufferFormats[PF_ETC2_RGB				] = { mtlpp::PixelFormat::Invalid, (uint8)EMetalBufferFormat::Unknown };
+	GAGXBufferFormats[PF_ETC2_RGBA			] = { mtlpp::PixelFormat::Invalid, (uint8)EMetalBufferFormat::Unknown };
+	GAGXBufferFormats[PF_R32G32B32A32_UINT	] = { mtlpp::PixelFormat::RGBA32Uint, (uint8)EMetalBufferFormat::RGBA32Uint };
+	GAGXBufferFormats[PF_R16G16_UINT			] = { mtlpp::PixelFormat::RG16Uint, (uint8)EMetalBufferFormat::RG16Uint };
+	GAGXBufferFormats[PF_R32G32_UINT			] = { mtlpp::PixelFormat::RG32Uint, (uint8)EMetalBufferFormat::RG32Uint };
 	GAGXBufferFormats[PF_ASTC_4x4             ] = { mtlpp::PixelFormat::Invalid, (uint8)EMetalBufferFormat::Unknown };
 	GAGXBufferFormats[PF_ASTC_6x6             ] = { mtlpp::PixelFormat::Invalid, (uint8)EMetalBufferFormat::Unknown };
 	GAGXBufferFormats[PF_ASTC_8x8             ] = { mtlpp::PixelFormat::Invalid, (uint8)EMetalBufferFormat::Unknown };
 	GAGXBufferFormats[PF_ASTC_10x10           ] = { mtlpp::PixelFormat::Invalid, (uint8)EMetalBufferFormat::Unknown };
 	GAGXBufferFormats[PF_ASTC_12x12           ] = { mtlpp::PixelFormat::Invalid, (uint8)EMetalBufferFormat::Unknown };
-	GAGXBufferFormats[PF_ASTC_4x4_HDR         ] = { mtlpp::PixelFormat::Invalid, (uint8)EMetalBufferFormat::Unknown };
-	GAGXBufferFormats[PF_ASTC_6x6_HDR         ] = { mtlpp::PixelFormat::Invalid, (uint8)EMetalBufferFormat::Unknown };
-	GAGXBufferFormats[PF_ASTC_8x8_HDR         ] = { mtlpp::PixelFormat::Invalid, (uint8)EMetalBufferFormat::Unknown };
-	GAGXBufferFormats[PF_ASTC_10x10_HDR       ] = { mtlpp::PixelFormat::Invalid, (uint8)EMetalBufferFormat::Unknown };
-	GAGXBufferFormats[PF_ASTC_12x12_HDR       ] = { mtlpp::PixelFormat::Invalid, (uint8)EMetalBufferFormat::Unknown };
-	GAGXBufferFormats[PF_BC6H                 ] = { mtlpp::PixelFormat::Invalid, (uint8)EMetalBufferFormat::Unknown };
-	GAGXBufferFormats[PF_BC7                  ] = { mtlpp::PixelFormat::Invalid, (uint8)EMetalBufferFormat::Unknown };
-	GAGXBufferFormats[PF_R8_UINT              ] = { mtlpp::PixelFormat::R8Uint, (uint8)EMetalBufferFormat::R8Uint };
-	GAGXBufferFormats[PF_R8                   ] = { mtlpp::PixelFormat::R8Unorm, (uint8)EMetalBufferFormat::R8Unorm };
-	GAGXBufferFormats[PF_L8                   ] = { mtlpp::PixelFormat::Invalid, (uint8)EMetalBufferFormat::R8Unorm };
-	GAGXBufferFormats[PF_XGXR8                ] = { mtlpp::PixelFormat::Invalid, (uint8)EMetalBufferFormat::Unknown };
-	GAGXBufferFormats[PF_R8G8B8A8_UINT        ] = { mtlpp::PixelFormat::RGBA8Uint, (uint8)EMetalBufferFormat::RGBA8Uint };
-	GAGXBufferFormats[PF_R8G8B8A8_SNORM       ] = { mtlpp::PixelFormat::RGBA8Snorm, (uint8)EMetalBufferFormat::RGBA8Snorm };
-	GAGXBufferFormats[PF_R16G16B16A16_UNORM   ] = { mtlpp::PixelFormat::RGBA16Unorm, (uint8)EMetalBufferFormat::RGBA16Unorm };
-	GAGXBufferFormats[PF_R16G16B16A16_SNORM   ] = { mtlpp::PixelFormat::RGBA16Snorm, (uint8)EMetalBufferFormat::RGBA16Snorm };
-	GAGXBufferFormats[PF_PLATFORM_HDR_0       ] = { mtlpp::PixelFormat::Invalid, (uint8)EMetalBufferFormat::Unknown };
-	GAGXBufferFormats[PF_PLATFORM_HDR_1       ] = { mtlpp::PixelFormat::Invalid, (uint8)EMetalBufferFormat::Unknown };
-	GAGXBufferFormats[PF_PLATFORM_HDR_2       ] = { mtlpp::PixelFormat::Invalid, (uint8)EMetalBufferFormat::Unknown };
-	GAGXBufferFormats[PF_NV12                 ] = { mtlpp::PixelFormat::Invalid, (uint8)EMetalBufferFormat::Unknown };
+	GAGXBufferFormats[PF_BC6H					] = { mtlpp::PixelFormat::Invalid, (uint8)EMetalBufferFormat::Unknown };
+	GAGXBufferFormats[PF_BC7					] = { mtlpp::PixelFormat::Invalid, (uint8)EMetalBufferFormat::Unknown };
+	GAGXBufferFormats[PF_R8_UINT				] = { mtlpp::PixelFormat::R8Uint, (uint8)EMetalBufferFormat::R8Uint };
+	GAGXBufferFormats[PF_R8					] = { mtlpp::PixelFormat::R8Unorm, (uint8)EMetalBufferFormat::R8Unorm };
+	GAGXBufferFormats[PF_L8					] = { mtlpp::PixelFormat::Invalid, (uint8)EMetalBufferFormat::R8Unorm };
+	GAGXBufferFormats[PF_XGXR8				] = { mtlpp::PixelFormat::Invalid, (uint8)EMetalBufferFormat::Unknown };
+	GAGXBufferFormats[PF_R8G8B8A8_UINT		] = { mtlpp::PixelFormat::RGBA8Uint, (uint8)EMetalBufferFormat::RGBA8Uint };
+	GAGXBufferFormats[PF_R8G8B8A8_SNORM		] = { mtlpp::PixelFormat::RGBA8Snorm, (uint8)EMetalBufferFormat::RGBA8Snorm };
+	GAGXBufferFormats[PF_R16G16B16A16_UNORM	] = { mtlpp::PixelFormat::RGBA16Unorm, (uint8)EMetalBufferFormat::RGBA16Unorm };
+	GAGXBufferFormats[PF_R16G16B16A16_SNORM	] = { mtlpp::PixelFormat::RGBA16Snorm, (uint8)EMetalBufferFormat::RGBA16Snorm };
+	GAGXBufferFormats[PF_PLATFORM_HDR_0		] = { mtlpp::PixelFormat::Invalid, (uint8)EMetalBufferFormat::Unknown };
+	GAGXBufferFormats[PF_PLATFORM_HDR_1		] = { mtlpp::PixelFormat::Invalid, (uint8)EMetalBufferFormat::Unknown };
+	GAGXBufferFormats[PF_PLATFORM_HDR_2		] = { mtlpp::PixelFormat::Invalid, (uint8)EMetalBufferFormat::Unknown };
+	GAGXBufferFormats[PF_NV12					] = { mtlpp::PixelFormat::Invalid, (uint8)EMetalBufferFormat::Unknown };
 	
-	GAGXBufferFormats[PF_ETC2_R11_EAC         ] = { mtlpp::PixelFormat::Invalid, (uint8)EMetalBufferFormat::Unknown };
-	GAGXBufferFormats[PF_ETC2_RG11_EAC        ] = { mtlpp::PixelFormat::Invalid, (uint8)EMetalBufferFormat::Unknown };
-
-	GAGXBufferFormats[PF_G16R16_SNORM         ] = { mtlpp::PixelFormat::RG16Snorm, (uint8)EMetalBufferFormat::RG16Snorm };
-	GAGXBufferFormats[PF_R8G8_UINT            ] = { mtlpp::PixelFormat::RG8Uint, (uint8)EMetalBufferFormat::RG8Uint };
-	GAGXBufferFormats[PF_R32G32B32_UINT       ] = { mtlpp::PixelFormat::Invalid, (uint8)EMetalBufferFormat::Unknown };
-	GAGXBufferFormats[PF_R32G32B32_SINT       ] = { mtlpp::PixelFormat::Invalid, (uint8)EMetalBufferFormat::Unknown };
-	GAGXBufferFormats[PF_R32G32B32F           ] = { mtlpp::PixelFormat::Invalid, (uint8)EMetalBufferFormat::Unknown };
-	GAGXBufferFormats[PF_R8_SINT              ] = { mtlpp::PixelFormat::R8Sint, (uint8)EMetalBufferFormat::R8Sint };
-	GAGXBufferFormats[PF_R64_UINT             ] = { mtlpp::PixelFormat::Invalid, (uint8)EMetalBufferFormat::Unknown };
-
+	GAGXBufferFormats[PF_ETC2_R11_EAC			] = { mtlpp::PixelFormat::Invalid, (uint8)EMetalBufferFormat::Unknown };
+	GAGXBufferFormats[PF_ETC2_RG11_EAC		] = { mtlpp::PixelFormat::Invalid, (uint8)EMetalBufferFormat::Unknown };
+		
 	// Initialize the platform pixel format map.
 	GPixelFormats[PF_Unknown			].PlatformFormat	= (uint32)mtlpp::PixelFormat::Invalid;
 	GPixelFormats[PF_A32B32G32R32F		].PlatformFormat	= (uint32)mtlpp::PixelFormat::RGBA32Float;
@@ -749,21 +734,6 @@ FAGXDynamicRHI::FAGXDynamicRHI(ERHIFeatureLevel::Type RequestedFeatureLevel)
 	GPixelFormats[PF_ASTC_12x12			].PlatformFormat	= (uint32)mtlpp::PixelFormat::ASTC_12x12_LDR;
 	GPixelFormats[PF_ASTC_12x12			].Supported			= true;
 
-#if !PLATFORM_TVOS
-	if ([GMtlDevice supportsFamily:MTLGPUFamilyApple6])
-	{
-		GPixelFormats[PF_ASTC_4x4_HDR].PlatformFormat		= (uint32)mtlpp::PixelFormat::ASTC_4x4_HDR;
-		GPixelFormats[PF_ASTC_4x4_HDR].Supported			= true;
-		GPixelFormats[PF_ASTC_6x6_HDR].PlatformFormat		= (uint32)mtlpp::PixelFormat::ASTC_6x6_HDR;
-		GPixelFormats[PF_ASTC_6x6_HDR].Supported			= true;
-		GPixelFormats[PF_ASTC_8x8_HDR].PlatformFormat		= (uint32)mtlpp::PixelFormat::ASTC_8x8_HDR;
-		GPixelFormats[PF_ASTC_8x8_HDR].Supported			= true;
-		GPixelFormats[PF_ASTC_10x10_HDR].PlatformFormat		= (uint32)mtlpp::PixelFormat::ASTC_10x10_HDR;
-		GPixelFormats[PF_ASTC_10x10_HDR].Supported			= true;
-		GPixelFormats[PF_ASTC_12x12_HDR].PlatformFormat		= (uint32)mtlpp::PixelFormat::ASTC_12x12_HDR;
-		GPixelFormats[PF_ASTC_12x12_HDR].Supported			= true;
-	}
-#endif
 	// used with virtual textures
 	GPixelFormats[PF_ETC2_RGB	  		].PlatformFormat	= (uint32)mtlpp::PixelFormat::ETC2_RGB8;
 	GPixelFormats[PF_ETC2_RGB			].Supported			= true;
@@ -824,7 +794,6 @@ FAGXDynamicRHI::FAGXDynamicRHI(ERHIFeatureLevel::Type RequestedFeatureLevel)
 	
 	GPixelFormats[PF_FloatRGB			].PlatformFormat	= (uint32)mtlpp::PixelFormat::RG11B10Float;
 	GPixelFormats[PF_FloatRGB			].BlockBytes		= 4;
-	
 	GPixelFormats[PF_FloatR11G11B10		].PlatformFormat	= (uint32)mtlpp::PixelFormat::RG11B10Float;
 	GPixelFormats[PF_FloatR11G11B10		].BlockBytes		= 4;
 	GPixelFormats[PF_FloatR11G11B10		].Supported			= true;
@@ -920,18 +889,6 @@ FAGXDynamicRHI::FAGXDynamicRHI(ERHIFeatureLevel::Type RequestedFeatureLevel)
 	GPixelFormats[PF_NV12				].PlatformFormat	= (uint32)mtlpp::PixelFormat::Invalid;
 	GPixelFormats[PF_NV12				].Supported			= false;
 	
-	GPixelFormats[PF_G16R16_SNORM		].PlatformFormat	= (uint32)mtlpp::PixelFormat::RG16Snorm;
-	GPixelFormats[PF_R8G8_UINT			].PlatformFormat	= (uint32)mtlpp::PixelFormat::RG8Uint;
-	GPixelFormats[PF_R32G32B32_UINT		].PlatformFormat	= (uint32)mtlpp::PixelFormat::Invalid;
-	GPixelFormats[PF_R32G32B32_UINT		].Supported			= false;
-	GPixelFormats[PF_R32G32B32_SINT		].PlatformFormat	= (uint32)mtlpp::PixelFormat::Invalid;
-	GPixelFormats[PF_R32G32B32_SINT		].Supported			= false;
-	GPixelFormats[PF_R32G32B32F			].PlatformFormat	= (uint32)mtlpp::PixelFormat::Invalid;
-	GPixelFormats[PF_R32G32B32F			].Supported			= false;
-	GPixelFormats[PF_R8_SINT			].PlatformFormat	= (uint32)mtlpp::PixelFormat::R8Sint;
-	GPixelFormats[PF_R64_UINT			].PlatformFormat	= (uint32)mtlpp::PixelFormat::Invalid;
-	GPixelFormats[PF_R64_UINT			].Supported			= false;
-
 #if METAL_DEBUG_OPTIONS
 	for (uint32 i = 0; i < PF_MAX; i++)
 	{
@@ -978,6 +935,10 @@ FAGXDynamicRHI::FAGXDynamicRHI(ERHIFeatureLevel::Type RequestedFeatureLevel)
 	{
 		FGPUDriverInfo GPUDriverInfo = FPlatformMisc::GetGPUDriverInfo(GRHIAdapterName);
 		
+		GRHIAdapterUserDriverVersion = GPUDriverInfo.UserDriverVersion;
+		GRHIAdapterInternalDriverVersion = GPUDriverInfo.InternalDriverVersion;
+		GRHIAdapterDriverDate = GPUDriverInfo.DriverDate;
+		
 		UE_LOG(LogAGX, Display, TEXT("    Adapter Name: %s"), *GRHIAdapterName);
 		UE_LOG(LogAGX, Display, TEXT("  Driver Version: %s (internal:%s, unified:%s)"), *GRHIAdapterUserDriverVersion, *GRHIAdapterInternalDriverVersion, *GPUDriverInfo.GetUnifiedDriverVersion());
 		UE_LOG(LogAGX, Display, TEXT("     Driver Date: %s"), *GRHIAdapterDriverDate);
@@ -1251,7 +1212,7 @@ bool FAGXDynamicRHI::RHIGetAvailableResolutions(FScreenResolutionArray& Resoluti
 void FAGXDynamicRHI::RHIFlushResources()
 {
 	@autoreleasepool {
-		((FAGXDeviceContext*)ImmediateContext.Context)->FlushFreeList();
+		((FAGXDeviceContext*)ImmediateContext.Context)->FlushFreeList(false);
 		ImmediateContext.Context->SubmitCommandBufferAndWait();
 		((FAGXDeviceContext*)ImmediateContext.Context)->ClearFreeList();
         ((FAGXDeviceContext*)ImmediateContext.Context)->DrainHeap();
diff --git a/Engine/Source/Runtime/Apple/AGXRHI/Private/AGXRHIContext.cpp b/Engine/Source/Runtime/Apple/AGXRHI/Private/AGXRHIContext.cpp
--- a/Engine/Source/Runtime/Apple/AGXRHI/Private/AGXRHIContext.cpp
+++ b/Engine/Source/Runtime/Apple/AGXRHI/Private/AGXRHIContext.cpp
@@ -53,6 +53,19 @@ void AGXSafeReleaseMetalBuffer(FAGXBuffer& Buffer)
 	}
 }
 
+void AGXSafeReleaseMetalFence(FAGXFence* Object)
+{
+	if(GIsAGXInitialized && GDynamicRHI && Object)
+	{
+		FAGXRHICommandContext* Context = static_cast<FAGXRHICommandContext*>(RHIGetDefaultContext());
+		if(Context)
+		{
+			((FAGXDeviceContext&)Context->GetInternalContext()).ReleaseFence(Object);
+			return;
+		}
+	}
+}
+
 FAGXRHICommandContext::FAGXRHICommandContext(class FAGXProfiler* InProfiler, FAGXContext* WrapContext)
 : Context(WrapContext)
 , Profiler(InProfiler)
@@ -73,6 +86,10 @@ FAGXRHICommandContext::~FAGXRHICommandContext()
 FAGXRHIComputeContext::FAGXRHIComputeContext(class FAGXProfiler* InProfiler, FAGXContext* WrapContext)
 : FAGXRHICommandContext(InProfiler, WrapContext)
 {
+	if (FAGXCommandQueue::SupportsFeature(EAGXFeaturesFences) && FApplePlatformMisc::IsOSAtLeastVersion((uint32[]){10, 14, 0}, (uint32[]){12, 0, 0}, (uint32[]){12, 0, 0}))
+	{
+		WrapContext->GetCurrentRenderPass().SetDispatchType(mtlpp::DispatchType::Concurrent);
+	}
 }
 
 FAGXRHIComputeContext::~FAGXRHIComputeContext()
diff --git a/Engine/Source/Runtime/Apple/AGXRHI/Private/AGXRHIPrivate.h b/Engine/Source/Runtime/Apple/AGXRHI/Private/AGXRHIPrivate.h
--- a/Engine/Source/Runtime/Apple/AGXRHI/Private/AGXRHIPrivate.h
+++ b/Engine/Source/Runtime/Apple/AGXRHI/Private/AGXRHIPrivate.h
@@ -15,16 +15,11 @@
 #import <Metal/Metal.h>
 #import <QuartzCore/CAMetalLayer.h>
 
-#define AGXRHI_TRUE  1
-#define AGXRHI_FALSE 0
-
 // Metal C++ wrapper
 THIRD_PARTY_INCLUDES_START
 #include "mtlpp.hpp"
 THIRD_PARTY_INCLUDES_END
 
-extern id<MTLDevice> AGXUtilGetDevice();
-
 // The Metal device object.
 extern id<MTLDevice> GMtlDevice;
 
@@ -57,54 +52,6 @@ const uint32 AGXBufferBytesSize = BufferOffsetAlignment * 32;
 #include "AGXDynamicRHI.h"
 #include "RHI.h"
 
-/**
- * EAGXDebugLevel: Level of AGX RHI debug features to be enabled.
- */
-enum EAGXDebugLevel
-{
-	EAGXDebugLevelOff,
-	EAGXDebugLevelFastValidation,
-	EAGXDebugLevelResetOnBind,
-	EAGXDebugLevelConditionalSubmit,
-	EAGXDebugLevelValidation,
-	EAGXDebugLevelLogOperations,
-	EAGXDebugLevelWaitForComplete,
-};
-
-/**
- * The sampler, buffer and texture resource limits as defined here:
- * https://developer.apple.com/library/ios/documentation/Miscellaneous/Conceptual/MetalProgrammingGuide/Render-Ctx/Render-Ctx.html
- */
-#if PLATFORM_IOS
-#define METAL_MAX_BUFFERS 31
-#define METAL_MAX_TEXTURES 31
-typedef uint32 FAGXTextureMask;
-#elif PLATFORM_MAC
-#define METAL_MAX_BUFFERS 31
-#define METAL_MAX_TEXTURES 128
-typedef __uint128_t FAGXTextureMask;
-#else
-#error "Unsupported Platform!"
-#endif
-typedef uint32 FAGXBufferMask;
-typedef uint16 FAGXSamplerMask;
-
-enum EAGXLimits
-{
-	ML_MaxSamplers = 16, /** Maximum number of samplers */
-	ML_MaxBuffers = METAL_MAX_BUFFERS, /** Maximum number of buffers */
-	ML_MaxTextures = METAL_MAX_TEXTURES, /** Maximum number of textures - there are more textures available on Mac than iOS */
-	ML_MaxViewports = 16 /** Technically this may be different at runtime, but this is the likely absolute upper-bound */
-};
-
-/** A structure for quick mask-testing of shader-stage resource bindings */
-struct FAGXDebugShaderResourceMask
-{
-	FAGXTextureMask TextureMask;
-	FAGXBufferMask BufferMask;
-	FAGXSamplerMask SamplerMask;
-};
-
 #define BUFFER_CACHE_MODE mtlpp::ResourceOptions::CpuCacheModeDefaultCache
 
 #if PLATFORM_MAC
@@ -161,6 +108,8 @@ extern FAGXBufferFormat GAGXBufferFormats[PF_MAX];
 #define METAL_DEBUG_LAYER(Level, Code)
 #endif
 
+extern bool GAGXCommandBufferDebuggingEnabled;
+
 /** Set to 1 to enable GPU events in Xcode frame debugger */
 #ifndef ENABLE_METAL_GPUEVENTS_IN_TEST
 	#define ENABLE_METAL_GPUEVENTS_IN_TEST 0
@@ -190,6 +139,19 @@ extern FAGXBufferFormat GAGXBufferFormats[PF_MAX];
 	#define METAL_IGNORED(Func)
 #endif
 
+struct FAGXDebugInfo
+{
+	uint32 CmdBuffIndex;
+	uint32 EncoderIndex;
+	uint32 ContextIndex;
+	uint32 CommandIndex;
+	uint64 CommandBuffer;
+	uint32 PSOSignature[4];
+};
+
+// Get a compute pipeline state used to implement some debug features.
+mtlpp::ComputePipelineState AGXGetMetalDebugComputeState();
+
 // Access the internal context for the device-owning DynamicRHI object
 FAGXDeviceContext& GetAGXDeviceContext();
 
@@ -202,6 +164,9 @@ void AGXSafeReleaseMetalTexture(FAGXTexture& Object);
 // Safely release a metal buffer, correctly handling the case where the RHI has been destructed first
 void AGXSafeReleaseMetalBuffer(FAGXBuffer& Buffer);
 
+// Safely release a fence, correctly handling cases where fences aren't supported or the debug implementation is used.
+void AGXSafeReleaseMetalFence(class FAGXFence* Object);
+
 // Safely release a render pass descriptor so that it may be reused.
 void AGXSafeReleaseMetalRenderPassDescriptor(mtlpp::RenderPassDescriptor& Desc);
 
@@ -270,7 +235,7 @@ extern uint32 AGXSafeGetRuntimeDebuggingLevel();
 
 extern int32 GAGXBufferZeroFill;
 
-mtlpp::LanguageVersion AGXValidateVersion(uint32 Version);
+mtlpp::LanguageVersion AGXValidateVersion(uint8 Version);
 
 // Needs to be the same as EShaderFrequency when all stages are supported, but unlike EShaderFrequency you can compile out stages.
 enum EAGXShaderStages
diff --git a/Engine/Source/Runtime/Apple/AGXRHI/Private/AGXRenderCommandEncoder.cpp b/Engine/Source/Runtime/Apple/AGXRHI/Private/AGXRenderCommandEncoder.cpp
new file mode 100644
--- /dev/null
+++ b/Engine/Source/Runtime/Apple/AGXRHI/Private/AGXRenderCommandEncoder.cpp
@@ -0,0 +1,780 @@
+// Copyright Epic Games, Inc. All Rights Reserved.
+
+/*=============================================================================
+	AGXRenderCommandEncoder.cpp: AGX RHI command encoder wrapper.
+=============================================================================*/
+
+#include "AGXRHIPrivate.h"
+
+#include "AGXRenderCommandEncoder.h"
+#include "AGXCommandBuffer.h"
+#include "AGXCommandQueue.h"
+#include "AGXFence.h"
+#include "AGXPipeline.h"
+
+#if MTLPP_CONFIG_VALIDATE && METAL_DEBUG_OPTIONS
+NS_ASSUME_NONNULL_BEGIN
+
+@interface FAGXDebugRenderCommandEncoder : FAGXDebugCommandEncoder
+{
+@public
+#pragma mark - Private Member Variables -
+#if METAL_DEBUG_OPTIONS
+	FAGXDebugShaderResourceMask ResourceMask[EAGXShaderRenderNum];
+	FAGXDebugBufferBindings ShaderBuffers[EAGXShaderRenderNum];
+	FAGXDebugTextureBindings ShaderTextures[EAGXShaderRenderNum];
+	FAGXDebugSamplerBindings ShaderSamplers[EAGXShaderRenderNum];
+	id<MTLRenderPipelineState> DebugState;
+#endif
+	MTLRenderPassDescriptor* RenderPassDesc;
+	id<MTLRenderCommandEncoder> Inner;
+	FAGXCommandBufferDebugging Buffer;
+	FAGXShaderPipeline* Pipeline;
+}
+
+/** Initialise the wrapper with the provided command-buffer. */
+-(id)initWithEncoder:(id<MTLRenderCommandEncoder>)Encoder fromDescriptor:(MTLRenderPassDescriptor*)RenderPassDesc andCommandBuffer:(FAGXCommandBufferDebugging const&)Buffer;
+
+@end
+
+@interface FAGXDebugParallelRenderCommandEncoder : FApplePlatformObject
+{
+@public
+	TArray<FAGXRenderCommandEncoderDebugging> RenderEncoders;
+	MTLRenderPassDescriptor* RenderPassDesc;
+	id<MTLParallelRenderCommandEncoder> Inner;
+	FAGXCommandBufferDebugging Buffer;
+}
+-(id)initWithEncoder:(id<MTLParallelRenderCommandEncoder>)Encoder fromDescriptor:(MTLRenderPassDescriptor*)RenderPassDesc andCommandBuffer:(FAGXCommandBufferDebugging const&)Buffer;
+@end
+
+
+#if METAL_DEBUG_OPTIONS
+static NSString* GAGXDebugVertexShader = @"#include <metal_stdlib>\n"
+@"using namespace metal;\n"
+@"struct VertexInput\n"
+@"{\n"
+@"};\n"
+@"vertex void WriteCommandIndexVS(VertexInput StageIn [[stage_in]], constant uint* Input [[ buffer(0) ]], device uint* Output  [[ buffer(1) ]])\n"
+@"{\n"
+@"	Output[0] = Input[0];\n"
+@"}\n";
+
+static id <MTLRenderPipelineState> GetDebugVertexShaderState(id<MTLDevice> Device, MTLRenderPassDescriptor* PassDesc)
+{
+	static id<MTLFunction> Func = nil;
+	static FCriticalSection Mutex;
+	static NSMutableDictionary* Dict = [NSMutableDictionary new];
+	if(!Func)
+	{
+		id<MTLLibrary> Lib = [Device newLibraryWithSource:GAGXDebugVertexShader options:nil error:nullptr];
+		check(Lib);
+		Func = [Lib newFunctionWithName:@"WriteCommandIndexVS"];
+		check(Func);
+		[Lib release];
+	}
+	
+	FScopeLock Lock(&Mutex);
+	id<MTLRenderPipelineState> State = [Dict objectForKey:PassDesc];
+	if (!State)
+	{
+		MTLRenderPipelineDescriptor* Desc = [MTLRenderPipelineDescriptor new];
+		
+		Desc.vertexFunction = Func;
+		
+		if (PassDesc.depthAttachment)
+		{
+			Desc.depthAttachmentPixelFormat = PassDesc.depthAttachment.texture.pixelFormat;
+		}
+		if (PassDesc.stencilAttachment)
+		{
+			Desc.stencilAttachmentPixelFormat = PassDesc.stencilAttachment.texture.pixelFormat;
+		}
+		if (PassDesc.colorAttachments)
+		{
+			for (NSUInteger i = 0; i < 8; i++)
+			{
+				MTLRenderPassColorAttachmentDescriptor* CD = [PassDesc.colorAttachments objectAtIndexedSubscript:i];
+				if (CD.texture.pixelFormat != MTLPixelFormatInvalid)
+				{
+					MTLRenderPipelineColorAttachmentDescriptor* CD0 = [[MTLRenderPipelineColorAttachmentDescriptor new] autorelease];
+					CD0.pixelFormat = CD.texture.pixelFormat;
+					[Desc.colorAttachments setObject:CD0 atIndexedSubscript:i];
+				}
+			}
+		}
+		Desc.rasterizationEnabled = false;
+		
+		State = [[Device newRenderPipelineStateWithDescriptor:Desc error:nil] autorelease];
+		check(State);
+		
+		[Dict setObject:State forKey:PassDesc];
+		
+		[Desc release];
+	}
+	check(State);
+	return State;
+}
+#endif
+
+@implementation FAGXDebugParallelRenderCommandEncoder
+
+APPLE_PLATFORM_OBJECT_ALLOC_OVERRIDES(FAGXDebugParallelRenderCommandEncoder)
+
+-(id)initWithEncoder:(id<MTLParallelRenderCommandEncoder>)Encoder fromDescriptor:(MTLRenderPassDescriptor*)Desc andCommandBuffer:(FAGXCommandBufferDebugging const&)SourceBuffer
+{
+	id Self = [super init];
+	if (Self)
+	{
+        Inner = Encoder;
+		Buffer = SourceBuffer;
+		RenderPassDesc = [Desc retain];
+	}
+	return Self;
+}
+
+-(void)dealloc
+{
+	[RenderPassDesc release];
+	[super dealloc];
+}
+
+@end
+
+@implementation FAGXDebugRenderCommandEncoder
+
+APPLE_PLATFORM_OBJECT_ALLOC_OVERRIDES(FAGXDebugRenderCommandEncoder)
+
+-(id)initWithEncoder:(id<MTLRenderCommandEncoder>)Encoder fromDescriptor:(MTLRenderPassDescriptor*)Desc andCommandBuffer:(FAGXCommandBufferDebugging const&)SourceBuffer
+{
+	id Self = [super init];
+	if (Self)
+	{
+        Inner = Encoder;
+		Buffer = SourceBuffer;
+		RenderPassDesc = [Desc retain];
+#if METAL_DEBUG_OPTIONS
+		DebugState = Buffer.GetPtr()->DebugLevel >= EAGXDebugLevelValidation ? [GetDebugVertexShaderState(Buffer.GetPtr()->InnerBuffer.device, Desc) retain] : nil;
+#endif
+        Pipeline = nil;
+	}
+	return Self;
+}
+
+-(void)dealloc
+{
+	[RenderPassDesc release];
+#if METAL_DEBUG_OPTIONS
+	[DebugState release];
+#endif
+	[Pipeline release];
+	[super dealloc];
+}
+
+@end
+
+NS_ASSUME_NONNULL_END
+
+void FAGXRenderCommandEncoderDebugging::InsertDebugDraw()
+{
+#if METAL_DEBUG_OPTIONS
+//	switch (((FAGXDebugRenderCommandEncoder*)m_ptr)->Buffer.GetPtr()->DebugLevel)
+//	{
+//		case EAGXDebugLevelConditionalSubmit:
+//		case EAGXDebugLevelWaitForComplete:
+//		case EAGXDebugLevelLogOperations:
+//		case EAGXDebugLevelValidation:
+//		{
+//			uint32 const Index = Buffer->DebugCommands.Num();
+//#if PLATFORM_MAC
+//			[Inner textureBarrier];
+//#endif
+//			[Inner setVertexBytes:&Index length:sizeof(Index) atIndex:0];
+//			[Inner setVertexBuffer:Buffer->DebugInfoBuffer offset:0 atIndex:1];
+//			[Inner setRenderPipelineState:DebugState];
+//			[Inner drawPrimitives:mtlpp::PrimitiveTypePoint vertexStart:0 vertexCount:1];
+//#if PLATFORM_MAC
+//			[Inner textureBarrier];
+//#endif
+//			if (Pipeline && Pipeline->RenderPipelineState)
+//			{
+//				[Inner setRenderPipelineState:Pipeline->RenderPipelineState];
+//			}
+//			
+//			if (ShaderBuffers[EAGXShaderVertex].Buffers[0])
+//			{
+//				[Inner setVertexBuffer:ShaderBuffers[EAGXShaderVertex].Buffers[0] offset:ShaderBuffers[EAGXShaderVertex].Offsets[0] atIndex:0];
+//			}
+//			else if (ShaderBuffers[EAGXShaderVertex].Bytes[0])
+//			{
+//				[Inner setVertexBytes:ShaderBuffers[EAGXShaderVertex].Bytes[0] length:ShaderBuffers[EAGXShaderVertex].Offsets[0] atIndex:0];
+//			}
+//			
+//			if (ShaderBuffers[EAGXShaderVertex].Buffers[1])
+//			{
+//				[Inner setVertexBuffer:ShaderBuffers[EAGXShaderVertex].Buffers[1] offset:ShaderBuffers[EAGXShaderVertex].Offsets[1] atIndex:1];
+//			}
+//			else if (ShaderBuffers[EAGXShaderVertex].Bytes[1])
+//			{
+//				[Inner setVertexBytes:ShaderBuffers[EAGXShaderVertex].Bytes[1] length:ShaderBuffers[EAGXShaderVertex].Offsets[1] atIndex:1];
+//			}
+//		}
+//		default:
+//		{
+//			break;
+//		}
+//	}
+#endif
+}
+
+FAGXRenderCommandEncoderDebugging::FAGXRenderCommandEncoderDebugging()
+{
+	
+}
+FAGXRenderCommandEncoderDebugging::FAGXRenderCommandEncoderDebugging(mtlpp::RenderCommandEncoder& Encoder, mtlpp::RenderPassDescriptor const& Desc, FAGXCommandBufferDebugging& Buffer)
+: FAGXCommandEncoderDebugging((FAGXDebugCommandEncoder*)[[[FAGXDebugRenderCommandEncoder alloc] initWithEncoder:Encoder.GetPtr() fromDescriptor:Desc.GetPtr() andCommandBuffer:Buffer] autorelease])
+{
+	Buffer.BeginRenderCommandEncoder([NSString stringWithFormat:@"Render: %@", Encoder.GetLabel().GetPtr()], Desc);
+	Encoder.SetAssociatedObject((void const*)&FAGXRenderCommandEncoderDebugging::Get, (FAGXCommandEncoderDebugging const&)*this);
+}
+FAGXRenderCommandEncoderDebugging::FAGXRenderCommandEncoderDebugging(FAGXDebugCommandEncoder* handle)
+: FAGXCommandEncoderDebugging(handle)
+{
+	
+}
+
+FAGXRenderCommandEncoderDebugging FAGXRenderCommandEncoderDebugging::Get(mtlpp::RenderCommandEncoder& Encoder)
+{
+	return Encoder.GetAssociatedObject<FAGXRenderCommandEncoderDebugging>((void const*)&FAGXRenderCommandEncoderDebugging::Get);
+}
+
+void FAGXRenderCommandEncoderDebugging::SetPipeline(FAGXShaderPipeline* Pipeline)
+{
+#if METAL_DEBUG_OPTIONS
+	((FAGXDebugRenderCommandEncoder*)m_ptr)->Pipeline = [Pipeline retain];
+	switch (((FAGXDebugRenderCommandEncoder*)m_ptr)->Buffer.GetPtr()->DebugLevel)
+	{
+		case EAGXDebugLevelConditionalSubmit:
+		case EAGXDebugLevelWaitForComplete:
+		case EAGXDebugLevelLogOperations:
+		{
+			((FAGXDebugRenderCommandEncoder*)m_ptr)->Buffer.SetPipeline(Pipeline->RenderPipelineState.GetLabel());
+		}
+		default:
+		{
+			break;
+		}
+	}
+#endif
+}
+
+void FAGXRenderCommandEncoderDebugging::SetBytes(EAGXShaderFrequency Freq, const void * bytes, NSUInteger length, NSUInteger index)
+{
+#if METAL_DEBUG_OPTIONS
+	switch (((FAGXDebugRenderCommandEncoder*)m_ptr)->Buffer.GetPtr()->DebugLevel)
+	{
+		case EAGXDebugLevelConditionalSubmit:
+		case EAGXDebugLevelWaitForComplete:
+		case EAGXDebugLevelLogOperations:
+		case EAGXDebugLevelValidation:
+		{
+			((FAGXDebugRenderCommandEncoder*)m_ptr)->ShaderBuffers[Freq].Buffers[index] = nil;
+			((FAGXDebugRenderCommandEncoder*)m_ptr)->ShaderBuffers[Freq].Bytes[index] = bytes;
+			((FAGXDebugRenderCommandEncoder*)m_ptr)->ShaderBuffers[Freq].Offsets[index] = length;
+			((FAGXDebugRenderCommandEncoder*)m_ptr)->ResourceMask[Freq].BufferMask = bytes ? (((FAGXDebugRenderCommandEncoder*)m_ptr)->ResourceMask[Freq].BufferMask | (1 << (index))) : (((FAGXDebugRenderCommandEncoder*)m_ptr)->ResourceMask[Freq].BufferMask & ~(1 << (index)));
+		}
+		default:
+		{
+			break;
+		}
+	}
+#endif
+}
+void FAGXRenderCommandEncoderDebugging::SetBuffer(EAGXShaderFrequency Freq,  FAGXBuffer const& buffer, NSUInteger offset, NSUInteger index)
+{
+#if METAL_DEBUG_OPTIONS
+	switch (((FAGXDebugRenderCommandEncoder*)m_ptr)->Buffer.GetPtr()->DebugLevel)
+	{
+		case EAGXDebugLevelConditionalSubmit:
+		case EAGXDebugLevelWaitForComplete:
+		case EAGXDebugLevelLogOperations:
+		case EAGXDebugLevelValidation:
+		{
+			((FAGXDebugRenderCommandEncoder*)m_ptr)->ShaderBuffers[Freq].Buffers[index] = buffer;
+			((FAGXDebugRenderCommandEncoder*)m_ptr)->ShaderBuffers[Freq].Bytes[index] = nil;
+			((FAGXDebugRenderCommandEncoder*)m_ptr)->ShaderBuffers[Freq].Offsets[index] = offset;
+			((FAGXDebugRenderCommandEncoder*)m_ptr)->ResourceMask[Freq].BufferMask = buffer ? (((FAGXDebugRenderCommandEncoder*)m_ptr)->ResourceMask[Freq].BufferMask | (1 << (index))) : (((FAGXDebugRenderCommandEncoder*)m_ptr)->ResourceMask[Freq].BufferMask & ~(1 << (index)));
+		}
+		default:
+		{
+			break;
+		}
+	}
+#endif
+}
+void FAGXRenderCommandEncoderDebugging::SetBufferOffset(EAGXShaderFrequency Freq, NSUInteger offset, NSUInteger index)
+{
+#if METAL_DEBUG_OPTIONS
+	switch (((FAGXDebugRenderCommandEncoder*)m_ptr)->Buffer.GetPtr()->DebugLevel)
+	{
+		case EAGXDebugLevelConditionalSubmit:
+		case EAGXDebugLevelWaitForComplete:
+		case EAGXDebugLevelLogOperations:
+		case EAGXDebugLevelValidation:
+		{
+			((FAGXDebugRenderCommandEncoder*)m_ptr)->ShaderBuffers[Freq].Offsets[index] = offset;
+			check(((FAGXDebugRenderCommandEncoder*)m_ptr)->ResourceMask[Freq].BufferMask & (1 << (index)));
+		}
+		default:
+		{
+			break;
+		}
+	}
+#endif
+}
+
+void FAGXRenderCommandEncoderDebugging::SetTexture(EAGXShaderFrequency Freq,  FAGXTexture const& texture, NSUInteger index)
+{
+#if METAL_DEBUG_OPTIONS
+	switch (((FAGXDebugRenderCommandEncoder*)m_ptr)->Buffer.GetPtr()->DebugLevel)
+	{
+		case EAGXDebugLevelConditionalSubmit:
+		case EAGXDebugLevelWaitForComplete:
+		case EAGXDebugLevelLogOperations:
+		case EAGXDebugLevelValidation:
+		{
+			((FAGXDebugRenderCommandEncoder*)m_ptr)->ShaderTextures[Freq].Textures[index] = texture;
+			((FAGXDebugRenderCommandEncoder*)m_ptr)->ResourceMask[Freq].TextureMask = texture ? (((FAGXDebugRenderCommandEncoder*)m_ptr)->ResourceMask[Freq].TextureMask | (FAGXTextureMask(1) << index)) : (((FAGXDebugRenderCommandEncoder*)m_ptr)->ResourceMask[Freq].TextureMask & ~(FAGXTextureMask(1) << index));
+		}
+		default:
+		{
+			break;
+		}
+	}
+#endif
+}
+
+void FAGXRenderCommandEncoderDebugging::SetSamplerState(EAGXShaderFrequency Freq,  mtlpp::SamplerState const& sampler, NSUInteger index)
+{
+#if METAL_DEBUG_OPTIONS
+	switch (((FAGXDebugRenderCommandEncoder*)m_ptr)->Buffer.GetPtr()->DebugLevel)
+	{
+		case EAGXDebugLevelConditionalSubmit:
+		case EAGXDebugLevelWaitForComplete:
+		case EAGXDebugLevelLogOperations:
+		case EAGXDebugLevelValidation:
+		{
+			((FAGXDebugRenderCommandEncoder*)m_ptr)->ShaderSamplers[Freq].Samplers[index] = sampler;
+			((FAGXDebugRenderCommandEncoder*)m_ptr)->ResourceMask[Freq].SamplerMask = sampler ? (((FAGXDebugRenderCommandEncoder*)m_ptr)->ResourceMask[Freq].SamplerMask | (1 << (index))) : (((FAGXDebugRenderCommandEncoder*)m_ptr)->ResourceMask[Freq].SamplerMask & ~(1 << (index)));
+		}
+		default:
+		{
+			break;
+		}
+	}
+#endif
+}
+
+void FAGXRenderCommandEncoderDebugging::SetSamplerState(EAGXShaderFrequency Freq,  mtlpp::SamplerState const& sampler, float lodMinClamp, float lodMaxClamp, NSUInteger index)
+{
+#if METAL_DEBUG_OPTIONS
+	switch (((FAGXDebugRenderCommandEncoder*)m_ptr)->Buffer.GetPtr()->DebugLevel)
+	{
+		case EAGXDebugLevelConditionalSubmit:
+		case EAGXDebugLevelWaitForComplete:
+		case EAGXDebugLevelLogOperations:
+		case EAGXDebugLevelValidation:
+		{
+			((FAGXDebugRenderCommandEncoder*)m_ptr)->ShaderSamplers[Freq].Samplers[index] = sampler;
+			((FAGXDebugRenderCommandEncoder*)m_ptr)->ResourceMask[Freq].SamplerMask = sampler ? (((FAGXDebugRenderCommandEncoder*)m_ptr)->ResourceMask[Freq].SamplerMask | (1 << (index))) : (((FAGXDebugRenderCommandEncoder*)m_ptr)->ResourceMask[Freq].SamplerMask & ~(1 << (index)));
+		}
+		default:
+		{
+			break;
+		}
+	}
+#endif
+}
+
+void FAGXRenderCommandEncoderDebugging::SetDepthStencilState( mtlpp::DepthStencilState const& depthStencilState)
+{
+}
+
+void FAGXRenderCommandEncoderDebugging::Draw(mtlpp::PrimitiveType primitiveType, NSUInteger vertexStart, NSUInteger vertexCount, NSUInteger instanceCount)
+{
+#if METAL_DEBUG_OPTIONS
+	switch(((FAGXDebugRenderCommandEncoder*)m_ptr)->Buffer.GetPtr()->DebugLevel)
+	{
+		case EAGXDebugLevelConditionalSubmit:
+		case EAGXDebugLevelWaitForComplete:
+		case EAGXDebugLevelLogOperations:
+		{
+			((FAGXDebugRenderCommandEncoder*)m_ptr)->Buffer.Draw([NSString stringWithFormat:@"%s", __PRETTY_FUNCTION__]);
+		}
+		case EAGXDebugLevelValidation:
+		{
+			Validate();
+		}
+		default:
+		{
+			break;
+		}
+	}
+#endif
+}
+
+void FAGXRenderCommandEncoderDebugging::Draw(mtlpp::PrimitiveType primitiveType, NSUInteger vertexStart, NSUInteger vertexCount)
+{
+#if METAL_DEBUG_OPTIONS
+	switch(((FAGXDebugRenderCommandEncoder*)m_ptr)->Buffer.GetPtr()->DebugLevel)
+	{
+		case EAGXDebugLevelConditionalSubmit:
+		case EAGXDebugLevelWaitForComplete:
+		case EAGXDebugLevelLogOperations:
+		{
+			((FAGXDebugRenderCommandEncoder*)m_ptr)->Buffer.Draw([NSString stringWithFormat:@"%s", __PRETTY_FUNCTION__]);
+		}
+		case EAGXDebugLevelValidation:
+		{
+			Validate();
+		}
+		default:
+		{
+			break;
+		}
+	}
+#endif
+}
+
+void FAGXRenderCommandEncoderDebugging::DrawIndexed(mtlpp::PrimitiveType primitiveType, NSUInteger indexCount, mtlpp::IndexType indexType, FAGXBuffer const& indexBuffer, NSUInteger indexBufferOffset, NSUInteger instanceCount)
+{
+#if METAL_DEBUG_OPTIONS
+	switch(((FAGXDebugRenderCommandEncoder*)m_ptr)->Buffer.GetPtr()->DebugLevel)
+	{
+		case EAGXDebugLevelConditionalSubmit:
+		case EAGXDebugLevelWaitForComplete:
+		case EAGXDebugLevelLogOperations:
+		{
+			((FAGXDebugRenderCommandEncoder*)m_ptr)->Buffer.Draw([NSString stringWithFormat:@"%s", __PRETTY_FUNCTION__]);
+		}
+		case EAGXDebugLevelValidation:
+		{
+			Validate();
+		}
+		default:
+		{
+			break;
+		}
+	}
+#endif
+}
+
+void FAGXRenderCommandEncoderDebugging::DrawIndexed(mtlpp::PrimitiveType primitiveType, NSUInteger indexCount, mtlpp::IndexType indexType, FAGXBuffer const& indexBuffer, NSUInteger indexBufferOffset)
+{
+#if METAL_DEBUG_OPTIONS
+	switch(((FAGXDebugRenderCommandEncoder*)m_ptr)->Buffer.GetPtr()->DebugLevel)
+	{
+		case EAGXDebugLevelConditionalSubmit:
+		case EAGXDebugLevelWaitForComplete:
+		case EAGXDebugLevelLogOperations:
+		{
+			((FAGXDebugRenderCommandEncoder*)m_ptr)->Buffer.Draw([NSString stringWithFormat:@"%s", __PRETTY_FUNCTION__]);
+		}
+		case EAGXDebugLevelValidation:
+		{
+			Validate();
+		}
+		default:
+		{
+			break;
+		}
+	}
+#endif
+}
+
+void FAGXRenderCommandEncoderDebugging::Draw(mtlpp::PrimitiveType primitiveType, NSUInteger vertexStart, NSUInteger vertexCount, NSUInteger instanceCount, NSUInteger baseInstance)
+{
+#if METAL_DEBUG_OPTIONS
+	switch(((FAGXDebugRenderCommandEncoder*)m_ptr)->Buffer.GetPtr()->DebugLevel)
+	{
+		case EAGXDebugLevelConditionalSubmit:
+		case EAGXDebugLevelWaitForComplete:
+		case EAGXDebugLevelLogOperations:
+		{
+			((FAGXDebugRenderCommandEncoder*)m_ptr)->Buffer.Draw([NSString stringWithFormat:@"%s:%u,%u,%u,%u,%u", __PRETTY_FUNCTION__, (uint32)primitiveType, (uint32)vertexStart, (uint32)vertexCount, (uint32)instanceCount, (uint32)baseInstance]);
+		}
+		case EAGXDebugLevelValidation:
+		{
+			Validate();
+		}
+		default:
+		{
+			break;
+		}
+	}
+#endif
+}
+
+void FAGXRenderCommandEncoderDebugging::DrawIndexed(mtlpp::PrimitiveType primitiveType, NSUInteger indexCount, mtlpp::IndexType indexType, FAGXBuffer const& indexBuffer, NSUInteger indexBufferOffset, NSUInteger instanceCount, NSInteger baseVertex, NSUInteger baseInstance)
+{
+#if METAL_DEBUG_OPTIONS
+	switch(((FAGXDebugRenderCommandEncoder*)m_ptr)->Buffer.GetPtr()->DebugLevel)
+	{
+		case EAGXDebugLevelConditionalSubmit:
+		case EAGXDebugLevelWaitForComplete:
+		case EAGXDebugLevelLogOperations:
+		{
+			((FAGXDebugRenderCommandEncoder*)m_ptr)->Buffer.Draw([NSString stringWithFormat:@"%s:%u,%u,%u,%u,%u,%u,%u", __PRETTY_FUNCTION__, (uint32)primitiveType, (uint32)indexCount, (uint32)indexType, (uint32)indexBufferOffset, (uint32)instanceCount, (uint32)baseVertex, (uint32)baseInstance]);
+		}
+		case EAGXDebugLevelValidation:
+		{
+			Validate();
+		}
+		default:
+		{
+			break;
+		}
+	}
+#endif
+}
+
+void FAGXRenderCommandEncoderDebugging::Draw(mtlpp::PrimitiveType primitiveType, FAGXBuffer const& indirectBuffer, NSUInteger indirectBufferOffset)
+{
+#if METAL_DEBUG_OPTIONS
+	switch(((FAGXDebugRenderCommandEncoder*)m_ptr)->Buffer.GetPtr()->DebugLevel)
+	{
+		case EAGXDebugLevelConditionalSubmit:
+		case EAGXDebugLevelWaitForComplete:
+		case EAGXDebugLevelLogOperations:
+		{
+			((FAGXDebugRenderCommandEncoder*)m_ptr)->Buffer.Draw([NSString stringWithFormat:@"%s", __PRETTY_FUNCTION__]);
+		}
+		case EAGXDebugLevelValidation:
+		{
+			Validate();
+		}
+		default:
+		{
+			break;
+		}
+	}
+#endif
+}
+
+void FAGXRenderCommandEncoderDebugging::DrawIndexed(mtlpp::PrimitiveType primitiveType, mtlpp::IndexType indexType, FAGXBuffer const& indexBuffer, NSUInteger indexBufferOffset, FAGXBuffer const& indirectBuffer, NSUInteger indirectBufferOffset)
+{
+#if METAL_DEBUG_OPTIONS
+	switch(((FAGXDebugRenderCommandEncoder*)m_ptr)->Buffer.GetPtr()->DebugLevel)
+	{
+		case EAGXDebugLevelConditionalSubmit:
+		case EAGXDebugLevelWaitForComplete:
+		case EAGXDebugLevelLogOperations:
+		{
+			((FAGXDebugRenderCommandEncoder*)m_ptr)->Buffer.Draw([NSString stringWithFormat:@"%s", __PRETTY_FUNCTION__]);
+		}
+		case EAGXDebugLevelValidation:
+		{
+			Validate();
+		}
+		default:
+		{
+			break;
+		}
+	}
+#endif
+}
+
+/** Validates the pipeline/binding state */
+bool FAGXRenderCommandEncoderDebugging::ValidateFunctionBindings(EAGXShaderFrequency Frequency)
+{
+	bool bOK = true;
+#if METAL_DEBUG_OPTIONS
+	switch (((FAGXDebugRenderCommandEncoder*)m_ptr)->Buffer.GetPtr()->DebugLevel)
+	{
+		case EAGXDebugLevelConditionalSubmit:
+		case EAGXDebugLevelWaitForComplete:
+		case EAGXDebugLevelLogOperations:
+		case EAGXDebugLevelValidation:
+		{
+			check(((FAGXDebugRenderCommandEncoder*)m_ptr)->Pipeline);
+			
+			MTLRenderPipelineReflection* Reflection = ((FAGXDebugRenderCommandEncoder*)m_ptr)->Pipeline->RenderPipelineReflection;
+			check(Reflection);
+			
+			NSArray<MTLArgument*>* Arguments = nil;
+			switch(Frequency)
+			{
+				case EAGXShaderVertex:
+				{
+					Arguments = Reflection.vertexArguments;
+					break;
+				}
+				case EAGXShaderFragment:
+				{
+					Arguments = Reflection.fragmentArguments;
+					break;
+				}
+				default:
+					check(false);
+					break;
+			}
+			
+			for (uint32 i = 0; i < Arguments.count; i++)
+			{
+				MTLArgument* Arg = [Arguments objectAtIndex:i];
+				check(Arg);
+				switch(Arg.type)
+				{
+					case MTLArgumentTypeBuffer:
+					{
+						checkf(Arg.index < ML_MaxBuffers, TEXT("Metal buffer index exceeded!"));
+						if ((((FAGXDebugRenderCommandEncoder*)m_ptr)->ShaderBuffers[Frequency].Buffers[Arg.index] == nil && ((FAGXDebugRenderCommandEncoder*)m_ptr)->ShaderBuffers[Frequency].Bytes[Arg.index] == nil))
+						{
+							bOK = false;
+							UE_LOG(LogAGX, Warning, TEXT("Unbound buffer at Metal index %u which will crash the driver: %s"), (uint32)Arg.index, *FString([Arg description]));
+						}
+						break;
+					}
+					case MTLArgumentTypeThreadgroupMemory:
+					{
+						break;
+					}
+					case MTLArgumentTypeTexture:
+					{
+						checkf(Arg.index < ML_MaxTextures, TEXT("Metal texture index exceeded!"));
+						if (((FAGXDebugRenderCommandEncoder*)m_ptr)->ShaderTextures[Frequency].Textures[Arg.index] == nil)
+						{
+							bOK = false;
+							UE_LOG(LogAGX, Warning, TEXT("Unbound texture at Metal index %u which will crash the driver: %s"), (uint32)Arg.index, *FString([Arg description]));
+						}
+						else if (((FAGXDebugRenderCommandEncoder*)m_ptr)->ShaderTextures[Frequency].Textures[Arg.index].textureType != Arg.textureType)
+						{
+							bOK = false;
+							UE_LOG(LogAGX, Warning, TEXT("Incorrect texture type bound at Metal index %u which will crash the driver: %s\n%s"), (uint32)Arg.index, *FString([Arg description]), *FString([((FAGXDebugRenderCommandEncoder*)m_ptr)->ShaderTextures[Frequency].Textures[Arg.index] description]));
+						}
+						break;
+					}
+					case MTLArgumentTypeSampler:
+					{
+						checkf(Arg.index < ML_MaxSamplers, TEXT("Metal sampler index exceeded!"));
+						if (((FAGXDebugRenderCommandEncoder*)m_ptr)->ShaderSamplers[Frequency].Samplers[Arg.index] == nil)
+						{
+							bOK = false;
+							UE_LOG(LogAGX, Warning, TEXT("Unbound sampler at Metal index %u which will crash the driver: %s"), (uint32)Arg.index, *FString([Arg description]));
+						}
+						break;
+					}
+					default:
+						check(false);
+						break;
+				}
+			}
+			break;
+		}
+		default:
+		{
+			break;
+		}
+	}
+#endif
+	return bOK;
+}
+void FAGXRenderCommandEncoderDebugging::Validate()
+{
+#if METAL_DEBUG_OPTIONS
+	bool bOK = ValidateFunctionBindings(EAGXShaderVertex);
+	if (!bOK)
+	{
+		UE_LOG(LogAGX, Error, TEXT("Metal Validation failures for vertex shader:\n%s"), ((FAGXDebugRenderCommandEncoder*)m_ptr)->Pipeline && ((FAGXDebugRenderCommandEncoder*)m_ptr)->Pipeline->VertexSource ? *FString(((FAGXDebugRenderCommandEncoder*)m_ptr)->Pipeline->VertexSource) : TEXT("nil"));
+	}
+	
+	bOK = ValidateFunctionBindings(EAGXShaderFragment);
+	if (!bOK)
+	{
+		UE_LOG(LogAGX, Error, TEXT("Metal Validation failures for fragment shader:\n%s"), ((FAGXDebugRenderCommandEncoder*)m_ptr)->Pipeline && ((FAGXDebugRenderCommandEncoder*)m_ptr)->Pipeline->FragmentSource ? *FString(((FAGXDebugRenderCommandEncoder*)m_ptr)->Pipeline->FragmentSource) : TEXT("nil"));
+	}
+#endif
+}
+
+void FAGXRenderCommandEncoderDebugging::InsertDebugSignpost(ns::String const& Label)
+{
+	((FAGXDebugRenderCommandEncoder*)m_ptr)->Buffer.InsertDebugSignpost(Label);
+}
+void FAGXRenderCommandEncoderDebugging::PushDebugGroup(ns::String const& Group)
+{
+    ((FAGXDebugRenderCommandEncoder*)m_ptr)->Buffer.PushDebugGroup(Group);
+}
+void FAGXRenderCommandEncoderDebugging::PopDebugGroup()
+{
+    ((FAGXDebugRenderCommandEncoder*)m_ptr)->Buffer.PopDebugGroup();
+#if METAL_DEBUG_OPTIONS
+	InsertDebugDraw();
+#endif
+}
+
+void FAGXRenderCommandEncoderDebugging::EndEncoder()
+{
+	((FAGXDebugRenderCommandEncoder*)m_ptr)->Buffer.EndCommandEncoder();
+}
+
+FAGXParallelRenderCommandEncoderDebugging::FAGXParallelRenderCommandEncoderDebugging()
+{
+
+}
+
+FAGXParallelRenderCommandEncoderDebugging::FAGXParallelRenderCommandEncoderDebugging(mtlpp::ParallelRenderCommandEncoder& Encoder, mtlpp::RenderPassDescriptor const& Desc, FAGXCommandBufferDebugging& Buffer)
+: ns::Object<FAGXDebugParallelRenderCommandEncoder*>([[FAGXDebugParallelRenderCommandEncoder alloc] initWithEncoder:Encoder.GetPtr() fromDescriptor:Desc.GetPtr() andCommandBuffer:Buffer], ns::Ownership::Assign)
+{
+	Buffer.BeginRenderCommandEncoder([NSString stringWithFormat:@"ParallelRender: %@", Encoder.GetLabel().GetPtr()], Desc);
+	Encoder.SetAssociatedObject((void const*)&FAGXParallelRenderCommandEncoderDebugging::Get, *this);
+}
+
+FAGXParallelRenderCommandEncoderDebugging::FAGXParallelRenderCommandEncoderDebugging(FAGXDebugParallelRenderCommandEncoder* handle)
+: ns::Object<FAGXDebugParallelRenderCommandEncoder*>(handle)
+{
+
+}
+
+FAGXDebugParallelRenderCommandEncoder* FAGXParallelRenderCommandEncoderDebugging::Get(mtlpp::ParallelRenderCommandEncoder& Buffer)
+{
+	return Buffer.GetAssociatedObject<FAGXParallelRenderCommandEncoderDebugging>((void const*)&FAGXParallelRenderCommandEncoderDebugging::Get);
+}
+
+FAGXRenderCommandEncoderDebugging FAGXParallelRenderCommandEncoderDebugging::GetRenderCommandEncoderDebugger(mtlpp::RenderCommandEncoder& Encoder)
+{
+	mtlpp::RenderPassDescriptor Desc(m_ptr->RenderPassDesc);
+	FAGXCommandBufferDebugging IndirectBuffer([[FAGXDebugCommandBuffer alloc] initWithCommandBuffer:m_ptr->Buffer.GetPtr()->InnerBuffer]);
+	FAGXRenderCommandEncoderDebugging EncoderDebugging(Encoder, Desc, IndirectBuffer);
+	m_ptr->RenderEncoders.Add(EncoderDebugging);
+	return EncoderDebugging;
+}
+
+void FAGXParallelRenderCommandEncoderDebugging::InsertDebugSignpost(ns::String const& Label)
+{
+	m_ptr->Buffer.InsertDebugSignpost(Label);
+}
+
+void FAGXParallelRenderCommandEncoderDebugging::PushDebugGroup(ns::String const& Group)
+{
+    m_ptr->Buffer.PushDebugGroup(Group);
+}
+
+void FAGXParallelRenderCommandEncoderDebugging::PopDebugGroup()
+{
+    m_ptr->Buffer.PopDebugGroup();
+}
+
+void FAGXParallelRenderCommandEncoderDebugging::EndEncoder()
+{
+	FAGXDebugCommandBuffer* CommandBuffer = m_ptr->Buffer.GetPtr();
+	for (FAGXRenderCommandEncoderDebugging& EncoderDebug : m_ptr->RenderEncoders)
+	{
+		FAGXDebugCommandBuffer* CmdBuffer = ((FAGXDebugRenderCommandEncoder*)EncoderDebug.GetPtr())->Buffer.GetPtr();
+		[CommandBuffer->DebugGroup addObjectsFromArray:CmdBuffer->DebugGroup];
+		CommandBuffer->DebugCommands.Append(CmdBuffer->DebugCommands);
+	}
+	m_ptr->Buffer.EndCommandEncoder();
+}
+
+#endif
+
diff --git a/Engine/Source/Runtime/Apple/AGXRHI/Private/AGXRenderCommandEncoder.h b/Engine/Source/Runtime/Apple/AGXRHI/Private/AGXRenderCommandEncoder.h
new file mode 100644
--- /dev/null
+++ b/Engine/Source/Runtime/Apple/AGXRHI/Private/AGXRenderCommandEncoder.h
@@ -0,0 +1,80 @@
+// Copyright Epic Games, Inc. All Rights Reserved.
+
+#pragma once
+
+#include "AGXDebugCommandEncoder.h"
+
+#if MTLPP_CONFIG_VALIDATE && METAL_DEBUG_OPTIONS
+NS_ASSUME_NONNULL_BEGIN
+
+@class FAGXShaderPipeline;
+@class FAGXDebugRenderCommandEncoder;
+@class FAGXDebugParallelRenderCommandEncoder;
+
+class FAGXRenderCommandEncoderDebugging : public FAGXCommandEncoderDebugging
+{
+	void InsertDebugDraw();
+public:
+	FAGXRenderCommandEncoderDebugging();
+	FAGXRenderCommandEncoderDebugging(mtlpp::RenderCommandEncoder& Encoder, mtlpp::RenderPassDescriptor const& Desc, FAGXCommandBufferDebugging& Buffer);
+	FAGXRenderCommandEncoderDebugging(FAGXDebugCommandEncoder* handle);
+	
+	static FAGXRenderCommandEncoderDebugging Get(mtlpp::RenderCommandEncoder& Buffer);
+	
+	void SetPipeline(FAGXShaderPipeline* Pipeline);
+	
+	void SetBytes(EAGXShaderFrequency Freq, const void * bytes, NSUInteger length, NSUInteger index);
+	void SetBuffer(EAGXShaderFrequency Freq,  FAGXBuffer const& buffer, NSUInteger offset, NSUInteger index);
+	void SetBufferOffset(EAGXShaderFrequency Freq, NSUInteger offset, NSUInteger index);
+	void SetTexture(EAGXShaderFrequency Freq,  FAGXTexture const& texture, NSUInteger index);
+	void SetSamplerState(EAGXShaderFrequency Freq,  mtlpp::SamplerState const& sampler, NSUInteger index);
+	void SetSamplerState(EAGXShaderFrequency Freq,  mtlpp::SamplerState const& sampler, float lodMinClamp, float lodMaxClamp, NSUInteger index);
+	void SetDepthStencilState( mtlpp::DepthStencilState const& depthStencilState);
+	
+	void Draw(mtlpp::PrimitiveType primitiveType, NSUInteger vertexStart, NSUInteger vertexCount, NSUInteger instanceCount);
+	
+	void Draw(mtlpp::PrimitiveType primitiveType, NSUInteger vertexStart, NSUInteger vertexCount);
+	
+	void DrawIndexed(mtlpp::PrimitiveType primitiveType, NSUInteger indexCount, mtlpp::IndexType indexType, FAGXBuffer const& indexBuffer, NSUInteger indexBufferOffset, NSUInteger instanceCount);
+	
+	void DrawIndexed(mtlpp::PrimitiveType primitiveType, NSUInteger indexCount, mtlpp::IndexType indexType, FAGXBuffer const& indexBuffer, NSUInteger indexBufferOffset);
+	
+	void Draw(mtlpp::PrimitiveType primitiveType, NSUInteger vertexStart, NSUInteger vertexCount, NSUInteger instanceCount, NSUInteger baseInstance);
+	
+	void DrawIndexed(mtlpp::PrimitiveType primitiveType, NSUInteger indexCount, mtlpp::IndexType indexType, FAGXBuffer const& indexBuffer, NSUInteger indexBufferOffset, NSUInteger instanceCount, NSInteger baseVertex, NSUInteger baseInstance);
+	
+	void Draw(mtlpp::PrimitiveType primitiveType, FAGXBuffer const& indirectBuffer, NSUInteger indirectBufferOffset);
+	
+	void DrawIndexed(mtlpp::PrimitiveType primitiveType, mtlpp::IndexType indexType, FAGXBuffer const& indexBuffer, NSUInteger indexBufferOffset, FAGXBuffer const& indirectBuffer, NSUInteger indirectBufferOffset);
+	
+	/** Validates the pipeline/binding state */
+	bool ValidateFunctionBindings(EAGXShaderFrequency Frequency);
+	void Validate();
+	
+	void InsertDebugSignpost(ns::String const& Label);
+	void PushDebugGroup(ns::String const& Group);
+	void PopDebugGroup();
+	
+	void EndEncoder();
+};
+
+class FAGXParallelRenderCommandEncoderDebugging : public ns::Object<FAGXDebugParallelRenderCommandEncoder*>
+{
+public:
+	FAGXParallelRenderCommandEncoderDebugging();
+	FAGXParallelRenderCommandEncoderDebugging(mtlpp::ParallelRenderCommandEncoder& Encoder, mtlpp::RenderPassDescriptor const& Desc, FAGXCommandBufferDebugging& Buffer);
+	FAGXParallelRenderCommandEncoderDebugging(FAGXDebugParallelRenderCommandEncoder* handle);
+
+	static FAGXDebugParallelRenderCommandEncoder* Get(mtlpp::ParallelRenderCommandEncoder& Buffer);
+
+	FAGXRenderCommandEncoderDebugging GetRenderCommandEncoderDebugger(mtlpp::RenderCommandEncoder& Encoder);
+
+	void InsertDebugSignpost(ns::String const& Label);
+	void PushDebugGroup(ns::String const& Group);
+	void PopDebugGroup();
+	
+	void EndEncoder();
+};
+
+NS_ASSUME_NONNULL_END
+#endif
diff --git a/Engine/Source/Runtime/Apple/AGXRHI/Private/AGXRenderPass.cpp b/Engine/Source/Runtime/Apple/AGXRHI/Private/AGXRenderPass.cpp
--- a/Engine/Source/Runtime/Apple/AGXRHI/Private/AGXRenderPass.cpp
+++ b/Engine/Source/Runtime/Apple/AGXRHI/Private/AGXRenderPass.cpp
@@ -10,6 +10,7 @@
 #include "AGXGraphicsPipelineState.h"
 #include "AGXVertexDeclaration.h"
 #include "AGXRenderPass.h"
+#include "AGXCommandBuffer.h"
 #include "AGXProfiler.h"
 #include "AGXFrameAllocator.h"
 
@@ -52,6 +53,7 @@ FAGXRenderPass::~FAGXRenderPass(void)
 {
 	check(!CurrentEncoder.GetCommandBuffer());
 	check(!PrologueEncoder.GetCommandBuffer());
+	check(!PassStartFence);
 }
 
 void FAGXRenderPass::SetDispatchType(mtlpp::DispatchType Type)
@@ -59,8 +61,27 @@ void FAGXRenderPass::SetDispatchType(mtlpp::DispatchType Type)
 	ComputeDispatchType = Type;
 }
 
-void FAGXRenderPass::Begin(bool bParallelBegin)
+void FAGXRenderPass::Begin(FAGXFence* Fence, bool const bParallelBegin)
 {
+	if (!bParallelBegin || !FAGXCommandQueue::SupportsFeature(EAGXFeaturesParallelRenderEncoders))
+	{
+		check(!PassStartFence || Fence == nullptr);
+		if (Fence)
+		{
+			PassStartFence = Fence;
+			PrologueStartEncoderFence = Fence;
+		}
+	}
+	else
+	{
+		check(!ParallelPassEndFence || Fence == nullptr);
+		if (Fence)
+		{
+			ParallelPassEndFence = Fence;
+			PrologueStartEncoderFence = Fence;
+		}
+	}
+	
 	if (!CmdList.IsParallel() && !CurrentEncoder.GetCommandBuffer())
 	{
 		CurrentEncoder.StartCommandBuffer();
@@ -68,14 +89,58 @@ void FAGXRenderPass::Begin(bool bParallelBegin)
 	}
 }
 
-void FAGXRenderPass::Submit(EAGXSubmitFlags Flags)
+void FAGXRenderPass::Wait(FAGXFence* Fence)
+{
+	if (Fence)
+	{
+		if (PrologueEncoder.IsBlitCommandEncoderActive() || PrologueEncoder.IsComputeCommandEncoderActive())
+		{
+			PrologueEncoder.WaitForFence(Fence);
+			METAL_DEBUG_LAYER(EAGXDebugLevelValidation, FAGXFence::ValidateUsage(Fence));
+		}
+		else if (CurrentEncoder.IsRenderCommandEncoderActive() || CurrentEncoder.IsBlitCommandEncoderActive() || CurrentEncoder.IsComputeCommandEncoderActive())
+		{
+			CurrentEncoder.WaitForFence(Fence);
+			METAL_DEBUG_LAYER(EAGXDebugLevelValidation, FAGXFence::ValidateUsage(Fence));
+		}
+        else
+        {
+            PassStartFence = Fence;
+			PrologueStartEncoderFence = Fence;
+        }
+	}
+		
+}
+
+void FAGXRenderPass::Update(FAGXFence* Fence)
+{
+	if (Fence)
+	{
+		// Force an encoder - possibly consuming the start fence so that we get the proper order
+		// the higher-level can generate empty contexts but we have no sane way to deal with that.
+		if (!CurrentEncoder.IsRenderCommandEncoderActive() && !CurrentEncoder.IsBlitCommandEncoderActive() && !CurrentEncoder.IsComputeCommandEncoderActive())
+		{
+			ConditionalSwitchToBlit();
+		}
+		CurrentEncoder.UpdateFence(Fence);
+		State.FlushVisibilityResults(CurrentEncoder);
+		TRefCountPtr<FAGXFence> NewFence = CurrentEncoder.EndEncoding();
+		check(!CurrentEncoderFence || !NewFence);
+		if (NewFence)
+		{
+			CurrentEncoderFence = NewFence;
+		}
+	}
+}
+
+TRefCountPtr<FAGXFence> const& FAGXRenderPass::Submit(EAGXSubmitFlags Flags)
 {
 	if (CurrentEncoder.GetCommandBuffer() || (Flags & EAGXSubmitFlagsAsyncCommandBuffer))
 	{
 		if (PrologueEncoder.IsBlitCommandEncoderActive() || PrologueEncoder.IsComputeCommandEncoderActive())
 		{
 			check(PrologueEncoder.GetCommandBuffer());
-			PrologueEncoder.EndEncoding();
+			PrologueEncoderFence = PrologueEncoder.EndEncoding();
 		}
 		if (PrologueEncoder.GetCommandBuffer())
 		{
@@ -110,7 +175,7 @@ void FAGXRenderPass::Submit(EAGXSubmitFlags Flags)
                 State.SetRenderStoreActions(CurrentEncoder, (Flags & EAGXSubmitFlagsBreakCommandBuffer));
 				State.FlushVisibilityResults(CurrentEncoder);
             }
-            CurrentEncoder.EndEncoding();
+            CurrentEncoderFence = CurrentEncoder.EndEncoding();
         }
 		
         CurrentEncoder.CommitCommandBuffer(Flags);
@@ -122,6 +187,8 @@ void FAGXRenderPass::Submit(EAGXSubmitFlags Flags)
 		PrologueEncoder.Reset();
 		CurrentEncoder.Reset();
 	}
+	
+	return CurrentEncoderFence;
 }
 
 void FAGXRenderPass::BeginParallelRenderPass(mtlpp::RenderPassDescriptor RenderPass, uint32 NumParallelContextsInPass)
@@ -135,12 +202,12 @@ void FAGXRenderPass::BeginParallelRenderPass(mtlpp::RenderPassDescriptor RenderP
 	{
 		if (PrologueEncoder.IsBlitCommandEncoderActive() || PrologueEncoder.IsComputeCommandEncoderActive())
 		{
-			PrologueEncoder.EndEncoding();
+			PrologueEncoderFence = PrologueEncoder.EndEncoding();
 		}
 		if (CurrentEncoder.IsRenderCommandEncoderActive() || CurrentEncoder.IsBlitCommandEncoderActive() || CurrentEncoder.IsComputeCommandEncoderActive())
 		{
 			State.FlushVisibilityResults(CurrentEncoder);
-			CurrentEncoder.EndEncoding();
+			CurrentEncoderFence = CurrentEncoder.EndEncoding();
 		}
 
 		CurrentEncoder.SetRenderPassDescriptor(RenderPass);
@@ -168,12 +235,12 @@ void FAGXRenderPass::BeginRenderPass(mtlpp::RenderPassDescriptor RenderPass)
 	// EndEncoding should provide the encoder fence...
 	if (PrologueEncoder.IsBlitCommandEncoderActive() || PrologueEncoder.IsComputeCommandEncoderActive())
 	{
-		PrologueEncoder.EndEncoding();
+		PrologueEncoderFence = PrologueEncoder.EndEncoding();
 	}
 	if (CurrentEncoder.IsRenderCommandEncoderActive() || CurrentEncoder.IsBlitCommandEncoderActive() || CurrentEncoder.IsComputeCommandEncoderActive())
 	{
 		State.FlushVisibilityResults(CurrentEncoder);
-		CurrentEncoder.EndEncoding();
+		CurrentEncoderFence = CurrentEncoder.EndEncoding();
 	}
 	State.SetStateDirty();
 	State.SetRenderTargetsActive(true);
@@ -185,6 +252,30 @@ void FAGXRenderPass::BeginRenderPass(mtlpp::RenderPassDescriptor RenderPass)
 	if (!GAGXDeferRenderPasses || !State.CanRestartRenderPass() || CmdList.IsParallel())
 	{
 		CurrentEncoder.BeginRenderCommandEncoding();
+		if (PassStartFence)
+		{
+			CurrentEncoder.WaitForFence(PassStartFence);
+			PassStartFence = nullptr;
+		}
+		if (ParallelPassEndFence)
+		{
+			CurrentEncoder.WaitForFence(ParallelPassEndFence);
+			ParallelPassEndFence = nullptr;
+		}
+		if (CurrentEncoderFence)
+		{
+			CurrentEncoder.WaitForFence(CurrentEncoderFence);
+			CurrentEncoderFence = nullptr;
+		}
+		if (PrologueEncoderFence)
+		{
+			// Consume on the current encoder but do not invalidate
+			CurrentEncoder.WaitForFence(PrologueEncoderFence);
+		}
+		if (PrologueEncoder.IsBlitCommandEncoderActive() || PrologueEncoder.IsComputeCommandEncoderActive())
+		{
+			CurrentEncoder.WaitForFence(PrologueEncoder.GetEncoderFence());
+		}
 		State.SetRenderStoreActions(CurrentEncoder, false);
 		check(CurrentEncoder.IsRenderCommandEncoderActive());
 	}
@@ -275,13 +366,37 @@ void FAGXRenderPass::RestartRenderPass(mtlpp::RenderPassDescriptor RenderPass)
 			State.SetRenderStoreActions(CurrentEncoder, true);
 			State.FlushVisibilityResults(CurrentEncoder);
 		}
-		CurrentEncoder.EndEncoding();
+		CurrentEncoderFence = CurrentEncoder.EndEncoding();
 	}
 	State.SetStateDirty();
 	State.SetRenderTargetsActive(true);
 	
 	CurrentEncoder.SetRenderPassDescriptor(RenderPassDesc);
 	CurrentEncoder.BeginRenderCommandEncoding();
+	if (PassStartFence)
+	{
+		CurrentEncoder.WaitForFence(PassStartFence);
+		PassStartFence = nullptr;
+	}
+	if (ParallelPassEndFence)
+	{
+		CurrentEncoder.WaitForFence(ParallelPassEndFence);
+		ParallelPassEndFence = nullptr;
+	}
+	if (CurrentEncoderFence)
+	{
+		CurrentEncoder.WaitForFence(CurrentEncoderFence);
+		CurrentEncoderFence = nullptr;
+	}
+	if (PrologueEncoderFence)
+	{
+		// Consume on the current encoder but do not invalidate
+		CurrentEncoder.WaitForFence(PrologueEncoderFence);
+	}
+	if (PrologueEncoder.IsBlitCommandEncoderActive() || PrologueEncoder.IsComputeCommandEncoderActive())
+	{
+		CurrentEncoder.WaitForFence(PrologueEncoder.GetEncoderFence());
+	}
 	State.SetRenderStoreActions(CurrentEncoder, false);
 	
 	check(CurrentEncoder.IsRenderCommandEncoderActive());
@@ -303,7 +418,20 @@ void FAGXRenderPass::DrawPrimitive(uint32 PrimitiveType, uint32 BaseVertexIndex,
 	
 	METAL_GPUPROFILE(FAGXProfiler::GetProfiler()->EncodeDraw(CurrentEncoder.GetCommandBufferStats(), __FUNCTION__, NumPrimitives, NumVertices, NumInstances));
 	CurrentEncoder.GetRenderCommandEncoder().Draw(AGXTranslatePrimitiveType(PrimitiveType), BaseVertexIndex, NumVertices, NumInstances);
+	METAL_DEBUG_LAYER(EAGXDebugLevelFastValidation, CurrentEncoder.GetRenderCommandEncoderDebugging().Draw(AGXTranslatePrimitiveType(PrimitiveType), BaseVertexIndex, NumVertices, NumInstances));
 
+	if (GAGXCommandBufferDebuggingEnabled)
+	{
+		FAGXCommandData Data;
+		Data.CommandType = FAGXCommandData::Type::DrawPrimitive;
+		Data.Draw.BaseInstance = 0;
+		Data.Draw.InstanceCount = NumInstances;
+		Data.Draw.VertexCount = NumVertices;
+		Data.Draw.VertexStart = BaseVertexIndex;
+		
+		InsertDebugDraw(Data);
+	}
+	
 	ConditionalSubmit();	
 }
 
@@ -322,6 +450,15 @@ void FAGXRenderPass::DrawPrimitiveIndirect(uint32 PrimitiveType, FAGXVertexBuffe
 		
 		METAL_GPUPROFILE(FAGXProfiler::GetProfiler()->EncodeDraw(CurrentEncoder.GetCommandBufferStats(), __FUNCTION__, 1, 1, 1));
 		CurrentEncoder.GetRenderCommandEncoder().Draw(AGXTranslatePrimitiveType(PrimitiveType), TheBackingBuffer, ArgumentOffset);
+		METAL_DEBUG_LAYER(EAGXDebugLevelFastValidation, CurrentEncoder.GetRenderCommandEncoderDebugging().Draw(AGXTranslatePrimitiveType(PrimitiveType), TheBackingBuffer, ArgumentOffset));
+
+		if (GAGXCommandBufferDebuggingEnabled)
+		{
+			FAGXCommandData Data;
+			Data.CommandType = FAGXCommandData::Type::DrawPrimitiveIndirect;
+
+			InsertDebugDraw(Data);
+		}
 
 		ConditionalSubmit();
 	}
@@ -393,10 +530,25 @@ void FAGXRenderPass::DrawIndexedPrimitive(FAGXBuffer const& IndexBuffer, uint32
 	if (GRHISupportsBaseVertexIndex && GRHISupportsFirstInstance)
 	{
 		CurrentEncoder.GetRenderCommandEncoder().DrawIndexed(AGXTranslatePrimitiveType(PrimitiveType), NumIndices, ((IndexStride == 2) ? mtlpp::IndexType::UInt16 : mtlpp::IndexType::UInt32), IndexBuffer, StartIndex * IndexStride, NumInstances, BaseVertexIndex, FirstInstance);
+		METAL_DEBUG_LAYER(EAGXDebugLevelFastValidation, CurrentEncoder.GetRenderCommandEncoderDebugging().DrawIndexed(AGXTranslatePrimitiveType(PrimitiveType), NumIndices, ((IndexStride == 2) ? mtlpp::IndexType::UInt16 : mtlpp::IndexType::UInt32), IndexBuffer, StartIndex * IndexStride, NumInstances, BaseVertexIndex, FirstInstance));
 	}
 	else
 	{
 		CurrentEncoder.GetRenderCommandEncoder().DrawIndexed(AGXTranslatePrimitiveType(PrimitiveType), NumIndices, ((IndexStride == 2) ? mtlpp::IndexType::UInt16 : mtlpp::IndexType::UInt32), IndexBuffer, StartIndex * IndexStride, NumInstances);
+		METAL_DEBUG_LAYER(EAGXDebugLevelFastValidation, CurrentEncoder.GetRenderCommandEncoderDebugging().DrawIndexed(AGXTranslatePrimitiveType(PrimitiveType), NumIndices, ((IndexStride == 2) ? mtlpp::IndexType::UInt16 : mtlpp::IndexType::UInt32), IndexBuffer, StartIndex * IndexStride, NumInstances));
+	}
+	
+	if (GAGXCommandBufferDebuggingEnabled)
+	{
+		FAGXCommandData Data;
+		Data.CommandType = FAGXCommandData::Type::DrawPrimitiveIndexed;
+		Data.DrawIndexed.BaseInstance = FirstInstance;
+		Data.DrawIndexed.BaseVertex = BaseVertexIndex;
+		Data.DrawIndexed.IndexCount = NumIndices;
+		Data.DrawIndexed.IndexStart = StartIndex;
+		Data.DrawIndexed.InstanceCount = NumInstances;
+		
+		InsertDebugDraw(Data);
 	}
 	
 	ConditionalSubmit();
@@ -423,7 +575,15 @@ void FAGXRenderPass::DrawIndexedIndirect(FAGXIndexBuffer* IndexBuffer, uint32 Pr
 		
 		METAL_GPUPROFILE(FAGXProfiler::GetProfiler()->EncodeDraw(CurrentEncoder.GetCommandBufferStats(), __FUNCTION__, 1, 1, 1));
 		CurrentEncoder.GetRenderCommandEncoder().DrawIndexed(AGXTranslatePrimitiveType(PrimitiveType), (mtlpp::IndexType)IndexBuffer->IndexType, TheBackingIndexBuffer, 0, TheBackingBuffer, (DrawArgumentsIndex * 5 * sizeof(uint32)));
+		METAL_DEBUG_LAYER(EAGXDebugLevelFastValidation, CurrentEncoder.GetRenderCommandEncoderDebugging().DrawIndexed(AGXTranslatePrimitiveType(PrimitiveType), (mtlpp::IndexType)IndexBuffer->IndexType, TheBackingIndexBuffer, 0, TheBackingBuffer, (DrawArgumentsIndex * 5 * sizeof(uint32))));
 
+		if (GAGXCommandBufferDebuggingEnabled)
+		{
+			FAGXCommandData Data;
+			Data.CommandType = FAGXCommandData::Type::DrawPrimitiveIndexedIndirect;
+			
+			InsertDebugDraw(Data);
+		}
 		ConditionalSubmit();
 	}
 	else
@@ -450,7 +610,16 @@ void FAGXRenderPass::DrawIndexedPrimitiveIndirect(uint32 PrimitiveType,FAGXIndex
 		
 		METAL_GPUPROFILE(FAGXProfiler::GetProfiler()->EncodeDraw(CurrentEncoder.GetCommandBufferStats(), __FUNCTION__, 1, 1, 1));
 		CurrentEncoder.GetRenderCommandEncoder().DrawIndexed(AGXTranslatePrimitiveType(PrimitiveType), (mtlpp::IndexType)IndexBuffer->IndexType, TheBackingIndexBuffer, 0, TheBackingBuffer, ArgumentOffset);
+		METAL_DEBUG_LAYER(EAGXDebugLevelFastValidation, CurrentEncoder.GetRenderCommandEncoderDebugging().DrawIndexed(AGXTranslatePrimitiveType(PrimitiveType), (mtlpp::IndexType)IndexBuffer->IndexType, TheBackingIndexBuffer, 0, TheBackingBuffer, ArgumentOffset));
 
+		if (GAGXCommandBufferDebuggingEnabled)
+		{
+			FAGXCommandData Data;
+			Data.CommandType = FAGXCommandData::Type::DrawPrimitiveIndirect;
+			
+			InsertDebugDraw(Data);
+		}
+		
 		ConditionalSubmit();
 	}
 	else
@@ -478,6 +647,7 @@ void FAGXRenderPass::Dispatch(uint32 ThreadGroupCountX, uint32 ThreadGroupCountY
 		check(ComputeShader->NumThreadsX > 0 && ComputeShader->NumThreadsY > 0 && ComputeShader->NumThreadsZ > 0);
 		mtlpp::Size Threadgroups = mtlpp::Size(ThreadGroupCountX, ThreadGroupCountY, ThreadGroupCountZ);
 		PrologueEncoder.GetComputeCommandEncoder().DispatchThreadgroups(Threadgroups, ThreadgroupCounts);
+		METAL_DEBUG_LAYER(EAGXDebugLevelFastValidation, PrologueEncoder.GetComputeCommandEncoderDebugging().DispatchThreadgroups(Threadgroups, ThreadgroupCounts));
 		
 		ConditionalSubmit();
 	}
@@ -498,6 +668,18 @@ void FAGXRenderPass::Dispatch(uint32 ThreadGroupCountX, uint32 ThreadGroupCountY
 		check(ComputeShader->NumThreadsX > 0 && ComputeShader->NumThreadsY > 0 && ComputeShader->NumThreadsZ > 0);
 		mtlpp::Size Threadgroups = mtlpp::Size(ThreadGroupCountX, ThreadGroupCountY, ThreadGroupCountZ);
 		CurrentEncoder.GetComputeCommandEncoder().DispatchThreadgroups(Threadgroups, ThreadgroupCounts);
+		METAL_DEBUG_LAYER(EAGXDebugLevelFastValidation, CurrentEncoder.GetComputeCommandEncoderDebugging().DispatchThreadgroups(Threadgroups, ThreadgroupCounts));
+		
+		if (GAGXCommandBufferDebuggingEnabled)
+		{
+			FAGXCommandData Data;
+			Data.CommandType = FAGXCommandData::Type::Dispatch;
+			Data.Dispatch.threadgroupsPerGrid[0] = ThreadGroupCountX;
+			Data.Dispatch.threadgroupsPerGrid[1] = ThreadGroupCountY;
+			Data.Dispatch.threadgroupsPerGrid[2] = ThreadGroupCountZ;
+			
+			InsertDebugDispatch(Data);
+		}
 		
 		ConditionalSubmit();
 	}
@@ -524,6 +706,7 @@ void FAGXRenderPass::DispatchIndirect(FAGXVertexBuffer* ArgumentBuffer, uint32 A
 		check(ComputeShader->NumThreadsX > 0 && ComputeShader->NumThreadsY > 0 && ComputeShader->NumThreadsZ > 0);
 		
 		PrologueEncoder.GetComputeCommandEncoder().DispatchThreadgroupsWithIndirectBuffer(ArgumentBuffer->GetCurrentBuffer(), ArgumentOffset, ThreadgroupCounts);
+		METAL_DEBUG_LAYER(EAGXDebugLevelFastValidation, PrologueEncoder.GetComputeCommandEncoderDebugging().DispatchThreadgroupsWithIndirectBuffer(ArgumentBuffer->GetCurrentBuffer(), ArgumentOffset, ThreadgroupCounts));
 		
 		ConditionalSubmit();
 	}
@@ -543,12 +726,23 @@ void FAGXRenderPass::DispatchIndirect(FAGXVertexBuffer* ArgumentBuffer, uint32 A
 		check(ComputeShader->NumThreadsX > 0 && ComputeShader->NumThreadsY > 0 && ComputeShader->NumThreadsZ > 0);
 
 		CurrentEncoder.GetComputeCommandEncoder().DispatchThreadgroupsWithIndirectBuffer(ArgumentBuffer->GetCurrentBuffer(), ArgumentOffset, ThreadgroupCounts);
+		METAL_DEBUG_LAYER(EAGXDebugLevelFastValidation, CurrentEncoder.GetComputeCommandEncoderDebugging().DispatchThreadgroupsWithIndirectBuffer(ArgumentBuffer->GetCurrentBuffer(), ArgumentOffset, ThreadgroupCounts));
 
+		if (GAGXCommandBufferDebuggingEnabled)
+		{
+			FAGXCommandData Data;
+			Data.CommandType = FAGXCommandData::Type::DispatchIndirect;
+			Data.DispatchIndirect.ArgumentBuffer = ArgumentBuffer->GetCurrentBuffer();
+			Data.DispatchIndirect.ArgumentOffset = ArgumentOffset;
+			
+			InsertDebugDispatch(Data);
+		}
+	
 		ConditionalSubmit();
 	}
 }
 
-void FAGXRenderPass::EndRenderPass(void)
+TRefCountPtr<FAGXFence> const& FAGXRenderPass::EndRenderPass(void)
 {
 	if (bWithinRenderPass)
 	{
@@ -558,6 +752,7 @@ void FAGXRenderPass::EndRenderPass(void)
 		// This just calls End - it exists only to enforce assumptions
 		End();
 	}
+	return CurrentEncoderFence;
 }
 
 void FAGXRenderPass::InsertTextureBarrier()
@@ -580,6 +775,7 @@ void FAGXRenderPass::CopyFromTextureToBuffer(FAGXTexture const& Texture, uint32
 	METAL_GPUPROFILE(FAGXProfiler::GetProfiler()->EncodeBlit(CurrentEncoder.GetCommandBufferStats(), __FUNCTION__));
 	{
 		MTLPP_VALIDATE(mtlpp::BlitCommandEncoder, Encoder, AGXSafeGetRuntimeDebuggingLevel() >= EAGXDebugLevelValidation, Copy(Texture, sourceSlice, sourceLevel, sourceOrigin, sourceSize, toBuffer, destinationOffset, destinationBytesPerRow, destinationBytesPerImage, options));
+		METAL_DEBUG_LAYER(EAGXDebugLevelFastValidation, CurrentEncoder.GetBlitCommandEncoderDebugging().Copy(Texture, sourceSlice, sourceLevel, sourceOrigin, sourceSize, toBuffer, destinationOffset, destinationBytesPerRow, destinationBytesPerImage, options));
 	}
 	ConditionalSubmit();
 }
@@ -595,10 +791,12 @@ void FAGXRenderPass::CopyFromBufferToTexture(FAGXBuffer const& Buffer, uint32 so
 	if (options == mtlpp::BlitOption::None)
 	{
 		MTLPP_VALIDATE(mtlpp::BlitCommandEncoder, Encoder, AGXSafeGetRuntimeDebuggingLevel() >= EAGXDebugLevelValidation, Copy(Buffer, sourceOffset, sourceBytesPerRow, sourceBytesPerImage, sourceSize, toTexture, destinationSlice, destinationLevel, destinationOrigin));
+		METAL_DEBUG_LAYER(EAGXDebugLevelFastValidation, CurrentEncoder.GetBlitCommandEncoderDebugging().Copy(Buffer, sourceOffset, sourceBytesPerRow, sourceBytesPerImage, sourceSize, toTexture, destinationSlice, destinationLevel, destinationOrigin));
 	}
 	else
 	{
 		MTLPP_VALIDATE(mtlpp::BlitCommandEncoder, Encoder, AGXSafeGetRuntimeDebuggingLevel() >= EAGXDebugLevelValidation, Copy(Buffer, sourceOffset, sourceBytesPerRow, sourceBytesPerImage, sourceSize, toTexture, destinationSlice, destinationLevel, destinationOrigin, options));
+		METAL_DEBUG_LAYER(EAGXDebugLevelFastValidation, CurrentEncoder.GetBlitCommandEncoderDebugging().Copy(Buffer, sourceOffset, sourceBytesPerRow, sourceBytesPerImage, sourceSize, toTexture, destinationSlice, destinationLevel, destinationOrigin, options));
 	}
 	ConditionalSubmit();
 }
@@ -611,6 +809,7 @@ void FAGXRenderPass::CopyFromTextureToTexture(FAGXTexture const& Texture, uint32
 	
 	METAL_GPUPROFILE(FAGXProfiler::GetProfiler()->EncodeBlit(CurrentEncoder.GetCommandBufferStats(), __FUNCTION__));
 	MTLPP_VALIDATE(mtlpp::BlitCommandEncoder, Encoder, AGXSafeGetRuntimeDebuggingLevel() >= EAGXDebugLevelValidation, Copy(Texture, sourceSlice, sourceLevel, sourceOrigin, sourceSize, toTexture, destinationSlice, destinationLevel, destinationOrigin));
+	METAL_DEBUG_LAYER(EAGXDebugLevelFastValidation, CurrentEncoder.GetBlitCommandEncoderDebugging().Copy(Texture, sourceSlice, sourceLevel, sourceOrigin, sourceSize, toTexture, destinationSlice, destinationLevel, destinationOrigin));
 	ConditionalSubmit();
 }
 
@@ -622,6 +821,7 @@ void FAGXRenderPass::CopyFromBufferToBuffer(FAGXBuffer const& SourceBuffer, NSUI
 	
 	METAL_GPUPROFILE(FAGXProfiler::GetProfiler()->EncodeBlit(CurrentEncoder.GetCommandBufferStats(), __FUNCTION__));
 	MTLPP_VALIDATE(mtlpp::BlitCommandEncoder, Encoder, AGXSafeGetRuntimeDebuggingLevel() >= EAGXDebugLevelValidation, Copy(SourceBuffer, SourceOffset, DestinationBuffer, DestinationOffset, Size));
+	METAL_DEBUG_LAYER(EAGXDebugLevelFastValidation, CurrentEncoder.GetBlitCommandEncoderDebugging().Copy(SourceBuffer, SourceOffset, DestinationBuffer, DestinationOffset, Size));
 	ConditionalSubmit();
 }
 
@@ -633,6 +833,7 @@ void FAGXRenderPass::PresentTexture(FAGXTexture const& Texture, uint32 sourceSli
 	
 	METAL_GPUPROFILE(FAGXProfiler::GetProfiler()->EncodeBlit(CurrentEncoder.GetCommandBufferStats(), __FUNCTION__));
 	MTLPP_VALIDATE(mtlpp::BlitCommandEncoder, Encoder, AGXSafeGetRuntimeDebuggingLevel() >= EAGXDebugLevelValidation, Copy(Texture, sourceSlice, sourceLevel, sourceOrigin, sourceSize, toTexture, destinationSlice, destinationLevel, destinationOrigin));
+	METAL_DEBUG_LAYER(EAGXDebugLevelFastValidation, CurrentEncoder.GetBlitCommandEncoderDebugging().Copy(Texture, sourceSlice, sourceLevel, sourceOrigin, sourceSize, toTexture, destinationSlice, destinationLevel, destinationOrigin));
 }
 
 void FAGXRenderPass::SynchronizeTexture(FAGXTexture const& Texture, uint32 Slice, uint32 Level)
@@ -645,6 +846,7 @@ void FAGXRenderPass::SynchronizeTexture(FAGXTexture const& Texture, uint32 Slice
 	
 	// METAL_GPUPROFILE(FAGXProfiler::GetProfiler()->EncodeBlit(CurrentEncoder.GetCommandBufferStats(), __FUNCTION__));
 	MTLPP_VALIDATE(mtlpp::BlitCommandEncoder, Encoder, AGXSafeGetRuntimeDebuggingLevel() >= EAGXDebugLevelValidation, Synchronize(Texture, Slice, Level));
+	METAL_DEBUG_LAYER(EAGXDebugLevelFastValidation, CurrentEncoder.GetBlitCommandEncoderDebugging().Synchronize(Texture, Slice, Level));
 	ConditionalSubmit();
 #endif
 }
@@ -659,6 +861,7 @@ void FAGXRenderPass::SynchroniseResource(mtlpp::Resource const& Resource)
 	
 	// METAL_GPUPROFILE(FAGXProfiler::GetProfiler()->EncodeBlit(CurrentEncoder.GetCommandBufferStats(), __FUNCTION__));
 	MTLPP_VALIDATE(mtlpp::BlitCommandEncoder, Encoder, AGXSafeGetRuntimeDebuggingLevel() >= EAGXDebugLevelValidation, Synchronize(Resource));
+	METAL_DEBUG_LAYER(EAGXDebugLevelFastValidation, CurrentEncoder.GetBlitCommandEncoderDebugging().Synchronize(Resource));
 	ConditionalSubmit();
 #endif
 }
@@ -668,23 +871,27 @@ void FAGXRenderPass::FillBuffer(FAGXBuffer const& Buffer, ns::Range Range, uint8
 	check(Buffer);
 	
 	mtlpp::BlitCommandEncoder TargetEncoder;
+	METAL_DEBUG_ONLY(FAGXBlitCommandEncoderDebugging Debugging);
 	bool bAsync = !CurrentEncoder.HasBufferBindingHistory(Buffer);
 	if(bAsync)
 	{
 		ConditionalSwitchToAsyncBlit();
 		TargetEncoder = PrologueEncoder.GetBlitCommandEncoder();
 		METAL_GPUPROFILE(FAGXProfiler::GetProfiler()->EncodeBlit(PrologueEncoder.GetCommandBufferStats(), FString::Printf(TEXT("FillBuffer: %p %llu %llu"), Buffer.GetPtr(), Buffer.GetOffset() + Range.Location, Range.Length)));
+		METAL_DEBUG_LAYER(EAGXDebugLevelFastValidation, Debugging = PrologueEncoder.GetBlitCommandEncoderDebugging());
 	}
 	else
 	{
 		ConditionalSwitchToBlit();
 		TargetEncoder = CurrentEncoder.GetBlitCommandEncoder();
 		METAL_GPUPROFILE(FAGXProfiler::GetProfiler()->EncodeBlit(CurrentEncoder.GetCommandBufferStats(), FString::Printf(TEXT("FillBuffer: %p %llu %llu"), Buffer.GetPtr(), Buffer.GetOffset() + Range.Location, Range.Length)));
+		METAL_DEBUG_LAYER(EAGXDebugLevelFastValidation, Debugging = CurrentEncoder.GetBlitCommandEncoderDebugging());
 	}
 	
 	check(TargetEncoder.GetPtr());
 	
 	MTLPP_VALIDATE(mtlpp::BlitCommandEncoder, TargetEncoder, AGXSafeGetRuntimeDebuggingLevel() >= EAGXDebugLevelValidation, Fill(Buffer, Range, Value));
+	METAL_DEBUG_LAYER(EAGXDebugLevelFastValidation, (bAsync ? PrologueEncoder.GetBlitCommandEncoderDebugging() : CurrentEncoder.GetBlitCommandEncoderDebugging()).Fill(Buffer, Range, Value));
 	
 	if (!bAsync)
 	{
@@ -695,18 +902,21 @@ void FAGXRenderPass::FillBuffer(FAGXBuffer const& Buffer, ns::Range Range, uint8
 bool FAGXRenderPass::AsyncCopyFromBufferToTexture(FAGXBuffer const& Buffer, uint32 sourceOffset, uint32 sourceBytesPerRow, uint32 sourceBytesPerImage, mtlpp::Size sourceSize, FAGXTexture const& toTexture, uint32 destinationSlice, uint32 destinationLevel, mtlpp::Origin destinationOrigin, mtlpp::BlitOption options)
 {
 	mtlpp::BlitCommandEncoder TargetEncoder;
+	METAL_DEBUG_ONLY(FAGXBlitCommandEncoderDebugging Debugging);
 	bool bAsync = !CurrentEncoder.HasTextureBindingHistory(toTexture);
 	if(bAsync)
 	{
 		ConditionalSwitchToAsyncBlit();
 		TargetEncoder = PrologueEncoder.GetBlitCommandEncoder();
 		METAL_GPUPROFILE(FAGXProfiler::GetProfiler()->EncodeBlit(PrologueEncoder.GetCommandBufferStats(), __FUNCTION__));
+		METAL_DEBUG_LAYER(EAGXDebugLevelFastValidation, Debugging = PrologueEncoder.GetBlitCommandEncoderDebugging());
 	}
 	else
 	{
 		ConditionalSwitchToBlit();
 		TargetEncoder = CurrentEncoder.GetBlitCommandEncoder();
 		METAL_GPUPROFILE(FAGXProfiler::GetProfiler()->EncodeBlit(CurrentEncoder.GetCommandBufferStats(), __FUNCTION__));
+		METAL_DEBUG_LAYER(EAGXDebugLevelFastValidation, Debugging = CurrentEncoder.GetBlitCommandEncoderDebugging());
 	}
 	
 	check(TargetEncoder.GetPtr());
@@ -714,10 +924,12 @@ bool FAGXRenderPass::AsyncCopyFromBufferToTexture(FAGXBuffer const& Buffer, uint
 	if (options == mtlpp::BlitOption::None)
 	{
 		MTLPP_VALIDATE(mtlpp::BlitCommandEncoder, TargetEncoder, AGXSafeGetRuntimeDebuggingLevel() >= EAGXDebugLevelValidation, Copy(Buffer, sourceOffset, sourceBytesPerRow, sourceBytesPerImage, sourceSize, toTexture, destinationSlice, destinationLevel, destinationOrigin));
+		METAL_DEBUG_LAYER(EAGXDebugLevelFastValidation, Debugging.Copy(Buffer, sourceOffset, sourceBytesPerRow, sourceBytesPerImage, sourceSize, toTexture, destinationSlice, destinationLevel, destinationOrigin));
 	}
 	else
 	{
 		MTLPP_VALIDATE(mtlpp::BlitCommandEncoder, TargetEncoder, AGXSafeGetRuntimeDebuggingLevel() >= EAGXDebugLevelValidation, Copy(Buffer, sourceOffset, sourceBytesPerRow, sourceBytesPerImage, sourceSize, toTexture, destinationSlice, destinationLevel, destinationOrigin, options));
+		METAL_DEBUG_LAYER(EAGXDebugLevelFastValidation, Debugging.Copy(Buffer, sourceOffset, sourceBytesPerRow, sourceBytesPerImage, sourceSize, toTexture, destinationSlice, destinationLevel, destinationOrigin, options));
 	}
 	
 	return bAsync;
@@ -726,23 +938,27 @@ bool FAGXRenderPass::AsyncCopyFromBufferToTexture(FAGXBuffer const& Buffer, uint
 bool FAGXRenderPass::AsyncCopyFromTextureToTexture(FAGXTexture const& Texture, uint32 sourceSlice, uint32 sourceLevel, mtlpp::Origin sourceOrigin, mtlpp::Size sourceSize, FAGXTexture const& toTexture, uint32 destinationSlice, uint32 destinationLevel, mtlpp::Origin destinationOrigin)
 {
 	mtlpp::BlitCommandEncoder TargetEncoder;
+	METAL_DEBUG_ONLY(FAGXBlitCommandEncoderDebugging Debugging);
 	bool bAsync = !CurrentEncoder.HasTextureBindingHistory(toTexture);
 	if(bAsync)
 	{
 		ConditionalSwitchToAsyncBlit();
 		TargetEncoder = PrologueEncoder.GetBlitCommandEncoder();
 		METAL_GPUPROFILE(FAGXProfiler::GetProfiler()->EncodeBlit(PrologueEncoder.GetCommandBufferStats(), __FUNCTION__));
+		METAL_DEBUG_LAYER(EAGXDebugLevelFastValidation, Debugging = PrologueEncoder.GetBlitCommandEncoderDebugging());
 	}
 	else
 	{
 		ConditionalSwitchToBlit();
 		TargetEncoder = CurrentEncoder.GetBlitCommandEncoder();
 		METAL_GPUPROFILE(FAGXProfiler::GetProfiler()->EncodeBlit(CurrentEncoder.GetCommandBufferStats(), __FUNCTION__));
+		METAL_DEBUG_LAYER(EAGXDebugLevelFastValidation, Debugging = CurrentEncoder.GetBlitCommandEncoderDebugging());
 	}
 	
 	check(TargetEncoder.GetPtr());
 	
 	MTLPP_VALIDATE(mtlpp::BlitCommandEncoder, TargetEncoder, AGXSafeGetRuntimeDebuggingLevel() >= EAGXDebugLevelValidation, Copy(Texture, sourceSlice, sourceLevel, sourceOrigin, sourceSize, toTexture, destinationSlice, destinationLevel, destinationOrigin));
+	METAL_DEBUG_LAYER(EAGXDebugLevelFastValidation, Debugging.Copy(Texture, sourceSlice, sourceLevel, sourceOrigin, sourceSize, toTexture, destinationSlice, destinationLevel, destinationOrigin));
 	
 	return bAsync;
 }
@@ -755,23 +971,27 @@ bool FAGXRenderPass::CanAsyncCopyToBuffer(FAGXBuffer const& DestinationBuffer)
 void FAGXRenderPass::AsyncCopyFromBufferToBuffer(FAGXBuffer const& SourceBuffer, NSUInteger SourceOffset, FAGXBuffer const& DestinationBuffer, NSUInteger DestinationOffset, NSUInteger Size)
 {
 	mtlpp::BlitCommandEncoder TargetEncoder;
+	METAL_DEBUG_ONLY(FAGXBlitCommandEncoderDebugging Debugging);
 	bool bAsync = !CurrentEncoder.HasBufferBindingHistory(DestinationBuffer);
 	if(bAsync)
 	{
 		ConditionalSwitchToAsyncBlit();
 		TargetEncoder = PrologueEncoder.GetBlitCommandEncoder();
 		METAL_GPUPROFILE(FAGXProfiler::GetProfiler()->EncodeBlit(PrologueEncoder.GetCommandBufferStats(), FString::Printf(TEXT("AsyncCopyFromBufferToBuffer: %p %llu %llu"), DestinationBuffer.GetPtr(), DestinationBuffer.GetOffset() + DestinationOffset, Size)));
+		METAL_DEBUG_LAYER(EAGXDebugLevelFastValidation, Debugging = PrologueEncoder.GetBlitCommandEncoderDebugging());
 	}
 	else
 	{
 		ConditionalSwitchToBlit();
 		TargetEncoder = CurrentEncoder.GetBlitCommandEncoder();
 		METAL_GPUPROFILE(FAGXProfiler::GetProfiler()->EncodeBlit(CurrentEncoder.GetCommandBufferStats(), FString::Printf(TEXT("AsyncCopyFromBufferToBuffer: %p %llu %llu"), DestinationBuffer.GetPtr(), DestinationBuffer.GetOffset() + DestinationOffset, Size)));
+		METAL_DEBUG_LAYER(EAGXDebugLevelFastValidation, Debugging = CurrentEncoder.GetBlitCommandEncoderDebugging());
 	}
 	
 	check(TargetEncoder.GetPtr());
 	
     MTLPP_VALIDATE(mtlpp::BlitCommandEncoder, TargetEncoder, AGXSafeGetRuntimeDebuggingLevel() >= EAGXDebugLevelValidation, Copy(SourceBuffer, SourceOffset, DestinationBuffer, DestinationOffset, Size));
+	METAL_DEBUG_LAYER(EAGXDebugLevelFastValidation, Debugging.Copy(SourceBuffer, SourceOffset, DestinationBuffer, DestinationOffset, Size));
 }
 
 FAGXBuffer FAGXRenderPass::AllocateTemporyBufferForCopy(FAGXBuffer const& DestinationBuffer, NSUInteger Size, NSUInteger Align)
@@ -799,14 +1019,15 @@ void FAGXRenderPass::AsyncGenerateMipmapsForTexture(FAGXTexture const& Texture)
 	
 	METAL_GPUPROFILE(FAGXProfiler::GetProfiler()->EncodeBlit(CurrentEncoder.GetCommandBufferStats(), __FUNCTION__));
 	MTLPP_VALIDATE(mtlpp::BlitCommandEncoder, Encoder, AGXSafeGetRuntimeDebuggingLevel() >= EAGXDebugLevelValidation, GenerateMipmaps(Texture));
+	METAL_DEBUG_LAYER(EAGXDebugLevelFastValidation, PrologueEncoder.GetBlitCommandEncoderDebugging().GenerateMipmaps(Texture));
 }
 
-void FAGXRenderPass::End(void)
+TRefCountPtr<FAGXFence> const& FAGXRenderPass::End(void)
 {
 	// EndEncoding should provide the encoder fence...
 	if (PrologueEncoder.IsBlitCommandEncoderActive() || PrologueEncoder.IsComputeCommandEncoderActive())
 	{
-		PrologueEncoder.EndEncoding();
+		PrologueEncoderFence = PrologueEncoder.EndEncoding();
 	}
 	
 	if (CmdList.IsImmediate() && IsWithinParallelPass() && CurrentEncoder.IsParallelRenderCommandEncoderActive())
@@ -815,18 +1036,35 @@ void FAGXRenderPass::End(void)
 		CurrentEncoder.EndEncoding();
 		
 		ConditionalSwitchToBlit();
-		CurrentEncoder.EndEncoding();
+		CurrentEncoderFence = CurrentEncoder.EndEncoding();
+		ParallelPassEndFence = nullptr;
+		PassStartFence = nullptr;
 	}
 	else if (CurrentEncoder.IsRenderCommandEncoderActive() || CurrentEncoder.IsBlitCommandEncoderActive() || CurrentEncoder.IsComputeCommandEncoderActive())
 	{
 		State.FlushVisibilityResults(CurrentEncoder);
-		CurrentEncoder.EndEncoding();
+		check(!CurrentEncoderFence);
+		check(!PassStartFence);
+		check(!ParallelPassEndFence);
+		CurrentEncoderFence = CurrentEncoder.EndEncoding();
+	}
+	else if (PassStartFence || ParallelPassEndFence)
+	{
+		ConditionalSwitchToBlit();
+		CurrentEncoderFence = CurrentEncoder.EndEncoding();
+		ParallelPassEndFence = nullptr;
+		PassStartFence = nullptr;
 	}
 	
+	check(!PassStartFence);
+	check(!ParallelPassEndFence);
+	
 	State.SetRenderTargetsActive(false);
 	
 	RenderPassDesc = nil;
 	bWithinRenderPass = false;
+	
+	return CurrentEncoderFence;
 }
 
 void FAGXRenderPass::InsertCommandBufferFence(FAGXCommandBufferFence& Fence, mtlpp::CommandBufferHandler Handler)
@@ -867,13 +1105,16 @@ void FAGXRenderPass::InsertDebugEncoder()
 	check(NewBuf);
 	
 	mtlpp::BlitCommandEncoder TargetEncoder;
+	METAL_DEBUG_ONLY(FAGXBlitCommandEncoderDebugging Debugging);
 	ConditionalSwitchToBlit();
 	TargetEncoder = CurrentEncoder.GetBlitCommandEncoder();
 	METAL_GPUPROFILE(FAGXProfiler::GetProfiler()->EncodeBlit(CurrentEncoder.GetCommandBufferStats(), __FUNCTION__));
+	METAL_DEBUG_LAYER(EAGXDebugLevelFastValidation, Debugging = CurrentEncoder.GetBlitCommandEncoderDebugging());
 	
 	check(TargetEncoder.GetPtr());
 	
 	MTLPP_VALIDATE(mtlpp::BlitCommandEncoder, TargetEncoder, AGXSafeGetRuntimeDebuggingLevel() >= EAGXDebugLevelValidation, Fill(NewBuf, ns::Range(0, BufferOffsetAlignment), 0xff));
+	METAL_DEBUG_LAYER(EAGXDebugLevelFastValidation, CurrentEncoder.GetBlitCommandEncoderDebugging().Fill(NewBuf, ns::Range(0, BufferOffsetAlignment), 0xff));
 	
 	ConditionalSubmit();
 }
@@ -941,7 +1182,7 @@ void FAGXRenderPass::ConditionalSwitchToRender(void)
 	
 	if (CurrentEncoder.IsComputeCommandEncoderActive() || CurrentEncoder.IsBlitCommandEncoderActive())
 	{
-		CurrentEncoder.EndEncoding();
+		CurrentEncoderFence = CurrentEncoder.EndEncoding();
 	}
 	
 	if (!CurrentEncoder.IsRenderCommandEncoderActive())
@@ -966,7 +1207,7 @@ void FAGXRenderPass::ConditionalSwitchToCompute(void)
 			State.SetRenderStoreActions(CurrentEncoder, true);
 			State.FlushVisibilityResults(CurrentEncoder);
 		}
-		CurrentEncoder.EndEncoding();
+		CurrentEncoderFence = CurrentEncoder.EndEncoding();
 		State.SetRenderTargetsActive(false);
 	}
 	
@@ -974,9 +1215,33 @@ void FAGXRenderPass::ConditionalSwitchToCompute(void)
 	{
 		State.SetStateDirty();
 		CurrentEncoder.BeginComputeCommandEncoding(ComputeDispatchType);
+		if (PassStartFence)
+		{
+			CurrentEncoder.WaitForFence(PassStartFence);
+			PassStartFence = nullptr;
+		}
+		if (ParallelPassEndFence)
+		{
+			CurrentEncoder.WaitForFence(ParallelPassEndFence);
+			ParallelPassEndFence = nullptr;
+		}
+		if (CurrentEncoderFence)
+		{
+			CurrentEncoder.WaitForFence(CurrentEncoderFence);
+			CurrentEncoderFence = nullptr;
+		}
+		if (PrologueEncoderFence)
+		{
+			CurrentEncoder.WaitForFence(PrologueEncoderFence);
+		}
 	}
 	
 	check(CurrentEncoder.IsComputeCommandEncoderActive());
+	
+	if (PrologueEncoder.IsBlitCommandEncoderActive() || PrologueEncoder.IsComputeCommandEncoderActive())
+	{
+		CurrentEncoder.WaitForFence(PrologueEncoder.GetEncoderFence());
+	}
 }
 
 void FAGXRenderPass::ConditionalSwitchToBlit(void)
@@ -993,16 +1258,40 @@ void FAGXRenderPass::ConditionalSwitchToBlit(void)
 			State.SetRenderStoreActions(CurrentEncoder, true);
 			State.FlushVisibilityResults(CurrentEncoder);
 		}
-		CurrentEncoder.EndEncoding();
+		CurrentEncoderFence = CurrentEncoder.EndEncoding();
 		State.SetRenderTargetsActive(false);
 	}
 	
 	if (!CurrentEncoder.IsBlitCommandEncoderActive())
 	{
 		CurrentEncoder.BeginBlitCommandEncoding();
+		if (PassStartFence)
+		{
+			CurrentEncoder.WaitForFence(PassStartFence);
+			PassStartFence = nullptr;
+		}
+		if (ParallelPassEndFence)
+		{
+			CurrentEncoder.WaitForFence(ParallelPassEndFence);
+			ParallelPassEndFence = nullptr;
+		}
+		if (CurrentEncoderFence)
+		{
+			CurrentEncoder.WaitForFence(CurrentEncoderFence);
+			CurrentEncoderFence = nullptr;
+		}
+		if (PrologueEncoderFence)
+		{
+			CurrentEncoder.WaitForFence(PrologueEncoderFence);
+		}
 	}
 	
 	check(CurrentEncoder.IsBlitCommandEncoderActive());
+	
+	if (PrologueEncoder.IsBlitCommandEncoderActive() || PrologueEncoder.IsComputeCommandEncoderActive())
+	{
+		CurrentEncoder.WaitForFence(PrologueEncoder.GetEncoderFence());
+	}
 }
 
 void FAGXRenderPass::ConditionalSwitchToAsyncBlit(void)
@@ -1011,7 +1300,7 @@ void FAGXRenderPass::ConditionalSwitchToAsyncBlit(void)
 	
 	if (PrologueEncoder.IsComputeCommandEncoderActive() || PrologueEncoder.IsRenderCommandEncoderActive())
 	{
-		PrologueEncoder.EndEncoding();
+		PrologueEncoderFence = PrologueEncoder.EndEncoding();
 	}
 	
 	if (!PrologueEncoder.IsBlitCommandEncoderActive())
@@ -1021,6 +1310,51 @@ void FAGXRenderPass::ConditionalSwitchToAsyncBlit(void)
 			PrologueEncoder.StartCommandBuffer();
 		}
 		PrologueEncoder.BeginBlitCommandEncoding();
+		if (PrologueStartEncoderFence)
+		{
+			if (PrologueStartEncoderFence->NeedsWait(mtlpp::RenderStages::Vertex))
+			{
+				PrologueEncoder.WaitForFence(PrologueStartEncoderFence);
+			}
+			else
+			{
+				PrologueEncoder.WaitAndUpdateFence(PrologueStartEncoderFence);
+			}
+			PrologueStartEncoderFence = nullptr;
+		}
+		if (PrologueEncoderFence)
+		{
+			if (PrologueEncoderFence->NeedsWait(mtlpp::RenderStages::Vertex))
+			{
+				PrologueEncoder.WaitForFence(PrologueEncoderFence);
+			}
+			else
+			{
+				PrologueEncoder.WaitAndUpdateFence(PrologueEncoderFence);
+			}
+			PrologueEncoderFence = nullptr;
+		}
+#if METAL_DEBUG_OPTIONS
+//		if (GetEmitDrawEvents() && PrologueEncoder.GetEncoderFence())
+//		{
+//			for (uint32 i = mtlpp::RenderStages::Vertex; i <= mtlpp::RenderStages::Fragment && PrologueEncoder.GetEncoderFence()->Get((mtlpp::RenderStages)i).GetPtr(); i++)
+//			{
+//				if (CmdList.GetCommandQueue().GetRuntimeDebuggingLevel() >= EAGXDebugLevelValidation)
+//				{
+//					PrologueEncoder.GetEncoderFence()->Get((mtlpp::RenderStages)i).GetPtr().label = [NSString stringWithFormat:@"Prologue %@", PrologueEncoderFence->Get((mtlpp::RenderStages)i).GetLabel().GetPtr()];
+//				}
+//				else
+//				{
+//					PrologueEncoder.GetEncoderFence()->Get((mtlpp::RenderStages)i).SetLabel([NSString stringWithFormat:@"Prologue %@", PrologueEncoderFence->Get((mtlpp::RenderStages)i).GetLabel().GetPtr()]);
+//				}
+//			}
+//		}
+#endif
+
+		if (CurrentEncoder.IsRenderCommandEncoderActive() || CurrentEncoder.IsComputeCommandEncoderActive() || CurrentEncoder.IsBlitCommandEncoderActive())
+		{
+			CurrentEncoder.WaitForFence(PrologueEncoder.GetEncoderFence());
+		}
 	}
 	
 	check(PrologueEncoder.IsBlitCommandEncoderActive());
@@ -1032,7 +1366,7 @@ void FAGXRenderPass::ConditionalSwitchToAsyncCompute(void)
 	
 	if (PrologueEncoder.IsBlitCommandEncoderActive() || PrologueEncoder.IsRenderCommandEncoderActive())
 	{
-		PrologueEncoder.EndEncoding();
+		PrologueEncoderFence = PrologueEncoder.EndEncoding();
 	}
 	
 	if (!PrologueEncoder.IsComputeCommandEncoderActive())
@@ -1043,6 +1377,64 @@ void FAGXRenderPass::ConditionalSwitchToAsyncCompute(void)
 		}
 		State.SetStateDirty();
 		PrologueEncoder.BeginComputeCommandEncoding(ComputeDispatchType);
+		
+		if (PrologueStartEncoderFence)
+		{
+			if (PrologueStartEncoderFence->NeedsWait(mtlpp::RenderStages::Vertex))
+			{
+				PrologueEncoder.WaitForFence(PrologueStartEncoderFence);
+			}
+			else
+			{
+				PrologueEncoder.WaitAndUpdateFence(PrologueStartEncoderFence);
+			}
+			PrologueStartEncoderFence = nullptr;
+		}
+		if (PrologueEncoderFence)
+		{
+			if (PrologueEncoderFence->NeedsWait(mtlpp::RenderStages::Vertex))
+			{
+				PrologueEncoder.WaitForFence(PrologueEncoderFence);
+			}
+			else
+			{
+				PrologueEncoder.WaitAndUpdateFence(PrologueEncoderFence);
+			}
+			PrologueEncoderFence = nullptr;
+		}
+		if (PassStartFence)
+		{
+			if(PassStartFence->NeedsWait(mtlpp::RenderStages::Vertex))
+			{
+				PrologueEncoder.WaitForFence(PassStartFence);
+			}
+			else
+			{
+				PrologueEncoder.WaitAndUpdateFence(PassStartFence);
+			}
+			PassStartFence = nullptr;
+		}
+#if METAL_DEBUG_OPTIONS
+//		if (GetEmitDrawEvents() && PrologueEncoder.GetEncoderFence())
+//		{
+//			for (uint32 i = mtlpp::RenderStages::Vertex; i <= mtlpp::RenderStages::Fragment && PrologueEncoder.GetEncoderFence()->Get((mtlpp::RenderStages)i).GetPtr(); i++)
+//			{
+//				if (CmdList.GetCommandQueue().GetRuntimeDebuggingLevel() >= EAGXDebugLevelValidation)
+//				{
+//					PrologueEncoder.GetEncoderFence()->Get((mtlpp::RenderStages)i).GetPtr().label = [NSString stringWithFormat:@"Prologue %@", PrologueEncoderFence->Get((mtlpp::RenderStages)i).GetLabel().GetPtr()];
+//				}
+//				else
+//				{
+//					PrologueEncoder.GetEncoderFence()->Get((mtlpp::RenderStages)i).SetLabel([NSString stringWithFormat:@"Prologue %@", PrologueEncoderFence->Get((mtlpp::RenderStages)i).GetLabel().GetPtr()]);
+//				}
+//			}
+//		}
+#endif
+		
+		if (CurrentEncoder.IsRenderCommandEncoderActive() || CurrentEncoder.IsComputeCommandEncoderActive() || CurrentEncoder.IsBlitCommandEncoderActive())
+		{
+			CurrentEncoder.WaitForFence(PrologueEncoder.GetEncoderFence());
+		}
 	}
 	
 	check(PrologueEncoder.IsComputeCommandEncoderActive());
@@ -1256,3 +1648,120 @@ uint32 FAGXRenderPass::GetCommandBufferIndex(void) const
 @protocol IMTLComputeCommandEncoder
 - (void)memoryBarrierWithResources:(const id<MTLResource>[])resources count:(NSUInteger)count;
 @end
+
+void FAGXRenderPass::InsertDebugDraw(FAGXCommandData& Data)
+{
+#if !PLATFORM_TVOS
+	if (GAGXCommandBufferDebuggingEnabled && (!FAGXCommandQueue::SupportsFeature(EAGXFeaturesValidation) || State.GetVisibilityResultMode() == mtlpp::VisibilityResultMode::Disabled))
+	{
+		FAGXGraphicsPipelineState* BoundShaderState = State.GetGraphicsPSO();
+		
+		uint32 NumCommands = CurrentEncoder.GetMarkers().AddCommand(GetCommandBufferIndex(), GetEncoderIndex(), CmdList.GetParallelIndex(), State.GetDebugBuffer(), BoundShaderState, nullptr, Data);
+		
+		if ((NumCommands % GAGXDebugOpsCount) == 0)
+		{
+			FAGXDebugInfo DebugInfo;
+			DebugInfo.EncoderIndex = GetEncoderIndex();
+			DebugInfo.ContextIndex = CmdList.GetParallelIndex();
+			DebugInfo.CommandIndex = NumCommands;
+			DebugInfo.CmdBuffIndex = GetCommandBufferIndex();
+			DebugInfo.CommandBuffer = reinterpret_cast<uintptr_t>(CurrentEncoder.GetCommandBuffer().GetPtr());
+			DebugInfo.PSOSignature[0] = BoundShaderState->VertexShader->SourceLen;
+			DebugInfo.PSOSignature[1] = BoundShaderState->VertexShader->SourceCRC;
+			if (IsValidRef(BoundShaderState->PixelShader))
+			{
+				DebugInfo.PSOSignature[2] = BoundShaderState->PixelShader->SourceLen;
+				DebugInfo.PSOSignature[3] = BoundShaderState->PixelShader->SourceCRC;
+			}
+			else
+			{
+				DebugInfo.PSOSignature[2] = 0;
+				DebugInfo.PSOSignature[3] = 0;
+			}
+			
+			FAGXShaderPipeline* PSO = State.GetPipelineState();
+			
+			CurrentEncoder.GetRenderCommandEncoder().SetRenderPipelineState(PSO->DebugPipelineState);
+
+			mtlpp::VisibilityResultMode VisMode = State.GetVisibilityResultMode();
+			uint32 VisibilityOffset = State.GetVisibilityResultOffset();
+			if (VisMode != mtlpp::VisibilityResultMode::Disabled)
+			{
+				CurrentEncoder.GetRenderCommandEncoder().SetVisibilityResultMode(mtlpp::VisibilityResultMode::Disabled, 0);
+			}
+
+		#if PLATFORM_MAC
+			id<MTLBuffer> DebugBufferPtr = State.GetDebugBuffer().GetPtr();
+			[(id<IMTLRenderCommandEncoder>)CurrentEncoder.GetRenderCommandEncoder().GetPtr() memoryBarrierWithResources:&DebugBufferPtr count:1 afterStages:MTLRenderStageFragment beforeStages:MTLRenderStageVertex];
+			
+			CurrentEncoder.SetShaderBytes(mtlpp::FunctionType::Vertex, (uint8 const*)&DebugInfo, sizeof(DebugInfo), 0);
+			State.SetShaderBufferDirty(EAGXShaderStages::Vertex, 0);
+
+			CurrentEncoder.SetShaderBuffer(mtlpp::FunctionType::Vertex, State.GetDebugBuffer(), 0, State.GetDebugBuffer().GetLength(), 1, mtlpp::ResourceUsage::Write);
+			State.SetShaderBufferDirty(EAGXShaderStages::Vertex, 1);
+
+			CurrentEncoder.GetRenderCommandEncoder().Draw(mtlpp::PrimitiveType::Point, 0, 1);
+			
+			[(id<IMTLRenderCommandEncoder>)CurrentEncoder.GetRenderCommandEncoder().GetPtr() memoryBarrierWithResources:&DebugBufferPtr count:1 afterStages:MTLRenderStageVertex beforeStages:MTLRenderStageVertex];
+		#else
+			CurrentEncoder.GetRenderCommandEncoder().SetTileData((uint8 const*)&DebugInfo, sizeof(DebugInfo), 0);
+			CurrentEncoder.GetRenderCommandEncoder().SetTileBuffer(State.GetDebugBuffer(), 0, 1);
+			mtlpp::Size ThreadsPerTile(1, 1, 1);
+			CurrentEncoder.GetRenderCommandEncoder().DispatchThreadsPerTile(ThreadsPerTile);
+		#endif
+
+			CurrentEncoder.GetRenderCommandEncoder().SetRenderPipelineState(PSO->RenderPipelineState);
+
+			if (VisMode != mtlpp::VisibilityResultMode::Disabled)
+			{
+				CurrentEncoder.GetRenderCommandEncoder().SetVisibilityResultMode(VisMode, VisibilityOffset);
+			}
+		}
+	}
+#endif
+}
+
+void FAGXRenderPass::InsertDebugDispatch(FAGXCommandData& Data)
+{
+#if !PLATFORM_TVOS
+	if (GAGXCommandBufferDebuggingEnabled)
+	{
+		FAGXComputeShader* BoundShaderState = State.GetComputeShader();
+		
+		uint32 NumCommands = CurrentEncoder.GetMarkers().AddCommand(GetCommandBufferIndex(), GetEncoderIndex(), CmdList.GetParallelIndex(), State.GetDebugBuffer(), nullptr, BoundShaderState, Data);
+		
+		if ((NumCommands % GAGXDebugOpsCount) == 0)
+		{
+			FAGXDebugInfo DebugInfo;
+			DebugInfo.EncoderIndex = GetEncoderIndex();
+			DebugInfo.ContextIndex = CmdList.GetParallelIndex();
+			DebugInfo.CommandIndex = NumCommands;
+			DebugInfo.CmdBuffIndex = GetCommandBufferIndex();
+			DebugInfo.CommandBuffer = reinterpret_cast<uintptr_t>(CurrentEncoder.GetCommandBuffer().GetPtr());
+			DebugInfo.PSOSignature[0] = BoundShaderState->SourceLen;
+			DebugInfo.PSOSignature[1] = BoundShaderState->SourceCRC;
+			DebugInfo.PSOSignature[2] = 0;
+			DebugInfo.PSOSignature[3] = 0;
+			
+			CurrentEncoder.GetComputeCommandEncoder().SetComputePipelineState(AGXGetMetalDebugComputeState());
+			
+			id<MTLBuffer> DebugBufferPtr = State.GetDebugBuffer().GetPtr();
+			[(id<IMTLComputeCommandEncoder>)CurrentEncoder.GetComputeCommandEncoder().GetPtr() memoryBarrierWithResources:&DebugBufferPtr count:1];
+			
+			CurrentEncoder.SetShaderBytes(mtlpp::FunctionType::Kernel, (uint8 const*)&DebugInfo, sizeof(DebugInfo), 0);
+			State.SetShaderBufferDirty(EAGXShaderStages::Compute, 0);
+			
+			CurrentEncoder.SetShaderBuffer(mtlpp::FunctionType::Kernel, State.GetDebugBuffer(), 0, State.GetDebugBuffer().GetLength(), 1, mtlpp::ResourceUsage::Write);
+			State.SetShaderBufferDirty(EAGXShaderStages::Compute, 1);
+			
+			mtlpp::Size ThreadsPerTile(1, 1, 1);
+			CurrentEncoder.GetComputeCommandEncoder().DispatchThreads(ThreadsPerTile, ThreadsPerTile);
+			
+			[(id<IMTLComputeCommandEncoder>)CurrentEncoder.GetComputeCommandEncoder().GetPtr() memoryBarrierWithResources:&DebugBufferPtr count:1];
+			
+			FAGXShaderPipeline* Pipeline = BoundShaderState->GetPipeline();
+			CurrentEncoder.GetComputeCommandEncoder().SetComputePipelineState(Pipeline->ComputePipelineState);
+		}
+	}
+#endif
+}
