Undo //UE5/Main/Engine/Source/Runtime/Apple/AGXRHI/Private/... changelist 18589140

diff --git a/Engine/Source/Runtime/Apple/AGXRHI/Private/AGXRenderPass.h b/Engine/Source/Runtime/Apple/AGXRHI/Private/AGXRenderPass.h
--- a/Engine/Source/Runtime/Apple/AGXRHI/Private/AGXRenderPass.h
+++ b/Engine/Source/Runtime/Apple/AGXRHI/Private/AGXRenderPass.h
@@ -5,6 +5,7 @@
 #include "AGXRHIPrivate.h"
 #include "AGXCommandEncoder.h"
 #include "AGXState.h"
+#include "AGXFence.h"
 
 class FAGXCommandList;
 class FAGXCommandQueue;
@@ -23,7 +24,11 @@ public:
 #pragma mark -
 	void SetDispatchType(mtlpp::DispatchType Type);
 	
-    void Begin(bool const bParallelBegin = false);
+    void Begin(FAGXFence* Fence, bool const bParallelBegin = false);
+	
+	void Wait(FAGXFence* Fence);
+
+	void Update(FAGXFence* Fence);
 	
     void BeginParallelRenderPass(mtlpp::RenderPassDescriptor RenderPass, uint32 NumParallelContextsInPass);
 
@@ -46,7 +51,7 @@ public:
     
     void DispatchIndirect(FAGXVertexBuffer* ArgumentBufferRHI, uint32 ArgumentOffset);
     
-    void EndRenderPass(void);
+    TRefCountPtr<FAGXFence> const& EndRenderPass(void);
     
     void CopyFromTextureToBuffer(FAGXTexture const& Texture, uint32 sourceSlice, uint32 sourceLevel, mtlpp::Origin sourceOrigin, mtlpp::Size sourceSize, FAGXBuffer const& toBuffer, uint32 destinationOffset, uint32 destinationBytesPerRow, uint32 destinationBytesPerImage, mtlpp::BlitOption options);
     
@@ -76,9 +81,9 @@ public:
 	
 	void AsyncGenerateMipmapsForTexture(FAGXTexture const& Texture);
 	
-    void Submit(EAGXSubmitFlags SubmissionFlags);
+    TRefCountPtr<FAGXFence> const& Submit(EAGXSubmitFlags SubmissionFlags);
     
-    void End(void);
+    TRefCountPtr<FAGXFence> const& End(void);
 	
 	void InsertCommandBufferFence(FAGXCommandBufferFence& Fence, mtlpp::CommandBufferHandler Handler);
 	
@@ -170,6 +175,9 @@ private:
 	uint32 GetEncoderIndex(void) const;
 	uint32 GetCommandBufferIndex(void) const;
 	
+	void InsertDebugDraw(FAGXCommandData& Data);
+	void InsertDebugDispatch(FAGXCommandData& Data);
+
 private:
 	FAGXCommandList& CmdList;
     FAGXStateCache& State;
@@ -187,6 +195,15 @@ private:
 	// Disjoint ranges *are* permitted!
 	TMap<id<MTLBuffer>, TArray<NSRange>> OutstandingBufferUploads;
 
+	// Fences for the current command encoder chain
+	TRefCountPtr<FAGXFence> PassStartFence;
+	TRefCountPtr<FAGXFence> CurrentEncoderFence;
+	TRefCountPtr<FAGXFence> ParallelPassEndFence;
+
+	// Fences for the prologue command encoder chain
+	TRefCountPtr<FAGXFence> PrologueStartEncoderFence;
+	TRefCountPtr<FAGXFence> PrologueEncoderFence;
+    
     mtlpp::RenderPassDescriptor RenderPassDesc;
     
 	mtlpp::DispatchType ComputeDispatchType;
diff --git a/Engine/Source/Runtime/Apple/AGXRHI/Private/AGXShaders.cpp b/Engine/Source/Runtime/Apple/AGXRHI/Private/AGXShaders.cpp
--- a/Engine/Source/Runtime/Apple/AGXRHI/Private/AGXShaders.cpp
+++ b/Engine/Source/Runtime/Apple/AGXRHI/Private/AGXShaders.cpp
@@ -17,6 +17,7 @@
 #include "MetalShaderResources.h"
 #include "AGXResources.h"
 #include "AGXProfiler.h"
+#include "AGXCommandBuffer.h"
 #include "Serialization/MemoryReader.h"
 #include "Misc/FileHelper.h"
 #include "Misc/ScopeRWLock.h"
@@ -44,25 +45,84 @@ NSString* AGXDecodeMetalSourceCode(uint32 CodeSize, TArray<uint8> const& Compres
 	return GlslCodeNSString;
 }
 
-mtlpp::LanguageVersion AGXValidateVersion(uint32 Version)
+mtlpp::LanguageVersion AGXValidateVersion(uint8 Version)
 {
-	mtlpp::LanguageVersion Result = mtlpp::LanguageVersion::Version2_2;
-
+	static uint32 MetalMacOSVersions[][3] = {
+		{10,11,6},
+		{10,11,6},
+		{10,12,6},
+		{10,13,0},
+		{10,14,0},
+		{10,15,0},
+		{11, 0,0},
+	};
+	static uint32 MetaliOSVersions[][3] = {
+		{ 8,0,0},
+		{ 9,0,0},
+		{10,0,0},
+		{11,0,0},
+		{12,0,0},
+		{13,0,0},
+		{14,0,0},
+	};
+	static TCHAR const* StandardNames[] =
+	{
+		TEXT("Metal 1.0"),
+		TEXT("Metal 1.1"),
+		TEXT("Metal 1.2"),
+		TEXT("Metal 2.0"),
+		TEXT("Metal 2.1"),
+		TEXT("Metal 2.2"),
+		TEXT("Metal 2.3"),
+	};
+	
+	mtlpp::LanguageVersion Result = mtlpp::LanguageVersion::Version1_1;
 	switch(Version)
 	{
-		case 7:
-			Result = mtlpp::LanguageVersion::Version2_4;
-			break;
 		case 6:
 			Result = mtlpp::LanguageVersion::Version2_3;
 			break;
 		case 5:
 			Result = mtlpp::LanguageVersion::Version2_2;
 			break;
+		case 4:
+			Result = mtlpp::LanguageVersion::Version2_1;
+			break;
+		case 3:
+			Result = mtlpp::LanguageVersion::Version2_0;
+			break;
+		case 2:
+			Result = mtlpp::LanguageVersion::Version1_2;
+			break;
+		case 1:
+			Result = mtlpp::LanguageVersion::Version1_1;
+			break;
+		case 0:
 		default:
-			Result = mtlpp::LanguageVersion::Version2_2;
+#if PLATFORM_MAC
+			Result = mtlpp::LanguageVersion::Version1_1;
+#else
+			Result = mtlpp::LanguageVersion::Version1_0;
+#endif
 			break;
 	}
-
+	
+	if (!FApplePlatformMisc::IsOSAtLeastVersion(MetalMacOSVersions[Version], MetaliOSVersions[Version], MetaliOSVersions[Version]))
+	{
+		FFormatNamedArguments Args;
+		Args.Add(TEXT("ShaderVersion"), FText::FromString(FString(StandardNames[Version])));
+#if PLATFORM_MAC
+		Args.Add(TEXT("RequiredOS"), FText::FromString(FString::Printf(TEXT("macOS %d.%d.%d"), MetalMacOSVersions[Version][0], MetalMacOSVersions[Version][1], MetalMacOSVersions[Version][2])));
+#else
+		Args.Add(TEXT("RequiredOS"), FText::FromString(FString::Printf(TEXT("%d.%d.%d"), MetaliOSVersions[Version][0], MetaliOSVersions[Version][1], MetaliOSVersions[Version][2])));
+#endif
+		FText LocalizedMsg = FText::Format(NSLOCTEXT("AGXRHI", "ShaderVersionUnsupported", "The current OS version does not support {ShaderVersion} required by the project. You must upgrade to {RequiredOS} to run this project."),Args);
+		
+		FText Title = NSLOCTEXT("AGXRHI", "ShaderVersionUnsupportedTitle", "Shader Version Unsupported");
+		FMessageDialog::Open(EAppMsgType::Ok, LocalizedMsg, &Title);
+		
+		FPlatformMisc::RequestExit(true);
+	}
+	
 	return Result;
 }
diff --git a/Engine/Source/Runtime/Apple/AGXRHI/Private/AGXState.cpp b/Engine/Source/Runtime/Apple/AGXRHI/Private/AGXState.cpp
--- a/Engine/Source/Runtime/Apple/AGXRHI/Private/AGXState.cpp
+++ b/Engine/Source/Runtime/Apple/AGXRHI/Private/AGXState.cpp
@@ -266,6 +266,11 @@ static FAGXSampler FindOrCreateSamplerState(const FSamplerStateInitializerRHI& I
 #if PLATFORM_MAC
 		[Desc setBorderColor:Initializer.BorderColor == 0 ? MTLSamplerBorderColorTransparentBlack : MTLSamplerBorderColorOpaqueWhite];
 #endif
+		if (FAGXCommandQueue::SupportsFeature(EAGXFeaturesIABs))
+		{
+			[Desc setSupportArgumentBuffers:YES];
+		}
+		
 		State = FAGXSampler([GMtlDevice newSamplerStateWithDescriptor:Desc], ns::Ownership::Assign);
 		
 		[Desc release];
diff --git a/Engine/Source/Runtime/Apple/AGXRHI/Private/AGXStateCache.cpp b/Engine/Source/Runtime/Apple/AGXRHI/Private/AGXStateCache.cpp
--- a/Engine/Source/Runtime/Apple/AGXRHI/Private/AGXStateCache.cpp
+++ b/Engine/Source/Runtime/Apple/AGXRHI/Private/AGXStateCache.cpp
@@ -7,6 +7,7 @@
 #include "AGXGraphicsPipelineState.h"
 #include "AGXStateCache.h"
 #include "AGXProfiler.h"
+#include "AGXCommandBuffer.h"
 
 #if PLATFORM_MAC
 	#ifndef UINT128_MAX
@@ -1711,7 +1712,7 @@ void FAGXStateCache::SetResourcesFromTables(ShaderType Shader, uint32 ShaderStag
 		const int32 BufferIndex = FMath::FloorLog2(LowestBitMask); // todo: This has a branch on zero, we know it could never be zero...
 		DirtyBits ^= LowestBitMask;
 		FAGXUniformBuffer* Buffer = (FAGXUniformBuffer*)GetBoundUniformBuffers(Frequency)[BufferIndex];
-		if (Buffer)
+		if (Buffer && !FAGXCommandQueue::SupportsFeature(EAGXFeaturesIABs))
 		{
 			check(BufferIndex < Shader->Bindings.ShaderResourceTable.ResourceTableLayoutHashes.Num());
 			check(Buffer->GetLayout().GetHash() == Shader->Bindings.ShaderResourceTable.ResourceTableLayoutHashes[BufferIndex]);
@@ -1870,14 +1871,16 @@ void FAGXStateCache::FlushVisibilityResults(FAGXCommandEncoder& CommandEncoder)
 #if PLATFORM_MAC
 	if(VisibilityResults && VisibilityResults->Buffer && VisibilityResults->Buffer.GetStorageMode() == mtlpp::StorageMode::Managed && VisibilityWritten && CommandEncoder.IsRenderCommandEncoderActive())
 	{
-		CommandEncoder.EndEncoding();
+		TRefCountPtr<FAGXFence> Fence = CommandEncoder.EndEncoding();
 		
         CommandEncoder.BeginBlitCommandEncoding();
+		CommandEncoder.WaitForFence(Fence);
 		
 		mtlpp::BlitCommandEncoder& Encoder = CommandEncoder.GetBlitCommandEncoder();
 
 		// METAL_GPUPROFILE(FAGXProfiler::GetProfiler()->EncodeBlit(CommandEncoder.GetCommandBufferStats(), __FUNCTION__));
 		MTLPP_VALIDATE(mtlpp::BlitCommandEncoder, Encoder, AGXSafeGetRuntimeDebuggingLevel() >= EAGXDebugLevelValidation, Synchronize(VisibilityResults->Buffer));
+		METAL_DEBUG_LAYER(EAGXDebugLevelFastValidation, CommandEncoder.GetBlitCommandEncoderDebugging().Synchronize(VisibilityResults->Buffer));
 		
 		VisibilityWritten = 0;
 	}
@@ -1992,9 +1995,9 @@ void FAGXStateCache::SetRenderPipelineState(FAGXCommandEncoder& CommandEncoder,
 		FAGXShaderPipeline* Pipeline = GetPipelineState();
 		EAGXShaderStages VertexStage = EAGXShaderStages::Vertex;
 		
-		FAGXDebugShaderResourceMask VertexMask = Pipeline->ResourceMask[SF_Vertex];
-		TArray<uint32>& MinVertexBufferSizes = Pipeline->BufferDataSizes[SF_Vertex];
-		const TMap<uint8, uint8>& VertexTexTypes = Pipeline->TextureTypes[SF_Vertex];
+		FAGXDebugShaderResourceMask VertexMask = Pipeline->ResourceMask[EAGXShaderVertex];
+		TArray<uint32>& MinVertexBufferSizes = Pipeline->BufferDataSizes[EAGXShaderVertex];
+		const TMap<uint8, uint8>& VertexTexTypes = Pipeline->TextureTypes[EAGXShaderVertex];
 		while(VertexMask.BufferMask)
 		{
 			uint32 Index = __builtin_ctz(VertexMask.BufferMask);
@@ -2042,9 +2045,9 @@ void FAGXStateCache::SetRenderPipelineState(FAGXCommandEncoder& CommandEncoder,
 			ensure(ShaderSamplers[VertexStage].Samplers[Index]);
 		}
 		
-		FAGXDebugShaderResourceMask FragmentMask = Pipeline->ResourceMask[SF_Pixel];
-		TArray<uint32>& MinFragmentBufferSizes = Pipeline->BufferDataSizes[SF_Pixel];
-		const TMap<uint8, uint8>& FragmentTexTypes = Pipeline->TextureTypes[SF_Pixel];
+		FAGXDebugShaderResourceMask FragmentMask = Pipeline->ResourceMask[EAGXShaderFragment];
+		TArray<uint32>& MinFragmentBufferSizes = Pipeline->BufferDataSizes[EAGXShaderFragment];
+		const TMap<uint8, uint8>& FragmentTexTypes = Pipeline->TextureTypes[EAGXShaderFragment];
 		while(FragmentMask.BufferMask)
 		{
 			uint32 Index = __builtin_ctz(FragmentMask.BufferMask);
@@ -2108,9 +2111,9 @@ void FAGXStateCache::SetComputePipelineState(FAGXCommandEncoder& CommandEncoder)
 		FAGXShaderPipeline* Pipeline = ComputeShader->GetPipeline();
 		check(Pipeline);
 		
-		FAGXDebugShaderResourceMask ComputeMask = Pipeline->ResourceMask[SF_Compute];
-		TArray<uint32>& MinComputeBufferSizes = Pipeline->BufferDataSizes[SF_Compute];
-		const TMap<uint8, uint8>& ComputeTexTypes = Pipeline->TextureTypes[SF_Compute];
+		FAGXDebugShaderResourceMask ComputeMask = Pipeline->ResourceMask[EAGXShaderCompute];
+		TArray<uint32>& MinComputeBufferSizes = Pipeline->BufferDataSizes[EAGXShaderCompute];
+		const TMap<uint8, uint8>& ComputeTexTypes = Pipeline->TextureTypes[EAGXShaderCompute];
 		while(ComputeMask.BufferMask)
 		{
 			uint32 Index = __builtin_ctz(ComputeMask.BufferMask);
@@ -2239,6 +2242,17 @@ void FAGXStateCache::CommitResourceTable(EAGXShaderStages const Frequency, mtlpp
 	}
 }
 
+FAGXBuffer& FAGXStateCache::GetDebugBuffer()
+{
+    if (!DebugBuffer)
+    {
+        // Assume worst case tiling (16x16) and render-target size (4096x4096) on iOS for now
+        uint32 NumTiles = PLATFORM_MAC ? 1 : 65536;
+        DebugBuffer = GetAGXDeviceContext().CreatePooledBuffer(FAGXPooledBufferArgs(NumTiles * sizeof(FAGXDebugInfo), BUF_Dynamic, mtlpp::StorageMode::Shared));
+    }
+    return DebugBuffer;
+}
+
 FTexture2DRHIRef FAGXStateCache::CreateFallbackDepthStencilSurface(uint32 Width, uint32 Height)
 {
 #if PLATFORM_MAC
diff --git a/Engine/Source/Runtime/Apple/AGXRHI/Private/AGXStateCache.h b/Engine/Source/Runtime/Apple/AGXRHI/Private/AGXStateCache.h
--- a/Engine/Source/Runtime/Apple/AGXRHI/Private/AGXStateCache.h
+++ b/Engine/Source/Runtime/Apple/AGXRHI/Private/AGXStateCache.h
@@ -116,6 +116,7 @@ public:
 	
 	bool PrepareToRestart(bool const bCurrentApplied);
 	
+	FAGXBuffer& GetDebugBuffer();
 	FAGXShaderParameterCache& GetShaderParameters(EAGXShaderStages const Stage) { return ShaderParameters[Stage]; }
 	FLinearColor const& GetBlendFactor() const { return BlendFactor; }
 	uint32 GetStencilRef() const { return StencilRef; }
@@ -286,6 +287,7 @@ private:
 	FTextureRHIRef DepthStencilResolve;
 	/** A fallback depth-stencil surface for draw calls that write to depth without a depth-stencil surface bound. */
 	FTexture2DRHIRef FallbackDepthStencilSurface;
+	FAGXBuffer DebugBuffer;
 	mtlpp::RenderPassDescriptor RenderPassDesc;
 	uint32 RasterBits;
     uint8 PipelineBits;
diff --git a/Engine/Source/Runtime/Apple/AGXRHI/Private/AGXTexture.cpp b/Engine/Source/Runtime/Apple/AGXRHI/Private/AGXTexture.cpp
--- a/Engine/Source/Runtime/Apple/AGXRHI/Private/AGXTexture.cpp
+++ b/Engine/Source/Runtime/Apple/AGXRHI/Private/AGXTexture.cpp
@@ -6,6 +6,7 @@
 
 #include "AGXRHIPrivate.h"
 #include "AGXProfiler.h" // for STAT_AGXTexturePageOffTime
+#include "AGXCommandBuffer.h"
 #include "RenderUtils.h"
 #include "Containers/ResourceArray.h"
 #include "Misc/ScopeRWLock.h"
@@ -398,12 +399,71 @@ void FAGXSurface::ReplaceTexture(FAGXContext& Context, FAGXTexture CurrentTextur
 
 void FAGXSurface::MakeAliasable(void)
 {
-	// TODO
+	check(!bTextureView);
+	check(ImageSurfaceRef == nullptr);
+	
+	static bool bSupportsHeaps = GetAGXDeviceContext().SupportsFeature(EAGXFeaturesHeaps);
+	if (bSupportsHeaps && Texture.GetStorageMode() == mtlpp::StorageMode::Private && Texture.GetHeap())
+	{
+		if (MSAATexture && (MSAATexture != Texture) && !MSAATexture.IsAliasable())
+		{
+			MSAATexture.MakeAliasable();
+#if STATS || ENABLE_LOW_LEVEL_MEM_TRACKER
+			AGXLLM::LogAliasTexture(MSAATexture);
+#endif
+		}
+		if (!Texture.IsAliasable())
+		{
+			Texture.MakeAliasable();
+#if STATS || ENABLE_LOW_LEVEL_MEM_TRACKER
+			AGXLLM::LogAliasTexture(Texture);
+#endif
+		}
+	}
 }
 
 void FAGXSurface::MakeUnAliasable(void)
 {
-	// TODO
+	check(!bTextureView);
+	check(ImageSurfaceRef == nullptr);
+	
+	static bool bSupportsHeaps = GetAGXDeviceContext().SupportsFeature(EAGXFeaturesHeaps);
+	if (bSupportsHeaps && Texture.GetStorageMode() == mtlpp::StorageMode::Private && Texture.GetHeap() && Texture.IsAliasable())
+	{
+		FAGXTexture OldTexture = Texture;
+		Texture = Reallocate(Texture, mtlpp::TextureUsage::Unknown);
+		AGXSafeReleaseMetalTexture(this, OldTexture);
+		if (MSAATexture && (MSAATexture != OldTexture))
+		{
+			FAGXTexture OldMSAATexture = MSAATexture;
+			MSAATexture = Reallocate(MSAATexture, mtlpp::TextureUsage::Unknown);
+			AGXSafeReleaseMetalTexture(this, OldMSAATexture);
+		}
+		else if(MSAATexture)
+		{
+			MSAATexture = Texture;
+		}
+		
+		for (FAGXShaderResourceView* SRV : SRVs)
+		{
+			if(OldTexture != SRV->TextureView->Texture)
+			{
+				AGXSafeReleaseMetalTexture(this, SRV->TextureView->Texture, true, false);
+			}
+			SRV->TextureView->Texture = nil;
+			
+			SRV->TextureView->MSAATexture = nil;
+			
+			if (SRV->Format == PF_Unknown)
+			{
+				SRV->TextureView->Init(*this, NSMakeRange(SRV->MipLevel, SRV->NumMips));
+			}
+			else
+			{
+				SRV->TextureView->Init(*this, NSMakeRange(SRV->MipLevel, SRV->NumMips), (EPixelFormat)SRV->Format, SRV->bSRGBForceDisable);
+			}
+		}
+	}
 }
 
 void FAGXSurface::Init(FAGXSurface& Source, NSRange MipRange)
diff --git a/Engine/Source/Runtime/Apple/AGXRHI/Private/AGXTransitionData.cpp b/Engine/Source/Runtime/Apple/AGXRHI/Private/AGXTransitionData.cpp
--- a/Engine/Source/Runtime/Apple/AGXRHI/Private/AGXTransitionData.cpp
+++ b/Engine/Source/Runtime/Apple/AGXRHI/Private/AGXTransitionData.cpp
@@ -12,7 +12,7 @@
 
 #pragma mark - AGX RHI Resource Transition Data Definitions -
 
-// TODO: Put AGXRHI resource transition data definitions here
+#define UE_METAL_TRANSITION_DATA_USES_EXPLICIT_FENCING					0
 
 
 //------------------------------------------------------------------------------
@@ -21,9 +21,9 @@
 
 
 FAGXTransitionData::FAGXTransitionData(ERHIPipeline                         InSrcPipelines,
-									   ERHIPipeline                         InDstPipelines,
-									   ERHITransitionCreateFlags            InCreateFlags,
-									   TArrayView<const FRHITransitionInfo> InInfos)
+										   ERHIPipeline                         InDstPipelines,
+										   ERHITransitionCreateFlags            InCreateFlags,
+										   TArrayView<const FRHITransitionInfo> InInfos)
 {
 	SrcPipelines   = InSrcPipelines;
 	DstPipelines   = InDstPipelines;
@@ -32,10 +32,48 @@ FAGXTransitionData::FAGXTransitionData(ERHIPipeline                         InSr
 	bCrossPipeline = (SrcPipelines != DstPipelines);
 
 	Infos.Append(InInfos.GetData(), InInfos.Num());
+
+	// TODO: Determine whether the AGX RHI needs to create a separate, per-transition fence.
+#if UE_METAL_TRANSITION_DATA_USES_EXPLICIT_FENCING
+	if (bCrossPipeline && !EnumHasAnyFlags(CreateFlags, ERHITransitionCreateFlags::NoFence))
+	{
+		// Get the current context pointer.
+		FAGXContext* Context = GetAGXDeviceContext().GetCurrentContext();
+
+		// Get the current render pass fence.
+		TRefCountPtr<FAGXFence> const& MetalFence = Context->GetCurrentRenderPass().End();
+
+		// Write it again as we may wait on this fence in two different encoders.
+		Context->GetCurrentRenderPass().Update(MetalFence);
+
+		// Write it into the transition data.
+		Fence = MetalFence;
+		Fence->AddRef();
+
+		if (GSupportsEfficientAsyncCompute)
+		{
+			static_cast<FAGXRHICommandContext*>(RHIGetDefaultContext())->RHISubmitCommandsHint();
+		}
+	}
+#endif // UE_METAL_TRANSITION_DATA_USES_EXPLICIT_FENCING
 }
 
 void FAGXTransitionData::BeginResourceTransitions() const
 {
+	// TODO: Determine whether the AGX RHI needs to do anything here.
+#if UE_METAL_TRANSITION_DATA_USES_EXPLICIT_FENCING
+	if (Fence.IsValid())
+	{
+		FAGXContext* Context = GetAGXDeviceContext().GetCurrentContext();
+
+		if (Context->GetCurrentCommandBuffer())
+		{
+			Context->SubmitCommandsHint(EAGXSubmitFlagsNone);
+		}
+
+		Context->GetCurrentRenderPass().Begin(Fence);
+	}
+#endif // UE_METAL_TRANSITION_DATA_USES_EXPLICIT_FENCING
 }
 
 void FAGXTransitionData::EndResourceTransitions() const
diff --git a/Engine/Source/Runtime/Apple/AGXRHI/Private/AGXTransitionData.h b/Engine/Source/Runtime/Apple/AGXRHI/Private/AGXTransitionData.h
--- a/Engine/Source/Runtime/Apple/AGXRHI/Private/AGXTransitionData.h
+++ b/Engine/Source/Runtime/Apple/AGXRHI/Private/AGXTransitionData.h
@@ -24,15 +24,15 @@ class FAGXTransitionData
 {
 public:
 	explicit FAGXTransitionData(ERHIPipeline                         InSrcPipelines,
-								ERHIPipeline                         InDstPipelines,
-								ERHITransitionCreateFlags            InCreateFlags,
-								TArrayView<const FRHITransitionInfo> InInfos);
+								  ERHIPipeline                         InDstPipelines,
+								  ERHITransitionCreateFlags            InCreateFlags,
+								  TArrayView<const FRHITransitionInfo> InInfos);
 
 	// The default destructor is sufficient.
 	~FAGXTransitionData() = default;
 
 	// Disallow default, copy and move constructors.
-	FAGXTransitionData()                           = delete;
+	FAGXTransitionData()                             = delete;
 	FAGXTransitionData(const FAGXTransitionData&)  = delete;
 	FAGXTransitionData(const FAGXTransitionData&&) = delete;
 
@@ -47,5 +47,6 @@ private:
 	ERHIPipeline              DstPipelines   = ERHIPipeline::Num;
 	ERHITransitionCreateFlags CreateFlags    = ERHITransitionCreateFlags::None;
 	bool                      bCrossPipeline = false;
-	FAGXTransitionInfoArray   Infos          = {};
+	FAGXTransitionInfoArray Infos          = {};
+	TRefCountPtr<FAGXFence> Fence          = nullptr;
 };
diff --git a/Engine/Source/Runtime/Apple/AGXRHI/Private/AGXVertexBuffer.cpp b/Engine/Source/Runtime/Apple/AGXRHI/Private/AGXVertexBuffer.cpp
--- a/Engine/Source/Runtime/Apple/AGXRHI/Private/AGXVertexBuffer.cpp
+++ b/Engine/Source/Runtime/Apple/AGXRHI/Private/AGXVertexBuffer.cpp
@@ -6,6 +6,7 @@
 
 #include "AGXRHIPrivate.h"
 #include "AGXProfiler.h"
+#include "AGXCommandBuffer.h"
 #include "AGXCommandQueue.h"
 #include "Containers/ResourceArray.h"
 #include "RenderUtils.h"
@@ -107,12 +108,13 @@ void FAGXRHIBuffer::Swap(FAGXRHIBuffer& Other)
 
 static bool CanUsePrivateMemory()
 {
-	return (FAGXCommandQueue::SupportsFeature(EAGXFeaturesEfficientBufferBlits));
+	return (FAGXCommandQueue::SupportsFeature(EAGXFeaturesEfficientBufferBlits) || FAGXCommandQueue::SupportsFeature(EAGXFeaturesIABs));
 }
 
 bool FAGXRHIBuffer::UsePrivateMemory() const
 {
-	return (FAGXCommandQueue::SupportsFeature(EAGXFeaturesEfficientBufferBlits) && EnumHasAnyFlags(Usage, BUF_Dynamic|BUF_Static));
+	return (FAGXCommandQueue::SupportsFeature(EAGXFeaturesEfficientBufferBlits) && EnumHasAnyFlags(Usage, BUF_Dynamic|BUF_Static))
+	|| (FAGXCommandQueue::SupportsFeature(EAGXFeaturesIABs) && EnumHasAnyFlags(Usage, BUF_ShaderResource|BUF_UnorderedAccess));
 }
 
 FAGXRHIBuffer::FAGXRHIBuffer(uint32 InSize, EBufferUsageFlags InUsage, EAGXBufferUsage InAgxUsage, ERHIResourceType InType)
@@ -762,7 +764,7 @@ FBufferRHIRef FAGXDynamicRHI::RHICreateBuffer(uint32 Size, EBufferUsageFlags Usa
 	{
 		check (!Buffer->TransferBuffer);
 
-		if (GAGXBufferZeroFill)
+		if (GAGXBufferZeroFill && !FAGXCommandQueue::SupportsFeature(EAGXFeaturesFences))
 		{
 			for(FAGXRHIBuffer::FAGXBufferAndViews& Backing : Buffer->BufferPool)
 			{
diff --git a/Engine/Source/Runtime/Apple/AGXRHI/Private/AGXViewport.cpp b/Engine/Source/Runtime/Apple/AGXRHI/Private/AGXViewport.cpp
--- a/Engine/Source/Runtime/Apple/AGXRHI/Private/AGXViewport.cpp
+++ b/Engine/Source/Runtime/Apple/AGXRHI/Private/AGXViewport.cpp
@@ -479,22 +479,29 @@ void FAGXViewport::Present(FAGXCommandQueue& CommandQueue, bool bLockToVsync)
 							NSUInteger Width = FMath::Min(Src.GetWidth(), Dst.GetWidth());
 							NSUInteger Height = FMath::Min(Src.GetHeight(), Dst.GetHeight());
 							
-							id<MTLBlitCommandEncoder> Encoder = [CurrentCommandBuffer.GetPtr() blitCommandEncoder];
-							check(Encoder);
-
+							mtlpp::BlitCommandEncoder Encoder = CurrentCommandBuffer.BlitCommandEncoder();
+							check(Encoder.GetPtr());
+#if MTLPP_CONFIG_VALIDATE && METAL_DEBUG_OPTIONS
+							FAGXBlitCommandEncoderDebugging Debugging;
+							if (AGXSafeGetRuntimeDebuggingLevel() >= EAGXDebugLevelFastValidation)
+							{
+								FAGXCommandBufferDebugging CmdDebug = FAGXCommandBufferDebugging::Get(CurrentCommandBuffer);
+								Debugging = FAGXBlitCommandEncoderDebugging(Encoder, CmdDebug);
+							}
+#endif
 							METAL_GPUPROFILE(Profiler->EncodeBlit(Stats, __FUNCTION__));
 
-							[Encoder	copyFromTexture:Src.GetPtr()
-											sourceSlice:0
-											sourceLevel:0
-										   sourceOrigin:MTLOriginMake(0, 0, 0)
-											 sourceSize:MTLSizeMake(Width, Height, 1)
-											  toTexture:Dst.GetPtr()
-									   destinationSlice:0
-									   destinationLevel:0
-									  destinationOrigin:MTLOriginMake(0, 0, 0)];
-
-							[Encoder endEncoding];
+							Encoder.Copy(Src, 0, 0, mtlpp::Origin(0, 0, 0), mtlpp::Size(Width, Height, 1), Dst, 0, 0, mtlpp::Origin(0, 0, 0));
+							METAL_DEBUG_LAYER(EAGXDebugLevelFastValidation, Debugging.Copy(Src, 0, 0, mtlpp::Origin(0, 0, 0), mtlpp::Size(Width, Height, 1), Dst, 0, 0, mtlpp::Origin(0, 0, 0)));
+
+							Encoder.EndEncoding();
+							METAL_DEBUG_LAYER(EAGXDebugLevelFastValidation, Debugging.EndEncoder());
+
+							mtlpp::CommandBufferHandler H = [Src, Dst](const mtlpp::CommandBuffer &) {
+								// void
+							};
+
+							CurrentCommandBuffer.AddCompletedHandler(H);
 
 							[Drawable release];
 							Drawable = nil;
diff --git a/Engine/Source/Runtime/Apple/AGXRHI/Private/Shaders/Types/AGXComputeShader.cpp b/Engine/Source/Runtime/Apple/AGXRHI/Private/Shaders/Types/AGXComputeShader.cpp
--- a/Engine/Source/Runtime/Apple/AGXRHI/Private/Shaders/Types/AGXComputeShader.cpp
+++ b/Engine/Source/Runtime/Apple/AGXRHI/Private/Shaders/Types/AGXComputeShader.cpp
@@ -55,7 +55,7 @@ FAGXShaderPipeline* FAGXComputeShader::GetPipeline()
 		{
 			ns::AutoReleased<ns::Array<mtlpp::PipelineBufferDescriptor>> PipelineBuffers = Descriptor.GetBuffers();
 
-			uint32 ImmutableBuffers = Bindings.ConstantBuffers;
+			uint32 ImmutableBuffers = Bindings.ConstantBuffers | Bindings.ArgumentBuffers;
 			while(ImmutableBuffers)
 			{
 				uint32 Index = __builtin_ctz(ImmutableBuffers);
diff --git a/Engine/Source/Runtime/Apple/AGXRHI/Private/Shaders/Types/Templates/AGXBaseShader.h b/Engine/Source/Runtime/Apple/AGXRHI/Private/Shaders/Types/Templates/AGXBaseShader.h
--- a/Engine/Source/Runtime/Apple/AGXRHI/Private/Shaders/Types/Templates/AGXBaseShader.h
+++ b/Engine/Source/Runtime/Apple/AGXRHI/Private/Shaders/Types/Templates/AGXBaseShader.h
@@ -20,7 +20,7 @@
 #pragma mark - AGX RHI Base Shader Class Support Routines
 
 
-extern mtlpp::LanguageVersion AGXValidateVersion(uint32 Version);
+extern mtlpp::LanguageVersion AGXValidateVersion(uint8 Version);
 
 
 //------------------------------------------------------------------------------
@@ -76,6 +76,12 @@ public:
 	// List of memory copies from RHIUniformBuffer to packed uniforms
 	TArray<CrossCompiler::FUniformBufferCopyInfo> UniformBuffersCopyInfo;
 
+	/* Argument encoders for shader IABs */
+	TMap<uint32, mtlpp::ArgumentEncoder> ArgumentEncoders;
+
+	/* Tier1 Argument buffer bitmasks */
+	TMap<uint32, TBitArray<>> ArgumentBitmasks;
+
 	/* Uniform buffer static slots */
 	TArray<FUniformBufferStaticSlot> StaticSlots;
 
@@ -296,18 +302,34 @@ void TAGXBaseShader<BaseResourceType, ShaderType>::Init(TArrayView<const uint8>
 			mtlpp::LanguageVersion MetalVersion;
 			switch(Header.Version)
 			{
-				case 7:
-					MetalVersion = mtlpp::LanguageVersion::Version2_4;
-					break;
 				case 6:
-					MetalVersion = mtlpp::LanguageVersion::Version2_3;
-					break;
 				case 5:
-					MetalVersion = mtlpp::LanguageVersion::Version2_2;
+				case 4:
+					MetalVersion = mtlpp::LanguageVersion::Version2_1;
+					break;
+				case 3:
+					MetalVersion = mtlpp::LanguageVersion::Version2_0;
+					break;
+				case 2:
+					MetalVersion = mtlpp::LanguageVersion::Version1_2;
+					break;
+				case 1:
+					MetalVersion = mtlpp::LanguageVersion::Version1_1;
+					break;
+				case 0:
+#if PLATFORM_MAC
+					MetalVersion = mtlpp::LanguageVersion::Version1_1;
+#else
+					MetalVersion = mtlpp::LanguageVersion::Version1_0;
+#endif
 					break;
 				default:
 					UE_LOG(LogRHI, Fatal, TEXT("Failed to create shader with unknown version %d: %s"), Header.Version, *FString(NewShaderString));
-					MetalVersion = mtlpp::LanguageVersion::Version2_2;
+#if PLATFORM_MAC
+					MetalVersion = mtlpp::LanguageVersion::Version1_1;
+#else
+					MetalVersion = mtlpp::LanguageVersion::Version1_0;
+#endif
 					break;
 				}
 			CompileOptions.SetLanguageVersion(MetalVersion);
@@ -466,6 +488,30 @@ mtlpp::Function TAGXBaseShader<BaseResourceType, ShaderType>::GetCompiledFunctio
 		}
 	}
 
+	if (FAGXCommandQueue::SupportsFeature(EAGXFeaturesIABs) && Bindings.ArgumentBuffers && ArgumentEncoders.Num() == 0)
+	{
+		uint32 ArgumentBuffers = Bindings.ArgumentBuffers;
+		while(ArgumentBuffers)
+		{
+			uint32 Index = __builtin_ctz(ArgumentBuffers);
+			ArgumentBuffers &= ~(1 << Index);
+
+			mtlpp::ArgumentEncoder ArgumentEncoder = Function.NewArgumentEncoderWithBufferIndex(Index);
+			ArgumentEncoders.Add(Index, ArgumentEncoder);
+
+			TBitArray<> Resources;
+			for (uint8 Id : Bindings.ArgumentBufferMasks[Index])
+			{
+				if (Id >= Resources.Num())
+				{
+					Resources.Add(false, (Id + 1) - Resources.Num());
+				}
+				Resources[Id] = true;
+			}
+			ArgumentBitmasks.Add(Index, Resources);
+		}
+	}
+
 	check(Func);
 	return Func;
 }
