Fix build break.

#rnx

diff --git a/Engine/Source/Runtime/Apple/MetalRHI/Private/MetalShaders.cpp b/Engine/Source/Runtime/Apple/MetalRHI/Private/MetalShaders.cpp
--- a/Engine/Source/Runtime/Apple/MetalRHI/Private/MetalShaders.cpp
+++ b/Engine/Source/Runtime/Apple/MetalRHI/Private/MetalShaders.cpp
@@ -55,35 +55,41 @@ mtlpp::LanguageVersion ValidateVersion(uint32 Version)
 	static uint32 MetaliOSVersions[][3] = {
 		{13,0,0},
 		{14,0,0},
-        {15,0,0},
+		{15,0,0},
 	};
 	static TCHAR const* StandardNames[] =
 	{
 		TEXT("Metal 2.2"),
 		TEXT("Metal 2.3"),
-        TEXT("Metal 2.4"),
+		TEXT("Metal 2.4"),
 	};
 	
 	mtlpp::LanguageVersion Result = mtlpp::LanguageVersion::Version2_2;
 	switch(Version)
 	{
-        case 7:
-            Result = mtlpp::LanguageVersion::Version2_4;
-            break;
+		case 7:
+			Result = mtlpp::LanguageVersion::Version2_4;
+			break;
+		
 		case 6:
 			Result = mtlpp::LanguageVersion::Version2_3;
 			break;
+		
 		case 5:
 			Result = mtlpp::LanguageVersion::Version2_2;
 			break;
-        case 0:
-            Version = 5;
+		
+		case 0:
+			Version = 5;
 			Result = mtlpp::LanguageVersion::Version2_2; // minimum version as of UE5.0
-        default:
-            //EMacMetalShaderStandard::MacMetalSLStandard_Minimum and EIOSMetalShaderStandard::IOSMetalSLStandard_Minimum is currently 2.2
-            UE_LOG(LogTemp, Warning, TEXT("The Metal version currently set is not supported anymore. Set it in the Project Settings. Defaulting to the minimum version."));
-            Version = 5;
-            Result = mtlpp::LanguageVersion::Version2_2;
-            break;
+		
+		default:
+			//EMacMetalShaderStandard::MacMetalSLStandard_Minimum and EIOSMetalShaderStandard::IOSMetalSLStandard_Minimum is currently 2.2
+			UE_LOG(LogTemp, Warning, TEXT("The Metal version currently set is not supported anymore. Set it in the Project Settings. Defaulting to the minimum version."));
+			Version = 5;
+			Result = mtlpp::LanguageVersion::Version2_2;
+			break;
+	}
+	
 	return Result;
 }
