Mac binaries would fail to locate shader libraries when running from sandboxes

It is possible to use the -sandbox= option to mount an alternative file system location from which to load data. On desktop platforms, if the -sandbox= is not specified then Saved/Cooked/[Platform] is implicitly created as a sandbox. This allows runtime binaries to run off cooked data and negate the need to stage. When cooking, the Metal shader compiler creates shader libraries with files in title case. UAT will convert these into lowercase when staging. The Metal RHI is unable to load the cooked libraries as it assumes it is loading lowercase ones from staged data. This change makes sure that all shader libraries are always lowercase. Secondly, the paths to .metallib files where preprocessed in a such a way that bypassed the sandbox file system, resulting in an incorrect native path. Now the API used to convert from IPlatformFile paths to native ones is always used such that sandboxing will get taken into account.

#rnx

diff --git a/Engine/Source/Developer/Apple/MetalShaderFormat/Private/MetalShaderFormat.cpp b/Engine/Source/Developer/Apple/MetalShaderFormat/Private/MetalShaderFormat.cpp
--- a/Engine/Source/Developer/Apple/MetalShaderFormat/Private/MetalShaderFormat.cpp
+++ b/Engine/Source/Developer/Apple/MetalShaderFormat/Private/MetalShaderFormat.cpp
@@ -137,6 +137,7 @@ public:
 
 		bool bOK = false;
 		FString LibraryPlatformName = FString::Printf(TEXT("%s_%s"), *LibraryName, *Format.GetPlainNameString());
+		LibraryPlatformName.ToLowerInline();
 		volatile int32 CompiledLibraries = 0;
 		TArray<FGraphEventRef> Tasks;
 
@@ -219,6 +220,7 @@ public:
 		if (CompiledLibraries == SubLibraries.Num())
 		{
 			FString BinaryShaderFile = (OutputDir / LibraryPlatformName) + FMetalCompilerToolchain::MetalMapExtension;
+			BinaryShaderFile.ToLowerInline();
 			FArchive* BinaryShaderAr = IFileManager::Get().CreateFileWriter(*BinaryShaderFile);
 			if (BinaryShaderAr != NULL)
 			{
diff --git a/Engine/Source/Runtime/Apple/AGXRHI/Private/AGXDynamicRHI_Shaders.cpp b/Engine/Source/Runtime/Apple/AGXRHI/Private/AGXDynamicRHI_Shaders.cpp
--- a/Engine/Source/Runtime/Apple/AGXRHI/Private/AGXDynamicRHI_Shaders.cpp
+++ b/Engine/Source/Runtime/Apple/AGXRHI/Private/AGXDynamicRHI_Shaders.cpp
@@ -103,10 +103,7 @@ FRHIShaderLibraryRef FAGXDynamicRHI::RHICreateShaderLibrary(EShaderPlatform Plat
 				for (uint32 i = 0; i < Header.NumLibraries; i++)
 				{
 					FString MetalLibraryFilePath = (FilePath / LibName) + FString::Printf(TEXT(".%d.metallib"), i);
-					MetalLibraryFilePath = FPaths::ConvertRelativePathToFull(MetalLibraryFilePath);
-#if !PLATFORM_MAC
 					MetalLibraryFilePath = IFileManager::Get().ConvertToAbsolutePathForExternalAppForRead(*MetalLibraryFilePath);
-#endif // !PLATFORM_MAC
 
 					METAL_GPUPROFILE(FAGXScopedCPUStats CPUStat(FString::Printf(TEXT("NewLibraryFile: %s"), *MetalLibraryFilePath)));
 					NSError* Error;
diff --git a/Engine/Source/Runtime/Apple/MetalRHI/Private/MetalDynamicRHI_Shaders.cpp b/Engine/Source/Runtime/Apple/MetalRHI/Private/MetalDynamicRHI_Shaders.cpp
--- a/Engine/Source/Runtime/Apple/MetalRHI/Private/MetalDynamicRHI_Shaders.cpp
+++ b/Engine/Source/Runtime/Apple/MetalRHI/Private/MetalDynamicRHI_Shaders.cpp
@@ -103,10 +103,7 @@ FRHIShaderLibraryRef FMetalDynamicRHI::RHICreateShaderLibrary(EShaderPlatform Pl
 				for (uint32 i = 0; i < Header.NumLibraries; i++)
 				{
 					FString MetalLibraryFilePath = (FilePath / LibName) + FString::Printf(TEXT(".%d.metallib"), i);
-					MetalLibraryFilePath = FPaths::ConvertRelativePathToFull(MetalLibraryFilePath);
-#if !PLATFORM_MAC
 					MetalLibraryFilePath = IFileManager::Get().ConvertToAbsolutePathForExternalAppForRead(*MetalLibraryFilePath);
-#endif // !PLATFORM_MAC
 
 					METAL_GPUPROFILE(FScopedMetalCPUStats CPUStat(FString::Printf(TEXT("NewLibraryFile: %s"), *MetalLibraryFilePath)));
 					NSError* Error;
