Replace lost break
#rnx

diff --git a/Engine/Source/Runtime/Apple/MetalRHI/Private/MetalShaders.cpp b/Engine/Source/Runtime/Apple/MetalRHI/Private/MetalShaders.cpp
--- a/Engine/Source/Runtime/Apple/MetalRHI/Private/MetalShaders.cpp
+++ b/Engine/Source/Runtime/Apple/MetalRHI/Private/MetalShaders.cpp
@@ -82,6 +82,7 @@ mtlpp::LanguageVersion ValidateVersion(uint32 Version)
 		case 0:
 			Version = 5;
 			Result = mtlpp::LanguageVersion::Version2_2; // minimum version as of UE5.0
+			break;
 		
 		default:
 			//EMacMetalShaderStandard::MacMetalSLStandard_Minimum and EIOSMetalShaderStandard::IOSMetalSLStandard_Minimum is currently 2.2
