Account for minimum alignment when sub-allocating buffers for linear texture allocation.

- This fix addresses sub-buffer allocations. We could get into a situation where alignment restrictions would cause us to oversubscribe (allocate beyond the bounds) the buffer. This was discovered while investigating UE-140266.

#rnx

diff --git a/Engine/Source/Runtime/Apple/MetalRHI/Private/MetalTexture.cpp b/Engine/Source/Runtime/Apple/MetalRHI/Private/MetalTexture.cpp
--- a/Engine/Source/Runtime/Apple/MetalRHI/Private/MetalTexture.cpp
+++ b/Engine/Source/Runtime/Apple/MetalRHI/Private/MetalTexture.cpp
@@ -923,15 +923,14 @@ FMetalSurface::FMetalSurface(ERHIResourceType ResourceType, EPixelFormat Format,
 		else
 		{
 			mtlpp::Device Device = GetMetalDeviceContext().GetDevice();
-			mtlpp::SizeAndAlign SizeAlign = Device.HeapTextureSizeAndAlign(Desc);
-
-			// Backing buffer resource options must match the texture we are going to create from it
-			FMetalPooledBufferArgs Args(Device, SizeAlign.Size, BUF_Dynamic, mtlpp::StorageMode::Private, Desc.GetCpuCacheMode());
-			FMetalBuffer Buffer = GetMetalDeviceContext().CreatePooledBuffer(Args);
-
+			
 			const uint32 MinimumByteAlignment = GetMetalDeviceContext().GetDevice().GetMinimumLinearTextureAlignmentForPixelFormat(MTLFormat);
 			const NSUInteger BytesPerRow = Align(Desc.GetWidth() * GPixelFormats[Format].BlockBytes, MinimumByteAlignment);
 			
+			// Backing buffer resource options must match the texture we are going to create from it
+			FMetalPooledBufferArgs Args(Device, BytesPerRow * Desc.GetHeight(), BUF_Dynamic, mtlpp::StorageMode::Private, Desc.GetCpuCacheMode());
+			FMetalBuffer Buffer = GetMetalDeviceContext().CreatePooledBuffer(Args);
+
 			Texture = Buffer.NewTexture(Desc, 0, BytesPerRow);
 		}
 		
