Always set the incoming renderpass descriptor when requested from the higher-level renderer code.

- This is a partial fix to address a number of ongoing and somewhat related issues. We cannot ignore the render pass information being sent down to us from the higher-level code now that RDG is driving things.

#rnx

diff --git a/Engine/Source/Runtime/Apple/MetalRHI/Private/MetalRHIContext.cpp b/Engine/Source/Runtime/Apple/MetalRHI/Private/MetalRHIContext.cpp
--- a/Engine/Source/Runtime/Apple/MetalRHI/Private/MetalRHIContext.cpp
+++ b/Engine/Source/Runtime/Apple/MetalRHI/Private/MetalRHIContext.cpp
@@ -143,31 +143,25 @@ FMetalRHIImmediateCommandContext::FMetalRHIImmediateCommandContext(class FMetalP
 
 void FMetalRHICommandContext::RHIBeginRenderPass(const FRHIRenderPassInfo& InInfo, const TCHAR* InName)
 {
-	@autoreleasepool {
-	bool bHasTarget = (InInfo.DepthStencilRenderTarget.DepthStencilTarget != nullptr || InInfo.GetNumColorRenderTargets() > 0);
+	SCOPED_AUTORELEASE_POOL;
 	
 	if (InInfo.bOcclusionQueries)
 	{
 		Context->GetCommandList().SetParallelIndex(0, 0);
 	}
 
-	// Ignore any attempt to "clear" the render-targets as that is senseless with the way MetalRHI has to try and coalesce passes.
-	if (bHasTarget)
-	{
-		Context->SetRenderPassInfo(InInfo);
+	Context->SetRenderPassInfo(InInfo);
 
-		// Set the viewport to the full size of render target 0.
-		if (InInfo.ColorRenderTargets[0].RenderTarget)
-		{
-			const FRHIRenderPassInfo::FColorEntry& RenderTargetView = InInfo.ColorRenderTargets[0];
-			FMetalSurface* RenderTarget = GetMetalSurfaceFromRHITexture(RenderTargetView.RenderTarget);
+	// Set the viewport to the full size of render target 0.
+	if (InInfo.ColorRenderTargets[0].RenderTarget)
+	{
+		const FRHIRenderPassInfo::FColorEntry& RenderTargetView = InInfo.ColorRenderTargets[0];
+		FMetalSurface* RenderTarget = GetMetalSurfaceFromRHITexture(RenderTargetView.RenderTarget);
 
-			uint32 Width = FMath::Max((uint32)(RenderTarget->Texture.GetWidth() >> RenderTargetView.MipIndex), (uint32)1);
-			uint32 Height = FMath::Max((uint32)(RenderTarget->Texture.GetHeight() >> RenderTargetView.MipIndex), (uint32)1);
+		uint32 Width = FMath::Max((uint32)(RenderTarget->Texture.GetWidth() >> RenderTargetView.MipIndex), (uint32)1);
+		uint32 Height = FMath::Max((uint32)(RenderTarget->Texture.GetHeight() >> RenderTargetView.MipIndex), (uint32)1);
 
-			RHISetViewport(0.0f, 0.0f, 0.0f, (float)Width, (float)Height, 1.0f);
-		}
-	}
+		RHISetViewport(0.0f, 0.0f, 0.0f, (float)Width, (float)Height, 1.0f);
 	}
 	
 	RenderPassInfo = InInfo;
