Fixing IOS build

diff --git a/Engine/Source/Runtime/Apple/AudioMixerAudioUnit/Private/AudioMixerPlatformAudioUnit.cpp b/Engine/Source/Runtime/Apple/AudioMixerAudioUnit/Private/AudioMixerPlatformAudioUnit.cpp
--- a/Engine/Source/Runtime/Apple/AudioMixerAudioUnit/Private/AudioMixerPlatformAudioUnit.cpp
+++ b/Engine/Source/Runtime/Apple/AudioMixerAudioUnit/Private/AudioMixerPlatformAudioUnit.cpp
@@ -397,19 +397,15 @@ namespace Audio
 		{
 			RuntimeFormat = Audio::NAME_ADPCM;
 		}
-		else if (InRuntimeFormat == Audio::NAME_BINKA)
-		{
-			return new FBinkAudioInfo();
-		}
 
 		return RuntimeFormat;
 	}
 
-	ICompressedAudioInfo* FMixerPlatformCoreAudio::CreateCompressedAudioInfo(const FName& InRuntimeFormat) const
+	ICompressedAudioInfo* FMixerPlatformAudioUnit::CreateCompressedAudioInfo(const FName& InRuntimeFormat) const
 	{
 		ICompressedAudioInfo* Decoder = nullptr;
 
-		if (RuntimeFormat == Audio::NAME_PLATFORM_SPECIFIC)
+		if (InRuntimeFormat == Audio::NAME_PLATFORM_SPECIFIC)
 		{
 			Decoder = new FADPCMAudioInfo();
 		}	
diff --git a/Engine/Source/Runtime/Apple/AudioMixerAudioUnit/Private/AudioMixerPlatformAudioUnit.h b/Engine/Source/Runtime/Apple/AudioMixerAudioUnit/Private/AudioMixerPlatformAudioUnit.h
--- a/Engine/Source/Runtime/Apple/AudioMixerAudioUnit/Private/AudioMixerPlatformAudioUnit.h
+++ b/Engine/Source/Runtime/Apple/AudioMixerAudioUnit/Private/AudioMixerPlatformAudioUnit.h
@@ -36,6 +36,7 @@ namespace Audio
 		virtual FAudioPlatformDeviceInfo GetPlatformDeviceInfo() const override;
 		virtual void SubmitBuffer(const uint8* Buffer) override;
 		virtual FName GetRuntimeFormat(const USoundWave* InSoundWave) const override;
+		virtual ICompressedAudioInfo* CreateCompressedAudioInfo(const FName& InRuntimeFormat) override;
 		virtual FString GetDefaultDeviceName() override;
 		virtual FAudioPlatformSettings GetPlatformSettings() const override;
 		virtual int32 GetNumFrames(const int32 InNumReqestedFrames) override;
diff --git a/Engine/Source/Runtime/IOS/IOSAudio/Private/IOSAudioBuffer.cpp b/Engine/Source/Runtime/IOS/IOSAudio/Private/IOSAudioBuffer.cpp
--- a/Engine/Source/Runtime/IOS/IOSAudio/Private/IOSAudioBuffer.cpp
+++ b/Engine/Source/Runtime/IOS/IOSAudio/Private/IOSAudioBuffer.cpp
@@ -31,8 +31,6 @@ FIOSAudioSoundBuffer::FIOSAudioSoundBuffer(FIOSAudioDevice* InAudioDevice, USoun
 {
 	if (!bIsProcedural)
 	{
-		DecompressionState = static_cast<FADPCMAudioInfo*>(InAudioDevice->CreateCompressedAudioInfo(InWave));
-
 		if (!ReadCompressedInfo(InWave))
 		{
 			return;
diff --git a/Engine/Source/Runtime/IOS/IOSAudio/Private/IOSAudioDevice.cpp b/Engine/Source/Runtime/IOS/IOSAudio/Private/IOSAudioDevice.cpp
--- a/Engine/Source/Runtime/IOS/IOSAudio/Private/IOSAudioDevice.cpp
+++ b/Engine/Source/Runtime/IOS/IOSAudio/Private/IOSAudioDevice.cpp
@@ -292,12 +292,6 @@ void FIOSAudioDevice::HandleError(const TCHAR* InLogOutput, bool bTeardown)
 	}
 }
 
-class ICompressedAudioInfo* FIOSAudioDevice::CreateCompressedAudioInfo(const FName& InRuntimeFormat) const
-{
-	return new FADPCMAudioInfo();
-}
-
-
 FAudioEffectsManager* FIOSAudioDevice::CreateEffectsManager()
 {
 	// Create the basic no-op effects manager
diff --git a/Engine/Source/Runtime/IOS/IOSAudio/Public/IOSAudioDevice.h b/Engine/Source/Runtime/IOS/IOSAudio/Public/IOSAudioDevice.h
--- a/Engine/Source/Runtime/IOS/IOSAudio/Public/IOSAudioDevice.h
+++ b/Engine/Source/Runtime/IOS/IOSAudio/Public/IOSAudioDevice.h
@@ -249,9 +249,7 @@ private:
 	// Used to support adpcm streaming
 	virtual bool HasCompressedAudioInfoClass(USoundWave* SoundWave) override { return true; }
 
-	virtual bool SupportsRealtimeDecompression() const override { return true; }
-	
-	virtual ICompressedAudioInfo* CreateCompressedAudioInfo(const FName& InRuntimeFormat) const override;
+	virtual bool SupportsRealtimeDecompression() const override { return true; }	
 
 	void HandleError(const TCHAR* InLogOutput, bool bTeardown = false);
 
