Fixing IOS build

diff --git a/Engine/Source/Runtime/Apple/AudioMixerAudioUnit/Private/AudioMixerPlatformAudioUnit.h b/Engine/Source/Runtime/Apple/AudioMixerAudioUnit/Private/AudioMixerPlatformAudioUnit.h
--- a/Engine/Source/Runtime/Apple/AudioMixerAudioUnit/Private/AudioMixerPlatformAudioUnit.h
+++ b/Engine/Source/Runtime/Apple/AudioMixerAudioUnit/Private/AudioMixerPlatformAudioUnit.h
@@ -36,7 +36,7 @@ namespace Audio
 		virtual FAudioPlatformDeviceInfo GetPlatformDeviceInfo() const override;
 		virtual void SubmitBuffer(const uint8* Buffer) override;
 		virtual FName GetRuntimeFormat(const USoundWave* InSoundWave) const override;
-		virtual ICompressedAudioInfo* CreateCompressedAudioInfo(const FName& InRuntimeFormat) override;
+		virtual ICompressedAudioInfo* CreateCompressedAudioInfo(const FName& InRuntimeFormat) const override;		
 		virtual FString GetDefaultDeviceName() override;
 		virtual FAudioPlatformSettings GetPlatformSettings() const override;
 		virtual int32 GetNumFrames(const int32 InNumReqestedFrames) override;
