Fix DumpGPU on macOS.  Begin Resource Dump occurs in: FSlateApplication::Get().Tick(ESlateTickType::PlatformAndInput); which is after GEngine->Tick(FApp::GetDeltaTime(), bIdleMode);.  Make the defered dump command request a GPU dump for the next frame by adding an InitDump function so it's clear for any platform when to begin.  Could add 1 when setting the DumpingFrameCounter_GameThread but that seems brittle across platforms.

Allow more flexibility of the Metal RHICopyToResolveTarget to include compatable texture view pixelformats.  Fixes validation error when resolving between sRGB and RGB formats.

FMetalContext: don't assert then go into the weeds on macOS when there are no render targets.

#rnx

diff --git a/Engine/Source/Runtime/Apple/MetalRHI/Private/MetalContext.cpp b/Engine/Source/Runtime/Apple/MetalRHI/Private/MetalContext.cpp
--- a/Engine/Source/Runtime/Apple/MetalRHI/Private/MetalContext.cpp
+++ b/Engine/Source/Runtime/Apple/MetalRHI/Private/MetalContext.cpp
@@ -1267,14 +1267,10 @@ bool FMetalContext::PrepareToDraw(uint32 PrimitiveType)
 	check(IsValidRef(CurrentPSO));
 	
 	// Enforce calls to SetRenderTarget prior to issuing draw calls.
-#if PLATFORM_MAC
-	check(StateCache.GetHasValidRenderTarget());
-#else
 	if (!StateCache.GetHasValidRenderTarget())
 	{
 		return false;
 	}
-#endif
 	
 	FMetalHashedVertexDescriptor const& VertexDesc = CurrentPSO->VertexDeclaration->Layout;
 	
diff --git a/Engine/Source/Runtime/Apple/MetalRHI/Private/MetalRenderTarget.cpp b/Engine/Source/Runtime/Apple/MetalRHI/Private/MetalRenderTarget.cpp
--- a/Engine/Source/Runtime/Apple/MetalRHI/Private/MetalRenderTarget.cpp
+++ b/Engine/Source/Runtime/Apple/MetalRHI/Private/MetalRenderTarget.cpp
@@ -134,10 +134,31 @@ void FMetalRHICommandContext::RHICopyToResolveTarget(FRHITexture* SourceTextureR
 			
 			Context->CopyFromTextureToTexture(Source->MSAAResolveTexture, SrcIndex, ResolveParams.MipIndex, Origin, Size, Destination->Texture, DestIndex, ResolveParams.MipIndex, Origin);
 		}
-		else
+		else if(Source->Texture.GetPixelFormat() == Destination->Texture.GetPixelFormat())
 		{
+			// Blit Copy for matching formats
 			Context->CopyFromTextureToTexture(Source->Texture, SrcIndex, ResolveParams.MipIndex, Origin, Size, Destination->Texture, DestIndex, ResolveParams.MipIndex, Origin);
 		}
+		else
+		{
+			const FPixelFormatInfo& SourceFormatInfo = GPixelFormats[Source->PixelFormat];
+			const FPixelFormatInfo& DestFormatInfo = GPixelFormats[Destination->PixelFormat];
+			bool bUsingPixelFormatView = (Source->Texture.GetUsage() & mtlpp::TextureUsage::PixelFormatView) != 0;
+			
+			// Attempt to Resolve with a texture view - source Texture doesn't have to be created with MTLTextureUsagePixelFormatView for these cases e.g:
+			// If we are resolving to/from sRGB linear color space within the same format OR using same bit length color format
+			if	(	SourceFormatInfo.BlockBytes == DestFormatInfo.BlockBytes &&
+					(bUsingPixelFormatView || SourceFormatInfo.NumComponents == DestFormatInfo.NumComponents)
+				)
+			{
+				FMetalTexture SourceTextureView = Source->Texture.NewTextureView(Destination->Texture.GetPixelFormat(), Source->Texture.GetTextureType(), ns::Range(ResolveParams.MipIndex, 1), ns::Range(SrcIndex, 1));
+				if(SourceTextureView)
+				{
+					Context->CopyFromTextureToTexture(SourceTextureView, 0, 0, Origin, Size, Destination->Texture, DestIndex, ResolveParams.MipIndex, Origin);
+					SafeReleaseMetalTexture(SourceTextureView);
+				}
+			}
+		}
 
 #if PLATFORM_MAC
 		if((Destination->GPUReadback & FMetalSurface::EMetalGPUReadbackFlags::ReadbackRequested) != 0)
diff --git a/Engine/Source/Runtime/Launch/Private/LaunchEngineLoop.cpp b/Engine/Source/Runtime/Launch/Private/LaunchEngineLoop.cpp
--- a/Engine/Source/Runtime/Launch/Private/LaunchEngineLoop.cpp
+++ b/Engine/Source/Runtime/Launch/Private/LaunchEngineLoop.cpp
@@ -5206,6 +5206,11 @@ void FEngineLoop::Tick()
                 SlateApp.FinishedInputThisFrame();
             }
 		}
+		
+		// init for RDG resource dump
+		#if WITH_ENGINE && WITH_DUMPGPU
+			FRDGBuilder::InitResourceDump();
+		#endif
 
 		// main game engine tick (world, game objects, etc.)
 		GEngine->Tick(FApp::GetDeltaTime(), bIdleMode);
diff --git a/Engine/Source/Runtime/RenderCore/Private/RenderGraphDumpResources.cpp b/Engine/Source/Runtime/RenderCore/Private/RenderGraphDumpResources.cpp
--- a/Engine/Source/Runtime/RenderCore/Private/RenderGraphDumpResources.cpp
+++ b/Engine/Source/Runtime/RenderCore/Private/RenderGraphDumpResources.cpp
@@ -1193,6 +1193,8 @@ struct FRDGResourceDumpContext
 		return bDumpPass;
 	}
 };
+
+// 0 = not dumping, MAX_uint64 dump request for next frame, otherwise dump frame counter
 static uint64 DumpingFrameCounter_GameThread = 0;
 FRDGResourceDumpContext GRDGResourceDumpContext;
 
@@ -1203,6 +1205,14 @@ bool IsDumpingRDGResources()
 	return GRDGResourceDumpContext.IsDumpingFrame();
 }
 
+void FRDGBuilder::InitResourceDump()
+{
+	if(DumpingFrameCounter_GameThread == MAX_uint64)
+	{
+		DumpingFrameCounter_GameThread = GFrameCounter;
+	}
+}
+
 FString FRDGBuilder::BeginResourceDump(const TArray<FString>& Args)
 {
 	check(IsInGameThread());
@@ -1319,7 +1329,8 @@ FString FRDGBuilder::BeginResourceDump(const TArray<FString>& Args)
 		#endif
 	});
 
-	DumpingFrameCounter_GameThread = GFrameCounter;
+	// Mark ready for dump on next available frame
+	DumpingFrameCounter_GameThread = MAX_uint64;
 
 	if (NewResourceDumpContext.bEnableDiskWrite)
 	{
@@ -1332,8 +1343,10 @@ void FRDGBuilder::EndResourceDump()
 {
 	check(IsInGameThread());
 
+	 // make sure at least one frame has passed since we start a resource dump and we are not waiting on the dump to begin
 	if (DumpingFrameCounter_GameThread == 0 ||
-		DumpingFrameCounter_GameThread >= GFrameCounter) // make sure at least one frame has passed since we start a resource dump
+		DumpingFrameCounter_GameThread == MAX_uint64 ||
+		DumpingFrameCounter_GameThread >= GFrameCounter)
 	{
 		return;
 	}
diff --git a/Engine/Source/Runtime/RenderCore/Public/RenderGraphBuilder.h b/Engine/Source/Runtime/RenderCore/Public/RenderGraphBuilder.h
--- a/Engine/Source/Runtime/RenderCore/Public/RenderGraphBuilder.h
+++ b/Engine/Source/Runtime/RenderCore/Public/RenderGraphBuilder.h
@@ -289,6 +289,7 @@ public:
 
 #if RDG_DUMP_RESOURCES
 	static FString BeginResourceDump(const TArray<FString>& Args);
+	static void InitResourceDump();
 	static void EndResourceDump();
 	static bool IsDumpingFrame();
 #else
@@ -770,4 +771,4 @@ private:
 	#define RDG_GPU_MASK_SCOPE(GraphBuilder, GPUMask)
 #endif
 
-#include "RenderGraphBuilder.inl"
\ No newline at end of file
+#include "RenderGraphBuilder.inl"
