Added support for Metal Binary PSO Cache

diff --git a/Engine/Source/Runtime/Apple/AGXRHI/Private/AGXPipeline.cpp b/Engine/Source/Runtime/Apple/AGXRHI/Private/AGXPipeline.cpp
--- a/Engine/Source/Runtime/Apple/AGXRHI/Private/AGXPipeline.cpp
+++ b/Engine/Source/Runtime/Apple/AGXRHI/Private/AGXPipeline.cpp
@@ -31,6 +31,12 @@ static FAutoConsoleVariableRef CVarAGXCacheMinSize(
 	GAGXCacheMinSize,
 	TEXT("Sets the minimum size that we expect the metal OS cache to be (in MB). This is used to determine if we need to cache PSOs again (Default: 32).\n"), ECVF_ReadOnly);
 
+static int32 GAGXBinaryCacheDebugOutput = 0;
+static FAutoConsoleVariableRef CVarAGXBinaryCacheDebugOutput(
+    TEXT("rhi.AGX.BinaryCacheDebugOutput"),
+    GAGXBinaryCacheDebugOutput,
+    TEXT("Enable to output logging information for PSO Binary cache default(0) \n"), ECVF_ReadOnly);
+
 static uint32 BlendBitOffsets[] = { Offset_BlendState0, Offset_BlendState1, Offset_BlendState2, Offset_BlendState3, Offset_BlendState4, Offset_BlendState5, Offset_BlendState6, Offset_BlendState7 };
 static uint32 RTBitOffsets[] = { Offset_RenderTargetFormat0, Offset_RenderTargetFormat1, Offset_RenderTargetFormat2, Offset_RenderTargetFormat3, Offset_RenderTargetFormat4, Offset_RenderTargetFormat5, Offset_RenderTargetFormat6, Offset_RenderTargetFormat7 };
 static_assert(Offset_RasterEnd < 64 && Offset_End < 128, "Offset_RasterEnd must be < 64 && Offset_End < 128");
@@ -704,8 +710,171 @@ static bool ConfigureRenderPipelineDescriptor(mtlpp::RenderPipelineDescriptor& R
 	return true;
 }
 
+
+
+/*
+* PSO Harvesting and Reuse
+*
+* Usage:
+*
+* To Harvest, run the game with -AGXPSOCache=recreate
+* All Render and Compute PSOs created will be harvested into the MTLBinaryArchive
+* Console command r.Agx.ShaderPipelineCache.Save will trigger the serialization to file.
+* The binary archive's location will be printed to the log.
+*
+* To reuse, run the game with -AGXPSOCache=use
+* The binary archive will be opened from the saved location.
+* The binary archive can be moved to another device, as long as it's the same GPU
+* and OS build.
+*
+*/
+
+enum class CacheMode
+{
+	Uninitialized,
+	Recreate,
+	Append,
+	Use,
+	Ignore
+};
+
+static CacheMode GPSOCacheMode = CacheMode::Uninitialized;
+static id< MTLBinaryArchive > GPSOBinaryArchive = nil;
+static uint32_t GPSOHarvestCount = 0;
+
+static NSURL * PipelineCacheSaveLocation()
+{
+	NSString * path = NSSearchPathForDirectoriesInDomains(NSDocumentDirectory, NSUserDomainMask, YES)[0];
+	if (!path)
+	{
+		UE_LOG(LogAGX, Error, TEXT("AGX Pipeline Cache: Unable to find Documents folder"));
+	}
+	NSURL * url = [NSURL fileURLWithPath : [NSString stringWithFormat : @"%@/mtlarchive.bin", path] ];
+	return url;
+}
+
+static void InitializeAGXPipelineCache()
+{
+	FString strCacheMode;
+	FParse::Value(FCommandLine::Get(), TEXT("AGXPSOCache="), strCacheMode);
+
+	if (strCacheMode.Compare(TEXT("recreate"), ESearchCase::IgnoreCase) == 0)
+	{
+		GPSOCacheMode = CacheMode::Recreate;
+		UE_LOG(LogAGX, Log, TEXT("AGX Pipeline Cache: recreate PSO cache"));
+	}
+	else if (strCacheMode.Compare(TEXT("append"), ESearchCase::IgnoreCase) == 0)
+	{
+		GPSOCacheMode = CacheMode::Append;
+		UE_LOG(LogAGX, Log, TEXT("AGX Pipeline Cache: append to PSO cache"));
+	}
+	else if (strCacheMode.Compare(TEXT("use"), ESearchCase::IgnoreCase) == 0)
+	{
+		GPSOCacheMode = CacheMode::Use;
+		UE_LOG(LogAGX, Log, TEXT("AGX Pipeline Cache: use PSO cache"));
+	}
+	else
+	{
+		GPSOCacheMode = CacheMode::Ignore;
+		UE_LOG(LogAGX, Log, TEXT("AGX Pipeline Cache: ignore PSO cache"));
+	}
+
+	if (GPSOCacheMode != CacheMode::Ignore)
+	{
+		NSURL* url = PipelineCacheSaveLocation();
+		UE_LOG(LogAGX, Log, TEXT("AGX Pipeline Cache: pso cache save location will be: %s"), *FString(url.absoluteString.UTF8String));
+		MTLBinaryArchiveDescriptor* archDesc = [MTLBinaryArchiveDescriptor new];
+		archDesc.url = ((GPSOCacheMode == CacheMode::Append) || (GPSOCacheMode == CacheMode::Use)) ? PipelineCacheSaveLocation() : nil;
+		id< MTLDevice > mtlDevice = GMtlDevice;
+		__autoreleasing NSError* err = nil;
+		GPSOBinaryArchive = [mtlDevice newBinaryArchiveWithDescriptor : archDesc error : &err];
+		if (err)
+		{
+			UE_LOG(LogAGX, Error, TEXT("Error adding Pipeline Functions to Binary Archive: %s"), *FString(err.localizedDescription.UTF8String));
+		}
+	}
+}
+
+static void RelatePipelineStateToCache(const mtlpp::RenderPipelineDescriptor & PipelineDesc, NSUInteger * outPipelineOpts)
+{
+	if (GPSOBinaryArchive && (GPSOCacheMode != CacheMode::Ignore))
+	{
+		if (GPSOCacheMode == CacheMode::Recreate || GPSOCacheMode == CacheMode::Append)
+		{
+			__autoreleasing NSError * err = nil;
+			bool bAddedBinaryPSO = [GPSOBinaryArchive addRenderPipelineFunctionsWithDescriptor : PipelineDesc.GetPtr() error : &err];
+			if (err)
+			{
+				UE_LOG(LogAGX, Warning, TEXT("AGX Pipeline Cache: Error adding Pipeline Functions to Binary Archive: %s"), *FString(err.localizedDescription.UTF8String));
+			}
+            else if(bAddedBinaryPSO)
+            {
+                GPSOHarvestCount++;
+                if(GAGXBinaryCacheDebugOutput && GPSOHarvestCount % 100)
+                {
+                    UE_LOG(LogAGX, Log, TEXT("AGX Pipeline Cache: Harvested PSO count: %d"), GPSOHarvestCount);
+                }
+            }
+		}
+	}
+}
+
+static void RelatePipelineStateToCache(const mtlpp::ComputePipelineDescriptor & PipelineDesc, NSUInteger * outPipelineOpts)
+{
+	if (GPSOBinaryArchive && (GPSOCacheMode != CacheMode::Ignore))
+	{
+		if (GPSOCacheMode == CacheMode::Recreate || GPSOCacheMode == CacheMode::Append)
+		{
+			__autoreleasing NSError* err = nil;
+            bool bAddedBinaryPSO = [GPSOBinaryArchive addComputePipelineFunctionsWithDescriptor : PipelineDesc.GetPtr() error : &err] ;
+			if (err)
+			{
+				UE_LOG(LogAGX, Warning, TEXT("AGX Pipeline Cache: Error adding Pipeline Functions to Binary Archive: %s"), *FString(err.localizedDescription.UTF8String));
+			}
+            else if(bAddedBinaryPSO)
+            {
+                GPSOHarvestCount++;
+                if(GAGXBinaryCacheDebugOutput && GPSOHarvestCount % 100)
+                {
+                    UE_LOG(LogAGX, Log, TEXT("AGX Pipeline Cache: Harvested PSO count: %d"), GPSOHarvestCount);
+                }
+            }
+		}
+	}
+}
+
+void AGXConsoleCommandSavePipelineFileCache()
+{
+	UE_LOG(LogAGX, Log, TEXT("AGX Pipeline Cache: requesting PSO save..."));
+
+	if (GPSOBinaryArchive && (GPSOCacheMode == CacheMode::Recreate || GPSOCacheMode == CacheMode::Append))
+	{
+		NSURL * url = PipelineCacheSaveLocation();
+		UE_LOG(LogAGX, Log, TEXT("AGX Pipeline Cache: Serialize harvested PSOs to: %s"), *FString(url.absoluteString.UTF8String));
+        UE_LOG(LogAGX, Log, TEXT("AGX Pipeline Cache: Serialized PSO Count: %d"), GPSOHarvestCount);
+        
+        __autoreleasing NSError * err = nil;
+		[GPSOBinaryArchive serializeToURL : url error : &err];
+				
+		if (err)
+		{
+			UE_LOG(LogAGX, Error, TEXT("AGX Pipeline Cache: Error Serializing binary archive: %s"), *FString(err.localizedDescription.UTF8String));
+		}
+	}
+}
+
+static FAutoConsoleCommand SavePipelineCacheCmd(
+        TEXT("rhi.Agx.ShaderPipelineCache.Save"),
+        TEXT("Save the current pipeline file cache."),
+        FConsoleCommandDelegate::CreateStatic(AGXConsoleCommandSavePipelineFileCache));
+
 static FAGXShaderPipeline* CreateMTLRenderPipeline(bool const bSync, FAGXGraphicsPipelineKey const& Key, const FGraphicsPipelineStateInitializer& Init)
 {
+	if (GPSOCacheMode == CacheMode::Uninitialized)
+	{
+		InitializeAGXPipelineCache();
+	}
+
     FAGXVertexShader* VertexShader = (FAGXVertexShader*)Init.BoundShaderState.VertexShaderRHI;
     FAGXPixelShader* PixelShader = (FAGXPixelShader*)Init.BoundShaderState.PixelShaderRHI;
     
@@ -763,6 +932,7 @@ static FAGXShaderPipeline* CreateMTLRenderPipeline(bool const bSync, FAGXGraphic
 		{
 			ns::AutoReleasedError RenderError;
 			METAL_GPUPROFILE(FAGXScopedCPUStats CPUStat(FString::Printf(TEXT("NewRenderPipeline: %s"), TEXT("")/**FString([RenderPipelineDesc.GetPtr() description])*/)));
+			RelatePipelineStateToCache(RenderPipelineDesc, &RenderOption);
 			Pipeline->RenderPipelineState = GMtlppDevice.NewRenderPipelineState(RenderPipelineDesc, (mtlpp::PipelineOption)RenderOption, Reflection, &RenderError);
 			if (Reflection)
 			{
diff --git a/Engine/Source/Runtime/Apple/MetalRHI/Private/MetalPipeline.cpp b/Engine/Source/Runtime/Apple/MetalRHI/Private/MetalPipeline.cpp
--- a/Engine/Source/Runtime/Apple/MetalRHI/Private/MetalPipeline.cpp
+++ b/Engine/Source/Runtime/Apple/MetalRHI/Private/MetalPipeline.cpp
@@ -32,6 +32,12 @@ static FAutoConsoleVariableRef CVarMetalCacheMinSize(
 	GMetalCacheMinSize,
 	TEXT("Sets the minimum size that we expect the metal OS cache to be (in MB). This is used to determine if we need to cache PSOs again (Default: 32).\n"), ECVF_ReadOnly);
 
+static int32 GMetalBinaryCacheDebugOutput = 0;
+static FAutoConsoleVariableRef CVarMetalBinaryCacheDebugOutput(
+    TEXT("rhi.Metal.BinaryCacheDebugOutput"),
+    GMetalBinaryCacheDebugOutput,
+    TEXT("Enable to output logging information for PSO Binary cache default(0) \n"), ECVF_ReadOnly);
+
 static uint32 BlendBitOffsets[] = { Offset_BlendState0, Offset_BlendState1, Offset_BlendState2, Offset_BlendState3, Offset_BlendState4, Offset_BlendState5, Offset_BlendState6, Offset_BlendState7 };
 static uint32 RTBitOffsets[] = { Offset_RenderTargetFormat0, Offset_RenderTargetFormat1, Offset_RenderTargetFormat2, Offset_RenderTargetFormat3, Offset_RenderTargetFormat4, Offset_RenderTargetFormat5, Offset_RenderTargetFormat6, Offset_RenderTargetFormat7 };
 static_assert(Offset_RasterEnd < 64 && Offset_End < 128, "Offset_RasterEnd must be < 64 && Offset_End < 128");
@@ -828,8 +834,170 @@ static bool ConfigureRenderPipelineDescriptor(mtlpp::RenderPipelineDescriptor& R
 	return true;
 }
 
+/*
+ * PSO Harvesting and Reuse
+ *
+ * Usage:
+ *
+ * To Harvest, run the game with -MetalPSOCache=recreate
+ * All Render and Compute PSOs created will be harvested into the MTLBinaryArchive
+ * Console command r.Metal.ShaderPipelineCache.Save will trigger the serialization to file.
+ * The binary archive's location will be printed to the log.
+ *
+ * To reuse, run the game with -MetalPSOCache=use
+ * The binary archive will be opened from the saved location.
+ * The binary archive can be moved to another device, as long as it's the same GPU
+ * and OS build.
+ *
+ */
+
+enum class CacheMode
+{
+	Uninitialized,
+	Recreate,
+	Append,
+	Use,
+	Ignore
+};
+
+static CacheMode GPSOCacheMode = CacheMode::Uninitialized;
+static id< MTLBinaryArchive > GPSOBinaryArchive = nil;
+static uint32_t GPSOHarvestCount = 0;
+
+static NSURL * PipelineCacheSaveLocation()
+{
+	NSString * path = NSSearchPathForDirectoriesInDomains(NSDocumentDirectory, NSUserDomainMask, YES)[0];
+	if (!path)
+	{
+		UE_LOG(LogMetal, Error, TEXT("Metal Pipeline Cache: Unable to find Documents folder"));
+	}
+	NSURL * url = [NSURL fileURLWithPath : [NSString stringWithFormat : @"%@/mtlarchive.bin", path] ];
+	return url;
+}
+
+static void InitializeMetalPipelineCache()
+{
+	FString strCacheMode;
+	FParse::Value(FCommandLine::Get(), TEXT("MetalPSOCache="), strCacheMode);
+		
+	if (strCacheMode.Compare(TEXT("recreate"), ESearchCase::IgnoreCase) == 0)
+	{
+		GPSOCacheMode = CacheMode::Recreate;
+        UE_LOG(LogMetal, Log, TEXT("Metal Pipeline Cache: recreate PSO cache"));
+	}
+	else if (strCacheMode.Compare(TEXT("append"), ESearchCase::IgnoreCase) == 0)
+	{
+		GPSOCacheMode = CacheMode::Append;
+        UE_LOG(LogMetal, Log, TEXT("Metal Pipeline Cache: append to PSO cache"));
+	}
+	else if (strCacheMode.Compare(TEXT("use"), ESearchCase::IgnoreCase) == 0)
+	{
+		GPSOCacheMode = CacheMode::Use;
+        UE_LOG(LogMetal, Log, TEXT("Metal Pipeline Cache: use PSO cache"));
+	}
+	else
+	{
+		GPSOCacheMode = CacheMode::Ignore;
+        UE_LOG(LogMetal, Log, TEXT("Metal Pipeline Cache: ignore PSO cache"));
+	}
+		
+	if (GPSOCacheMode != CacheMode::Ignore)
+	{
+		NSURL * url = PipelineCacheSaveLocation();
+        UE_LOG(LogMetal, Log, TEXT("Metal Pipeline Cache: pso cache save location will be: %s"), *FString(url.absoluteString.UTF8String));
+        
+		MTLBinaryArchiveDescriptor * archDesc = [MTLBinaryArchiveDescriptor new];
+		archDesc.url = ((GPSOCacheMode == CacheMode::Append) || (GPSOCacheMode == CacheMode::Use)) ? PipelineCacheSaveLocation() : nil;
+		id< MTLDevice > mtlDevice = GetMetalDeviceContext().GetDevice().GetPtr();
+		__autoreleasing NSError * err = nil;
+		GPSOBinaryArchive = [mtlDevice newBinaryArchiveWithDescriptor : archDesc error : &err];
+		if (err)
+		{
+			UE_LOG(LogMetal, Error, TEXT("Error adding Pipeline Functions to Binary Archive: %s"), *FString(err.localizedDescription.UTF8String));
+		}
+	}
+}
+
+static void RelatePipelineStateToCache(const mtlpp::RenderPipelineDescriptor & PipelineDesc, NSUInteger * outPipelineOpts)
+{
+	if (GPSOBinaryArchive && (GPSOCacheMode != CacheMode::Ignore))
+	{
+		if (GPSOCacheMode == CacheMode::Recreate || GPSOCacheMode == CacheMode::Append)
+		{
+			__autoreleasing NSError * err = nil;
+			bool bAddedBinaryPSO = [GPSOBinaryArchive addRenderPipelineFunctionsWithDescriptor : PipelineDesc.GetPtr() error : &err];
+			if (err)
+			{
+				UE_LOG(LogMetal, Warning, TEXT("Metal Pipeline Cache: Error adding Pipeline Functions to Binary Archive: %s"), *FString(err.localizedDescription.UTF8String));
+			}
+            else if(bAddedBinaryPSO)
+            {
+                GPSOHarvestCount++;
+                if(GMetalBinaryCacheDebugOutput && GPSOHarvestCount % 100)
+                {
+                    UE_LOG(LogMetal, Log, TEXT("Metal Pipeline Cache: Harvested PSO count: %d"), GPSOHarvestCount);
+                }
+            }
+		}
+	}
+}
+
+static void RelatePipelineStateToCache(const mtlpp::ComputePipelineDescriptor & PipelineDesc, NSUInteger * outPipelineOpts)
+{
+	if (GPSOBinaryArchive && (GPSOCacheMode != CacheMode::Ignore))
+	{
+		if (GPSOCacheMode == CacheMode::Recreate || GPSOCacheMode == CacheMode::Append)
+		{
+			__autoreleasing NSError * err = nil;
+            bool bAddedBinaryPSO = [GPSOBinaryArchive addComputePipelineFunctionsWithDescriptor : PipelineDesc.GetPtr() error : &err];
+			if (err)
+			{
+				UE_LOG(LogMetal, Warning, TEXT("Metal Pipeline Cache: Error adding Pipeline Functions to Binary Archive: %s"), *FString(err.localizedDescription.UTF8String));
+			}
+            else if(bAddedBinaryPSO)
+            {
+                GPSOHarvestCount++;
+                if(GMetalBinaryCacheDebugOutput && GPSOHarvestCount % 100)
+                {
+                    UE_LOG(LogMetal, Log, TEXT("Metal Pipeline Cache: Harvested PSO count: %d"), GPSOHarvestCount);
+                }
+            }
+		}
+	}
+}
+
+void MetalConsoleCommandSavePipelineFileCache()
+{
+	UE_LOG(LogMetal, Log, TEXT("Metal Pipeline Cache: requesting PSO save..."));
+	{
+		if (GPSOBinaryArchive && (GPSOCacheMode == CacheMode::Recreate || GPSOCacheMode == CacheMode::Append))
+		{
+			NSURL * url = PipelineCacheSaveLocation();
+			UE_LOG(LogMetal, Log, TEXT("Metal Pipeline Cache: Serialize harvested PSOs to: %s"), *FString(url.absoluteString.UTF8String));
+            UE_LOG(LogMetal, Log, TEXT("Metal Pipeline Cache: Serialized PSO Count: %d"), GPSOHarvestCount);
+            
+			__autoreleasing NSError * err = nil;
+			[GPSOBinaryArchive serializeToURL : url error : &err];
+			if (err)
+			{
+				UE_LOG(LogMetal, Error, TEXT("Metal Pipeline Cache: Error Serializing binary archive: %s"), *FString(err.localizedDescription.UTF8String));
+			}
+		}
+	}
+}
+
+static FAutoConsoleCommand SavePipelineCacheCmd(
+    TEXT("rhi.Metal.ShaderPipelineCache.Save"),
+    TEXT("Save the current pipeline file cache."),
+    FConsoleCommandDelegate::CreateStatic(MetalConsoleCommandSavePipelineFileCache));
+
 static FMetalShaderPipeline* CreateMTLRenderPipeline(bool const bSync, FMetalGraphicsPipelineKey const& Key, const FGraphicsPipelineStateInitializer& Init)
 {
+	if (GPSOCacheMode == CacheMode::Uninitialized)
+	{
+		InitializeMetalPipelineCache();
+	}
+
     FMetalVertexShader* VertexShader = (FMetalVertexShader*)Init.BoundShaderState.VertexShaderRHI;
     FMetalPixelShader* PixelShader = (FMetalPixelShader*)Init.BoundShaderState.PixelShaderRHI;
     
@@ -888,6 +1056,7 @@ static FMetalShaderPipeline* CreateMTLRenderPipeline(bool const bSync, FMetalGra
 		{
 			ns::AutoReleasedError RenderError;
 			METAL_GPUPROFILE(FScopedMetalCPUStats CPUStat(FString::Printf(TEXT("NewRenderPipeline: %s"), TEXT("")/**FString([RenderPipelineDesc.GetPtr() description])*/)));
+			RelatePipelineStateToCache(RenderPipelineDesc, &RenderOption);
 			Pipeline->RenderPipelineState = Device.NewRenderPipelineState(RenderPipelineDesc, (mtlpp::PipelineOption)RenderOption, Reflection, &RenderError);
 			if (Reflection)
 			{
