Fixed Metal mobile desktop renderer crash. RHIDrawPrimitiveIndirect was not using the platform flags to determine if the feature was available on iOS.

diff --git a/Engine/Source/Runtime/Apple/AGXRHI/Private/AGXCommands.cpp b/Engine/Source/Runtime/Apple/AGXRHI/Private/AGXCommands.cpp
--- a/Engine/Source/Runtime/Apple/AGXRHI/Private/AGXCommands.cpp
+++ b/Engine/Source/Runtime/Apple/AGXRHI/Private/AGXCommands.cpp
@@ -510,18 +510,20 @@ void FAGXRHICommandContext::RHIDrawPrimitive(uint32 BaseVertexIndex, uint32 NumP
 void FAGXRHICommandContext::RHIDrawPrimitiveIndirect(FRHIBuffer* ArgumentBufferRHI, uint32 ArgumentOffset)
 {
 	@autoreleasepool {
-#if PLATFORM_IOS
-	NOT_SUPPORTED("RHIDrawPrimitiveIndirect");
-#else
-	SCOPE_CYCLE_COUNTER(STAT_AGXDrawCallTime);
-	uint32 PrimitiveType = Context->GetCurrentState().GetPrimitiveType();
-	
-	
-	RHI_DRAW_CALL_STATS(PrimitiveType,1);
-	FAGXResourceMultiBuffer* ArgumentBuffer = ResourceCast(ArgumentBufferRHI);
-	
-	Context->DrawPrimitiveIndirect(PrimitiveType, ArgumentBuffer, ArgumentOffset);
-#endif
+        if (GetAGXDeviceContext().SupportsFeature(EAGXFeaturesIndirectBuffer))
+        {
+            SCOPE_CYCLE_COUNTER(STAT_AGXDrawCallTime);
+            uint32 PrimitiveType = Context->GetCurrentState().GetPrimitiveType();
+            
+            RHI_DRAW_CALL_STATS(PrimitiveType,1);
+            FAGXResourceMultiBuffer* ArgumentBuffer = ResourceCast(ArgumentBufferRHI);
+            
+            Context->DrawPrimitiveIndirect(PrimitiveType, ArgumentBuffer, ArgumentOffset);
+        }
+        else
+        {
+            NOT_SUPPORTED("RHIDrawPrimitiveIndirect");
+        }
 	}
 }
 
diff --git a/Engine/Source/Runtime/Apple/MetalRHI/Private/MetalCommands.cpp b/Engine/Source/Runtime/Apple/MetalRHI/Private/MetalCommands.cpp
--- a/Engine/Source/Runtime/Apple/MetalRHI/Private/MetalCommands.cpp
+++ b/Engine/Source/Runtime/Apple/MetalRHI/Private/MetalCommands.cpp
@@ -511,18 +511,21 @@ void FMetalRHICommandContext::RHIDrawPrimitive(uint32 BaseVertexIndex, uint32 Nu
 void FMetalRHICommandContext::RHIDrawPrimitiveIndirect(FRHIBuffer* ArgumentBufferRHI, uint32 ArgumentOffset)
 {
 	@autoreleasepool {
-#if PLATFORM_IOS
-	NOT_SUPPORTED("RHIDrawPrimitiveIndirect");
-#else
-	SCOPE_CYCLE_COUNTER(STAT_MetalDrawCallTime);
-	uint32 PrimitiveType = Context->GetCurrentState().GetPrimitiveType();
-	
-	
-	RHI_DRAW_CALL_STATS(PrimitiveType,1);
-	FMetalResourceMultiBuffer* ArgumentBuffer = ResourceCast(ArgumentBufferRHI);
-	
-	Context->DrawPrimitiveIndirect(PrimitiveType, ArgumentBuffer, ArgumentOffset);
-#endif
+        if (GetMetalDeviceContext().SupportsFeature(EMetalFeaturesIndirectBuffer))
+        {
+            SCOPE_CYCLE_COUNTER(STAT_MetalDrawCallTime);
+            uint32 PrimitiveType = Context->GetCurrentState().GetPrimitiveType();
+            
+            
+            RHI_DRAW_CALL_STATS(PrimitiveType,1);
+            FMetalResourceMultiBuffer* ArgumentBuffer = ResourceCast(ArgumentBufferRHI);
+            
+            Context->DrawPrimitiveIndirect(PrimitiveType, ArgumentBuffer, ArgumentOffset);
+        }
+        else
+        {
+            NOT_SUPPORTED("RHIDrawPrimitiveIndirect");
+        }
 	}
 }
 
