Assert if running on an Apple A8 GPU but the Support Apple A8 checkbox is not checked

diff --git a/Engine/Source/Runtime/Apple/MetalRHI/Private/MetalRHI.cpp b/Engine/Source/Runtime/Apple/MetalRHI/Private/MetalRHI.cpp
--- a/Engine/Source/Runtime/Apple/MetalRHI/Private/MetalRHI.cpp
+++ b/Engine/Source/Runtime/Apple/MetalRHI/Private/MetalRHI.cpp
@@ -187,6 +187,14 @@ FMetalDynamicRHI::FMetalDynamicRHI(ERHIFeatureLevel::Type RequestedFeatureLevel)
 	GRHISupportsDrawIndirect = Device.SupportsFeatureSet(mtlpp::FeatureSet::tvOS_GPUFamily2_v1);
 	GRHISupportsPixelShaderUAVs = Device.SupportsFeatureSet(mtlpp::FeatureSet::tvOS_GPUFamily2_v1);
 #else
+	{
+		bool bSupportAppleA8 = false;
+		GConfig->GetBool(TEXT("/Script/IOSRuntimeSettings.IOSRuntimeSettings"), TEXT("bSupportAppleA8"), bSupportAppleA8, GEngineIni);
+		if (!Device.SupportsFeatureSet(mtlpp::FeatureSet::iOS_GPUFamily3_v1) && !bSupportAppleA8)
+		{
+			UE_LOG(LogMetal, Fatal, TEXT("This device does not supports the Apple A9 feature set which is the minimum for this build. Please check the Support Apple A8 checkbox in the IOS Project Settings."));
+		}
+	}
 	GRHISupportsRWTextureBuffers = Device.SupportsFeatureSet(mtlpp::FeatureSet::iOS_GPUFamily4_v1);
 	GRHISupportsDrawIndirect = Device.SupportsFeatureSet(mtlpp::FeatureSet::iOS_GPUFamily3_v1);
 	GRHISupportsPixelShaderUAVs = Device.SupportsFeatureSet(mtlpp::FeatureSet::iOS_GPUFamily3_v1);
