#ue5 - Fix build: resolve ambiguous type.

diff --git a/Engine/Source/Runtime/Apple/AGXRHI/Private/AGXStateCache.cpp b/Engine/Source/Runtime/Apple/AGXRHI/Private/AGXStateCache.cpp
--- a/Engine/Source/Runtime/Apple/AGXRHI/Private/AGXStateCache.cpp
+++ b/Engine/Source/Runtime/Apple/AGXRHI/Private/AGXStateCache.cpp
@@ -352,8 +352,8 @@ void FAGXStateCache::SetScissorRect(bool const bEnable, mtlpp::ScissorRect const
 		// Clamp to framebuffer size - Metal doesn't allow scissor to be larger.
 		Scissor[0].x = Scissor[0].x;
 		Scissor[0].y = Scissor[0].y;
-		Scissor[0].width = FMath::Max((Scissor[0].x + Scissor[0].width <= FMath::RoundToInt(FrameBufferSize.width)) ? Scissor[0].width : FMath::RoundToInt(FrameBufferSize.width) - Scissor[0].x, (NSUInteger)1u);
-		Scissor[0].height = FMath::Max((Scissor[0].y + Scissor[0].height <= FMath::RoundToInt(FrameBufferSize.height)) ? Scissor[0].height : FMath::RoundToInt(FrameBufferSize.height) - Scissor[0].y, (NSUInteger)1u);
+		Scissor[0].width = FMath::Max((Scissor[0].x + Scissor[0].width <= FMath::RoundToInt32(FrameBufferSize.width)) ? Scissor[0].width : FMath::RoundToInt32(FrameBufferSize.width) - Scissor[0].x, (NSUInteger)1u);
+		Scissor[0].height = FMath::Max((Scissor[0].y + Scissor[0].height <= FMath::RoundToInt32(FrameBufferSize.height)) ? Scissor[0].height : FMath::RoundToInt32(FrameBufferSize.height) - Scissor[0].y, (NSUInteger)1u);
 		
 		RasterBits |= EAGXRenderFlagScissorRect;
 	}
diff --git a/Engine/Source/Runtime/Apple/MetalRHI/Private/MetalStateCache.cpp b/Engine/Source/Runtime/Apple/MetalRHI/Private/MetalStateCache.cpp
--- a/Engine/Source/Runtime/Apple/MetalRHI/Private/MetalStateCache.cpp
+++ b/Engine/Source/Runtime/Apple/MetalRHI/Private/MetalStateCache.cpp
@@ -353,8 +353,8 @@ void FMetalStateCache::SetScissorRect(bool const bEnable, mtlpp::ScissorRect con
 		// Clamp to framebuffer size - Metal doesn't allow scissor to be larger.
 		Scissor[0].x = Scissor[0].x;
 		Scissor[0].y = Scissor[0].y;
-		Scissor[0].width = FMath::Max((Scissor[0].x + Scissor[0].width <= FMath::RoundToInt(FrameBufferSize.width)) ? Scissor[0].width : FMath::RoundToInt(FrameBufferSize.width) - Scissor[0].x, (NSUInteger)1u);
-		Scissor[0].height = FMath::Max((Scissor[0].y + Scissor[0].height <= FMath::RoundToInt(FrameBufferSize.height)) ? Scissor[0].height : FMath::RoundToInt(FrameBufferSize.height) - Scissor[0].y, (NSUInteger)1u);
+		Scissor[0].width = FMath::Max((Scissor[0].x + Scissor[0].width <= FMath::RoundToInt32(FrameBufferSize.width)) ? Scissor[0].width : FMath::RoundToInt32(FrameBufferSize.width) - Scissor[0].x, (NSUInteger)1u);
+		Scissor[0].height = FMath::Max((Scissor[0].y + Scissor[0].height <= FMath::RoundToInt32(FrameBufferSize.height)) ? Scissor[0].height : FMath::RoundToInt32(FrameBufferSize.height) - Scissor[0].y, (NSUInteger)1u);
 		
 		RasterBits |= EMetalRenderFlagScissorRect;
 	}
