Moving AGX and Metal implementations of ApplyStaticUniformBuffers out of headers to remove header depenency on RHICore.

diff --git a/Engine/Source/Runtime/Apple/AGXRHI/Private/AGXCommands.cpp b/Engine/Source/Runtime/Apple/AGXRHI/Private/AGXCommands.cpp
--- a/Engine/Source/Runtime/Apple/AGXRHI/Private/AGXCommands.cpp
+++ b/Engine/Source/Runtime/Apple/AGXRHI/Private/AGXCommands.cpp
@@ -18,6 +18,7 @@
 #include "StaticBoundShaderState.h"
 #include "EngineGlobals.h"
 #include "PipelineStateCache.h"
+#include "RHICoreShader.h"
 
 static const bool GUsesInvertedZ = true;
 
@@ -120,6 +121,15 @@ void FAGXRHICommandContext::RHISetStreamSource(uint32 StreamIndex, FRHIBuffer* V
 	}
 }
 
+template <typename TRHIShader>
+void FAGXRHICommandContext::ApplyStaticUniformBuffers(TRHIShader* Shader)
+{
+	if (Shader)
+	{
+		UE::RHICore::ApplyStaticUniformBuffers(this, Shader, Shader->StaticSlots, Shader->Bindings.ShaderResourceTable.ResourceTableLayoutHashes, GlobalUniformBuffers);
+	}
+}
+
 void FAGXRHICommandContext::RHISetComputePipelineState(FRHIComputePipelineState* ComputePipelineState)
 {
 	@autoreleasepool {
diff --git a/Engine/Source/Runtime/Apple/AGXRHI/Public/AGXRHIContext.h b/Engine/Source/Runtime/Apple/AGXRHI/Public/AGXRHIContext.h
--- a/Engine/Source/Runtime/Apple/AGXRHI/Public/AGXRHIContext.h
+++ b/Engine/Source/Runtime/Apple/AGXRHI/Public/AGXRHIContext.h
@@ -6,7 +6,6 @@
 #include "AGXState.h"
 #include "AGXResources.h"
 #include "AGXViewport.h"
-#include "RHICoreShader.h"
 
 #define UE_METAL_RHI_SUPPORT_CLEAR_UAV_WITH_BLIT_ENCODER	1
 
@@ -228,13 +227,7 @@ protected:
 	uint32 PendingNumPrimitives;
 
 	template <typename TRHIShader>
-	void ApplyStaticUniformBuffers(TRHIShader* Shader)
-	{
-		if (Shader)
-		{
-			UE::RHICore::ApplyStaticUniformBuffers(this, Shader, Shader->StaticSlots, Shader->Bindings.ShaderResourceTable.ResourceTableLayoutHashes, GlobalUniformBuffers);
-		}
-	}
+	void ApplyStaticUniformBuffers(TRHIShader* Shader);
 
 	TArray<FRHIUniformBuffer*> GlobalUniformBuffers;
 
diff --git a/Engine/Source/Runtime/Apple/MetalRHI/Private/MetalCommands.cpp b/Engine/Source/Runtime/Apple/MetalRHI/Private/MetalCommands.cpp
--- a/Engine/Source/Runtime/Apple/MetalRHI/Private/MetalCommands.cpp
+++ b/Engine/Source/Runtime/Apple/MetalRHI/Private/MetalCommands.cpp
@@ -19,6 +19,7 @@
 #include "StaticBoundShaderState.h"
 #include "EngineGlobals.h"
 #include "PipelineStateCache.h"
+#include "RHICoreShader.h"
 
 static const bool GUsesInvertedZ = true;
 
@@ -121,6 +122,15 @@ void FMetalRHICommandContext::RHISetStreamSource(uint32 StreamIndex, FRHIBuffer*
 	}
 }
 
+template <typename TRHIShader>
+void FMetalRHICommandContext::ApplyStaticUniformBuffers(TRHIShader* Shader)
+{
+	if (Shader)
+	{
+		UE::RHICore::ApplyStaticUniformBuffers(this, Shader, Shader->StaticSlots, Shader->Bindings.ShaderResourceTable.ResourceTableLayoutHashes, GlobalUniformBuffers);
+	}
+}
+
 void FMetalRHICommandContext::RHISetComputePipelineState(FRHIComputePipelineState* ComputePipelineState)
 {
 	@autoreleasepool {
diff --git a/Engine/Source/Runtime/Apple/MetalRHI/Public/MetalRHIContext.h b/Engine/Source/Runtime/Apple/MetalRHI/Public/MetalRHIContext.h
--- a/Engine/Source/Runtime/Apple/MetalRHI/Public/MetalRHIContext.h
+++ b/Engine/Source/Runtime/Apple/MetalRHI/Public/MetalRHIContext.h
@@ -7,7 +7,6 @@
 #include "MetalState.h"
 #include "MetalResources.h"
 #include "MetalViewport.h"
-#include "RHICoreShader.h"
 
 #define UE_METAL_RHI_SUPPORT_CLEAR_UAV_WITH_BLIT_ENCODER	1
 
@@ -229,13 +228,7 @@ protected:
 	uint32 PendingNumPrimitives;
 
 	template <typename TRHIShader>
-	void ApplyStaticUniformBuffers(TRHIShader* Shader)
-	{
-		if (Shader)
-		{
-			UE::RHICore::ApplyStaticUniformBuffers(this, Shader, Shader->StaticSlots, Shader->Bindings.ShaderResourceTable.ResourceTableLayoutHashes, GlobalUniformBuffers);
-		}
-	}
+	void ApplyStaticUniformBuffers(TRHIShader* Shader);
 
 	TArray<FRHIUniformBuffer*> GlobalUniformBuffers;
 
