Remove Metal MRT assert when using MSAA, for SceneDepthAux
Copy unzipped IOS frameworks to the remote Mac when using the remote toolchain

diff --git a/Engine/Build/Rsync/RsyncEngine.txt b/Engine/Build/Rsync/RsyncEngine.txt
--- a/Engine/Build/Rsync/RsyncEngine.txt
+++ b/Engine/Build/Rsync/RsyncEngine.txt
@@ -189,6 +189,8 @@ protect /Source/Programs/Shared/EpicGames.Build/bin/**
 + /Source/ThirdParty/libSampleRate/**.cpp
 + /Source/ThirdParty/Eigen/**
 + /Source/ThirdParty/mimalloc/src/**
++ /Source/ThirdParty/**/IOS/**.framework/**
++ /Source/ThirdParty/**/TVOS/**.framework/**
 - /Source/ThirdParty/**
 
 # Source: Include everything else
diff --git a/Engine/Source/Runtime/Apple/AGXRHI/Private/AGXStateCache.cpp b/Engine/Source/Runtime/Apple/AGXRHI/Private/AGXStateCache.cpp
--- a/Engine/Source/Runtime/Apple/AGXRHI/Private/AGXStateCache.cpp
+++ b/Engine/Source/Runtime/Apple/AGXRHI/Private/AGXStateCache.cpp
@@ -589,8 +589,6 @@ bool FAGXStateCache::SetRenderPassInfo(FRHIRenderPassInfo const& InRenderTargets
 					ColorAttachment.SetStoreAction(!bMemoryless && GRHIDeviceId > 2 ? mtlpp::StoreAction::Unknown : NewColorStore[RenderTargetIndex]);
 					ColorAttachment.SetResolveTexture(Surface.MSAAResolveTexture ? Surface.MSAAResolveTexture : Surface.Texture);
 					SampleCount = Surface.MSAATexture.GetSampleCount();
-					// only allow one MRT with msaa
-					checkf(NumColorRenderTargets == 1, TEXT("Only expected one MRT when using MSAA"));
 				}
 				else
 				{
diff --git a/Engine/Source/Runtime/Apple/MetalRHI/Private/MetalStateCache.cpp b/Engine/Source/Runtime/Apple/MetalRHI/Private/MetalStateCache.cpp
--- a/Engine/Source/Runtime/Apple/MetalRHI/Private/MetalStateCache.cpp
+++ b/Engine/Source/Runtime/Apple/MetalRHI/Private/MetalStateCache.cpp
@@ -595,8 +595,6 @@ bool FMetalStateCache::SetRenderPassInfo(FRHIRenderPassInfo const& InRenderTarge
 					ColorAttachment.SetStoreAction(!bMemoryless && GRHIDeviceId > 2 ? mtlpp::StoreAction::Unknown : NewColorStore[RenderTargetIndex]);
 					ColorAttachment.SetResolveTexture(Surface.MSAAResolveTexture ? Surface.MSAAResolveTexture : Surface.Texture);
 					SampleCount = Surface.MSAATexture.GetSampleCount();
-					// only allow one MRT with msaa
-					checkf(NumColorRenderTargets == 1, TEXT("Only expected one MRT when using MSAA"));
 				}
 				else
 				{
