Copying //UE5/Dev-ParallelRendering to //UE5/Main

This change is the copy-up of work from //UE5/Dev-ParallelRendering to unify the various RHI texture types
 - The separate RHI texture types (FRHITexture2D, FRHITexture2DArray, FRHITexture3D, FRHITextureCube, FRHITextureCubeArray) have been replaced with a single type: FRHITexture.
 - Includes the required refactoring / fixes for the various platform RHIs.

The old texture types are still supported via typedefs, but will soon be marked deprecated. Future rendering code should make use of the FRHITexture type instead.

diff --git a/Engine/Extras/Android/SetupAndroid.bat b/Engine/Extras/Android/SetupAndroid.bat
--- a/Engine/Extras/Android/SetupAndroid.bat
+++ b/Engine/Extras/Android/SetupAndroid.bat
@@ -74,7 +74,7 @@ FOR /F "tokens=2*" %%A IN ('REG.exe query "%KEY_NAME%" /v "%VALUE_NAME%"') DO (s
 
 where.exe /Q adb.exe
 IF /I "%ERRORLEVEL%" NEQ "0" (
-	echo Current user path: %USERPATH%
+	echo Current user path: "%USERPATH%"
 	setx PATH "%USERPATH%;%PLATFORMTOOLS%"
 	echo Added %PLATFORMTOOLS% to path
 )
diff --git a/Engine/Plugins/Media/AndroidCamera/Source/AndroidCamera/Private/Player/AndroidCameraPlayer.cpp b/Engine/Plugins/Media/AndroidCamera/Source/AndroidCamera/Private/Player/AndroidCameraPlayer.cpp
--- a/Engine/Plugins/Media/AndroidCamera/Source/AndroidCamera/Private/Player/AndroidCameraPlayer.cpp
+++ b/Engine/Plugins/Media/AndroidCamera/Source/AndroidCamera/Private/Player/AndroidCameraPlayer.cpp
@@ -302,7 +302,7 @@ static void DoUpdateExternalCameraSampleExecute(TWeakPtr<FJavaAndroidCameraPlaye
 	if (VideoTexture == nullptr)
 	{
 		FRHIResourceCreateInfo CreateInfo(TEXT("VideoTexture"));
-		VideoTexture = GDynamicRHI->RHICreateTextureExternal2D(1, 1, PF_R8G8B8A8, 1, 1, TexCreate_None, ERHIAccess::Unknown, CreateInfo);
+		VideoTexture = RHICreateTextureExternal2D(1, 1, PF_R8G8B8A8, 1, 1, TexCreate_None, ERHIAccess::Unknown, CreateInfo);
 		PinnedJavaCameraPlayer->SetVideoTexture(VideoTexture);
 
 		if (VideoTexture == nullptr)
diff --git a/Engine/Plugins/Media/AndroidMedia/Source/AndroidMedia/Private/Player/AndroidMediaPlayer.cpp b/Engine/Plugins/Media/AndroidMedia/Source/AndroidMedia/Private/Player/AndroidMediaPlayer.cpp
--- a/Engine/Plugins/Media/AndroidMedia/Source/AndroidMedia/Private/Player/AndroidMediaPlayer.cpp
+++ b/Engine/Plugins/Media/AndroidMedia/Source/AndroidMedia/Private/Player/AndroidMediaPlayer.cpp
@@ -462,7 +462,7 @@ static void DoUpdateExternalMediaSampleExecute(TWeakPtr<FJavaAndroidMediaPlayer,
 	if (VideoTexture == nullptr)
 	{
 		FRHIResourceCreateInfo CreateInfo(TEXT("VideoTexture"));
-		VideoTexture = GDynamicRHI->RHICreateTextureExternal2D(1, 1, PF_R8G8B8A8, 1, 1, TexCreate_None, ERHIAccess::SRVGraphics, CreateInfo);
+		VideoTexture = RHICreateTextureExternal2D(1, 1, PF_R8G8B8A8, 1, 1, TexCreate_None, ERHIAccess::SRVGraphics, CreateInfo);
 		PinnedJavaMediaPlayer->SetVideoTexture(VideoTexture);
 
 		if (VideoTexture == nullptr)
diff --git a/Engine/Plugins/Media/BinkMedia/Source/BinkMediaPlayer/Public/BinkMediaTexture.h b/Engine/Plugins/Media/BinkMedia/Source/BinkMediaPlayer/Public/BinkMediaTexture.h
--- a/Engine/Plugins/Media/BinkMedia/Source/BinkMediaPlayer/Public/BinkMediaTexture.h
+++ b/Engine/Plugins/Media/BinkMedia/Source/BinkMediaPlayer/Public/BinkMediaTexture.h
@@ -9,7 +9,7 @@
 
 class UBinkMediaPlayer;
 class FBinkMediaPlayer;
-enum EPixelFormat;
+enum EPixelFormat : uint8;
 enum TextureAddress;
 
 /**
diff --git a/Engine/Plugins/Media/ElectraUtil/Source/ElectraSamples/Public/GStreamer/ElectraTextureSample.h b/Engine/Plugins/Media/ElectraUtil/Source/ElectraSamples/Public/GStreamer/ElectraTextureSample.h
--- a/Engine/Plugins/Media/ElectraUtil/Source/ElectraSamples/Public/GStreamer/ElectraTextureSample.h
+++ b/Engine/Plugins/Media/ElectraUtil/Source/ElectraSamples/Public/GStreamer/ElectraTextureSample.h
@@ -9,7 +9,7 @@
 #include "RHIUtilities.h"
 
 class MEDIArendererVideoUE;
-class FRHITexture2D;
+class FRHITexture;
 
 class FElectraTextureSampleUnix 
 	: public IMediaTextureSample
diff --git a/Engine/Plugins/Media/PixelStreaming/Source/PixelStreaming/Private/EncoderFrameFactory.cpp b/Engine/Plugins/Media/PixelStreaming/Source/PixelStreaming/Private/EncoderFrameFactory.cpp
--- a/Engine/Plugins/Media/PixelStreaming/Source/PixelStreaming/Private/EncoderFrameFactory.cpp
+++ b/Engine/Plugins/Media/PixelStreaming/Source/PixelStreaming/Private/EncoderFrameFactory.cpp
@@ -39,10 +39,10 @@ void UE::PixelStreaming::FEncoderFrameFactory::RemoveStaleTextures()
 {
 	// Remove any textures whose only reference is the one held by this class
 
-	TMap<FTexture2DRHIRef, TSharedPtr<AVEncoder::FVideoEncoderInputFrame>>::TIterator Iter = TextureToFrameMapping.CreateIterator();
+	TMap<FTextureRHIRef, TSharedPtr<AVEncoder::FVideoEncoderInputFrame>>::TIterator Iter = TextureToFrameMapping.CreateIterator();
 	for(; Iter; ++Iter)
 	{
-		FTexture2DRHIRef& Tex = Iter.Key();
+		FTextureRHIRef& Tex = Iter.Key();
 
 		if(Tex.GetRefCount() == 1) {
 			Iter.RemoveCurrent();
@@ -50,7 +50,7 @@ void UE::PixelStreaming::FEncoderFrameFactory::RemoveStaleTextures()
 	}
 }
 
-TSharedPtr<AVEncoder::FVideoEncoderInputFrame> UE::PixelStreaming::FEncoderFrameFactory::GetOrCreateFrame(const FTexture2DRHIRef InTexture)
+TSharedPtr<AVEncoder::FVideoEncoderInputFrame> UE::PixelStreaming::FEncoderFrameFactory::GetOrCreateFrame(const FTextureRHIRef InTexture)
 {
 	check(EncoderInput.IsValid());
 
@@ -82,7 +82,7 @@ TSharedPtr<AVEncoder::FVideoEncoderInputFrame> UE::PixelStreaming::FEncoderFrame
 	return OutFrame;
 }
 
-TSharedPtr<AVEncoder::FVideoEncoderInputFrame> UE::PixelStreaming::FEncoderFrameFactory::GetFrameAndSetTexture(FTexture2DRHIRef InTexture)
+TSharedPtr<AVEncoder::FVideoEncoderInputFrame> UE::PixelStreaming::FEncoderFrameFactory::GetFrameAndSetTexture(FTextureRHIRef InTexture)
 {
 	check(EncoderInput.IsValid());
 
@@ -165,7 +165,7 @@ TSharedPtr<AVEncoder::FVideoEncoderInput> UE::PixelStreaming::FEncoderFrameFacto
 	return nullptr;
 }
 
-void UE::PixelStreaming::FEncoderFrameFactory::SetTexture(TSharedPtr<AVEncoder::FVideoEncoderInputFrame> InputFrame, const FTexture2DRHIRef& Texture)
+void UE::PixelStreaming::FEncoderFrameFactory::SetTexture(TSharedPtr<AVEncoder::FVideoEncoderInputFrame> InputFrame, const FTextureRHIRef& Texture)
 {
 	const ERHIInterfaceType RHIType = RHIGetInterfaceType();
 
@@ -216,10 +216,9 @@ void UE::PixelStreaming::FEncoderFrameFactory::SetTexture(TSharedPtr<AVEncoder::
 	}
 }
 
-void UE::PixelStreaming::FEncoderFrameFactory::SetTextureCUDAVulkan(TSharedPtr<AVEncoder::FVideoEncoderInputFrame> InputFrame, const FTexture2DRHIRef& Texture)
+void UE::PixelStreaming::FEncoderFrameFactory::SetTextureCUDAVulkan(TSharedPtr<AVEncoder::FVideoEncoderInputFrame> InputFrame, const FTextureRHIRef& Texture)
 {
 	IVulkanDynamicRHI* VulkanRHI = GetIVulkanDynamicRHI();
-
 	const FVulkanRHIAllocationInfo TextureAllocationInfo = VulkanRHI->RHIGetAllocationInfo(Texture.GetReference());
 
 	VkDevice Device = VulkanRHI->RHIGetVkDevice();
@@ -371,9 +370,9 @@ void UE::PixelStreaming::FEncoderFrameFactory::SetTextureCUDAVulkan(TSharedPtr<A
 }
 
 #if PLATFORM_WINDOWS
-void UE::PixelStreaming::FEncoderFrameFactory::SetTextureCUDAD3D11(TSharedPtr<AVEncoder::FVideoEncoderInputFrame> InputFrame, const FTexture2DRHIRef& Texture)
+void UE::PixelStreaming::FEncoderFrameFactory::SetTextureCUDAD3D11(TSharedPtr<AVEncoder::FVideoEncoderInputFrame> InputFrame, const FTextureRHIRef& Texture)
 {
-	FD3D11TextureBase* D3D11Texture = GetD3D11TextureFromRHITexture(Texture);
+	FD3D11Texture* D3D11Texture = GetD3D11TextureFromRHITexture(Texture);
 	unsigned long long TextureMemorySize = D3D11Texture->GetMemorySize();
 
 	ID3D11Texture2D* D3D11NativeTexture = static_cast<ID3D11Texture2D*>(D3D11Texture->GetResource());
@@ -477,7 +476,7 @@ void UE::PixelStreaming::FEncoderFrameFactory::SetTextureCUDAD3D11(TSharedPtr<AV
 	});
 }
 
-void UE::PixelStreaming::FEncoderFrameFactory::SetTextureCUDAD3D12(TSharedPtr<AVEncoder::FVideoEncoderInputFrame> InputFrame, const FTexture2DRHIRef& Texture)
+void UE::PixelStreaming::FEncoderFrameFactory::SetTextureCUDAD3D12(TSharedPtr<AVEncoder::FVideoEncoderInputFrame> InputFrame, const FTextureRHIRef& Texture)
 {
 	ID3D12Resource* NativeD3D12Resource = GetID3D12DynamicRHI()->RHIGetResource(Texture);
 	const int64 TextureMemorySize = GetID3D12DynamicRHI()->RHIGetResourceMemorySize(Texture);
diff --git a/Engine/Plugins/Media/PixelStreaming/Source/PixelStreaming/Private/EncoderFrameFactory.h b/Engine/Plugins/Media/PixelStreaming/Source/PixelStreaming/Private/EncoderFrameFactory.h
--- a/Engine/Plugins/Media/PixelStreaming/Source/PixelStreaming/Private/EncoderFrameFactory.h
+++ b/Engine/Plugins/Media/PixelStreaming/Source/PixelStreaming/Private/EncoderFrameFactory.h
@@ -6,7 +6,7 @@
 #include "Templates/RefCounting.h"
 #include "Templates/SharedPointer.h"
 
-class FRHITexture2D;
+class FRHITexture;
 
 namespace UE
 {
@@ -23,26 +23,26 @@ namespace UE
 		public:
 			FEncoderFrameFactory();
 			~FEncoderFrameFactory();
-			TSharedPtr<AVEncoder::FVideoEncoderInputFrame> GetFrameAndSetTexture(TRefCountPtr<FRHITexture2D> InTexture);
+			TSharedPtr<AVEncoder::FVideoEncoderInputFrame> GetFrameAndSetTexture(TRefCountPtr<FRHITexture> InTexture);
 			TSharedPtr<AVEncoder::FVideoEncoderInput> GetOrCreateVideoEncoderInput();
 
 		private:
-			TSharedPtr<AVEncoder::FVideoEncoderInputFrame> GetOrCreateFrame(const TRefCountPtr<FRHITexture2D> InTexture);
+			TSharedPtr<AVEncoder::FVideoEncoderInputFrame> GetOrCreateFrame(const TRefCountPtr<FRHITexture> InTexture);
 			void RemoveStaleTextures();
 			void FlushFrames();
 			TSharedPtr<AVEncoder::FVideoEncoderInput> CreateVideoEncoderInput() const;
-			void SetTexture(TSharedPtr<AVEncoder::FVideoEncoderInputFrame> Frame, const TRefCountPtr<FRHITexture2D>& Texture);
-			void SetTextureCUDAVulkan(TSharedPtr<AVEncoder::FVideoEncoderInputFrame> Frame, const TRefCountPtr<FRHITexture2D>& Texture);
+			void SetTexture(TSharedPtr<AVEncoder::FVideoEncoderInputFrame> Frame, const TRefCountPtr<FRHITexture>& Texture);
+			void SetTextureCUDAVulkan(TSharedPtr<AVEncoder::FVideoEncoderInputFrame> Frame, const TRefCountPtr<FRHITexture>& Texture);
 #if PLATFORM_WINDOWS
-			void SetTextureCUDAD3D11(TSharedPtr<AVEncoder::FVideoEncoderInputFrame> Frame, const TRefCountPtr<FRHITexture2D>& Texture);
-			void SetTextureCUDAD3D12(TSharedPtr<AVEncoder::FVideoEncoderInputFrame> Frame, const TRefCountPtr<FRHITexture2D>& Texture);
+			void SetTextureCUDAD3D11(TSharedPtr<AVEncoder::FVideoEncoderInputFrame> Frame, const TRefCountPtr<FRHITexture>& Texture);
+			void SetTextureCUDAD3D12(TSharedPtr<AVEncoder::FVideoEncoderInputFrame> Frame, const TRefCountPtr<FRHITexture>& Texture);
 #endif // PLATFORM_WINDOWS
 
 		private:
 			uint64 FrameId = 0;
 			TSharedPtr<AVEncoder::FVideoEncoderInput> EncoderInput;
 			// Store a mapping between raw textures and the FVideoEncoderInputFrames that wrap them
-			TMap<TRefCountPtr<FRHITexture2D>, TSharedPtr<AVEncoder::FVideoEncoderInputFrame>> TextureToFrameMapping;
+			TMap<TRefCountPtr<FRHITexture>, TSharedPtr<AVEncoder::FVideoEncoderInputFrame>> TextureToFrameMapping;
 		};
 	} // namespace PixelStreaming
 } // namespace UE
diff --git a/Engine/Plugins/Media/PixelStreaming/Source/PixelStreaming/Private/Utils.h b/Engine/Plugins/Media/PixelStreaming/Source/PixelStreaming/Private/Utils.h
--- a/Engine/Plugins/Media/PixelStreaming/Source/PixelStreaming/Private/Utils.h
+++ b/Engine/Plugins/Media/PixelStreaming/Source/PixelStreaming/Private/Utils.h
@@ -86,19 +86,34 @@ namespace UE::PixelStreaming
 	inline FTexture2DRHIRef CreateTexture(uint32 Width, uint32 Height)
 	{
 		// Create empty texture
-		FRHIResourceCreateInfo CreateInfo(TEXT("BlankTexture"));
+		ETextureCreateFlags Flags;
+		ERHIAccess InitialState;
 
-		FTexture2DRHIRef Texture;
 
 		if (RHIGetInterfaceType() == ERHIInterfaceType::Vulkan)
 		{
-			Texture = GDynamicRHI->RHICreateTexture2D(Width, Height, EPixelFormat::PF_B8G8R8A8, 1, 1, TexCreate_RenderTargetable | TexCreate_External, ERHIAccess::Present, CreateInfo);
+			Flags = TexCreate_RenderTargetable | TexCreate_External;
+			InitialState = ERHIAccess::Present;
 		}
 		else
 		{
-			Texture = GDynamicRHI->RHICreateTexture2D(Width, Height, EPixelFormat::PF_B8G8R8A8, 1, 1, TexCreate_Shared | TexCreate_RenderTargetable, ERHIAccess::CopyDest, CreateInfo);
+			Flags = TexCreate_RenderTargetable | TexCreate_Shared;
+			InitialState = ERHIAccess::CopyDest;
 		}
-		return Texture;
+
+		FRHITextureCreateDesc Desc = FRHITextureCreateDesc::Create2D(
+			TEXT("PixelStreamingBlankTexture"),
+			{ (int32)Width, (int32)Height },
+			EPixelFormat::PF_B8G8R8A8,
+			FClearValueBinding::None,
+			Flags,
+			1,
+			1,
+			0,
+			InitialState
+		);
+
+		return GDynamicRHI->RHICreateTexture(Desc);
 	}
 
 	inline void ReadTextureToCPU(FRHICommandListImmediate& RHICmdList, FTexture2DRHIRef& TextureRef, TArray<FColor>& OutPixels)
diff --git a/Engine/Plugins/Media/PixelStreaming/Source/PixelStreaming/Public/PixelStreamingFrameBuffer.h b/Engine/Plugins/Media/PixelStreaming/Source/PixelStreaming/Public/PixelStreamingFrameBuffer.h
--- a/Engine/Plugins/Media/PixelStreaming/Source/PixelStreaming/Public/PixelStreamingFrameBuffer.h
+++ b/Engine/Plugins/Media/PixelStreaming/Source/PixelStreaming/Public/PixelStreamingFrameBuffer.h
@@ -9,7 +9,7 @@
 #include "api/video/i420_buffer.h"
 #include "api/scoped_refptr.h"
 
-class FRHITexture2D;
+class FRHITexture;
 class IPixelStreamingTextureSource;
 
 namespace UE
@@ -111,7 +111,7 @@ public:
 	virtual ~FLayerFrameBuffer();
 
 	virtual UE::PixelStreaming::EFrameBufferType GetFrameBufferType() const { return UE::PixelStreaming::EFrameBufferType::Layer; }
-	TRefCountPtr<FRHITexture2D> GetFrame() const;
+	TRefCountPtr<FRHITexture> GetFrame() const;
 
 	virtual int width() const override;
 	virtual int height() const override;
@@ -131,7 +131,7 @@ public:
 	virtual ~FFrameBufferI420();
 
 	virtual UE::PixelStreaming::EFrameBufferType GetFrameBufferType() const { return UE::PixelStreaming::EFrameBufferType::Layer; }
-	TRefCountPtr<FRHITexture2D> GetFrame() const;
+	TRefCountPtr<FRHITexture> GetFrame() const;
 
 	virtual int width() const override;
 	virtual int height() const override;
diff --git a/Engine/Plugins/Media/WebMMedia/Source/WebMMedia/Public/WebMMediaTextureSample.h b/Engine/Plugins/Media/WebMMedia/Source/WebMMedia/Public/WebMMediaTextureSample.h
--- a/Engine/Plugins/Media/WebMMedia/Source/WebMMedia/Public/WebMMediaTextureSample.h
+++ b/Engine/Plugins/Media/WebMMedia/Source/WebMMedia/Public/WebMMediaTextureSample.h
@@ -8,7 +8,7 @@
 #include "RHI.h"
 #include "RHIUtilities.h"
 
-class FRHITexture2D;
+class FRHITexture;
 
 class WEBMMEDIA_API FWebMMediaTextureSample
 	: public IMediaTextureSample
diff --git a/Engine/Plugins/Runtime/Nvidia/Reflex/Source/Private/ReflexLatencyMarkers.cpp b/Engine/Plugins/Runtime/Nvidia/Reflex/Source/Private/ReflexLatencyMarkers.cpp
--- a/Engine/Plugins/Runtime/Nvidia/Reflex/Source/Private/ReflexLatencyMarkers.cpp
+++ b/Engine/Plugins/Runtime/Nvidia/Reflex/Source/Private/ReflexLatencyMarkers.cpp
@@ -28,7 +28,7 @@ static FAutoConsoleVariableRef CVarDisableLatencyMarkers(
 
 void FReflexLatencyMarkers::Initialize()
 {
-	if (IsRHIDeviceNVIDIA())
+	if ((RHIGetInterfaceType() == ERHIInterfaceType::D3D11 || RHIGetInterfaceType() == ERHIInterfaceType::D3D12) && IsRHIDeviceNVIDIA())
 	{
 		FString RHIName = GDynamicRHI->GetName();
 		if (RHIName.StartsWith(TEXT("Vulkan")))
diff --git a/Engine/Plugins/Runtime/Nvidia/Reflex/Source/Private/ReflexMaxTickRateHandler.cpp b/Engine/Plugins/Runtime/Nvidia/Reflex/Source/Private/ReflexMaxTickRateHandler.cpp
--- a/Engine/Plugins/Runtime/Nvidia/Reflex/Source/Private/ReflexMaxTickRateHandler.cpp
+++ b/Engine/Plugins/Runtime/Nvidia/Reflex/Source/Private/ReflexMaxTickRateHandler.cpp
@@ -42,7 +42,7 @@ DEFINE_LOG_CATEGORY_STATIC(LogMaxTickRateHandler, Log, All);
 
 void FReflexMaxTickRateHandler::Initialize()
 {
-	if (IsRHIDeviceNVIDIA())
+	if ((RHIGetInterfaceType() == ERHIInterfaceType::D3D11 || RHIGetInterfaceType() == ERHIInterfaceType::D3D12) && IsRHIDeviceNVIDIA())
 	{
 		FString RHIName = GDynamicRHI->GetName();
 		if (RHIName.StartsWith(TEXT("Vulkan")))
diff --git a/Engine/Plugins/Runtime/Oculus/OculusVR/Source/OculusHMD/Private/OculusHMD.cpp b/Engine/Plugins/Runtime/Oculus/OculusVR/Source/OculusHMD/Private/OculusHMD.cpp
--- a/Engine/Plugins/Runtime/Oculus/OculusVR/Source/OculusHMD/Private/OculusHMD.cpp
+++ b/Engine/Plugins/Runtime/Oculus/OculusVR/Source/OculusHMD/Private/OculusHMD.cpp
@@ -1372,7 +1372,7 @@ namespace OculusHMD
 
 
 
-	void FOculusHMD::RenderTexture_RenderThread(class FRHICommandListImmediate& RHICmdList, class FRHITexture2D* BackBuffer, class FRHITexture2D* SrcTexture, FVector2D WindowSize) const
+	void FOculusHMD::RenderTexture_RenderThread(class FRHICommandListImmediate& RHICmdList, class FRHITexture* BackBuffer, class FRHITexture* SrcTexture, FVector2D WindowSize) const
 	{
 		CheckInRenderThread();
 		check(CustomPresent);
diff --git a/Engine/Plugins/Runtime/Oculus/OculusVR/Source/OculusHMD/Private/OculusHMD.h b/Engine/Plugins/Runtime/Oculus/OculusVR/Source/OculusHMD/Private/OculusHMD.h
--- a/Engine/Plugins/Runtime/Oculus/OculusVR/Source/OculusHMD/Private/OculusHMD.h
+++ b/Engine/Plugins/Runtime/Oculus/OculusVR/Source/OculusHMD/Private/OculusHMD.h
@@ -167,7 +167,7 @@ public:
 	virtual FMatrix GetStereoProjectionMatrix(const int32 ViewIndex) const override;
 	virtual void InitCanvasFromView(class FSceneView* InView, class UCanvas* Canvas) override;
 	//virtual void GetEyeRenderParams_RenderThread(const struct FHeadMountedDisplayPassContext& Context, FVector2D& EyeToSrcUVScaleValue, FVector2D& EyeToSrcUVOffsetValue) const override;
-	virtual void RenderTexture_RenderThread(class FRHICommandListImmediate& RHICmdList, class FRHITexture2D* BackBuffer, class FRHITexture2D* SrcTexture, FVector2D WindowSize) const override;
+	virtual void RenderTexture_RenderThread(class FRHICommandListImmediate& RHICmdList, class FRHITexture* BackBuffer, class FRHITexture* SrcTexture, FVector2D WindowSize) const override;
 	//virtual void SetClippingPlanes(float NCP, float FCP) override;
 	virtual IStereoRenderTargetManager* GetRenderTargetManager() override { return this; }
 	virtual IStereoLayers* GetStereoLayers() override { return this; }
diff --git a/Engine/Plugins/Runtime/OpenXR/Source/OpenXRHMD/Private/OpenXRHMD.cpp b/Engine/Plugins/Runtime/OpenXR/Source/OpenXRHMD/Private/OpenXRHMD.cpp
--- a/Engine/Plugins/Runtime/OpenXR/Source/OpenXRHMD/Private/OpenXRHMD.cpp
+++ b/Engine/Plugins/Runtime/OpenXR/Source/OpenXRHMD/Private/OpenXRHMD.cpp
@@ -2502,7 +2502,7 @@ void FOpenXRHMD::CopyTexture_RenderThread(FRHICommandListImmediate& RHICmdList,
 	CopyTexture_RenderThread(RHICmdList, SrcTexture, SrcRect, DstTexture, DstRect, bClearBlack, bNoAlpha, ERenderTargetActions::DontLoad_Store);
 }
 
-void FOpenXRHMD::RenderTexture_RenderThread(class FRHICommandListImmediate& RHICmdList, class FRHITexture2D* BackBuffer, class FRHITexture2D* SrcTexture, FVector2D WindowSize) const
+void FOpenXRHMD::RenderTexture_RenderThread(class FRHICommandListImmediate& RHICmdList, class FRHITexture* BackBuffer, class FRHITexture* SrcTexture, FVector2D WindowSize) const
 {
 	if (SpectatorScreenController)
 	{
diff --git a/Engine/Plugins/Runtime/OpenXR/Source/OpenXRHMD/Private/OpenXRHMD.h b/Engine/Plugins/Runtime/OpenXR/Source/OpenXRHMD/Private/OpenXRHMD.h
--- a/Engine/Plugins/Runtime/OpenXR/Source/OpenXRHMD/Private/OpenXRHMD.h
+++ b/Engine/Plugins/Runtime/OpenXR/Source/OpenXRHMD/Private/OpenXRHMD.h
@@ -239,7 +239,7 @@ public:
 	virtual FMatrix GetStereoProjectionMatrix(const int32 StereoViewIndex) const override;
 	virtual void GetEyeRenderParams_RenderThread(const struct FHeadMountedDisplayPassContext& Context, FVector2D& EyeToSrcUVScaleValue, FVector2D& EyeToSrcUVOffsetValue) const override;
 	virtual IStereoRenderTargetManager* GetRenderTargetManager() override;
-	virtual void RenderTexture_RenderThread(class FRHICommandListImmediate& RHICmdList, class FRHITexture2D* BackBuffer, class FRHITexture2D* SrcTexture, FVector2D WindowSize) const override;
+	virtual void RenderTexture_RenderThread(class FRHICommandListImmediate& RHICmdList, class FRHITexture* BackBuffer, class FRHITexture* SrcTexture, FVector2D WindowSize) const override;
 
 	/** ISceneViewExtension interface */
 	virtual void SetupViewFamily(FSceneViewFamily& InViewFamily) override;
diff --git a/Engine/Plugins/Runtime/Steam/SteamVR/Source/SteamVR/Private/SteamVRHMD.h b/Engine/Plugins/Runtime/Steam/SteamVR/Source/SteamVR/Private/SteamVRHMD.h
--- a/Engine/Plugins/Runtime/Steam/SteamVR/Source/SteamVR/Private/SteamVRHMD.h
+++ b/Engine/Plugins/Runtime/Steam/SteamVR/Source/SteamVR/Private/SteamVRHMD.h
@@ -166,7 +166,7 @@ public:
 	virtual void AdjustViewRect(int32 ViewIndex, int32& X, int32& Y, uint32& SizeX, uint32& SizeY) const override;
 	virtual void CalculateStereoViewOffset(const int32 ViewIndex, FRotator& ViewRotation, const float MetersToWorld, FVector& ViewLocation) override;
 	virtual FMatrix GetStereoProjectionMatrix(const int32 ViewIndex) const override;
-	virtual void RenderTexture_RenderThread(FRHICommandListImmediate& RHICmdList, FRHITexture2D* BackBuffer, FRHITexture2D* SrcTexture, FVector2D WindowSize) const override;
+	virtual void RenderTexture_RenderThread(FRHICommandListImmediate& RHICmdList, FRHITexture* BackBuffer, FRHITexture* SrcTexture, FVector2D WindowSize) const override;
 	virtual void GetEyeRenderParams_RenderThread(const FHeadMountedDisplayPassContext& Context, FVector2D& EyeToSrcUVScaleValue, FVector2D& EyeToSrcUVOffsetValue) const override;
 	virtual IStereoRenderTargetManager* GetRenderTargetManager() override { return this; }
 	virtual IStereoLayers* GetStereoLayers() override;
diff --git a/Engine/Plugins/Runtime/Steam/SteamVR/Source/SteamVR/Private/SteamVRRender.cpp b/Engine/Plugins/Runtime/Steam/SteamVR/Source/SteamVR/Private/SteamVRRender.cpp
--- a/Engine/Plugins/Runtime/Steam/SteamVR/Source/SteamVR/Private/SteamVRRender.cpp
+++ b/Engine/Plugins/Runtime/Steam/SteamVR/Source/SteamVR/Private/SteamVRRender.cpp
@@ -38,7 +38,7 @@ void FSteamVRHMD::DrawDistortionMesh_RenderThread(struct FHeadMountedDisplayPass
 	check(0);
 }
 
-void FSteamVRHMD::RenderTexture_RenderThread(FRHICommandListImmediate& RHICmdList, FRHITexture2D* BackBuffer, FRHITexture2D* SrcTexture, FVector2D WindowSize) const
+void FSteamVRHMD::RenderTexture_RenderThread(FRHICommandListImmediate& RHICmdList, FRHITexture* BackBuffer, FRHITexture* SrcTexture, FVector2D WindowSize) const
 {
 	check(IsInRenderingThread());
 
diff --git a/Engine/Plugins/Runtime/WindowsMixedReality/Source/WindowsMixedRealityHMD/Private/WindowsMixedRealityHMD.cpp b/Engine/Plugins/Runtime/WindowsMixedReality/Source/WindowsMixedRealityHMD/Private/WindowsMixedRealityHMD.cpp
--- a/Engine/Plugins/Runtime/WindowsMixedReality/Source/WindowsMixedRealityHMD/Private/WindowsMixedRealityHMD.cpp
+++ b/Engine/Plugins/Runtime/WindowsMixedReality/Source/WindowsMixedRealityHMD/Private/WindowsMixedRealityHMD.cpp
@@ -1544,7 +1544,7 @@ namespace WindowsMixedReality
 #endif
 	}
 
-	void FWindowsMixedRealityHMD::RenderTexture_RenderThread(class FRHICommandListImmediate& RHICmdList, class FRHITexture2D* BackBuffer, class FRHITexture2D* SrcTexture, FVector2D WindowSize) const
+	void FWindowsMixedRealityHMD::RenderTexture_RenderThread(class FRHICommandListImmediate& RHICmdList, class FRHITexture* BackBuffer, class FRHITexture* SrcTexture, FVector2D WindowSize) const
 	{
 		if (SpectatorScreenController)
 		{
diff --git a/Engine/Plugins/Runtime/WindowsMixedReality/Source/WindowsMixedRealityHMD/Private/WindowsMixedRealityHMD.h b/Engine/Plugins/Runtime/WindowsMixedReality/Source/WindowsMixedRealityHMD/Private/WindowsMixedRealityHMD.h
--- a/Engine/Plugins/Runtime/WindowsMixedReality/Source/WindowsMixedRealityHMD/Private/WindowsMixedRealityHMD.h
+++ b/Engine/Plugins/Runtime/WindowsMixedReality/Source/WindowsMixedRealityHMD/Private/WindowsMixedRealityHMD.h
@@ -283,8 +283,8 @@ namespace WindowsMixedReality
 		virtual void RenderTexture_RenderThread
 		(
 			class FRHICommandListImmediate & RHICmdList,
-			class FRHITexture2D * BackBuffer,
-			class FRHITexture2D * SrcTexture,
+			class FRHITexture * BackBuffer,
+			class FRHITexture * SrcTexture,
 			FVector2D WindowSize
 		) const override;
 		virtual bool AllocateRenderTargetTexture(
diff --git a/Engine/Plugins/Runtime/nDisplay/Source/DisplayCluster/Private/Render/Device/DisplayClusterDeviceBase.cpp b/Engine/Plugins/Runtime/nDisplay/Source/DisplayCluster/Private/Render/Device/DisplayClusterDeviceBase.cpp
--- a/Engine/Plugins/Runtime/nDisplay/Source/DisplayCluster/Private/Render/Device/DisplayClusterDeviceBase.cpp
+++ b/Engine/Plugins/Runtime/nDisplay/Source/DisplayCluster/Private/Render/Device/DisplayClusterDeviceBase.cpp
@@ -463,7 +463,7 @@ DECLARE_GPU_STAT_NAMED(nDisplay_Device_RenderTexture, TEXT("nDisplay RenderDevic
 #include "RenderGraphUtils.h"
 #include "RenderGraphBuilder.h"
 
-void FDisplayClusterDeviceBase::RenderTexture_RenderThread(FRHICommandListImmediate& RHICmdList, FRHITexture2D* BackBuffer, FRHITexture2D* SrcTexture, FVector2D WindowSize) const
+void FDisplayClusterDeviceBase::RenderTexture_RenderThread(FRHICommandListImmediate& RHICmdList, FRHITexture* BackBuffer, FRHITexture* SrcTexture, FVector2D WindowSize) const
 {
 	SCOPED_GPU_STAT(RHICmdList, nDisplay_Device_RenderTexture);
 	SCOPED_DRAW_EVENT(RHICmdList, nDisplay_Device_RenderTexture);
diff --git a/Engine/Plugins/Runtime/nDisplay/Source/DisplayCluster/Private/Render/Device/DisplayClusterDeviceBase.h b/Engine/Plugins/Runtime/nDisplay/Source/DisplayCluster/Private/Render/Device/DisplayClusterDeviceBase.h
--- a/Engine/Plugins/Runtime/nDisplay/Source/DisplayCluster/Private/Render/Device/DisplayClusterDeviceBase.h
+++ b/Engine/Plugins/Runtime/nDisplay/Source/DisplayCluster/Private/Render/Device/DisplayClusterDeviceBase.h
@@ -58,7 +58,7 @@ public:
 	virtual void AdjustViewRect(int32 ViewIndex, int32& X, int32& Y, uint32& SizeX, uint32& SizeY) const override;
 	virtual void CalculateStereoViewOffset(const int32 ViewIndex, FRotator& ViewRotation, const float WorldToMeters, FVector& ViewLocation) override;
 	virtual FMatrix GetStereoProjectionMatrix(const int32 ViewIndex) const override;
-	virtual void RenderTexture_RenderThread(FRHICommandListImmediate& RHICmdList, FRHITexture2D* BackBuffer, FRHITexture2D* SrcTexture, FVector2D WindowSize) const override;
+	virtual void RenderTexture_RenderThread(FRHICommandListImmediate& RHICmdList, FRHITexture* BackBuffer, FRHITexture* SrcTexture, FVector2D WindowSize) const override;
 	
 	virtual int32 GetDesiredNumberOfViews(bool bStereoRequested) const override
 	{
diff --git a/Engine/Source/Runtime/Apple/AGXRHI/Private/AGXCommandEncoder.cpp b/Engine/Source/Runtime/Apple/AGXRHI/Private/AGXCommandEncoder.cpp
--- a/Engine/Source/Runtime/Apple/AGXRHI/Private/AGXCommandEncoder.cpp
+++ b/Engine/Source/Runtime/Apple/AGXRHI/Private/AGXCommandEncoder.cpp
@@ -1128,21 +1128,23 @@ void FAGXCommandEncoder::FenceResource(mtlpp::Buffer const& Resource)
 
 void FAGXCommandEncoder::SetShaderBufferInternal(mtlpp::FunctionType Function, uint32 Index)
 {
-	NSUInteger Offset = ShaderBuffers[uint32(Function)].Offsets[Index];
-	mtlpp::ResourceUsage Usage = ShaderBuffers[uint32(Function)].Usage[Index];
-	bool bBufferHasBytes = ShaderBuffers[uint32(Function)].Bytes[Index] != nil;
-	if (!ShaderBuffers[uint32(Function)].Buffers[Index] && bBufferHasBytes && !bSupportsMetalFeaturesSetBytes)
+	FAGXBufferBindings& Binding = ShaderBuffers[uint32(Function)];
+
+	NSUInteger Offset = Binding.Offsets[Index];
+	
+	bool bBufferHasBytes = Binding.Bytes[Index] != nil;
+	if (!Binding.Buffers[Index] && bBufferHasBytes && !bSupportsMetalFeaturesSetBytes)
 	{
-		uint8 const* Bytes = (((uint8 const*)ShaderBuffers[uint32(Function)].Bytes[Index]->Data) + ShaderBuffers[uint32(Function)].Offsets[Index]);
-		uint32 Len = ShaderBuffers[uint32(Function)].Bytes[Index]->Len - ShaderBuffers[uint32(Function)].Offsets[Index];
+		uint8 const* Bytes = (((uint8 const*)Binding.Bytes[Index]->Data) + Binding.Offsets[Index]);
+		uint32 Len = Binding.Bytes[Index]->Len - Binding.Offsets[Index];
 		
 		Offset = 0;
-		ShaderBuffers[uint32(Function)].Buffers[Index] = RingBuffer.NewBuffer(Len, BufferOffsetAlignment);
+		Binding.Buffers[Index] = RingBuffer.NewBuffer(Len, BufferOffsetAlignment);
 		
-		FMemory::Memcpy(((uint8*)ShaderBuffers[uint32(Function)].Buffers[Index].GetContents()) + Offset, Bytes, Len);
+		FMemory::Memcpy(((uint8*)Binding.Buffers[Index].GetContents()) + Offset, Bytes, Len);
 	}
 	
-	ns::AutoReleased<FAGXBuffer>& Buffer = ShaderBuffers[uint32(Function)].Buffers[Index];
+	ns::AutoReleased<FAGXBuffer>& Buffer = Binding.Buffers[Index];
 	if (Buffer)
 	{
 #if METAL_DEBUG_OPTIONS
@@ -1155,22 +1157,25 @@ void FAGXCommandEncoder::SetShaderBufferInternal(mtlpp::FunctionType Function, u
 		switch (Function)
 		{
 			case mtlpp::FunctionType::Vertex:
-				ShaderBuffers[uint32(Function)].Bound |= (1 << Index);
+				Binding.Bound |= (1 << Index);
 				check(RenderCommandEncoder);
 				FenceResource(Buffer);
 				RenderCommandEncoder.SetVertexBuffer(Buffer, Offset, Index);
 				break;
+
 			case mtlpp::FunctionType::Fragment:
-				ShaderBuffers[uint32(Function)].Bound |= (1 << Index);
+				Binding.Bound |= (1 << Index);
 				check(RenderCommandEncoder);
 				RenderCommandEncoder.SetFragmentBuffer(Buffer, Offset, Index);
 				break;
+
 			case mtlpp::FunctionType::Kernel:
-				ShaderBuffers[uint32(Function)].Bound |= (1 << Index);
+				Binding.Bound |= (1 << Index);
 				check(ComputeCommandEncoder);
 				FenceResource(Buffer);
 				ComputeCommandEncoder.SetBuffer(Buffer, Offset, Index);
 				break;
+
 			default:
 				check(false);
 				break;
@@ -1178,34 +1183,37 @@ void FAGXCommandEncoder::SetShaderBufferInternal(mtlpp::FunctionType Function, u
 		
 		if (Buffer.IsSingleUse())
 		{
-			ShaderBuffers[uint32(Function)].Usage[Index] = mtlpp::ResourceUsage(0);
-			ShaderBuffers[uint32(Function)].Offsets[Index] = 0;
-			ShaderBuffers[uint32(Function)].Buffers[Index] = nil;
-			ShaderBuffers[uint32(Function)].Bound &= ~(1 << Index);
+			Binding.Usage[Index] = mtlpp::ResourceUsage(0);
+			Binding.Offsets[Index] = 0;
+			Binding.Buffers[Index] = nil;
+			Binding.Bound &= ~(1 << Index);
 		}
 	}
 	else if (bBufferHasBytes && bSupportsMetalFeaturesSetBytes)
 	{
-		uint8 const* Bytes = (((uint8 const*)ShaderBuffers[uint32(Function)].Bytes[Index]->Data) + ShaderBuffers[uint32(Function)].Offsets[Index]);
-		uint32 Len = ShaderBuffers[uint32(Function)].Bytes[Index]->Len - ShaderBuffers[uint32(Function)].Offsets[Index];
+		uint8 const* Bytes = (((uint8 const*)Binding.Bytes[Index]->Data) + Binding.Offsets[Index]);
+		uint32 Len = Binding.Bytes[Index]->Len - Binding.Offsets[Index];
 		
 		switch (Function)
 		{
 			case mtlpp::FunctionType::Vertex:
-				ShaderBuffers[uint32(Function)].Bound |= (1 << Index);
+				Binding.Bound |= (1 << Index);
 				check(RenderCommandEncoder);
 				RenderCommandEncoder.SetVertexData(Bytes, Len, Index);
 				break;
+
 			case mtlpp::FunctionType::Fragment:
-				ShaderBuffers[uint32(Function)].Bound |= (1 << Index);
+				Binding.Bound |= (1 << Index);
 				check(RenderCommandEncoder);
 				RenderCommandEncoder.SetFragmentData(Bytes, Len, Index);
 				break;
+
 			case mtlpp::FunctionType::Kernel:
-				ShaderBuffers[uint32(Function)].Bound |= (1 << Index);
+				Binding.Bound |= (1 << Index);
 				check(ComputeCommandEncoder);
 				ComputeCommandEncoder.SetBytes(Bytes, Len, Index);
 				break;
+
 			default:
 				check(false);
 				break;
diff --git a/Engine/Source/Runtime/Apple/AGXRHI/Private/AGXCommands.cpp b/Engine/Source/Runtime/Apple/AGXRHI/Private/AGXCommands.cpp
--- a/Engine/Source/Runtime/Apple/AGXRHI/Private/AGXCommands.cpp
+++ b/Engine/Source/Runtime/Apple/AGXRHI/Private/AGXCommands.cpp
@@ -307,7 +307,7 @@ void FAGXRHICommandContext::RHISetShaderTexture(FRHIGraphicsShader* ShaderRHI, u
 	EAGXShaderStages Stage = GetShaderStage(ShaderRHI);
 	if (Surface != nullptr)
 	{
-        if (Surface->Texture || !(Surface->Flags & TexCreate_Presentable))
+        if (Surface->Texture || !(Surface->GetDesc().Flags & TexCreate_Presentable))
         {
             Context->GetCurrentState().SetShaderTexture(Stage, Surface->Texture, TextureIndex, (mtlpp::ResourceUsage)(mtlpp::ResourceUsage::Read|mtlpp::ResourceUsage::Sample));
         }
@@ -330,7 +330,7 @@ void FAGXRHICommandContext::RHISetShaderTexture(FRHIComputeShader* ComputeShader
 	FAGXSurface* Surface = AGXGetMetalSurfaceFromRHITexture(NewTextureRHI);
 	if (Surface != nullptr)
     {
-        if (Surface->Texture || !(Surface->Flags & TexCreate_Presentable))
+        if (Surface->Texture || !(Surface->GetDesc().Flags & TexCreate_Presentable))
         {
             Context->GetCurrentState().SetShaderTexture(EAGXShaderStages::Compute, Surface->Texture, TextureIndex, (mtlpp::ResourceUsage)(mtlpp::ResourceUsage::Read|mtlpp::ResourceUsage::Sample));
         }
diff --git a/Engine/Source/Runtime/Apple/AGXRHI/Private/AGXContext.cpp b/Engine/Source/Runtime/Apple/AGXRHI/Private/AGXContext.cpp
--- a/Engine/Source/Runtime/Apple/AGXRHI/Private/AGXContext.cpp
+++ b/Engine/Source/Runtime/Apple/AGXRHI/Private/AGXContext.cpp
@@ -1085,36 +1085,10 @@ void FAGXContext::TransitionResource(FRHIUnorderedAccessView* InResource)
 {
 	FAGXUnorderedAccessView* UAV = ResourceCast(InResource);
 
-	// figure out which one of the resources we need to set
-	FAGXStructuredBuffer* StructuredBuffer = UAV->SourceView->SourceStructuredBuffer.GetReference();
-	FAGXVertexBuffer*     VertexBuffer     = UAV->SourceView->SourceVertexBuffer.GetReference();
-	FAGXIndexBuffer*      IndexBuffer      = UAV->SourceView->SourceIndexBuffer.GetReference();
-	FRHITexture*            Texture          = UAV->SourceView->SourceTexture.GetReference();
-	FAGXSurface*          Surface          = UAV->SourceView->TextureView;
-
-	if (StructuredBuffer)
+	if (UAV->bTexture)
 	{
-		RenderPass.TransitionResources(StructuredBuffer->GetCurrentBuffer());
-	}
-	else if (VertexBuffer && VertexBuffer->GetCurrentBufferOrNil())
-	{
-		RenderPass.TransitionResources(VertexBuffer->GetCurrentBuffer());
-	}
-	else if (IndexBuffer)
-	{
-		RenderPass.TransitionResources(IndexBuffer->GetCurrentBuffer());
-	}
-	else if (Surface)
-	{
-		RenderPass.TransitionResources(Surface->Texture.GetParentTexture());
-	}
-	else if (Texture)
-	{
-		if (!Surface)
-		{
-			Surface = AGXGetMetalSurfaceFromRHITexture(Texture);
-		}
-		if ((Surface != nullptr) && Surface->Texture)
+		FAGXSurface* Surface = UAV->GetSourceTexture();
+		if (Surface->Texture)
 		{
 			RenderPass.TransitionResources(Surface->Texture);
 			if (Surface->MSAATexture)
@@ -1123,6 +1097,11 @@ void FAGXContext::TransitionResource(FRHIUnorderedAccessView* InResource)
 			}
 		}
 	}
+	else
+	{
+		FAGXResourceMultiBuffer* Buffer = UAV->GetSourceBuffer();
+		RenderPass.TransitionResources(Buffer->GetCurrentBuffer());
+	}
 }
 
 void FAGXContext::TransitionResource(FRHITexture* InResource)
diff --git a/Engine/Source/Runtime/Apple/AGXRHI/Private/AGXDynamicRHI.h b/Engine/Source/Runtime/Apple/AGXRHI/Private/AGXDynamicRHI.h
--- a/Engine/Source/Runtime/Apple/AGXRHI/Private/AGXDynamicRHI.h
+++ b/Engine/Source/Runtime/Apple/AGXRHI/Private/AGXDynamicRHI.h
@@ -22,7 +22,9 @@ class FAGXDeviceContext;
 /** The interface which is implemented by the dynamically bound RHI. */
 class FAGXDynamicRHI : public FDynamicRHI
 {
+	static inline FAGXDynamicRHI* Singleton = nullptr;
 public:
+	static inline FAGXDynamicRHI& Get() { return *Singleton; }
 
 	/** Initialization constructor. */
 	FAGXDynamicRHI(ERHIFeatureLevel::Type RequestedFeatureLevel);
@@ -67,18 +69,13 @@ public:
 	virtual FShaderResourceViewRHIRef RHICreateShaderResourceView(const FShaderResourceViewInitializer& Initializer) final override;
 	virtual void RHIUpdateShaderResourceView(FRHIShaderResourceView* SRV, FRHIBuffer* Buffer, uint32 Stride, uint8 Format) final override;
 	virtual void RHIUpdateShaderResourceView(FRHIShaderResourceView* SRV, FRHIBuffer* Buffer) final override;
-	virtual uint64 RHICalcTexture2DPlatformSize(uint32 SizeX, uint32 SizeY, uint8 Format, uint32 NumMips, uint32 NumSamples, ETextureCreateFlags Flags, const FRHIResourceCreateInfo& CreateInfo, uint32& OutAlign) final override;
-	virtual uint64 RHICalcTexture3DPlatformSize(uint32 SizeX, uint32 SizeY, uint32 SizeZ, uint8 Format, uint32 NumMips, ETextureCreateFlags Flags, const FRHIResourceCreateInfo& CreateInfo, uint32& OutAlign) final override;
-	virtual uint64 RHICalcTextureCubePlatformSize(uint32 Size, uint8 Format, uint32 NumMips, ETextureCreateFlags Flags, const FRHIResourceCreateInfo& CreateInfo, uint32& OutAlign) final override;
+	virtual FRHICalcTextureSizeResult RHICalcTexturePlatformSize(FRHITextureDesc const& Desc, uint32 FirstMipIndex) final override;
 	virtual uint64 RHIGetMinimumAlignmentForBufferBackedSRV(EPixelFormat Format) final override;
 	virtual void RHIGetTextureMemoryStats(FTextureMemoryStats& OutStats) final override;
 	virtual bool RHIGetTextureMemoryVisualizeData(FColor* TextureData,int32 SizeX,int32 SizeY,int32 Pitch,int32 PixelSize) final override;
-	virtual FTexture2DRHIRef RHICreateTexture2D(uint32 SizeX, uint32 SizeY, uint8 Format, uint32 NumMips, uint32 NumSamples, ETextureCreateFlags Flags, ERHIAccess InResourceState, FRHIResourceCreateInfo& CreateInfo) final override;
-	virtual FTexture2DRHIRef RHICreateTextureExternal2D(uint32 SizeX, uint32 SizeY, uint8 Format, uint32 NumMips, uint32 NumSamples, ETextureCreateFlags Flags, ERHIAccess InResourceState, FRHIResourceCreateInfo& Info) final override;
+	virtual FTextureRHIRef RHICreateTexture(const FRHITextureCreateDesc& CreateDesc) final override;
 	virtual FTexture2DRHIRef RHIAsyncCreateTexture2D(uint32 SizeX, uint32 SizeY, uint8 Format, uint32 NumMips, ETextureCreateFlags Flags, ERHIAccess InResourceState, void** InitialMipData, uint32 NumInitialMips) final override;
 	virtual void RHICopySharedMips(FRHITexture2D* DestTexture2D, FRHITexture2D* SrcTexture2D) final override;
-	virtual FTexture2DArrayRHIRef RHICreateTexture2DArray(uint32 SizeX, uint32 SizeY, uint32 SizeZ, uint8 Format, uint32 NumMips, uint32 NumSamples, ETextureCreateFlags Flags, ERHIAccess InResourceState, FRHIResourceCreateInfo& CreateInfo) final override;
-	virtual FTexture3DRHIRef RHICreateTexture3D(uint32 SizeX, uint32 SizeY, uint32 SizeZ, uint8 Format, uint32 NumMips, ETextureCreateFlags Flags, ERHIAccess InResourceState, FRHIResourceCreateInfo& CreateInfo) final override;
 	virtual void RHIGenerateMips(FRHITexture* Texture) final override;
 	virtual uint32 RHIComputeMemorySize(FRHITexture* TextureRHI) final override;
 	virtual FTexture2DRHIRef RHIAsyncReallocateTexture2D(FRHITexture2D* Texture2D, int32 NewMipCount, int32 NewSizeX, int32 NewSizeY, FThreadSafeCounter* RequestStatus) final override;
@@ -90,8 +87,6 @@ public:
 	virtual void RHIUnlockTexture2DArray(FRHITexture2DArray* Texture, uint32 TextureIndex, uint32 MipIndex, bool bLockWithinMiptail) final override;
 	virtual void RHIUpdateTexture2D(FRHITexture2D* Texture, uint32 MipIndex, const struct FUpdateTextureRegion2D& UpdateRegion, uint32 SourcePitch, const uint8* SourceData) final override;
 	virtual void RHIUpdateTexture3D(FRHITexture3D* Texture, uint32 MipIndex, const struct FUpdateTextureRegion3D& UpdateRegion, uint32 SourceRowPitch, uint32 SourceDepthPitch, const uint8* SourceData) final override;
-	virtual FTextureCubeRHIRef RHICreateTextureCube(uint32 Size, uint8 Format, uint32 NumMips, ETextureCreateFlags Flags, ERHIAccess InResourceState, FRHIResourceCreateInfo& CreateInfo) final override;
-	virtual FTextureCubeRHIRef RHICreateTextureCubeArray(uint32 Size, uint32 ArraySize, uint8 Format, uint32 NumMips, ETextureCreateFlags Flags, ERHIAccess InResourceState, FRHIResourceCreateInfo& CreateInfo) final override;
 	virtual void* RHILockTextureCubeFace(FRHITextureCube* Texture, uint32 FaceIndex, uint32 ArrayIndex, uint32 MipIndex, EResourceLockMode LockMode, uint32& DestStride, bool bLockWithinMiptail) final override;
 	virtual void RHIUnlockTextureCubeFace(FRHITextureCube* Texture, uint32 FaceIndex, uint32 ArrayIndex, uint32 MipIndex, bool bLockWithinMiptail) final override;
 	virtual void RHIBindDebugLabelName(FRHITexture* Texture, const TCHAR* Name) final override;
@@ -139,11 +134,7 @@ public:
 	virtual void EndUpdateTexture3D_RenderThread(class FRHICommandListImmediate& RHICmdList, FUpdateTexture3DData& UpdateData) final override;
 	virtual void UpdateTexture3D_RenderThread(class FRHICommandListImmediate& RHICmdList, FRHITexture3D* Texture, uint32 MipIndex, const struct FUpdateTextureRegion3D& UpdateRegion, uint32 SourceRowPitch, uint32 SourceDepthPitch, const uint8* SourceData) final override;
 	
-	virtual FTexture2DRHIRef RHICreateTexture2D_RenderThread(class FRHICommandListImmediate& RHICmdList, uint32 SizeX, uint32 SizeY, uint8 Format, uint32 NumMips, uint32 NumSamples, ETextureCreateFlags Flags, ERHIAccess InResourceState, FRHIResourceCreateInfo& CreateInfo) final override;
-	virtual FTexture2DArrayRHIRef RHICreateTexture2DArray_RenderThread(class FRHICommandListImmediate& RHICmdList, uint32 SizeX, uint32 SizeY, uint32 SizeZ, uint8 Format, uint32 NumMips, uint32 NumSamples, ETextureCreateFlags Flags, ERHIAccess InResourceState, FRHIResourceCreateInfo& CreateInfo) final override;
-	virtual FTexture3DRHIRef RHICreateTexture3D_RenderThread(class FRHICommandListImmediate& RHICmdList, uint32 SizeX, uint32 SizeY, uint32 SizeZ, uint8 Format, uint32 NumMips, ETextureCreateFlags Flags, ERHIAccess InResourceState, FRHIResourceCreateInfo& CreateInfo) final override;
-	virtual FTextureCubeRHIRef RHICreateTextureCube_RenderThread(class FRHICommandListImmediate& RHICmdList, uint32 Size, uint8 Format, uint32 NumMips, ETextureCreateFlags Flags, ERHIAccess InResourceState, FRHIResourceCreateInfo& CreateInfo) final override;
-	virtual FTextureCubeRHIRef RHICreateTextureCubeArray_RenderThread(class FRHICommandListImmediate& RHICmdList, uint32 Size, uint32 ArraySize, uint8 Format, uint32 NumMips, ETextureCreateFlags Flags, ERHIAccess InResourceState, FRHIResourceCreateInfo& CreateInfo) final override;
+	virtual FTextureRHIRef RHICreateTexture_RenderThread(class FRHICommandListImmediate& RHICmdList, const FRHITextureCreateDesc& CreateDesc) final override;
 	
 	virtual FRenderQueryRHIRef RHICreateRenderQuery_RenderThread(class FRHICommandListImmediate& RHICmdList, ERenderQueryType QueryType) final override;
 	
diff --git a/Engine/Source/Runtime/Apple/AGXRHI/Private/AGXLLM.cpp b/Engine/Source/Runtime/Apple/AGXRHI/Private/AGXLLM.cpp
--- a/Engine/Source/Runtime/Apple/AGXRHI/Private/AGXLLM.cpp
+++ b/Engine/Source/Runtime/Apple/AGXRHI/Private/AGXLLM.cpp
@@ -143,47 +143,6 @@ static EPixelFormat MetalToRHIPixelFormat(mtlpp::PixelFormat Format)
 	return PF_MAX;
 }
 
-static mtlpp::SizeAndAlign TextureSizeAndAlign(mtlpp::TextureType TextureType, uint32 Width, uint32 Height, uint32 Depth, mtlpp::PixelFormat Format, uint32 MipCount, uint32 SampleCount, uint32 ArrayCount)
-{
-	mtlpp::SizeAndAlign SizeAlign;
-	SizeAlign.Size = 0;
-	SizeAlign.Align = 0;
-	
-	uint32 Align = 0;
-	FRHIResourceCreateInfo CreateInfo(TEXT(""));
-	switch (TextureType)
-	{
-		case mtlpp::TextureType::Texture2D:
-		case mtlpp::TextureType::Texture2DMultisample:
-			SizeAlign.Size = RHICalcTexture2DPlatformSize(Width, Height, MetalToRHIPixelFormat(Format), MipCount, SampleCount, TexCreate_None, CreateInfo, Align);
-			SizeAlign.Align = Align;
-			break;
-		case mtlpp::TextureType::Texture2DArray:
-			SizeAlign.Size = RHICalcTexture2DPlatformSize(Width, Height, MetalToRHIPixelFormat(Format), MipCount, SampleCount, TexCreate_None, CreateInfo, Align) * ArrayCount;
-			SizeAlign.Align = Align;
-			break;
-		case mtlpp::TextureType::TextureCube:
-			SizeAlign.Size = RHICalcTextureCubePlatformSize(Width, MetalToRHIPixelFormat(Format), MipCount, TexCreate_None, CreateInfo, Align);
-			SizeAlign.Align = Align;
-			break;
-		case mtlpp::TextureType::TextureCubeArray:
-			SizeAlign.Size = RHICalcTextureCubePlatformSize(Width, MetalToRHIPixelFormat(Format), MipCount, TexCreate_None, CreateInfo, Align) * ArrayCount;
-			SizeAlign.Align = Align;
-			break;
-		case mtlpp::TextureType::Texture3D:
-			SizeAlign.Size = RHICalcTexture3DPlatformSize(Width, Height, Depth, MetalToRHIPixelFormat(Format), MipCount, TexCreate_None, CreateInfo, Align);
-			SizeAlign.Align = Align;
-			break;
-		case mtlpp::TextureType::Texture1D:
-		case mtlpp::TextureType::Texture1DArray:
-		default:
-			check(false);
-			break;
-	}
-	
-	return SizeAlign;
-}
-
 void AGXLLM::LogAllocTexture(mtlpp::TextureDescriptor const& Desc, mtlpp::Texture const& Texture)
 {
 	MTLSizeAndAlign SizeAlign;
diff --git a/Engine/Source/Runtime/Apple/AGXRHI/Private/AGXRHI.cpp b/Engine/Source/Runtime/Apple/AGXRHI/Private/AGXRHI.cpp
--- a/Engine/Source/Runtime/Apple/AGXRHI/Private/AGXRHI.cpp
+++ b/Engine/Source/Runtime/Apple/AGXRHI/Private/AGXRHI.cpp
@@ -170,6 +170,9 @@ FAGXDynamicRHI::FAGXDynamicRHI(ERHIFeatureLevel::Type RequestedFeatureLevel)
 : ImmediateContext(nullptr, FAGXDeviceContext::CreateDeviceContext())
 , AsyncComputeContext(nullptr)
 {
+	check(Singleton == nullptr);
+	Singleton = this;
+
 	@autoreleasepool {
 	// This should be called once at the start 
 	check( IsInGameThread() );
@@ -1040,28 +1043,12 @@ FAGXDynamicRHI::~FAGXDynamicRHI()
 #endif
 }
 
-uint64 FAGXDynamicRHI::RHICalcTexture2DPlatformSize(uint32 SizeX, uint32 SizeY, uint8 Format, uint32 NumMips, uint32 NumSamples, ETextureCreateFlags Flags, const FRHIResourceCreateInfo& CreateInfo, uint32& OutAlign)
+FDynamicRHI::FRHICalcTextureSizeResult FAGXDynamicRHI::RHICalcTexturePlatformSize(FRHITextureDesc const& Desc, uint32 FirstMipIndex)
 {
-	@autoreleasepool {
-	OutAlign = 0;
-	return CalcTextureSize(SizeX, SizeY, (EPixelFormat)Format, NumMips);
-	}
-}
-
-uint64 FAGXDynamicRHI::RHICalcTexture3DPlatformSize(uint32 SizeX, uint32 SizeY, uint32 SizeZ, uint8 Format, uint32 NumMips, ETextureCreateFlags Flags, const FRHIResourceCreateInfo& CreateInfo, uint32& OutAlign)
-{
-	@autoreleasepool {
-	OutAlign = 0;
-	return CalcTextureSize3D(SizeX, SizeY, SizeZ, (EPixelFormat)Format, NumMips);
-	}
-}
-
-uint64 FAGXDynamicRHI::RHICalcTextureCubePlatformSize(uint32 Size, uint8 Format, uint32 NumMips, ETextureCreateFlags Flags, const FRHIResourceCreateInfo& CreateInfo, uint32& OutAlign)
-{
-	@autoreleasepool {
-	OutAlign = 0;
-	return CalcTextureSize(Size, Size, (EPixelFormat)Format, NumMips) * 6;
-	}
+	FDynamicRHI::FRHICalcTextureSizeResult Result;
+	Result.Size = Desc.CalcMemorySizeEstimate(FirstMipIndex);
+	Result.Align = 0;
+	return Result;
 }
 
 uint64 FAGXDynamicRHI::RHIGetMinimumAlignmentForBufferBackedSRV(EPixelFormat Format)
diff --git a/Engine/Source/Runtime/Apple/AGXRHI/Private/AGXRHIPrivate.h b/Engine/Source/Runtime/Apple/AGXRHI/Private/AGXRHIPrivate.h
--- a/Engine/Source/Runtime/Apple/AGXRHI/Private/AGXRHIPrivate.h
+++ b/Engine/Source/Runtime/Apple/AGXRHI/Private/AGXRHIPrivate.h
@@ -266,6 +266,11 @@ static FORCEINLINE typename TAGXResourceTraits<TRHIType>::TConcreteType* Resourc
 	return static_cast<typename TAGXResourceTraits<TRHIType>::TConcreteType*>(Resource);
 }
 
+static FORCEINLINE FAGXSurface* ResourceCast(FRHITexture* Texture)
+{
+	return AGXGetMetalSurfaceFromRHITexture(Texture);
+}
+
 extern uint32 AGXSafeGetRuntimeDebuggingLevel();
 
 extern int32 GAGXBufferZeroFill;
diff --git a/Engine/Source/Runtime/Apple/AGXRHI/Private/AGXRenderTarget.cpp b/Engine/Source/Runtime/Apple/AGXRHI/Private/AGXRenderTarget.cpp
--- a/Engine/Source/Runtime/Apple/AGXRHI/Private/AGXRenderTarget.cpp
+++ b/Engine/Source/Runtime/Apple/AGXRHI/Private/AGXRenderTarget.cpp
@@ -40,34 +40,39 @@ void FAGXRHICommandContext::RHICopyToResolveTarget(FRHITexture* SourceTextureRHI
 		// nothing to do if one of the textures is null!
 		return;
 	}
-	if(SourceTextureRHI != DestTextureRHI)
+	if (SourceTextureRHI != DestTextureRHI)
 	{
 		FAGXSurface* Source = AGXGetMetalSurfaceFromRHITexture(SourceTextureRHI);
 		FAGXSurface* Destination = AGXGetMetalSurfaceFromRHITexture(DestTextureRHI);
 		
-		switch (Source->Type)
+		const FRHITextureDesc& SourceDesc = Source->GetDesc();
+		const FRHITextureDesc& DestinationDesc = Destination->GetDesc();
+
+		// Only valid to have nil Metal textures when they are TexCreate_Presentable
+		if (!Source->Texture)
 		{
-			case RRT_Texture2D:
-				break;
-			case RRT_TextureCube:
-				check(Source->SizeZ == 6); // Arrays might not work yet.
-				break;
-			default:
-				check(false); // Only Tex2D & Cube are tested to work so far!
-				break;
+			// Source RHI texture is valid with no Presentable Metal texture - there is nothing to copy from
+			check(EnumHasAnyFlags(SourceDesc.Flags, TexCreate_Presentable));
+			return;
 		}
-		switch (Destination->Type)
+		if (!Destination->Texture)
 		{
-			case RRT_Texture2D:
-				break;
-			case RRT_TextureCube:
-				check(Destination->SizeZ == 6); // Arrays might not work yet.
-				break;
-			default:
-				check(false); // Only Tex2D & Cube are tested to work so far!
-				break;
+			// Destination RHI texture is valid with no Presentable Metal texture - force fetch it now so we can complete the copy
+			check(EnumHasAnyFlags(DestinationDesc.Flags, TexCreate_Presentable));
+			Destination->GetDrawableTexture();
+			if(!Destination->Texture)
+			{
+				UE_LOG(LogRHI, Error, TEXT("Drawable for destination texture resolve target unavailable"));
+				return;
+			}
 		}
-		
+
+		checkf( SourceDesc.IsTexture2D()   || SourceDesc.IsTextureCube(), TEXT("Only Tex2D & Cube are tested to work so far!"));
+		checkf(!SourceDesc.IsTextureCube() || SourceDesc.ArraySize == 1,  TEXT("Cube arrays might not work yet."));
+
+		checkf( DestinationDesc.IsTexture2D()   || DestinationDesc.IsTextureCube(), TEXT("Only Tex2D & Cube are tested to work so far!"));
+		checkf(!DestinationDesc.IsTextureCube() || DestinationDesc.ArraySize == 1,  TEXT("Cube arrays might not work yet."));
+
 		mtlpp::Origin Origin(0, 0, 0);
 		mtlpp::Size Size(0, 0, 1);
 		if (ResolveParams.Rect.IsValid())
@@ -84,12 +89,15 @@ void FAGXRHICommandContext::RHICopyToResolveTarget(FRHITexture* SourceTextureRHI
 			Origin.x = 0;
 			Origin.y = 0;
 			
-			Size.width = FMath::Max<uint32>(1, Source->SizeX >> ResolveParams.MipIndex);
-			Size.height = FMath::Max<uint32>(1, Source->SizeY >> ResolveParams.MipIndex);
+			Size.width = FMath::Max<uint32>(1, SourceDesc.Extent.X >> ResolveParams.MipIndex);
+			Size.height = FMath::Max<uint32>(1, SourceDesc.Extent.Y >> ResolveParams.MipIndex);
+			// clamp to a destination size
+			Size.width = FMath::Min<uint32>(Size.width, DestinationDesc.Extent.X >> ResolveParams.MipIndex);
+			Size.height = FMath::Min<uint32>(Size.height, DestinationDesc.Extent.Y >> ResolveParams.MipIndex);
 		}
 		
-		const bool bSrcCubemap  = Source->bIsCubemap;
-		const bool bDestCubemap = Destination->bIsCubemap;
+		const bool bSrcCubemap  = SourceDesc.IsTextureCube();
+		const bool bDestCubemap = DestinationDesc.IsTextureCube();
 		
 		uint32 DestIndex = ResolveParams.DestArrayIndex * (bDestCubemap ? 6 : 1) + (bDestCubemap ? uint32(ResolveParams.CubeFace) : 0);
 		uint32 SrcIndex  = ResolveParams.SourceArrayIndex * (bSrcCubemap ? 6 : 1) + (bSrcCubemap ? uint32(ResolveParams.CubeFace) : 0);
@@ -101,7 +109,7 @@ void FAGXRHICommandContext::RHICopyToResolveTarget(FRHITexture* SourceTextureRHI
 
 		const bool bMSAASource = Source->MSAATexture;
         const bool bMSAADest = Destination->MSAATexture;
-        const bool bDepthStencil = Source->PixelFormat == PF_DepthStencil;
+        const bool bDepthStencil = SourceDesc.Format == PF_DepthStencil;
 		if (bMSAASource && !bMSAADest)
 		{
 			// Resolve required - Device must support this - Using Shader for resolve not supported amd NumSamples should be 1
@@ -111,10 +119,31 @@ void FAGXRHICommandContext::RHICopyToResolveTarget(FRHITexture* SourceTextureRHI
 			
 			Context->CopyFromTextureToTexture(Source->MSAAResolveTexture, SrcIndex, ResolveParams.MipIndex, Origin, Size, Destination->Texture, DestIndex, ResolveParams.MipIndex, Origin);
 		}
-		else
+		else if(Source->Texture.GetPixelFormat() == Destination->Texture.GetPixelFormat())
 		{
+			// Blit Copy for matching formats
 			Context->CopyFromTextureToTexture(Source->Texture, SrcIndex, ResolveParams.MipIndex, Origin, Size, Destination->Texture, DestIndex, ResolveParams.MipIndex, Origin);
 		}
+		else
+		{
+			const FPixelFormatInfo& SourceFormatInfo = GPixelFormats[SourceDesc.Format];
+			const FPixelFormatInfo& DestFormatInfo = GPixelFormats[DestinationDesc.Format];
+			bool bUsingPixelFormatView = (Source->Texture.GetUsage() & mtlpp::TextureUsage::PixelFormatView) != 0;
+			
+			// Attempt to Resolve with a texture view - source Texture doesn't have to be created with MTLTextureUsagePixelFormatView for these cases e.g:
+			// If we are resolving to/from sRGB linear color space within the same format OR using same bit length color format
+			if	(	SourceFormatInfo.BlockBytes == DestFormatInfo.BlockBytes &&
+					(bUsingPixelFormatView || SourceFormatInfo.NumComponents == DestFormatInfo.NumComponents)
+				)
+			{
+				FAGXTexture SourceTextureView = Source->Texture.NewTextureView(Destination->Texture.GetPixelFormat(), Source->Texture.GetTextureType(), ns::Range(ResolveParams.MipIndex, 1), ns::Range(SrcIndex, 1));
+				if(SourceTextureView)
+				{
+					Context->CopyFromTextureToTexture(SourceTextureView, 0, 0, Origin, Size, Destination->Texture, DestIndex, ResolveParams.MipIndex, Origin);
+					AGXSafeReleaseMetalTexture(SourceTextureView);
+				}
+			}
+		}
 
 #if PLATFORM_MAC
 		if((Destination->GPUReadback & FAGXSurface::EAGXGPUReadbackFlags::ReadbackRequested) != 0)
@@ -239,7 +268,7 @@ void FAGXDynamicRHI::RHIReadSurfaceData(FRHITexture* TextureRHI, FIntRect Rect,
 	mtlpp::Region Region(Rect.Min.X, Rect.Min.Y, SizeX, SizeY);
     
 	FAGXTexture Texture = Surface->Texture;
-    if(!Texture && EnumHasAnyFlags(Surface->Flags, TexCreate_Presentable))
+    if(!Texture && EnumHasAnyFlags(Surface->GetDesc().Flags, TexCreate_Presentable))
     {
         Texture = Surface->GetCurrentTexture();
     }
@@ -249,7 +278,7 @@ void FAGXDynamicRHI::RHIReadSurfaceData(FRHITexture* TextureRHI, FIntRect Rect,
         return;
     }
 
-	if (GAGXUseTexGetBytes && Surface->PixelFormat != PF_DepthStencil && Surface->PixelFormat != PF_ShadowDepth)
+	if (GAGXUseTexGetBytes && Surface->GetDesc().Format != PF_DepthStencil && Surface->GetDesc().Format != PF_ShadowDepth)
 	{
 		SCOPE_CYCLE_COUNTER(STAT_AGXTexturePageOffTime);
 		
@@ -261,7 +290,7 @@ void FAGXDynamicRHI::RHIReadSurfaceData(FRHITexture* TextureRHI, FIntRect Rect,
 #else
 			mtlpp::StorageMode StorageMode = mtlpp::StorageMode::Shared;
 #endif
-			mtlpp::PixelFormat MetalFormat = (mtlpp::PixelFormat)GPixelFormats[Surface->PixelFormat].PlatformFormat;
+			mtlpp::PixelFormat MetalFormat = (mtlpp::PixelFormat)GPixelFormats[Surface->GetDesc().Format].PlatformFormat;
 			mtlpp::TextureDescriptor Desc;
 			Desc.SetTextureType(Texture.GetTextureType());
 			Desc.SetPixelFormat(Texture.GetPixelFormat());
@@ -297,7 +326,7 @@ void FAGXDynamicRHI::RHIReadSurfaceData(FRHITexture* TextureRHI, FIntRect Rect,
 		//kick the current command buffer.
 		ImmediateContext.Context->SubmitCommandBufferAndWait();
 		
-		const uint32 Stride = GPixelFormats[Surface->PixelFormat].BlockBytes * SizeX;
+		const uint32 Stride = GPixelFormats[Surface->GetDesc().Format].BlockBytes * SizeX;
 		const uint32 BytesPerImage = Stride * SizeY;
 
 		TArray<uint8> Data;
@@ -305,7 +334,7 @@ void FAGXDynamicRHI::RHIReadSurfaceData(FRHITexture* TextureRHI, FIntRect Rect,
 		
 		Texture.GetBytes(Data.GetData(), Stride, BytesPerImage, Region, 0, 0);
 		
-		ConvertSurfaceDataToFColor(Surface->PixelFormat, SizeX, SizeY, (uint8*)Data.GetData(), Stride, OutDataPtr, InFlags);
+		ConvertSurfaceDataToFColor(Surface->GetDesc().Format, SizeX, SizeY, (uint8*)Data.GetData(), Stride, OutDataPtr, InFlags);
 		
 		if (TempTexture)
 		{
@@ -314,7 +343,7 @@ void FAGXDynamicRHI::RHIReadSurfaceData(FRHITexture* TextureRHI, FIntRect Rect,
 	}
 	else
 	{
-		uint32 BytesPerPixel = (Surface->PixelFormat != PF_DepthStencil || !InFlags.GetOutputStencil()) ? GPixelFormats[Surface->PixelFormat].BlockBytes : 1;
+		uint32 BytesPerPixel = (Surface->GetDesc().Format != PF_DepthStencil || !InFlags.GetOutputStencil()) ? GPixelFormats[Surface->GetDesc().Format].BlockBytes : 1;
 		const uint32 Stride = BytesPerPixel * SizeX;
 		const uint32 Alignment = PLATFORM_MAC ? 1u : 64u; // Mac permits natural row alignment (tightly-packed) but iOS does not.
 		const uint32 AlignedStride = ((Stride - 1) & ~(Alignment - 1)) + Alignment;
@@ -324,7 +353,7 @@ void FAGXDynamicRHI::RHIReadSurfaceData(FRHITexture* TextureRHI, FIntRect Rect,
 			// Synchronise the texture with the CPU
 			SCOPE_CYCLE_COUNTER(STAT_AGXTexturePageOffTime);
 			
-			if (Surface->PixelFormat != PF_DepthStencil)
+			if (Surface->GetDesc().Format != PF_DepthStencil)
 			{
 				ImmediateContext.Context->CopyFromTextureToBuffer(Texture, 0, InFlags.GetMip(), Region.origin, Region.size, Buffer, 0, AlignedStride, BytesPerImage, mtlpp::BlitOption::None);
 			}
@@ -343,7 +372,7 @@ void FAGXDynamicRHI::RHIReadSurfaceData(FRHITexture* TextureRHI, FIntRect Rect,
 			//kick the current command buffer.
 			ImmediateContext.Context->SubmitCommandBufferAndWait();
 			
-			ConvertSurfaceDataToFColor(Surface->PixelFormat, SizeX, SizeY, (uint8*)Buffer.GetContents(), AlignedStride, OutDataPtr, InFlags);
+			ConvertSurfaceDataToFColor(Surface->GetDesc().Format, SizeX, SizeY, (uint8*)Buffer.GetContents(), AlignedStride, OutDataPtr, InFlags);
 		}
 		((FAGXDeviceContext*)ImmediateContext.Context)->ReleaseBuffer(Buffer);
 	}
@@ -354,7 +383,6 @@ void FAGXDynamicRHI::RHIMapStagingSurface(FRHITexture* TextureRHI, FRHIGPUFence*
 {
 	@autoreleasepool {
     FAGXSurface* Surface = AGXGetMetalSurfaceFromRHITexture(TextureRHI);
-    FAGXTexture2D* Texture = (FAGXTexture2D*)TextureRHI->GetTexture2D();
     
 #if PLATFORM_MAC
 	uint16 FencePoll = FenceRHI && FenceRHI->Poll() ? 1 : 0;
@@ -362,8 +390,8 @@ void FAGXDynamicRHI::RHIMapStagingSurface(FRHITexture* TextureRHI, FRHIGPUFence*
 #endif
     
     uint32 Stride = 0;
-    OutWidth = Texture->GetSizeX();
-    OutHeight = Texture->GetSizeY();
+    OutWidth = Surface->GetSizeX();
+    OutHeight = Surface->GetSizeY();
     OutData = Surface->Lock(0, 0, RLM_ReadOnly, Stride);
     
 #if PLATFORM_MAC
@@ -387,7 +415,7 @@ void FAGXDynamicRHI::RHIReadSurfaceFloatData(FRHITexture* TextureRHI, FIntRect R
 	FAGXSurface* Surface = AGXGetMetalSurfaceFromRHITexture(TextureRHI);
 	
     FAGXTexture Texture = Surface->Texture;
-    if(!Texture && EnumHasAnyFlags(Surface->Flags, TexCreate_Presentable))
+    if(!Texture && EnumHasAnyFlags(Surface->GetDesc().Flags, TexCreate_Presentable))
     {
 		Texture = Surface->GetCurrentTexture();
     }
@@ -398,7 +426,7 @@ void FAGXDynamicRHI::RHIReadSurfaceFloatData(FRHITexture* TextureRHI, FIntRect R
     }
     
 	// verify the input image format (but don't crash)
-	if (Surface->PixelFormat != PF_FloatRGBA)
+	if (Surface->GetDesc().Format != PF_FloatRGBA)
 	{
 		UE_LOG(LogRHI, Log, TEXT("Trying to read non-FloatRGBA surface."));
 	}
@@ -419,7 +447,7 @@ void FAGXDynamicRHI::RHIReadSurfaceFloatData(FRHITexture* TextureRHI, FIntRect R
 	mtlpp::Region Region = mtlpp::Region(Rect.Min.X, Rect.Min.Y, SizeX, SizeY);
 	
 	// function wants details about the destination, not the source
-	const uint32 Stride = GPixelFormats[Surface->PixelFormat].BlockBytes * SizeX;
+	const uint32 Stride = GPixelFormats[Surface->GetDesc().Format].BlockBytes * SizeX;
 	const uint32 Alignment = PLATFORM_MAC ? 1u : 64u; // Mac permits natural row alignment (tightly-packed) but iOS does not.
 	const uint32 AlignedStride = ((Stride - 1) & ~(Alignment - 1)) + Alignment;
 	const uint32 BytesPerImage = AlignedStride  * SizeY;
@@ -470,7 +498,7 @@ void FAGXDynamicRHI::RHIRead3DSurfaceFloatData(FRHITexture* TextureRHI,FIntRect
 	}
 	
 	// verify the input image format (but don't crash)
-	if (Surface->PixelFormat != PF_FloatRGBA)
+	if (Surface->GetDesc().Format != PF_FloatRGBA)
 	{
 		UE_LOG(LogRHI, Log, TEXT("Trying to read non-FloatRGBA surface."));
 	}
@@ -485,7 +513,7 @@ void FAGXDynamicRHI::RHIRead3DSurfaceFloatData(FRHITexture* TextureRHI,FIntRect
 	mtlpp::Region Region = mtlpp::Region(InRect.Min.X, InRect.Min.Y, ZMinMax.X, SizeX, SizeY, SizeZ);
 	
 	// function wants details about the destination, not the source
-	const uint32 Stride = GPixelFormats[Surface->PixelFormat].BlockBytes * SizeX;
+	const uint32 Stride = GPixelFormats[Surface->GetDesc().Format].BlockBytes * SizeX;
 	const uint32 Alignment = PLATFORM_MAC ? 1u : 64u; // Mac permits natural row alignment (tightly-packed) but iOS does not.
 	const uint32 AlignedStride = ((Stride - 1) & ~(Alignment - 1)) + Alignment;
 	const uint32 BytesPerImage = AlignedStride  * SizeY;
diff --git a/Engine/Source/Runtime/Apple/AGXRHI/Private/AGXStateCache.cpp b/Engine/Source/Runtime/Apple/AGXRHI/Private/AGXStateCache.cpp
--- a/Engine/Source/Runtime/Apple/AGXRHI/Private/AGXStateCache.cpp
+++ b/Engine/Source/Runtime/Apple/AGXRHI/Private/AGXStateCache.cpp
@@ -500,8 +500,8 @@ bool FAGXStateCache::SetRenderPassInfo(FRHIRenderPassInfo const& InRenderTargets
 				FAGXSurface& Surface = *AGXGetMetalSurfaceFromRHITexture(RenderTargetView.RenderTarget);
 				FormatKey = Surface.FormatKey;
 				
-				uint32 Width = FMath::Max((uint32)(Surface.SizeX >> RenderTargetView.MipIndex), (uint32)1);
-				uint32 Height = FMath::Max((uint32)(Surface.SizeY >> RenderTargetView.MipIndex), (uint32)1);
+				uint32 Width = FMath::Max((uint32)(Surface.GetDesc().Extent.X >> RenderTargetView.MipIndex), (uint32)1);
+				uint32 Height = FMath::Max((uint32)(Surface.GetDesc().Extent.Y >> RenderTargetView.MipIndex), (uint32)1);
 				if(!bFramebufferSizeSet)
 				{
 					bFramebufferSizeSet = true;
@@ -516,9 +516,14 @@ bool FAGXStateCache::SetRenderPassInfo(FRHIRenderPassInfo const& InRenderTargets
 	
 				// if this is the back buffer, make sure we have a usable drawable
 				ConditionalUpdateBackBuffer(Surface);
-	
+				FAGXSurface* ResolveSurface = AGXGetMetalSurfaceFromRHITexture(RenderTargetView.ResolveTarget);
+				if (ResolveSurface)
+				{
+					ConditionalUpdateBackBuffer(*ResolveSurface);
+				}
+
 				BoundTargets |= 1 << RenderTargetIndex;
-            
+
 #if !PLATFORM_MAC
                 if (Surface.Texture.GetPtr() == nil)
                 {
@@ -527,22 +532,23 @@ bool FAGXStateCache::SetRenderPassInfo(FRHIRenderPassInfo const& InRenderTargets
                     return true;
                 }
 #endif
-				
+
 				// The surface cannot be nil - we have to have a valid render-target array after this call.
 				check (Surface.Texture);
 	
 				// user code generally passes -1 as a default, but we need 0
 				uint32 ArraySliceIndex = RenderTargetView.ArraySlice == 0xFFFFFFFF ? 0 : RenderTargetView.ArraySlice;
-				if (Surface.bIsCubemap)
+				if (Surface.GetDesc().IsTextureCube())
 				{
 					ArraySliceIndex = GetMetalCubeFace((ECubeFace)ArraySliceIndex);
 				}
 				
-				switch(Surface.Type)
+				switch(Surface.GetDesc().Dimension)
 				{
-					case RRT_Texture2DArray:
-					case RRT_Texture3D:
-					case RRT_TextureCube:
+					case ETextureDimension::Texture2DArray:
+					case ETextureDimension::Texture3D:
+					case ETextureDimension::TextureCube:
+					case ETextureDimension::TextureCubeArray:
 						if(RenderTargetView.ArraySlice == 0xFFFFFFFF)
 						{
 							ArrayTargets |= (1 << RenderTargetIndex);
@@ -573,9 +579,9 @@ bool FAGXStateCache::SetRenderPassInfo(FRHIRenderPassInfo const& InRenderTargets
 					HighLevelLoadAction == ERenderTargetLoadAction::ELoad);
 #endif
 				
+				bool bMemoryless = false;
 				if (Surface.MSAATexture && !bUseResolvedTexture)
 				{
-					bool bMemoryless = false;
 #if PLATFORM_IOS
 					if (Surface.MSAATexture.GetStorageMode() == mtlpp::StorageMode::Memoryless)
 					{
@@ -592,7 +598,6 @@ bool FAGXStateCache::SetRenderPassInfo(FRHIRenderPassInfo const& InRenderTargets
 				}
 				else
 				{
-					bool bMemoryless = false;
 #if PLATFORM_IOS
 					if (Surface.Texture.GetStorageMode() == mtlpp::StorageMode::Memoryless)
 					{
@@ -609,7 +614,7 @@ bool FAGXStateCache::SetRenderPassInfo(FRHIRenderPassInfo const& InRenderTargets
 				}
 				
 				ColorAttachment.SetLevel(RenderTargetView.MipIndex);
-				if(Surface.Type == RRT_Texture3D)
+				if(Surface.GetDesc().IsTexture3D())
 				{
 					ColorAttachment.SetSlice(0);
 					ColorAttachment.SetDepthPlane(ArraySliceIndex);
@@ -631,7 +636,9 @@ bool FAGXStateCache::SetRenderPassInfo(FRHIRenderPassInfo const& InRenderTargets
 					ColorAttachment.SetClearColor(mtlpp::ClearColor(ClearColor.R, ClearColor.G, ClearColor.B, ClearColor.A));
 				}
 
-				bCanRestartRenderPass &= (SampleCount <= 1) && (ColorAttachment.GetLoadAction() == mtlpp::LoadAction::Load) && (HighLevelStoreAction == ERenderTargetStoreAction::EStore);
+				bCanRestartRenderPass &= 	!bMemoryless &&
+											ColorAttachment.GetLoadAction() == mtlpp::LoadAction::Load &&
+											HighLevelStoreAction != ERenderTargetStoreAction::ENoAction;
 	
 				bHasValidRenderTarget = true;
 				bHasValidColorTarget = true;
@@ -670,11 +677,12 @@ bool FAGXStateCache::SetRenderPassInfo(FRHIRenderPassInfo const& InRenderTargets
 		{
 			FAGXSurface& Surface = *AGXGetMetalSurfaceFromRHITexture(RenderPassInfo.DepthStencilRenderTarget.DepthStencilTarget);
 			
-			switch(Surface.Type)
+			switch(Surface.GetDesc().Dimension)
 			{
-				case RRT_Texture2DArray:
-				case RRT_Texture3D:
-				case RRT_TextureCube:
+				case ETextureDimension::Texture2DArray:
+				case ETextureDimension::Texture3D:
+				case ETextureDimension::TextureCube:
+				case ETextureDimension::TextureCubeArray:
 					ArrayRenderLayers = Surface.GetNumFaces();
 					break;
 				default:
@@ -693,13 +701,13 @@ bool FAGXStateCache::SetRenderPassInfo(FRHIRenderPassInfo const& InRenderTargets
 			if(!bFramebufferSizeSet)
 			{
 				bFramebufferSizeSet = true;
-				FrameBufferSize.width = Surface.SizeX;
-				FrameBufferSize.height = Surface.SizeY;
+				FrameBufferSize.width  = Surface.GetDesc().Extent.X;
+				FrameBufferSize.height = Surface.GetDesc().Extent.Y;
 			}
 			else
 			{
-				FrameBufferSize.width = FMath::Min(FrameBufferSize.width, (CGFloat)Surface.SizeX);
-				FrameBufferSize.height = FMath::Min(FrameBufferSize.height, (CGFloat)Surface.SizeY);
+				FrameBufferSize.width = FMath::Min(FrameBufferSize.width, (CGFloat)Surface.GetDesc().Extent.X);
+				FrameBufferSize.height = FMath::Min(FrameBufferSize.height, (CGFloat)Surface.GetDesc().Extent.Y);
 			}
 			
 			EPixelFormat DepthStencilPixelFormat = RenderPassInfo.DepthStencilRenderTarget.DepthStencilTarget->GetFormat();
@@ -859,7 +867,12 @@ bool FAGXStateCache::SetRenderPassInfo(FRHIRenderPassInfo const& InRenderTargets
 				bHasValidRenderTarget = true;
 				bFallbackDepthStencilBound = (RenderPassInfo.DepthStencilRenderTarget.DepthStencilTarget == FallbackDepthStencilSurface);
 
-				bCanRestartRenderPass &= (SampleCount <= 1) && ((RenderPassInfo.DepthStencilRenderTarget.DepthStencilTarget == FallbackDepthStencilSurface) || ((DepthAttachment.GetLoadAction() == mtlpp::LoadAction::Load) && (!RenderPassInfo.DepthStencilRenderTarget.ExclusiveDepthStencil.IsDepthWrite() || (DepthStoreAction == ERenderTargetStoreAction::EStore))));
+				bool bDepthMSAARestart = !bDepthTextureMemoryless && HighLevelStoreAction == ERenderTargetStoreAction::EMultisampleResolve;
+				bCanRestartRenderPass &=	(DepthSampleCount <= 1 || bDepthMSAARestart) &&
+											(
+												(RenderPassInfo.DepthStencilRenderTarget.DepthStencilTarget == FallbackDepthStencilSurface) ||
+												((DepthAttachment.GetLoadAction() == mtlpp::LoadAction::Load) && (bDepthMSAARestart || !RenderPassInfo.DepthStencilRenderTarget.ExclusiveDepthStencil.IsDepthWrite() || DepthStoreAction == ERenderTargetStoreAction::EStore))
+											);
 				
 				// and assign it
 				RenderPass.SetDepthAttachment(DepthAttachment);
@@ -921,7 +934,12 @@ bool FAGXStateCache::SetRenderPassInfo(FRHIRenderPassInfo const& InRenderTargets
 				
 				// @todo Stencil writes that need to persist must use ERenderTargetStoreAction::EStore on iOS.
 				// We should probably be using deferred store actions so that we can safely lazily instantiate encoders.
-				bCanRestartRenderPass &= (SampleCount <= 1) && ((RenderPassInfo.DepthStencilRenderTarget.DepthStencilTarget == FallbackDepthStencilSurface) || ((StencilAttachment.GetLoadAction() == mtlpp::LoadAction::Load) && (1 || !RenderPassInfo.DepthStencilRenderTarget.ExclusiveDepthStencil.IsStencilWrite() || (StencilStoreAction == ERenderTargetStoreAction::EStore))));
+				bool bStencilMSAARestart = !bStencilMemoryless && HighLevelStoreAction != ERenderTargetStoreAction::ENoAction;
+				bCanRestartRenderPass &= 	(bStencilMSAARestart || SampleCount <= 1) &&
+											(
+												(RenderPassInfo.DepthStencilRenderTarget.DepthStencilTarget == FallbackDepthStencilSurface) ||
+												((StencilAttachment.GetLoadAction() == mtlpp::LoadAction::Load) && (bStencilMSAARestart || !RenderPassInfo.DepthStencilRenderTarget.ExclusiveDepthStencil.IsStencilWrite() || (StencilStoreAction == ERenderTargetStoreAction::EStore)))
+											);
 				
 				// and assign it
 				RenderPass.SetStencilAttachment(StencilAttachment);
@@ -1178,7 +1196,7 @@ void FAGXStateCache::SetVisibilityResultMode(mtlpp::VisibilityResultMode const M
 void FAGXStateCache::ConditionalUpdateBackBuffer(FAGXSurface& Surface)
 {
 	// are we setting the back buffer? if so, make sure we have the drawable
-	if (EnumHasAnyFlags(Surface.Flags, TexCreate_Presentable))
+	if (EnumHasAnyFlags(Surface.GetDesc().Flags, TexCreate_Presentable))
 	{
 		// update the back buffer texture the first time used this frame
 		if (Surface.Texture.GetPtr() == nil)
@@ -1436,48 +1454,33 @@ void FAGXStateCache::SetShaderResourceView(FAGXContext* Context, EAGXShaderStage
 {
 	if (SRV)
 	{
-		FRHITexture* Texture = SRV->SourceTexture.GetReference();
-		FAGXVertexBuffer* VB = SRV->SourceVertexBuffer.GetReference();
-		FAGXIndexBuffer* IB = SRV->SourceIndexBuffer.GetReference();
-		FAGXStructuredBuffer* SB = SRV->SourceStructuredBuffer.GetReference();
-		if (Texture)
+		if (SRV->bTexture)
 		{
-			FAGXSurface* Surface = SRV->TextureView;
-			if (Surface != nullptr)
+			FAGXTexture const& View = SRV->GetTextureView();
+			if (View)
 			{
-				SetShaderTexture(ShaderStage, Surface->Texture, BindIndex, mtlpp::ResourceUsage(mtlpp::ResourceUsage::Read|mtlpp::ResourceUsage::Sample));
+				SetShaderTexture(ShaderStage, View, BindIndex, mtlpp::ResourceUsage(mtlpp::ResourceUsage::Read | mtlpp::ResourceUsage::Sample));
 			}
 			else
 			{
 				SetShaderTexture(ShaderStage, nil, BindIndex, mtlpp::ResourceUsage(0));
 			}
 		}
-		else if (IsLinearBuffer(ShaderStage, BindIndex) && SRV->GetLinearTexture(false))
-		{
-			ns::AutoReleased<FAGXTexture> Tex;
-			Tex = SRV->GetLinearTexture(false);
-			
-			SetShaderTexture(ShaderStage, Tex, BindIndex, mtlpp::ResourceUsage(mtlpp::ResourceUsage::Read|mtlpp::ResourceUsage::Sample));
-			if (VB)
-            {
-				SetShaderBuffer(ShaderStage, VB->GetCurrentBufferOrNil(), VB->Data, SRV->Offset, VB->GetSize(), BindIndex, mtlpp::ResourceUsage::Read, (EPixelFormat)SRV->Format);
-            }
-            else if (IB)
-            {
-                SetShaderBuffer(ShaderStage, IB->GetCurrentBufferOrNil(), nil, SRV->Offset, IB->GetSize(), BindIndex, mtlpp::ResourceUsage::Read, (EPixelFormat)SRV->Format);
-            }
-		}
-		else if (VB)
-		{
-			SetShaderBuffer(ShaderStage, VB->GetCurrentBufferOrNil(), VB->Data, SRV->Offset, VB->GetSize(), BindIndex, mtlpp::ResourceUsage::Read, (EPixelFormat)SRV->Format);
-		}
-		else if (IB)
-		{
-			SetShaderBuffer(ShaderStage, IB->GetCurrentBufferOrNil(), nil, SRV->Offset, IB->GetSize(), BindIndex, mtlpp::ResourceUsage::Read, (EPixelFormat)SRV->Format);
-		}
-		else if (SB)
+		else
 		{
-			SetShaderBuffer(ShaderStage, SB->GetCurrentBufferOrNil(), nil, SRV->Offset, SB->GetSize(), BindIndex, mtlpp::ResourceUsage::Read);
+			FAGXResourceMultiBuffer* Buffer = SRV->GetSourceBuffer();
+
+			if (IsLinearBuffer(ShaderStage, BindIndex) && SRV->GetLinearTexture())
+			{
+				ns::AutoReleased<FAGXTexture> Tex;
+				Tex = SRV->GetLinearTexture();
+
+				SetShaderTexture(ShaderStage, Tex, BindIndex, mtlpp::ResourceUsage(mtlpp::ResourceUsage::Read | mtlpp::ResourceUsage::Sample));
+			}
+			else 
+			{
+				SetShaderBuffer(ShaderStage, Buffer->GetCurrentBufferOrNil(), Buffer->Data, SRV->Offset, Buffer->GetSize(), BindIndex, mtlpp::ResourceUsage::Read, (EPixelFormat)SRV->Format);
+			}
 		}
 	}
 }
@@ -1512,68 +1515,21 @@ void FAGXStateCache::SetShaderUnorderedAccessView(EAGXShaderStages ShaderStage,
 {
 	if (UAV)
 	{
-		// figure out which one of the resources we need to set
-		FAGXStructuredBuffer* StructuredBuffer = UAV->SourceView->SourceStructuredBuffer.GetReference();
-		FAGXVertexBuffer* VertexBuffer = UAV->SourceView->SourceVertexBuffer.GetReference();
-		FAGXIndexBuffer* IndexBuffer = UAV->SourceView->SourceIndexBuffer.GetReference();
-		FRHITexture* Texture = UAV->SourceView->SourceTexture.GetReference();
-		FAGXSurface* Surface = UAV->SourceView->TextureView;
-		if (StructuredBuffer)
-		{
-			SetShaderBuffer(ShaderStage, StructuredBuffer->GetCurrentBufferOrNil(), nil, 0, StructuredBuffer->GetSize(), BindIndex, mtlpp::ResourceUsage(mtlpp::ResourceUsage::Read | mtlpp::ResourceUsage::Write));
-		}
-		else if (VertexBuffer)
-		{
-			check(!VertexBuffer->Data && VertexBuffer->GetCurrentBufferOrNil());
-			if (IsLinearBuffer(ShaderStage, BindIndex) && UAV->SourceView->GetLinearTexture(true))
-			{
-				ns::AutoReleased<FAGXTexture> Tex;
-				Tex = UAV->SourceView->GetLinearTexture(true);
-				SetShaderTexture(ShaderStage, Tex, BindIndex, mtlpp::ResourceUsage(mtlpp::ResourceUsage::Read | mtlpp::ResourceUsage::Write));
-                
-                SetShaderBuffer(ShaderStage, VertexBuffer->GetCurrentBufferOrNil(), VertexBuffer->Data, 0, VertexBuffer->GetSize(), BindIndex, mtlpp::ResourceUsage(mtlpp::ResourceUsage::Read | mtlpp::ResourceUsage::Write), (EPixelFormat)UAV->SourceView->Format);
-			}
-			else
-			{
-				SetShaderBuffer(ShaderStage, VertexBuffer->GetCurrentBufferOrNil(), VertexBuffer->Data, 0, VertexBuffer->GetSize(), BindIndex, mtlpp::ResourceUsage(mtlpp::ResourceUsage::Read | mtlpp::ResourceUsage::Write), (EPixelFormat)UAV->SourceView->Format);
-			}
-		}
-		else if (IndexBuffer)
-		{
-			check(IndexBuffer->GetCurrentBufferOrNil());
-			if (IsLinearBuffer(ShaderStage, BindIndex) && UAV->SourceView->GetLinearTexture(true))
-			{
-				ns::AutoReleased<FAGXTexture> Tex;
-				Tex = UAV->SourceView->GetLinearTexture(true);
-				SetShaderTexture(ShaderStage, Tex, BindIndex, mtlpp::ResourceUsage(mtlpp::ResourceUsage::Read | mtlpp::ResourceUsage::Write));
-				
-				SetShaderBuffer(ShaderStage, IndexBuffer->GetCurrentBufferOrNil(), nullptr, 0, IndexBuffer->GetSize(), BindIndex, mtlpp::ResourceUsage(mtlpp::ResourceUsage::Read | mtlpp::ResourceUsage::Write), (EPixelFormat)UAV->SourceView->Format);
-			}
-			else
-			{
-				SetShaderBuffer(ShaderStage, IndexBuffer->GetCurrentBufferOrNil(), nullptr, 0, IndexBuffer->GetSize(), BindIndex, mtlpp::ResourceUsage(mtlpp::ResourceUsage::Read | mtlpp::ResourceUsage::Write), (EPixelFormat)UAV->SourceView->Format);
-			}
-		}
-		else if (Texture)
+		if (UAV->bTexture)
 		{
-			if (!Surface)
-			{
-				Surface = AGXGetMetalSurfaceFromRHITexture(Texture);
-			}
-			if (Surface != nullptr)
+			FAGXSurface* Surface = UAV->GetSourceTexture();
+			FAGXTexture const& View = UAV->GetTextureView();
+
+			if (View)
 			{
-				FAGXSurface* Source = AGXGetMetalSurfaceFromRHITexture(Texture);
-				
 				FPlatformAtomics::InterlockedExchange(&Surface->Written, 1);
-				FPlatformAtomics::InterlockedExchange(&Source->Written, 1);
-				
-				SetShaderTexture(ShaderStage, Surface->Texture, BindIndex, mtlpp::ResourceUsage(mtlpp::ResourceUsage::Read | mtlpp::ResourceUsage::Write));
-				
-				if (Surface->Texture.GetBuffer() &&
-					(EnumHasAllFlags(Source->Flags, TexCreate_UAV | TexCreate_NoTiling) || EnumHasAllFlags(Source->Flags, TexCreate_AtomicCompatible)))
+
+				SetShaderTexture(ShaderStage, View, BindIndex, mtlpp::ResourceUsage(mtlpp::ResourceUsage::Read | mtlpp::ResourceUsage::Write));
+
+				if (Surface->Texture.GetBuffer() && (EnumHasAllFlags(Surface->GetDesc().Flags, TexCreate_UAV | TexCreate_NoTiling) || EnumHasAllFlags(Surface->GetDesc().Flags, TexCreate_AtomicCompatible)))
 				{
 					uint32 BytesPerRow = Surface->Texture.GetBufferBytesPerRow();
-					uint32 ElementsPerRow = BytesPerRow / GPixelFormats[(EPixelFormat)Texture->GetFormat()].BlockBytes;
+					uint32 ElementsPerRow = BytesPerRow / GPixelFormats[(EPixelFormat)Surface->GetFormat()].BlockBytes;
 					
 					FAGXBuffer Buffer(Surface->Texture.GetBuffer(), false);
 					const uint32 BufferOffset = Surface->Texture.GetBufferOffset();
@@ -1586,6 +1542,20 @@ void FAGXStateCache::SetShaderUnorderedAccessView(EAGXShaderStages ShaderStage,
 				SetShaderTexture(ShaderStage, nil, BindIndex, mtlpp::ResourceUsage(0));
 			}
 		}
+		else
+		{
+			FAGXResourceMultiBuffer* Buffer = UAV->GetSourceBuffer();
+			check(!Buffer->Data && Buffer->GetCurrentBufferOrNil());
+
+			if (IsLinearBuffer(ShaderStage, BindIndex) && UAV->GetLinearTexture())
+			{
+				ns::AutoReleased<FAGXTexture> Tex;
+				Tex = UAV->GetLinearTexture();
+				SetShaderTexture(ShaderStage, Tex, BindIndex, mtlpp::ResourceUsage(mtlpp::ResourceUsage::Read | mtlpp::ResourceUsage::Write));
+			}
+
+			SetShaderBuffer(ShaderStage, Buffer->GetCurrentBufferOrNil(), Buffer->Data, 0, Buffer->GetSize(), BindIndex, mtlpp::ResourceUsage(mtlpp::ResourceUsage::Read | mtlpp::ResourceUsage::Write), (EPixelFormat)UAV->Format);
+		}
 	}
 }
 
@@ -1764,67 +1734,74 @@ bool FAGXStateCache::PrepareToRestart(bool const bCurrentApplied)
 	}
 	else
 	{
-		if (SampleCount <= 1)
+		FRHIRenderPassInfo Info = GetRenderPassInfo();
+		
+		ERenderTargetActions DepthActions = GetDepthActions(Info.DepthStencilRenderTarget.Action);
+		ERenderTargetActions StencilActions = GetStencilActions(Info.DepthStencilRenderTarget.Action);
+		ERenderTargetLoadAction DepthLoadAction = GetLoadAction(DepthActions);
+		ERenderTargetStoreAction DepthStoreAction = GetStoreAction(DepthActions);
+		ERenderTargetLoadAction StencilLoadAction = GetLoadAction(StencilActions);
+		ERenderTargetStoreAction StencilStoreAction = GetStoreAction(StencilActions);
+
+		if (Info.DepthStencilRenderTarget.DepthStencilTarget)
 		{
-			FRHIRenderPassInfo Info = GetRenderPassInfo();
-			
-			ERenderTargetActions DepthActions = GetDepthActions(Info.DepthStencilRenderTarget.Action);
-			ERenderTargetActions StencilActions = GetStencilActions(Info.DepthStencilRenderTarget.Action);
-			ERenderTargetLoadAction DepthLoadAction = GetLoadAction(DepthActions);
-			ERenderTargetStoreAction DepthStoreAction = GetStoreAction(DepthActions);
-			ERenderTargetLoadAction StencilLoadAction = GetLoadAction(StencilActions);
-			ERenderTargetStoreAction StencilStoreAction = GetStoreAction(StencilActions);
-
-			if (Info.DepthStencilRenderTarget.DepthStencilTarget)
+			if(bCurrentApplied && Info.DepthStencilRenderTarget.ExclusiveDepthStencil.IsDepthWrite() && DepthStoreAction == ERenderTargetStoreAction::ENoAction)
 			{
-				if (bCurrentApplied || DepthLoadAction != ERenderTargetLoadAction::EClear)
-				{
-					DepthLoadAction = ERenderTargetLoadAction::ELoad;
-				}
-				if (Info.DepthStencilRenderTarget.ExclusiveDepthStencil.IsDepthWrite())
-				{
-					DepthStoreAction = ERenderTargetStoreAction::EStore;
-				}
+				return false;
+			}
+			if (bCurrentApplied && Info.DepthStencilRenderTarget.ExclusiveDepthStencil.IsStencilWrite() && StencilStoreAction == ERenderTargetStoreAction::ENoAction)
+			{
+				return false;
+			}
+		
+			if (bCurrentApplied || DepthLoadAction != ERenderTargetLoadAction::EClear)
+			{
+				DepthLoadAction = ERenderTargetLoadAction::ELoad;
+			}
+			if (Info.DepthStencilRenderTarget.ExclusiveDepthStencil.IsDepthWrite())
+			{
+				DepthStoreAction = ERenderTargetStoreAction::EStore;
+			}
 
-				if (bCurrentApplied || StencilLoadAction != ERenderTargetLoadAction::EClear)
-				{
-					StencilLoadAction = ERenderTargetLoadAction::ELoad;
-				}
-				if (Info.DepthStencilRenderTarget.ExclusiveDepthStencil.IsStencilWrite())
-				{
-					StencilStoreAction = ERenderTargetStoreAction::EStore;
-				}
-				
-				DepthActions = MakeRenderTargetActions(DepthLoadAction, DepthStoreAction);
-				StencilActions = MakeRenderTargetActions(StencilLoadAction, StencilStoreAction);
-				Info.DepthStencilRenderTarget.Action = MakeDepthStencilTargetActions(DepthActions, StencilActions);
+			if (bCurrentApplied || StencilLoadAction != ERenderTargetLoadAction::EClear)
+			{
+				StencilLoadAction = ERenderTargetLoadAction::ELoad;
 			}
-			
-			for (int32 RenderTargetIndex = 0; RenderTargetIndex < Info.GetNumColorRenderTargets(); RenderTargetIndex++)
+			if (Info.DepthStencilRenderTarget.ExclusiveDepthStencil.IsStencilWrite())
 			{
-				FRHIRenderPassInfo::FColorEntry& RenderTargetView = Info.ColorRenderTargets[RenderTargetIndex];
-				ERenderTargetLoadAction LoadAction = GetLoadAction(RenderTargetView.Action);
-				ERenderTargetStoreAction StoreAction = GetStoreAction(RenderTargetView.Action);
-				
-				if (!bCurrentApplied && LoadAction == ERenderTargetLoadAction::EClear)
-				{
-					StoreAction == ERenderTargetStoreAction::EStore;
-				}
-				else
-				{
-					LoadAction = ERenderTargetLoadAction::ELoad;
-				}
-				RenderTargetView.Action = MakeRenderTargetActions(LoadAction, StoreAction);
-				check(RenderTargetView.RenderTarget == nil || GetStoreAction(RenderTargetView.Action) == ERenderTargetStoreAction::EStore);
+				StencilStoreAction = ERenderTargetStoreAction::EStore;
 			}
 			
-			InvalidateRenderTargets();
-			return SetRenderPassInfo(Info, GetVisibilityResultsBuffer(), true) && CanRestartRenderPass();
+			DepthActions = MakeRenderTargetActions(DepthLoadAction, DepthStoreAction);
+			StencilActions = MakeRenderTargetActions(StencilLoadAction, StencilStoreAction);
+			Info.DepthStencilRenderTarget.Action = MakeDepthStencilTargetActions(DepthActions, StencilActions);
 		}
-		else
+		
+		for (int32 RenderTargetIndex = 0; RenderTargetIndex < Info.GetNumColorRenderTargets(); RenderTargetIndex++)
 		{
-			return false;
+			FRHIRenderPassInfo::FColorEntry& RenderTargetView = Info.ColorRenderTargets[RenderTargetIndex];
+			ERenderTargetLoadAction LoadAction = GetLoadAction(RenderTargetView.Action);
+			ERenderTargetStoreAction StoreAction = GetStoreAction(RenderTargetView.Action);
+			
+			if(bCurrentApplied && StoreAction == ERenderTargetStoreAction::ENoAction)
+			{
+				return false;
+			}
+			
+			if (!bCurrentApplied && LoadAction == ERenderTargetLoadAction::EClear)
+			{
+				StoreAction == ERenderTargetStoreAction::EStore;
+			}
+			else
+			{
+				LoadAction = ERenderTargetLoadAction::ELoad;
+			}
+			RenderTargetView.Action = MakeRenderTargetActions(LoadAction, StoreAction);
+			check(RenderTargetView.RenderTarget == nil || GetStoreAction(RenderTargetView.Action) != ERenderTargetStoreAction::ENoAction);
 		}
+		
+		InvalidateRenderTargets();
+		return SetRenderPassInfo(Info, GetVisibilityResultsBuffer(), true) && CanRestartRenderPass();
 	}
 }
 
