Copying //UE5/Dev-ParallelRendering to //UE5/Main

diff --git a/Engine/Source/Runtime/Apple/AGXRHI/Private/AGXTexture.cpp b/Engine/Source/Runtime/Apple/AGXRHI/Private/AGXTexture.cpp
--- a/Engine/Source/Runtime/Apple/AGXRHI/Private/AGXTexture.cpp
+++ b/Engine/Source/Runtime/Apple/AGXRHI/Private/AGXTexture.cpp
@@ -28,15 +28,9 @@ FAutoConsoleVariableRef CVarAGXForceIOSTexturesShared(TEXT("rhi.AGX.ForceIOSText
 /** Given a pointer to a RHI texture that was created by the AGX RHI, returns a pointer to the FAGXTextureBase it encapsulates. */
 FAGXSurface* AGXGetMetalSurfaceFromRHITexture(FRHITexture* Texture)
 {
-	if (!Texture)
-	{
-		return NULL;
-	}
-	else
-	{
-		FAGXSurface* Surface = (FAGXSurface*)Texture->GetTextureBaseRHI();
-		return Surface;
-	}
+	return Texture
+		? static_cast<FAGXSurface*>(Texture->GetTextureBaseRHI())
+		: nullptr;
 }
 
 static bool IsRenderTarget(ETextureCreateFlags Flags)
@@ -133,11 +127,11 @@ static bool IsPixelFormatPVRTCCompressed(EPixelFormat Format)
 	}
 }
 
-void AGXSafeReleaseMetalTexture(FAGXSurface* Surface, FAGXTexture& Texture, bool bTextureView, bool bAVFoundationTexture)
+void AGXSafeReleaseMetalTexture(FAGXSurface* Surface, FAGXTexture& Texture, bool bAVFoundationTexture)
 {
 	if(GIsAGXInitialized && GDynamicRHI)
 	{
-		if (!bAVFoundationTexture && !bTextureView)
+		if (!bAVFoundationTexture)
 		{
 			GetAGXDeviceContext().ReleaseTexture(Surface, Texture);
 		}
@@ -220,6 +214,21 @@ static mtlpp::PixelFormat AGX_ToSRGBFormat_AppleGPU(mtlpp::PixelFormat MTLFormat
 		case mtlpp::PixelFormat::ASTC_12x12_LDR:
 			MTLFormat = mtlpp::PixelFormat::ASTC_12x12_sRGB;
 			break;
+#if PLATFORM_MAC
+		// Fix for Apple silicon M1 macs that can support BC pixel formats even though they are Apple family GPUs.
+		case mtlpp::PixelFormat::BC1_RGBA:
+			MTLFormat = mtlpp::PixelFormat::BC1_RGBA_sRGB;
+			break;
+		case mtlpp::PixelFormat::BC2_RGBA:
+			MTLFormat = mtlpp::PixelFormat::BC2_RGBA_sRGB;
+			break;
+		case mtlpp::PixelFormat::BC3_RGBA:
+			MTLFormat = mtlpp::PixelFormat::BC3_RGBA_sRGB;
+			break;
+		case mtlpp::PixelFormat::BC7_RGBAUnorm:
+			MTLFormat = mtlpp::PixelFormat::BC7_RGBAUnorm_sRGB;
+			break;
+#endif
 		default:
 			break;
 	}
@@ -252,7 +261,6 @@ void FAGXSurface::PrepareTextureView()
 #endif
 	if(!(Usage & mtlpp::TextureUsage::PixelFormatView) && !bMemoryLess)
 	{
-		check(!bTextureView);
 		check(ImageSurfaceRef == nullptr);
 		
 		check(Texture);
@@ -262,18 +270,18 @@ void FAGXSurface::PrepareTextureView()
 		{
 			FAGXTexture OldMSAATexture = MSAATexture;
 			MSAATexture = Reallocate(MSAATexture, mtlpp::TextureUsage::PixelFormatView);
-			AGXSafeReleaseMetalTexture(this, OldMSAATexture, bTextureView, ImageSurfaceRef != nullptr);
+			AGXSafeReleaseMetalTexture(this, OldMSAATexture, ImageSurfaceRef != nullptr);
 		}
 		if (MSAAResolveTexture && !bMSAAResolveTextureIsTexture)
 		{
 			FAGXTexture OldMSAAResolveTexture = MSAAResolveTexture;
 			MSAAResolveTexture = Reallocate(MSAAResolveTexture, mtlpp::TextureUsage::PixelFormatView);
-			AGXSafeReleaseMetalTexture(this, OldMSAAResolveTexture, bTextureView, ImageSurfaceRef != nullptr);
+			AGXSafeReleaseMetalTexture(this, OldMSAAResolveTexture, ImageSurfaceRef != nullptr);
 		}
 		
 		FAGXTexture OldTexture = Texture;
 		Texture = Reallocate(Texture, mtlpp::TextureUsage::PixelFormatView);
-		AGXSafeReleaseMetalTexture(this, OldTexture, bTextureView, ImageSurfaceRef != nullptr);
+		AGXSafeReleaseMetalTexture(this, OldTexture, ImageSurfaceRef != nullptr);
 		
 		if (bMSAATextureIsTexture)
 		{
@@ -310,359 +318,75 @@ FAGXTexture FAGXSurface::Reallocate(FAGXTexture InTexture, mtlpp::TextureUsage U
 	return NewTex;
 }
 
-void FAGXSurface::ReplaceTexture(FAGXContext& Context, FAGXTexture CurrentTexture, FAGXTexture NewTexture)
-{
-	check(NewTexture);
-	{
-		// Copy texture content...
-		const uint32 BlockSizeX = GPixelFormats[PixelFormat].BlockSizeX;
-		const uint32 BlockSizeY = GPixelFormats[PixelFormat].BlockSizeY;
-		const uint32 BlockSizeZ = GPixelFormats[PixelFormat].BlockSizeZ;
-		
-		mtlpp::Origin Origin(0,0,0);
-		
-		// DXT/BC formats on Mac actually do have mip-tails that are smaller than the block size, they end up being uncompressed.
-		bool const bPixelFormatASTC = IsPixelFormatASTCCompressed(PixelFormat);
-		
-		for (uint32 SliceIndex = 0; SliceIndex < CurrentTexture.GetArrayLength(); SliceIndex++)
-		{
-			for (uint32 MipIndex = 0; MipIndex < CurrentTexture.GetMipmapLevelCount(); ++MipIndex)
-			{
-				const uint32 UnalignedMipSizeX = FMath::Max<uint32>(1, CurrentTexture.GetWidth() >> (MipIndex));
-				const uint32 UnalignedMipSizeY = FMath::Max<uint32>(1, CurrentTexture.GetHeight() >> (MipIndex));
-				const uint32 UnalignedMipSizeZ = FMath::Max<uint32>(1, CurrentTexture.GetDepth() >> (MipIndex));
-				const uint32 MipSizeX = (bPixelFormatASTC) ? AlignArbitrary(UnalignedMipSizeX, BlockSizeX) : UnalignedMipSizeX;
-				const uint32 MipSizeY = (bPixelFormatASTC) ? AlignArbitrary(UnalignedMipSizeY, BlockSizeY) : UnalignedMipSizeY;
-				const uint32 MipSizeZ = (bPixelFormatASTC) ? AlignArbitrary(UnalignedMipSizeZ, BlockSizeZ) : UnalignedMipSizeZ;
-				
-				Context.CopyFromTextureToTexture(CurrentTexture, SliceIndex, MipIndex, Origin, mtlpp::Size(MipSizeX, MipSizeY, MipSizeZ), NewTexture, SliceIndex, MipIndex, Origin);
-			}
-		}
-		
-	}
-	
-	FAGXTexture OldTexture = Texture;
-	if (Texture && Texture == CurrentTexture)
-	{
-		AGXSafeReleaseMetalTexture(OldTexture);
-		Texture = NewTexture;
-	}
-	if (MSAATexture && MSAATexture == CurrentTexture)
-	{
-		if (MSAATexture && (MSAATexture != OldTexture))
-		{
-			FAGXTexture OldMSAATexture = MSAATexture;
-			AGXSafeReleaseMetalTexture(OldMSAATexture);
-			MSAATexture = NewTexture;
-		}
-		else if(MSAATexture)
-		{
-			MSAATexture = NewTexture;
-		}
-		
-		if (MSAAResolveTexture && MSAAResolveTexture == CurrentTexture)
-		{
-			if (MSAAResolveTexture && (MSAAResolveTexture != OldTexture))
-			{
-				FAGXTexture OldMSAAResolveTexture = MSAAResolveTexture;
-				AGXSafeReleaseMetalTexture(OldMSAAResolveTexture);
-				MSAAResolveTexture = NewTexture;
-			}
-			else if (MSAAResolveTexture)
-			{
-				MSAAResolveTexture = NewTexture;
-			}
-		}
-	}
-	
-	for (FAGXShaderResourceView* SRV : SRVs)
-	{
-		if(OldTexture != SRV->TextureView->Texture)
-		{
-			AGXSafeReleaseMetalTexture(nullptr, SRV->TextureView->Texture, true, false);
-		}
-		SRV->TextureView->Texture = nil;
-		
-		SRV->TextureView->MSAATexture = nil;
-		
-		if (SRV->Format == PF_Unknown)
-		{
-			SRV->TextureView->Init(*this, NSMakeRange(SRV->MipLevel, SRV->NumMips));
-		}
-		else
-		{
-			SRV->TextureView->Init(*this, NSMakeRange(SRV->MipLevel, SRV->NumMips), (EPixelFormat)SRV->Format, SRV->bSRGBForceDisable);
-		}
-	}
-}
-
 void FAGXSurface::MakeAliasable(void)
 {
 	// TODO
 }
 
-void FAGXSurface::MakeUnAliasable(void)
-{
-	// TODO
-}
-
-void FAGXSurface::Init(FAGXSurface& Source, NSRange MipRange)
+uint8 AGXGetMetalPixelFormatKey(mtlpp::PixelFormat Format)
 {
-#if PLATFORM_IOS
-	// Mmeory;ess targets can't have texture views (SRVs or UAVs)
-	if (Source.Texture.GetStorageMode() == mtlpp::StorageMode::Memoryless)
-	{
-		return;
-	}
-#endif
-	auto TextureViewType = [&Source]()
-	{
-		return Source.bIsCubemap ? mtlpp::TextureType::Texture2DArray : Source.Texture.GetTextureType();
-	};
-
-	mtlpp::PixelFormat MetalFormat = (mtlpp::PixelFormat)GPixelFormats[PixelFormat].PlatformFormat;
-	if (EnumHasAnyFlags(Flags, TexCreate_SRGB))
-	{
-		// Ensure we have the correct sRGB target format if we create a new texture view rather than using the source texture
-		MetalFormat = AGXToSRGBFormat(MetalFormat);
-	}
-	
-	bool const bUseSourceTex = (Source.PixelFormat != PF_DepthStencil) && MipRange.location == 0 && MipRange.length == Source.Texture.GetMipmapLevelCount();
-	
-	// Recreate the texture to enable MTLTextureUsagePixelFormatView which must be off unless we definitely use this feature or we are throwing ~4% performance vs. Windows on the floor.
-	mtlpp::TextureUsage Usage = Source.Texture.GetUsage();
-	if(!(Usage & mtlpp::TextureUsage::PixelFormatView) && (Source.PixelFormat != PF_DepthStencil) && !bUseSourceTex)
-	{
-		Source.PrepareTextureView();
-	}
-	
-	ns::Range Slices(0, Source.Texture.GetArrayLength() * (bIsCubemap ? 6 : 1));
-	// Stencil requires a format conversion, so this will access depth only, consequently there are no mip levels other than 0, so this path requires no real special casing.
-	if(Source.PixelFormat != PF_DepthStencil && !bUseSourceTex)
-	{
-		Texture = Source.Texture.NewTextureView(MetalFormat, TextureViewType(), ns::Range(MipRange.location, MipRange.length), Slices);
-	}
-	else
+	struct FAGXPixelFormatKeyMap
 	{
-		Texture = Source.Texture;
-	}
-	
-	const uint32 BlockSizeX = GPixelFormats[PixelFormat].BlockSizeX;
-	const uint32 BlockSizeY = GPixelFormats[PixelFormat].BlockSizeY;
-	const uint32 BlockBytes = GPixelFormats[PixelFormat].BlockBytes;
-	SizeX = FMath::Max(SizeX >> MipRange.location, BlockSizeX);
-	SizeY = FMath::Max(SizeY >> MipRange.location, BlockSizeY);
-	SizeZ = (Type != RRT_Texture3D) ? SizeZ : FMath::Max(SizeZ >> MipRange.location, 1u);
-}
+		FRWLock Mutex;
+		uint8 NextKey = 1; // 0 is reserved for mtlpp::PixelFormat::Invalid
+		TMap<uint64, uint8> Map;
 
-void FAGXSurface::Init(FAGXSurface& Source, NSRange MipRange, EPixelFormat Format, bool bSRGBForceDisable)
-{
-	check(!Source.MSAATexture || Format == PF_X24_G8);
-#if PLATFORM_IOS
-	// Mmeory;ess targets can't have texture views (SRVs or UAVs)
-	if (Source.Texture.GetStorageMode() == mtlpp::StorageMode::Memoryless)
-	{
-		return;
-	}
-#endif
-	
-	mtlpp::PixelFormat MetalFormat = (mtlpp::PixelFormat)GPixelFormats[PixelFormat].PlatformFormat;
-	
-	bool bUseSourceTex = (Source.PixelFormat != PF_DepthStencil) && Source.PixelFormat == Format && MipRange.location == 0 && MipRange.length == Source.Texture.GetMipmapLevelCount();
-	
-	if (EnumHasAnyFlags(Flags, TexCreate_SRGB))
-	{
-		if(bSRGBForceDisable)
+		uint8 Get(mtlpp::PixelFormat Format)
 		{
-			// Force no sRGB, should create new view without sRGB format
-			bUseSourceTex = false;
-		}
-		else
-		{
-			// Ensure we have the correct sRGB target format if we create a new texture view rather than using the source texture
-			MetalFormat = AGXToSRGBFormat(MetalFormat);
+			FRWScopeLock Lock(Mutex, SLT_ReadOnly);
+			uint8* Key = Map.Find((uint64)Format);
+			if (Key == nullptr)
+			{
+				Lock.ReleaseReadOnlyLockAndAcquireWriteLock_USE_WITH_CAUTION();
+				Key = Map.Find((uint64)Format);
+				if (Key == nullptr)
+				{
+					Key = &Map.Add((uint64)Format, NextKey++);
+					// only giving 6 bits to the key
+					checkf(NextKey < 64, TEXT("Too many unique pixel formats to fit into the PipelineStateHash"));
+				}
+			}
+			return *Key;
 		}
-	}
-	
-	// Recreate the texture to enable MTLTextureUsagePixelFormatView which must be off unless we definitely use this feature or we are throwing ~4% performance vs. Windows on the floor.
-	mtlpp::TextureUsage Usage = Source.Texture.GetUsage();
-	if(!(Usage & mtlpp::TextureUsage::PixelFormatView) && (Source.PixelFormat == PF_DepthStencil && Format == PF_X24_G8) && !bUseSourceTex)
-	{
-		Source.PrepareTextureView();
-	}
-	
-	ns::Range Slices(0, Source.Texture.GetArrayLength() * (bIsCubemap ? 6 : 1));
-	if(Source.PixelFormat != PF_DepthStencil && !bUseSourceTex)
-	{
-		Texture = Source.Texture.NewTextureView(MetalFormat, Source.Texture.GetTextureType(), ns::Range(MipRange.location, MipRange.length), Slices);
-	}
-	else if (Source.PixelFormat == PF_DepthStencil && Format == PF_X24_G8)
-	{
-		switch (Source.Texture.GetPixelFormat())
+
+		FAGXPixelFormatKeyMap()
 		{
-#if PLATFORM_MAC // Currently only required on Mac as iOS uses separate textures
-			case mtlpp::PixelFormat::Depth24Unorm_Stencil8:
-				MetalFormat = mtlpp::PixelFormat::X24_Stencil8;
-				break;
+			// Add depth stencil formats first, so we don't have to use 6 bits for them in the pipeline hash
+			Get(mtlpp::PixelFormat::Depth32Float);
+			Get(mtlpp::PixelFormat::Stencil8);
+			Get(mtlpp::PixelFormat::Depth32Float_Stencil8);
+#if PLATFORM_MAC
+			Get(mtlpp::PixelFormat::Depth24Unorm_Stencil8);
+			Get(mtlpp::PixelFormat::Depth16Unorm);
 #endif
-			case mtlpp::PixelFormat::Depth32Float_Stencil8:
-				MetalFormat = mtlpp::PixelFormat::X32_Stencil8;
-				break;
-			default:
-				check(false);
-				break;
 		}
-		
-		Texture = Source.Texture.NewTextureView(MetalFormat, Source.Texture.GetTextureType(), ns::Range(MipRange.location, MipRange.length), Slices);
-	}
-	else
-	{
-		Texture = Source.Texture;
-	}
-	
-	const uint32 BlockSizeX = GPixelFormats[PixelFormat].BlockSizeX;
-	const uint32 BlockSizeY = GPixelFormats[PixelFormat].BlockSizeY;
-	const uint32 BlockBytes = GPixelFormats[PixelFormat].BlockBytes;
-	SizeX = FMath::Max(SizeX >> MipRange.location, BlockSizeX);
-	SizeY = FMath::Max(SizeY >> MipRange.location, BlockSizeY);
-	SizeZ = (Type != RRT_Texture3D) ? SizeZ : FMath::Max(SizeZ >> MipRange.location, 1u);
-}
-
-FAGXSurface::FAGXSurface(FAGXSurface& Source, NSRange MipRange)
-: Type(Source.Type)
-, PixelFormat(Source.PixelFormat)
-, Texture(nil)
-, MSAATexture(nil)
-, MSAAResolveTexture(nil)
-, SizeX(Source.SizeX)
-, SizeY(Source.SizeY)
-, SizeZ(Source.SizeZ)
-, bIsCubemap(Source.bIsCubemap)
-, Flags(Source.Flags)
-, BufferLocks(0)
-, TotalTextureSize(0)
-, Viewport(nullptr)
-, ImageSurfaceRef(nullptr)
-, bTextureView(true)
-{
-	FPlatformAtomics::InterlockedExchange(&Written, Source.Written);
-	
-	Init(Source, MipRange);
-}
-
-FAGXSurface::FAGXSurface(FAGXSurface& Source, NSRange const MipRange, EPixelFormat Format, bool bSRGBForceDisable)
-: Type(Source.Type)
-, PixelFormat(Format)
-, Texture(nil)
-, MSAATexture(nil)
-, MSAAResolveTexture(nil)
-, SizeX(Source.SizeX)
-, SizeY(Source.SizeY)
-, SizeZ(Source.SizeZ)
-, bIsCubemap(Source.bIsCubemap)
-, Flags(Source.Flags)
-, BufferLocks(0)
-, TotalTextureSize(0)
-, Viewport(nullptr)
-, ImageSurfaceRef(nullptr)
-, bTextureView(true)
-{
-	FPlatformAtomics::InterlockedExchange(&Written, Source.Written);
-	
-	Init(Source, MipRange, Format, bSRGBForceDisable);
-}
+	} static PixelFormatKeyMap;
 
-static TMap<uint64, uint8>& GetMetalPixelFormatKeyMap()
-{
-	static TMap<uint64, uint8> PixelFormatKeyMap;
-	return PixelFormatKeyMap;
+	return PixelFormatKeyMap.Get(Format);
 }
 
-uint8 AGXGetMetalPixelFormatKey(mtlpp::PixelFormat Format)
+FAGXTextureDesc::FAGXTextureDesc(FRHITextureDesc const& InDesc)
+	: bIsRenderTarget(IsRenderTarget(InDesc.Flags))
 {
-	static FRWLock Mutex;
-	static uint8 NextKey = 1; // 0 is reserved for mtlpp::PixelFormat::Invalid
-	FRWScopeLock Lock(Mutex, SLT_ReadOnly);
-	uint8* Key = GetMetalPixelFormatKeyMap().Find((uint64)Format);
-	if (Key == NULL)
-	{
-		Lock.ReleaseReadOnlyLockAndAcquireWriteLock_USE_WITH_CAUTION();
-		Key = GetMetalPixelFormatKeyMap().Find((uint64)Format);
-		if (Key == NULL)
-		{
-			Key = &GetMetalPixelFormatKeyMap().Add((uint64)Format, NextKey++);
-			// only giving 6 bits to the key
-			checkf(NextKey < 64, TEXT("Too many unique pixel formats to fit into the PipelineStateHash"));
-		}
-	}
-	return *Key;
-}
-
-FAGXSurface::FAGXSurface(ERHIResourceType ResourceType, EPixelFormat Format, uint32 InSizeX, uint32 InSizeY, uint32 InSizeZ, uint32 NumSamples, bool bArray, uint32 ArraySize, uint32 NumMips, ETextureCreateFlags InFlags, FResourceBulkDataInterface* BulkData)
-: Type(ResourceType)
-, PixelFormat(Format)
-, Texture(nil)
-, MSAATexture(nil)
-, MSAAResolveTexture(nil)
-, SizeX(InSizeX)
-, SizeY(InSizeY)
-, SizeZ(InSizeZ)
-, bIsCubemap(false)
-, Flags(InFlags)
-, BufferLocks(0)
-, TotalTextureSize(0)
-, Viewport(nullptr)
-, ImageSurfaceRef(nullptr)
-, bTextureView(false)
-{
-	check(SizeX > 0 && SizeY > 0 && NumMips > 0);
-	
-	// get a unique key for this surface's format
-	TMap<uint64, uint8>& PixelFormatKeyMap = GetMetalPixelFormatKeyMap();
-	if (PixelFormatKeyMap.Num() == 0)
-	{
-		// Add depth stencil formats first, so we don't have to use 6 bits for them in the pipeline hash
-		AGXGetMetalPixelFormatKey(mtlpp::PixelFormat::Depth32Float);
-		AGXGetMetalPixelFormatKey(mtlpp::PixelFormat::Stencil8);
-		AGXGetMetalPixelFormatKey(mtlpp::PixelFormat::Depth32Float_Stencil8);
-#if PLATFORM_MAC
-		AGXGetMetalPixelFormatKey(mtlpp::PixelFormat::Depth24Unorm_Stencil8);
-		AGXGetMetalPixelFormatKey(mtlpp::PixelFormat::Depth16Unorm);
-#endif
-	}
-	
-	FPlatformAtomics::InterlockedExchange(&Written, 0);
-	mtlpp::PixelFormat MTLFormat = (mtlpp::PixelFormat)GPixelFormats[Format].PlatformFormat;
-	
-	if (EnumHasAnyFlags(Flags, TexCreate_SRGB))
+	MTLFormat = (mtlpp::PixelFormat)GPixelFormats[InDesc.Format].PlatformFormat;
+	if (EnumHasAnyFlags(InDesc.Flags, TexCreate_SRGB))
 	{
 		MTLFormat = AGXToSRGBFormat(MTLFormat);
 	}
-	
-	// set the key
+
+	// get a unique key for this surface's format
 	FormatKey = AGXGetMetalPixelFormatKey(MTLFormat);
-	
-	
-	// the special back buffer surface will be updated in GetAGXDeviceContext().BeginDrawingViewport - no need to set the texture here
-	if (EnumHasAnyFlags(Flags, TexCreate_Presentable))
-	{
-		return;
-	}
-	
-	bool bIsRenderTarget = IsRenderTarget(Flags);
-	mtlpp::TextureDescriptor Desc;
-	
-	if (ResourceType == RRT_TextureCube)
+
+	if (InDesc.IsTextureCube())
 	{
-		Desc = mtlpp::TextureDescriptor::TextureCubeDescriptor(MTLFormat, SizeX, (NumMips > 1));
-		bIsCubemap = true;
+		Desc = mtlpp::TextureDescriptor::TextureCubeDescriptor(MTLFormat, InDesc.Extent.X, (InDesc.NumMips > 1));
 	}
-	else if (ResourceType == RRT_Texture3D)
+	else if (InDesc.IsTexture3D())
 	{
 		Desc.SetTextureType(mtlpp::TextureType::Texture3D);
-		Desc.SetWidth(SizeX);
-		Desc.SetHeight(SizeY);
-		Desc.SetDepth(SizeZ);
+		Desc.SetWidth(InDesc.Extent.X);
+		Desc.SetHeight(InDesc.Extent.Y);
+		Desc.SetDepth(InDesc.Depth);
 		Desc.SetPixelFormat(MTLFormat);
 		Desc.SetArrayLength(1);
 		Desc.SetMipmapLevelCount(1);
@@ -670,15 +394,15 @@ FAGXSurface::FAGXSurface(ERHIResourceType ResourceType, EPixelFormat Format, uin
 	}
 	else
 	{
-		Desc = mtlpp::TextureDescriptor::Texture2DDescriptor(MTLFormat, SizeX, SizeY, (NumMips > 1));
-		Desc.SetDepth(SizeZ);
+		Desc = mtlpp::TextureDescriptor::Texture2DDescriptor(MTLFormat, InDesc.Extent.X, InDesc.Extent.Y, (InDesc.NumMips > 1));
+		Desc.SetArrayLength(InDesc.ArraySize);
 	}
-	
+
 	// flesh out the descriptor
-	if (bArray)
+	if (InDesc.IsTextureArray())
 	{
-		Desc.SetArrayLength(ArraySize);
-		if (bIsCubemap)
+		Desc.SetArrayLength(InDesc.ArraySize);
+		if (InDesc.IsTextureCube())
 		{
 			if (FAGXCommandQueue::SupportsFeature(EAGXFeaturesCubemapArrays))
 			{
@@ -687,7 +411,7 @@ FAGXSurface::FAGXSurface(ERHIResourceType ResourceType, EPixelFormat Format, uin
 			else
 			{
 				Desc.SetTextureType(mtlpp::TextureType::Texture2DArray);
-				Desc.SetArrayLength(ArraySize * 6);
+				Desc.SetArrayLength(InDesc.ArraySize * 6);
 			}
 		}
 		else
@@ -695,90 +419,90 @@ FAGXSurface::FAGXSurface(ERHIResourceType ResourceType, EPixelFormat Format, uin
 			Desc.SetTextureType(mtlpp::TextureType::Texture2DArray);
 		}
 	}
-	Desc.SetMipmapLevelCount(NumMips);
-	
+	Desc.SetMipmapLevelCount(InDesc.NumMips);
+
 	{
-		Desc.SetUsage(ConvertFlagsToUsage(Flags));
-		
-		if (EnumHasAnyFlags(Flags, TexCreate_CPUReadback) && !EnumHasAnyFlags(Flags, TexCreate_RenderTargetable|TexCreate_DepthStencilTargetable|TexCreate_FastVRAM))
+		Desc.SetUsage(ConvertFlagsToUsage(InDesc.Flags));
+
+		if (EnumHasAnyFlags(InDesc.Flags, TexCreate_CPUReadback) && !EnumHasAnyFlags(InDesc.Flags, TexCreate_RenderTargetable | TexCreate_DepthStencilTargetable | TexCreate_FastVRAM))
 		{
 			Desc.SetCpuCacheMode(mtlpp::CpuCacheMode::DefaultCache);
 #if PLATFORM_MAC
 			Desc.SetStorageMode(mtlpp::StorageMode::Managed);
-			Desc.SetResourceOptions((mtlpp::ResourceOptions)(mtlpp::ResourceOptions::CpuCacheModeDefaultCache|mtlpp::ResourceOptions::StorageModeManaged));
+			Desc.SetResourceOptions((mtlpp::ResourceOptions)(mtlpp::ResourceOptions::CpuCacheModeDefaultCache | mtlpp::ResourceOptions::StorageModeManaged));
 #else
 			Desc.SetStorageMode(mtlpp::StorageMode::Shared);
-			Desc.SetResourceOptions((mtlpp::ResourceOptions)(mtlpp::ResourceOptions::CpuCacheModeDefaultCache|mtlpp::ResourceOptions::StorageModeShared));
+			Desc.SetResourceOptions((mtlpp::ResourceOptions)(mtlpp::ResourceOptions::CpuCacheModeDefaultCache | mtlpp::ResourceOptions::StorageModeShared));
 #endif
 		}
-		else if(EnumHasAnyFlags(Flags, TexCreate_NoTiling) && !EnumHasAnyFlags(Flags, TexCreate_FastVRAM|TexCreate_DepthStencilTargetable|TexCreate_RenderTargetable|TexCreate_UAV))
+		else if (EnumHasAnyFlags(InDesc.Flags, TexCreate_NoTiling) && !EnumHasAnyFlags(InDesc.Flags, TexCreate_FastVRAM | TexCreate_DepthStencilTargetable | TexCreate_RenderTargetable | TexCreate_UAV))
 		{
 			Desc.SetCpuCacheMode(mtlpp::CpuCacheMode::DefaultCache);
 #if PLATFORM_MAC
 			Desc.SetStorageMode(mtlpp::StorageMode::Managed);
-			Desc.SetResourceOptions((mtlpp::ResourceOptions)(mtlpp::ResourceOptions::CpuCacheModeDefaultCache|mtlpp::ResourceOptions::StorageModeManaged));
+			Desc.SetResourceOptions((mtlpp::ResourceOptions)(mtlpp::ResourceOptions::CpuCacheModeDefaultCache | mtlpp::ResourceOptions::StorageModeManaged));
 #else
 			Desc.SetStorageMode(mtlpp::StorageMode::Shared);
-			Desc.SetResourceOptions((mtlpp::ResourceOptions)(mtlpp::ResourceOptions::CpuCacheModeDefaultCache|mtlpp::ResourceOptions::StorageModeShared));
+			Desc.SetResourceOptions((mtlpp::ResourceOptions)(mtlpp::ResourceOptions::CpuCacheModeDefaultCache | mtlpp::ResourceOptions::StorageModeShared));
 #endif
 		}
-		else if (EnumHasAnyFlags(Flags, TexCreate_RenderTargetable|TexCreate_DepthStencilTargetable|TexCreate_ResolveTargetable|TexCreate_DepthStencilResolveTarget))
+		else if (EnumHasAnyFlags(InDesc.Flags, TexCreate_RenderTargetable | TexCreate_DepthStencilTargetable | TexCreate_ResolveTargetable | TexCreate_DepthStencilResolveTarget))
 		{
-			check(!(Flags & TexCreate_CPUReadback));
+			check(!(InDesc.Flags & TexCreate_CPUReadback));
 			Desc.SetCpuCacheMode(mtlpp::CpuCacheMode::DefaultCache);
 #if PLATFORM_MAC
 			Desc.SetStorageMode(mtlpp::StorageMode::Private);
-			Desc.SetResourceOptions((mtlpp::ResourceOptions)(mtlpp::ResourceOptions::CpuCacheModeDefaultCache|mtlpp::ResourceOptions::StorageModePrivate));
+			Desc.SetResourceOptions((mtlpp::ResourceOptions)(mtlpp::ResourceOptions::CpuCacheModeDefaultCache | mtlpp::ResourceOptions::StorageModePrivate));
 #else
-			if(GAGXForceIOSTexturesShared)
+			if (GAGXForceIOSTexturesShared)
 			{
 				Desc.SetStorageMode(mtlpp::StorageMode::Shared);
-				Desc.SetResourceOptions((mtlpp::ResourceOptions)(mtlpp::ResourceOptions::CpuCacheModeDefaultCache|mtlpp::ResourceOptions::StorageModeShared));
+				Desc.SetResourceOptions((mtlpp::ResourceOptions)(mtlpp::ResourceOptions::CpuCacheModeDefaultCache | mtlpp::ResourceOptions::StorageModeShared));
 			}
 			else
 			{
 				Desc.SetStorageMode(mtlpp::StorageMode::Private);
-				Desc.SetResourceOptions((mtlpp::ResourceOptions)(mtlpp::ResourceOptions::CpuCacheModeDefaultCache|mtlpp::ResourceOptions::StorageModePrivate));
+				Desc.SetResourceOptions((mtlpp::ResourceOptions)(mtlpp::ResourceOptions::CpuCacheModeDefaultCache | mtlpp::ResourceOptions::StorageModePrivate));
 			}
 #endif
 		}
 		else
 		{
-			check(!(Flags & TexCreate_CPUReadback));
+			check(!(InDesc.Flags & TexCreate_CPUReadback));
 			Desc.SetCpuCacheMode(mtlpp::CpuCacheMode::DefaultCache);
 #if PLATFORM_MAC
 			Desc.SetStorageMode(mtlpp::StorageMode::Private);
-			Desc.SetResourceOptions((mtlpp::ResourceOptions)(mtlpp::ResourceOptions::CpuCacheModeDefaultCache|mtlpp::ResourceOptions::StorageModePrivate));
+			Desc.SetResourceOptions((mtlpp::ResourceOptions)(mtlpp::ResourceOptions::CpuCacheModeDefaultCache | mtlpp::ResourceOptions::StorageModePrivate));
 #else
-			if(GAGXForceIOSTexturesShared)
+			if (GAGXForceIOSTexturesShared)
 			{
 				Desc.SetStorageMode(mtlpp::StorageMode::Shared);
-				Desc.SetResourceOptions((mtlpp::ResourceOptions)(mtlpp::ResourceOptions::CpuCacheModeDefaultCache|mtlpp::ResourceOptions::StorageModeShared));
+				Desc.SetResourceOptions((mtlpp::ResourceOptions)(mtlpp::ResourceOptions::CpuCacheModeDefaultCache | mtlpp::ResourceOptions::StorageModeShared));
 			}
 			// No private storage for PVRTC as it messes up the blit-encoder usage.
 			// note: this is set to always be on and will be re-addressed in a future release
 			else
 			{
-				if (IsPixelFormatPVRTCCompressed(Format))
+				if (IsPixelFormatPVRTCCompressed(InDesc.Format))
 				{
 					Desc.SetStorageMode(mtlpp::StorageMode::Shared);
-					Desc.SetResourceOptions((mtlpp::ResourceOptions)(mtlpp::ResourceOptions::CpuCacheModeDefaultCache|mtlpp::ResourceOptions::StorageModeShared));
+					Desc.SetResourceOptions((mtlpp::ResourceOptions)(mtlpp::ResourceOptions::CpuCacheModeDefaultCache | mtlpp::ResourceOptions::StorageModeShared));
 				}
 				else
 				{
 					Desc.SetStorageMode(mtlpp::StorageMode::Private);
-					Desc.SetResourceOptions((mtlpp::ResourceOptions)(mtlpp::ResourceOptions::CpuCacheModeDefaultCache|mtlpp::ResourceOptions::StorageModePrivate));
+					Desc.SetResourceOptions((mtlpp::ResourceOptions)(mtlpp::ResourceOptions::CpuCacheModeDefaultCache | mtlpp::ResourceOptions::StorageModePrivate));
 				}
 			}
 #endif
 		}
-		
+
 #if PLATFORM_IOS
-		if (EnumHasAnyFlags(Flags, TexCreate_Memoryless))
+		if (EnumHasAnyFlags(InDesc.Flags, TexCreate_Memoryless))
 		{
-			ensure(EnumHasAnyFlags(Flags, (TexCreate_RenderTargetable | TexCreate_DepthStencilTargetable)));
-			ensure(!EnumHasAnyFlags(Flags, (TexCreate_CPUReadback | TexCreate_CPUWritable)));
-			ensure(!EnumHasAnyFlags(Flags, TexCreate_UAV));
+			ensure(EnumHasAnyFlags(InDesc.Flags, (TexCreate_RenderTargetable | TexCreate_DepthStencilTargetable)));
+			ensure(!EnumHasAnyFlags(InDesc.Flags, (TexCreate_CPUReadback | TexCreate_CPUWritable)));
+			ensure(!EnumHasAnyFlags(InDesc.Flags, TexCreate_UAV));
 			Desc.SetStorageMode(mtlpp::StorageMode::Memoryless);
 			Desc.SetResourceOptions((mtlpp::ResourceOptions)(mtlpp::ResourceOptions::CpuCacheModeDefaultCache | mtlpp::ResourceOptions::StorageModeMemoryless));
 		}
@@ -788,164 +512,212 @@ FAGXSurface::FAGXSurface(ERHIResourceType ResourceType, EPixelFormat Format, uin
 		Desc.SetResourceOptions((mtlpp::ResourceOptions)(Desc.GetResourceOptions() | GeneralResourceOption));
 	}
 	
+	if (!FParse::Param(FCommandLine::Get(), TEXT("nomsaa")))
+	{
+		if (InDesc.NumSamples > 1)
+		{
+			check(bIsRenderTarget);
+			Desc.SetTextureType(mtlpp::TextureType::Texture2DMultisample);
+
+			// allow commandline to override
+			uint32 NewNumSamples;
+			if (FParse::Value(FCommandLine::Get(), TEXT("msaa="), NewNumSamples))
+			{
+				Desc.SetSampleCount(NewNumSamples);
+			}
+			else
+			{
+				Desc.SetSampleCount(InDesc.NumSamples);
+			}
+
+#if PLATFORM_IOS
+			if (GMaxRHIFeatureLevel < ERHIFeatureLevel::SM5)
+			{
+				bMemoryless = true;
+				Desc.SetStorageMode(mtlpp::StorageMode::Memoryless);
+				Desc.SetResourceOptions(mtlpp::ResourceOptions::StorageModeMemoryless);
+			}
+#endif
+		}
+	}
+}
+
+FAGXSurface::FAGXSurface(FAGXTextureCreateDesc const& CreateDesc)
+	: FRHITexture       (CreateDesc)
+	, FormatKey         (CreateDesc.FormatKey)
+	, Texture           (nil)
+	, MSAATexture       (nil)
+	, MSAAResolveTexture(nil)
+	, BufferLocks       (0)
+	, TotalTextureSize  (0)
+	, Viewport          (nullptr)
+	, ImageSurfaceRef   (nullptr)
+{
+	FPlatformAtomics::InterlockedExchange(&Written, 0);
+
+	check(CreateDesc.Extent.X > 0 && CreateDesc.Extent.Y > 0 && CreateDesc.NumMips > 0);
+
+	// the special back buffer surface will be updated in GetMetalDeviceContext().BeginDrawingViewport - no need to set the texture here
+	if (EnumHasAnyFlags(CreateDesc.Flags, TexCreate_Presentable))
+	{
+		return;
+	}
+
+	FResourceBulkDataInterface* BulkData = CreateDesc.BulkData;
+
+	// The bulk data interface can be used to create external textures for VR and media player.
+	// Handle these first.
 	if (BulkData != nullptr)
 	{
-		if (BulkData->GetResourceType() == FResourceBulkDataInterface::EBulkDataType::MediaTexture)
+		switch (BulkData->GetResourceType())
 		{
-			checkf(NumMips == 1 && ArraySize == 1, TEXT("Only handling bulk data with 1 mip and 1 array length"));
-			ImageSurfaceRef = (CFTypeRef)BulkData->GetResourceBulkData();
-			CFRetain(ImageSurfaceRef);
+		case FResourceBulkDataInterface::EBulkDataType::MediaTexture:
 			{
+				checkf(CreateDesc.NumMips == 1 && CreateDesc.ArraySize == 1, TEXT("Only handling bulk data with 1 mip and 1 array length"));
+				ImageSurfaceRef = (CFTypeRef)BulkData->GetResourceBulkData();
+				CFRetain(ImageSurfaceRef);
+				
 #if !COREVIDEO_SUPPORTS_METAL
-				Texture = MTLPP_VALIDATE(mtlpp::Device, GMtlppDevice, AGXSafeGetRuntimeDebuggingLevel() >= EAGXDebugLevelValidation, NewTexture(Desc, CVPixelBufferGetIOSurface((CVPixelBufferRef)ImageSurfaceRef), 0));
+				Texture = MTLPP_VALIDATE(mtlpp::Device, GMtlppDevice, AGXSafeGetRuntimeDebuggingLevel() >= EAGXDebugLevelValidation, NewTexture(CreateDesc.Desc, CVPixelBufferGetIOSurface((CVPixelBufferRef)ImageSurfaceRef), 0));
 #else
 				Texture = CVMetalTextureGetTexture((CVMetalTextureRef)ImageSurfaceRef);
 #endif
-				METAL_FATAL_ASSERT(Texture, TEXT("Failed to create texture, desc %s"), *FString([Desc description]));
+				METAL_FATAL_ASSERT(Texture, TEXT("Failed to create texture, desc %s"), *FString([CreateDesc.Desc.GetPtr() description]));
+
+				BulkData->Discard();
+				BulkData = nullptr;
 			}
-			
-			BulkData->Discard();
-		}
+			break;
+
 #if PLATFORM_MAC
-		else if (BulkData->GetResourceType() == FResourceBulkDataInterface::EBulkDataType::VREyeBuffer)
-		{
-			ImageSurfaceRef = (CFTypeRef)BulkData->GetResourceBulkData();
-			CFRetain(ImageSurfaceRef);
-			
-			Desc.SetStorageMode(mtlpp::StorageMode::Managed);
-			Desc.SetResourceOptions((mtlpp::ResourceOptions)((Desc.GetResourceOptions() & ~(mtlpp::ResourceStorageModeMask)) | mtlpp::ResourceOptions::StorageModeManaged));
-			
-			ns::IOSurface MtlppIOSurface((IOSurfaceRef)ImageSurfaceRef, ns::Ownership::AutoRelease);
-			Texture = MTLPP_VALIDATE(mtlpp::Device, GMtlppDevice, AGXSafeGetRuntimeDebuggingLevel() >= EAGXDebugLevelValidation, NewTextureWithDescriptor(Desc, MtlppIOSurface, 0));
-			
-			METAL_FATAL_ASSERT(Texture, TEXT("Failed to create texture, desc %s"), *FString([Desc description]));
-			
-			BulkData->Discard();
-		}
+		case FResourceBulkDataInterface::EBulkDataType::VREyeBuffer:
+			{
+				ImageSurfaceRef = (CFTypeRef)BulkData->GetResourceBulkData();
+				CFRetain(ImageSurfaceRef);
+
+				mtlpp::TextureDescriptor DescCopy([CreateDesc.Desc.GetPtr() copy], ns::Ownership::Assign);
+				DescCopy.SetStorageMode(mtlpp::StorageMode::Managed);
+				DescCopy.SetResourceOptions((mtlpp::ResourceOptions)((DescCopy.GetResourceOptions() & ~(mtlpp::ResourceStorageModeMask)) | mtlpp::ResourceOptions::StorageModeManaged));
+
+				ns::IOSurface MtlppIOSurface((IOSurfaceRef)ImageSurfaceRef, ns::Ownership::AutoRelease);
+				Texture = MTLPP_VALIDATE(mtlpp::Device, GMtlppDevice, AGXSafeGetRuntimeDebuggingLevel() >= EAGXDebugLevelValidation, NewTextureWithDescriptor(DescCopy, MtlppIOSurface, 0));
+				
+				METAL_FATAL_ASSERT(Texture, TEXT("Failed to create texture, desc %s"), *FString([DescCopy.GetPtr() description]));
+
+				BulkData->Discard();
+				BulkData = nullptr;
+			}
+			break;
 #endif
+		}
 	}
-	else
+
+	if (!Texture)
 	{
-		const bool bBufferCompatibleOption = (Desc.GetTextureType() == mtlpp::TextureType::Texture2D || Desc.GetTextureType() == mtlpp::TextureType::TextureBuffer) && NumMips == 1;
+		// Non VR/media texture case (i.e. a regular texture).
+		// Create the actual texture resource.
 
-		if (!bBufferCompatibleOption || (!EnumHasAllFlags(Flags, TexCreate_UAV|TexCreate_NoTiling) && !EnumHasAllFlags(Flags, TexCreate_AtomicCompatible)))
+		const bool bBufferCompatibleOption = (CreateDesc.Desc.GetTextureType() == mtlpp::TextureType::Texture2D || CreateDesc.Desc.GetTextureType() == mtlpp::TextureType::TextureBuffer) && CreateDesc.NumMips == 1;
+		if (!bBufferCompatibleOption || (!EnumHasAllFlags(CreateDesc.Flags, TexCreate_UAV | TexCreate_NoTiling) && !EnumHasAllFlags(CreateDesc.Flags, TexCreate_AtomicCompatible)))
 		{
-			Texture = GetAGXDeviceContext().CreateTexture(this, Desc);
+			Texture = GetAGXDeviceContext().CreateTexture(this, CreateDesc.Desc);
 		}
 		else
 		{
-			mtlpp::SizeAndAlign SizeAlign = GMtlppDevice.HeapTextureSizeAndAlign(Desc);
+			const uint32 MinimumByteAlignment = GMtlppDevice.GetMinimumLinearTextureAlignmentForPixelFormat(CreateDesc.MTLFormat);
+			const NSUInteger BytesPerRow = Align(CreateDesc.Desc.GetWidth() * GPixelFormats[CreateDesc.Format].BlockBytes, MinimumByteAlignment);
 
 			// Backing buffer resource options must match the texture we are going to create from it
-			FAGXPooledBufferArgs Args(SizeAlign.Size, BUF_Dynamic, mtlpp::StorageMode::Private, Desc.GetCpuCacheMode());
+			FAGXPooledBufferArgs Args(BytesPerRow * CreateDesc.Desc.GetHeight(), BUF_Dynamic, mtlpp::StorageMode::Private, CreateDesc.Desc.GetCpuCacheMode());
 			FAGXBuffer Buffer = GetAGXDeviceContext().CreatePooledBuffer(Args);
 
-			const uint32 MinimumByteAlignment = GMtlppDevice.GetMinimumLinearTextureAlignmentForPixelFormat(MTLFormat);
-			const NSUInteger BytesPerRow = Align(Desc.GetWidth() * GPixelFormats[Format].BlockBytes, MinimumByteAlignment);
-			
-			Texture = Buffer.NewTexture(Desc, 0, BytesPerRow);
-		}
-		
-		METAL_FATAL_ASSERT(Texture, TEXT("Failed to create texture, desc %s"), *FString([Desc description]));
-		
-		// upload existing bulkdata
-		if (BulkData)
-		{
-			UE_LOG(LogAGX, Display, TEXT("Got a bulk data texture, with %d mips"), NumMips);
-			checkf(NumMips == 1&& ArraySize == 1, TEXT("Only handling bulk data with 1 mip and 1 array length"));
-			uint32 Stride;
-			
-			// lock, copy, unlock
-			void* LockedData = Lock(0, 0, RLM_WriteOnly, Stride);
-			FMemory::Memcpy(LockedData, BulkData->GetResourceBulkData(), BulkData->GetResourceBulkDataSize());
-			Unlock(0, 0, false);
-			
-			// bulk data can be unloaded now
-			BulkData->Discard();
+			Texture = Buffer.NewTexture(CreateDesc.Desc, 0, BytesPerRow);
 		}
+
+		METAL_FATAL_ASSERT(Texture, TEXT("Failed to create texture, desc %s"), *FString([CreateDesc.Desc description]));
 	}
-	
+
+	if (BulkData)
+	{
+		// Regular texture has some bulk data to handle
+		UE_LOG(LogAGX, Display, TEXT("Got a bulk data texture, with %d mips"), CreateDesc.NumMips);
+		checkf(CreateDesc.NumMips == 1, TEXT("Only handling bulk data with 1 mip and 1 array length"));
+
+		check(IsInRenderingThread());
+		FRHICommandListImmediate& RHICmdList = FRHICommandListExecutor::GetImmediateCommandList();
+
+		// lock, copy, unlock
+		uint32 Stride;
+		void* LockedData = FAGXDynamicRHI::Get().LockTexture2D_RenderThread(RHICmdList, this, 0, RLM_WriteOnly, Stride, false);
+		check(LockedData);
+		FMemory::Memcpy(LockedData, BulkData->GetResourceBulkData(), BulkData->GetResourceBulkDataSize());
+		FAGXDynamicRHI::Get().UnlockTexture2D_RenderThread(RHICmdList, this, 0, false);
+
+		// bulk data can be unloaded now
+		BulkData->Discard();
+		BulkData = nullptr;
+	}
+
 	// calculate size of the texture
 	TotalTextureSize = GetMemorySize();
-	
-	const bool bSupportsMSAADepthResolve = GetAGXDeviceContext().SupportsFeature(EAGXFeaturesMSAADepthResolve);
-	if (!FParse::Param(FCommandLine::Get(), TEXT("nomsaa")))
+
+	if (CreateDesc.NumSamples > 1 && !FParse::Param(FCommandLine::Get(), TEXT("nomsaa")))
 	{
-		if (NumSamples > 1)
+		MSAATexture = GetAGXDeviceContext().CreateTexture(this, CreateDesc.Desc);
+
+		//device doesn't support HW depth resolve.  This case only valid on mobile renderer or
+		//on Mac where RHISupportsSeparateMSAAAndResolveTextures is true.
+		const bool bSupportsMSAADepthResolve = GetAGXDeviceContext().SupportsFeature(EAGXFeaturesMSAADepthResolve);
+		const bool bDepthButNoResolveSupported = CreateDesc.Format == PF_DepthStencil && !bSupportsMSAADepthResolve;
+		if (bDepthButNoResolveSupported)
 		{
-			check(bIsRenderTarget);
-			Desc.SetTextureType(mtlpp::TextureType::Texture2DMultisample);
-			
-			// allow commandline to override
-			FParse::Value(FCommandLine::Get(), TEXT("msaa="), NumSamples);
-			Desc.SetSampleCount(NumSamples);
-			
-			bool bMemoryless = false;
-#if PLATFORM_IOS
-			if (GMaxRHIFeatureLevel < ERHIFeatureLevel::SM5)
-			{
-				bMemoryless = true;
-				Desc.SetStorageMode(mtlpp::StorageMode::Memoryless);
-				Desc.SetResourceOptions(mtlpp::ResourceOptions::StorageModeMemoryless);
-			}
-#endif
-			
-			MSAATexture = GetAGXDeviceContext().CreateTexture(this, Desc);
-			
-			//device doesn't support HW depth resolve.  This case only valid on mobile renderer or
-			//on Mac where RHISupportsSeparateMSAAAndResolveTextures is true.
-			const bool bDepthButNoResolveSupported = Format == PF_DepthStencil && !bSupportsMSAADepthResolve;
-			if (bDepthButNoResolveSupported)
+			Texture = MSAATexture;
+
+			// we don't have the resolve texture, so we just update the memory size with the MSAA size
+			TotalTextureSize = TotalTextureSize * CreateDesc.NumSamples;
+		}
+		else if (!CreateDesc.bMemoryless)
+		{
+			// an MSAA render target takes NumSamples more space, in addition to the resolve texture
+			TotalTextureSize += TotalTextureSize * CreateDesc.NumSamples;
+		}
+
+		if (MSAATexture != Texture)
+		{
+			check(!MSAAResolveTexture);
+
+			//if bSupportsSeparateMSAAAndResolve then the high level expect to binds the MSAA when binding shader params.
+			const bool bSupportsSeparateMSAAAndResolve = FAGXCommandQueue::SupportsSeparateMSAAAndResolveTarget();
+			if (bSupportsSeparateMSAAAndResolve)
 			{
+				MSAAResolveTexture = Texture;
 				Texture = MSAATexture;
-				
-				// we don't have the resolve texture, so we just update the memory size with the MSAA size
-				TotalTextureSize = TotalTextureSize * NumSamples;
 			}
-			else if (!bMemoryless)
-			{
-				// an MSAA render target takes NumSamples more space, in addition to the resolve texture
-				TotalTextureSize += TotalTextureSize * NumSamples;
-			}
-			
-			if (MSAATexture != Texture)
-			{
-				check(!MSAAResolveTexture);
-				
-				//if bSupportsSeparateMSAAAndResolve then the high level expect to binds the MSAA when binding shader params.
-				const bool bSupportsSeparateMSAAAndResolve = FAGXCommandQueue::SupportsSeparateMSAAAndResolveTarget();
-				if (bSupportsSeparateMSAAAndResolve)
-				{
-					MSAAResolveTexture = Texture;
-					
-					Texture = MSAATexture;
-				}
-				else
-				{
-					MSAAResolveTexture = Texture;
-				}
-			}
-			
-			//we always require an MSAAResolveTexture if MSAATexture is active.
-			check(!MSAATexture || MSAAResolveTexture || bDepthButNoResolveSupported);
-			
-			NSLog(@"Creating %dx MSAA %d x %d %s surface", (int32)Desc.GetSampleCount(), SizeX, SizeY, EnumHasAnyFlags(Flags, TexCreate_RenderTargetable) ? "Color" : "Depth");
-			if (MSAATexture.GetPtr() == nil)
+			else
 			{
-				NSLog(@"Failed to create texture, desc  %@", Desc.GetPtr());
+				MSAAResolveTexture = Texture;
 			}
 		}
+
+		// We always require an MSAAResolveTexture if MSAATexture is active.
+		check(!MSAATexture || MSAAResolveTexture || bDepthButNoResolveSupported);
+
+		if (MSAATexture.GetPtr() == nil)
+		{
+			UE_LOG(LogAGX, Warning, TEXT("Failed to create MSAA texture with descriptor: %s"), *FString([CreateDesc.Desc.GetPtr() description]));
+		}
 	}
-	
+
 	// create a stencil buffer if needed
-	if (Format == PF_DepthStencil)
+	if (CreateDesc.Format == PF_DepthStencil)
 	{
 		// 1 byte per texel
-		TotalTextureSize += SizeX * SizeY;
+		TotalTextureSize += CreateDesc.Extent.X * CreateDesc.Extent.Y;
 	}
-	
+
 	// track memory usage
-	if (bIsRenderTarget)
+
+	if (CreateDesc.bIsRenderTarget)
 	{
 		GCurrentRendertargetMemorySize += Align(TotalTextureSize, 1024) / 1024;
 	}
@@ -953,11 +725,11 @@ FAGXSurface::FAGXSurface(ERHIResourceType ResourceType, EPixelFormat Format, uin
 	{
 		GCurrentTextureMemorySize += Align(TotalTextureSize, 1024) / 1024;
 	}
-	
+
 #if STATS
-	if (ResourceType == RRT_TextureCube)
+	if (CreateDesc.IsTextureCube())
 	{
-		if (bIsRenderTarget)
+		if (CreateDesc.bIsRenderTarget)
 		{
 			INC_MEMORY_STAT_BY(STAT_RenderTargetMemoryCube, TotalTextureSize);
 		}
@@ -966,9 +738,9 @@ FAGXSurface::FAGXSurface(ERHIResourceType ResourceType, EPixelFormat Format, uin
 			INC_MEMORY_STAT_BY(STAT_TextureMemoryCube, TotalTextureSize);
 		}
 	}
-	else if (ResourceType == RRT_Texture3D)
+	else if (CreateDesc.IsTexture3D())
 	{
-		if (bIsRenderTarget)
+		if (CreateDesc.bIsRenderTarget)
 		{
 			INC_MEMORY_STAT_BY(STAT_RenderTargetMemory3D, TotalTextureSize);
 		}
@@ -979,7 +751,7 @@ FAGXSurface::FAGXSurface(ERHIResourceType ResourceType, EPixelFormat Format, uin
 	}
 	else
 	{
-		if (bIsRenderTarget)
+		if (CreateDesc.bIsRenderTarget)
 		{
 			INC_MEMORY_STAT_BY(STAT_RenderTargetMemory2D, TotalTextureSize);
 		}
@@ -995,7 +767,7 @@ FAGXSurface::FAGXSurface(ERHIResourceType ResourceType, EPixelFormat Format, uin
 {
 @public
 	uint64 TextureSize;
-	ERHIResourceType TexType;
+	ETextureDimension Dimension;
 	bool bIsRenderTarget;
 }
 @end
@@ -1005,7 +777,7 @@ APPLE_PLATFORM_OBJECT_ALLOC_OVERRIDES(FAGXDeferredStats)
 -(void)dealloc
 {
 #if STATS
-	if (TexType == RRT_TextureCube)
+	if (Dimension == ETextureDimension::TextureCube || Dimension == ETextureDimension::TextureCubeArray)
 	{
 		if (bIsRenderTarget)
 		{
@@ -1016,7 +788,7 @@ APPLE_PLATFORM_OBJECT_ALLOC_OVERRIDES(FAGXDeferredStats)
 			DEC_MEMORY_STAT_BY(STAT_TextureMemoryCube, TextureSize);
 		}
 	}
-	else if (TexType == RRT_Texture3D)
+	else if (Dimension == ETextureDimension::Texture3D)
 	{
 		if (bIsRenderTarget)
 		{
@@ -1053,13 +825,13 @@ APPLE_PLATFORM_OBJECT_ALLOC_OVERRIDES(FAGXDeferredStats)
 
 FAGXSurface::~FAGXSurface()
 {
-	bool const bIsRenderTarget = IsRenderTarget(Flags);
+	bool const bIsRenderTarget = IsRenderTarget(GetDesc().Flags);
 	
 	if (MSAATexture.GetPtr())
 	{
 		if (Texture.GetPtr() != MSAATexture.GetPtr())
 		{
-			AGXSafeReleaseMetalTexture(this, MSAATexture, bTextureView, false);
+			AGXSafeReleaseMetalTexture(this, MSAATexture, false);
 		}
 	}
 	
@@ -1069,13 +841,13 @@ FAGXSurface::~FAGXSurface()
 	{
 		if (Texture.GetPtr() != MSAAResolveTexture.GetPtr())
 		{
-			AGXSafeReleaseMetalTexture(this, MSAAResolveTexture, bTextureView, false);
+			AGXSafeReleaseMetalTexture(this, MSAAResolveTexture, false);
 		}
 	}
 	
-	if (!(Flags & TexCreate_Presentable) && Texture.GetPtr())
+	if (!(GetDesc().Flags & TexCreate_Presentable) && Texture.GetPtr())
 	{
-		AGXSafeReleaseMetalTexture(this, Texture, bTextureView, (ImageSurfaceRef != nullptr));
+		AGXSafeReleaseMetalTexture(this, Texture, (ImageSurfaceRef != nullptr));
 	}
 	
 	MSAATexture = nil;
@@ -1084,7 +856,7 @@ FAGXSurface::~FAGXSurface()
 	
 	// track memory usage
 	FAGXDeferredStats* Block = [FAGXDeferredStats new];
-	Block->TexType = Type;
+	Block->Dimension = GetDesc().Dimension;
 	Block->TextureSize = TotalTextureSize;
 	Block->bIsRenderTarget = bIsRenderTarget;
 	AGXSafeReleaseMetalObject(Block);
@@ -1117,9 +889,13 @@ id <MTLBuffer> FAGXSurface::AllocSurface(uint32 MipIndex, uint32 ArrayIndex, ERe
 	
 #if PLATFORM_MAC
 	// Expand R8_sRGB into RGBA8_sRGB for Mac.
-	if (PixelFormat == PF_G8 && EnumHasAnyFlags(Flags, TexCreate_SRGB) && Type == RRT_Texture2D && LockMode == RLM_WriteOnly)
+	if (   GetDesc().Format == PF_G8
+		&& GetDesc().Dimension == ETextureDimension::Texture2D
+		&& EnumHasAnyFlags(GetDesc().Flags, TexCreate_SRGB)
+		&& LockMode == RLM_WriteOnly
+		&& Texture.GetPixelFormat() == mtlpp::PixelFormat::RGBA8Unorm_sRGB)
 	{
-		DestStride = FMath::Max<uint32>(SizeX >> MipIndex, 1);
+		DestStride = FMath::Max<uint32>(GetDesc().Extent.X >> MipIndex, 1);
 	}
 #endif
 	
@@ -1139,29 +915,42 @@ void FAGXSurface::UpdateSurfaceAndDestroySourceBuffer(id <MTLBuffer> SourceBuffe
 	uint32 BytesPerImage = GetMipSize(MipIndex, &Stride, true);
 	
 	mtlpp::Region Region;
-	if (SizeZ <= 1 || bIsCubemap)
+	if (GetDesc().IsTexture3D())
 	{
 		// upload the texture to the texture slice
-		Region = mtlpp::Region(0, 0, FMath::Max<uint32>(SizeX >> MipIndex, 1), FMath::Max<uint32>(SizeY >> MipIndex, 1));
+		Region = mtlpp::Region(
+			0, 0, 0,
+			FMath::Max<uint32>(GetDesc().Extent.X >> MipIndex, 1),
+			FMath::Max<uint32>(GetDesc().Extent.Y >> MipIndex, 1),
+			FMath::Max<uint32>(GetDesc().Depth    >> MipIndex, 1)
+		);
 	}
 	else
 	{
 		// upload the texture to the texture slice
-		Region = mtlpp::Region(0, 0, 0, FMath::Max<uint32>(SizeX >> MipIndex, 1), FMath::Max<uint32>(SizeY >> MipIndex, 1), FMath::Max<uint32>(SizeZ >> MipIndex, 1));
+		Region = mtlpp::Region(
+			0, 0,
+			FMath::Max<uint32>(GetDesc().Extent.X >> MipIndex, 1),
+			FMath::Max<uint32>(GetDesc().Extent.Y >> MipIndex, 1)
+		);
 	}
+
 #if PLATFORM_MAC
 	// Expand R8_sRGB into RGBA8_sRGB for Mac.
-	if (PixelFormat == PF_G8 && EnumHasAnyFlags(Flags, TexCreate_SRGB) && Type == RRT_Texture2D)
+	if (   GetDesc().Format == PF_G8
+		&& GetDesc().Dimension == ETextureDimension::Texture2D
+		&& EnumHasAnyFlags(GetDesc().Flags, TexCreate_SRGB)
+		&& Texture.GetPixelFormat() == mtlpp::PixelFormat::RGBA8Unorm_sRGB)
 	{
 		TArray<uint8> Data;
 		uint8* ExpandedMem = (uint8*) SourceBuffer.contents;
 		check(ExpandedMem);
 		Data.Append(ExpandedMem, BytesPerImage);
-		uint32 SrcStride = FMath::Max<uint32>(SizeX >> MipIndex, 1);
-		for(uint y = 0; y < FMath::Max<uint32>(SizeY >> MipIndex, 1); y++)
+		uint32 SrcStride = FMath::Max<uint32>(GetDesc().Extent.X >> MipIndex, 1);
+		for(uint y = 0; y < FMath::Max<uint32>(GetDesc().Extent.Y >> MipIndex, 1); y++)
 		{
 			uint8* RowDest = ExpandedMem;
-			for(uint x = 0; x < FMath::Max<uint32>(SizeX >> MipIndex, 1); x++)
+			for(uint x = 0; x < FMath::Max<uint32>(GetDesc().Extent.X >> MipIndex, 1); x++)
 			{
 				*(RowDest++) = Data[(y * SrcStride) + x];
 				*(RowDest++) = Data[(y * SrcStride) + x];
@@ -1278,15 +1067,23 @@ void* FAGXSurface::Lock(uint32 MipIndex, uint32 ArrayIndex, EResourceLockMode Lo
 			const bool bIssueImmediateCommands = RHICmdList.Bypass() || IsInRHIThread();
 			
 			mtlpp::Region Region;
-			if (SizeZ <= 1 || bIsCubemap)
+			if (GetDesc().IsTexture3D())
 			{
 				// upload the texture to the texture slice
-				Region = mtlpp::Region(0, 0, FMath::Max<uint32>(SizeX >> MipIndex, 1), FMath::Max<uint32>(SizeY >> MipIndex, 1));
+				Region = mtlpp::Region(
+					0, 0, 0,
+					FMath::Max<uint32>(GetDesc().Extent.X >> MipIndex, 1),
+					FMath::Max<uint32>(GetDesc().Extent.Y >> MipIndex, 1),
+					FMath::Max<uint32>(GetDesc().Depth    >> MipIndex, 1));
 			}
 			else
 			{
 				// upload the texture to the texture slice
-				Region = mtlpp::Region(0, 0, 0, FMath::Max<uint32>(SizeX >> MipIndex, 1), FMath::Max<uint32>(SizeY >> MipIndex, 1), FMath::Max<uint32>(SizeZ >> MipIndex, 1));
+				Region = mtlpp::Region(
+					0, 0,
+					FMath::Max<uint32>(GetDesc().Extent.X >> MipIndex, 1),
+					FMath::Max<uint32>(GetDesc().Extent.Y >> MipIndex, 1)
+				);
 			}
 			
 			if (Texture.GetStorageMode() == mtlpp::StorageMode::Private)
@@ -1339,7 +1136,7 @@ void* FAGXSurface::Lock(uint32 MipIndex, uint32 ArrayIndex, EResourceLockMode Lo
 				
 				// This block breaks the texture atlas system in Ocean, which depends on nonzero strides coming back from compressed textures. Turning off.
 #if 0
-				if (PixelFormat == PF_PVRTC2 || PixelFormat == PF_PVRTC4)
+				if (GetDesc().Format == PF_PVRTC2 || GetDesc().Format == PF_PVRTC4)
 				{
 					// for compressed textures metal debug RT expects 0 for rowBytes and imageBytes.
 					DestStride = 0;
@@ -1347,7 +1144,7 @@ void* FAGXSurface::Lock(uint32 MipIndex, uint32 ArrayIndex, EResourceLockMode Lo
 				}
 #endif
 				uint32 BytesPerRow = DestStride;
-				if (PixelFormat == PF_PVRTC2 || PixelFormat == PF_PVRTC4)
+				if (GetDesc().Format == PF_PVRTC2 || GetDesc().Format == PF_PVRTC4)
 				{
 					// for compressed textures metal debug RT expects 0 for rowBytes and imageBytes.
 					BytesPerRow = 0;
@@ -1358,17 +1155,20 @@ void* FAGXSurface::Lock(uint32 MipIndex, uint32 ArrayIndex, EResourceLockMode Lo
 			
 #if PLATFORM_MAC
 			// Pack RGBA8_sRGB into R8_sRGB for Mac.
-			if (PixelFormat == PF_G8 && EnumHasAnyFlags(Flags, TexCreate_SRGB) && Type == RRT_Texture2D)
+			if (   GetDesc().Format == PF_G8
+				&& GetDesc().Dimension == ETextureDimension::Texture2D
+				&& EnumHasAnyFlags(GetDesc().Flags, TexCreate_SRGB)
+				&& Texture.GetPixelFormat() == mtlpp::PixelFormat::RGBA8Unorm_sRGB)
 			{
 				TArray<uint8> Data;
 				uint8* ExpandedMem = (uint8*)SourceData.GetContents();
 				Data.Append(ExpandedMem, MipBytes);
 				uint32 SrcStride = DestStride;
-				DestStride = FMath::Max<uint32>(SizeX >> MipIndex, 1);
-				for(uint y = 0; y < FMath::Max<uint32>(SizeY >> MipIndex, 1); y++)
+				DestStride = FMath::Max<uint32>(GetDesc().Extent.X >> MipIndex, 1);
+				for(uint y = 0; y < FMath::Max<uint32>(GetDesc().Extent.Y >> MipIndex, 1); y++)
 				{
 					uint8* RowDest = ExpandedMem;
-					for(uint x = 0; x < FMath::Max<uint32>(SizeX >> MipIndex, 1); x++)
+					for(uint x = 0; x < FMath::Max<uint32>(GetDesc().Extent.X >> MipIndex, 1); x++)
 					{
 						*(RowDest++) = Data[(y * SrcStride) + (x * 4)];
 					}
@@ -1480,6 +1280,8 @@ void FAGXSurface::AsyncUnlock(id <MTLBuffer> SourceData, uint32 MipIndex, uint32
 
 uint32 FAGXSurface::GetMipSize(uint32 MipIndex, uint32* Stride, bool bSingleLayer)
 {
+	EPixelFormat PixelFormat = GetDesc().Format;
+
 	// DXT/BC formats on Mac actually do have mip-tails that are smaller than the block size, they end up being uncompressed.
 	bool const bPixelFormatASTC = IsPixelFormatASTCCompressed(PixelFormat);
 	
@@ -1488,12 +1290,12 @@ uint32 FAGXSurface::GetMipSize(uint32 MipIndex, uint32* Stride, bool bSingleLaye
 	const uint32 BlockSizeY = GPixelFormats[PixelFormat].BlockSizeY;
 	const uint32 BlockBytes = GPixelFormats[PixelFormat].BlockBytes;
 	const uint32 Alignment = 1u; // Apparently we always want natural row alignment (tightly-packed) even though the docs say iOS doesn't support it - this may be because we don't upload texture data from one contiguous buffer.
-	const uint32 UnalignedMipSizeX = FMath::Max(SizeX >> MipIndex, BlockSizeX);
-	const uint32 UnalignedMipSizeY = FMath::Max(SizeY >> MipIndex, BlockSizeY);
+	const uint32 UnalignedMipSizeX = FMath::Max<uint32>(GetDesc().Extent.X >> MipIndex, BlockSizeX);
+	const uint32 UnalignedMipSizeY = FMath::Max<uint32>(GetDesc().Extent.Y >> MipIndex, BlockSizeY);
 	const uint32 MipSizeX = (bPixelFormatASTC) ? AlignArbitrary(UnalignedMipSizeX, BlockSizeX) : UnalignedMipSizeX;
 	const uint32 MipSizeY = (bPixelFormatASTC) ? AlignArbitrary(UnalignedMipSizeY, BlockSizeY) : UnalignedMipSizeY;
 	
-	const uint32 MipSizeZ = bSingleLayer ? 1 : FMath::Max(SizeZ >> MipIndex, 1u);
+	const uint32 MipSizeZ = bSingleLayer ? 1 : FMath::Max<uint32>(GetDesc().Depth >> MipIndex, 1u);
 	uint32 NumBlocksX = (MipSizeX + BlockSizeX - 1) / BlockSizeX;
 	uint32 NumBlocksY = (MipSizeY + BlockSizeY - 1) / BlockSizeY;
 	if (PixelFormat == PF_PVRTC2 || PixelFormat == PF_PVRTC4)
@@ -1503,7 +1305,7 @@ uint32 FAGXSurface::GetMipSize(uint32 MipIndex, uint32* Stride, bool bSingleLaye
 		NumBlocksY = FMath::Max<uint32>(NumBlocksY, 2);
 	}
 #if PLATFORM_MAC
-	else if (PixelFormat == PF_G8 && EnumHasAnyFlags(Flags, TexCreate_SRGB))
+	else if (PixelFormat == PF_G8 && EnumHasAnyFlags(GetDesc().Flags, TexCreate_SRGB) && Texture.GetPixelFormat() == mtlpp::PixelFormat::RGBA8Unorm_sRGB)
 	{
 		// RGBA_sRGB is the closest match - so expand the data.
 		NumBlocksX *= 4;
@@ -1547,23 +1349,12 @@ uint32 FAGXSurface::GetMemorySize()
 
 uint32 FAGXSurface::GetNumFaces()
 {
-	switch (Type)
-	{
-		case RRT_Texture2DArray:
-		case RRT_Texture3D:
-		case RRT_TextureCube:
-			return SizeZ * Texture.GetArrayLength();
-			
-		case RRT_Texture:
-		case RRT_Texture2D:
-		default:
-			return 1;
-	}
+	return GetDesc().Depth * GetDesc().ArraySize;
 }
 
 FAGXTexture FAGXSurface::GetDrawableTexture()
 {
-	if (!Texture && EnumHasAnyFlags(Flags, TexCreate_Presentable))
+	if (!Texture && EnumHasAnyFlags(GetDesc().Flags, TexCreate_Presentable))
 	{
 		check(Viewport);
 		Texture = Viewport->GetDrawableTexture(EAGXViewportAccessRHI);
@@ -1574,7 +1365,7 @@ FAGXTexture FAGXSurface::GetDrawableTexture()
 ns::AutoReleased<FAGXTexture> FAGXSurface::GetCurrentTexture()
 {
 	ns::AutoReleased<FAGXTexture> Tex;
-	if (Viewport && EnumHasAnyFlags(Flags, TexCreate_Presentable))
+	if (Viewport && EnumHasAnyFlags(GetDesc().Flags, TexCreate_Presentable))
 	{
 		check(Viewport);
 		Tex = Viewport->GetCurrentTexture(EAGXViewportAccessRHI);
@@ -1632,17 +1423,17 @@ uint32 FAGXDynamicRHI::RHIComputeMemorySize(FRHITexture* TextureRHI)
  2D texture support.
  -----------------------------------------------------------------------------*/
 
-FTexture2DRHIRef FAGXDynamicRHI::RHICreateTexture2D(uint32 SizeX, uint32 SizeY, uint8 Format, uint32 NumMips, uint32 NumSamples, ETextureCreateFlags Flags, ERHIAccess InResourceState, FRHIResourceCreateInfo& CreateInfo)
+FTextureRHIRef FAGXDynamicRHI::RHICreateTexture_RenderThread(class FRHICommandListImmediate& RHICmdList, const FRHITextureCreateDesc& CreateDesc)
 {
-	@autoreleasepool {
-		return new FAGXTexture2D((EPixelFormat)Format, SizeX, SizeY, NumMips, NumSamples, Flags, CreateInfo.BulkData, CreateInfo.ClearValueBinding);
+	@autoreleasepool{
+		return this->RHICreateTexture(CreateDesc);
 	}
 }
 
-FTexture2DRHIRef FAGXDynamicRHI::RHICreateTextureExternal2D(uint32 SizeX, uint32 SizeY, uint8 Format, uint32 NumMips, uint32 NumSamples, ETextureCreateFlags Flags, ERHIAccess InResourceState, FRHIResourceCreateInfo& CreateInfo)
+FTextureRHIRef FAGXDynamicRHI::RHICreateTexture(const FRHITextureCreateDesc& CreateDesc)
 {
-	@autoreleasepool {
-		return new FAGXTexture2D((EPixelFormat)Format, SizeX, SizeY, NumMips, NumSamples, Flags, CreateInfo.BulkData, CreateInfo.ClearValueBinding);
+	@autoreleasepool{
+		return new FAGXSurface(CreateDesc);
 	}
 }
 
@@ -1657,20 +1448,6 @@ void FAGXDynamicRHI::RHICopySharedMips(FRHITexture2D* DestTexture2D, FRHITexture
 	NOT_SUPPORTED("RHICopySharedMips");
 }
 
-FTexture2DArrayRHIRef FAGXDynamicRHI::RHICreateTexture2DArray(uint32 SizeX, uint32 SizeY, uint32 SizeZ, uint8 Format, uint32 NumMips, uint32 NumSamples, ETextureCreateFlags Flags, ERHIAccess InResourceState, FRHIResourceCreateInfo& CreateInfo)
-{
-	@autoreleasepool {
-		return new FAGXTexture2DArray((EPixelFormat)Format, SizeX, SizeY, SizeZ, NumMips, Flags, CreateInfo.BulkData, CreateInfo.ClearValueBinding);
-	}
-}
-
-FTexture3DRHIRef FAGXDynamicRHI::RHICreateTexture3D(uint32 SizeX, uint32 SizeY, uint32 SizeZ, uint8 Format, uint32 NumMips, ETextureCreateFlags Flags, ERHIAccess InResourceState, FRHIResourceCreateInfo& CreateInfo)
-{
-	@autoreleasepool {
-		return new FAGXTexture3D((EPixelFormat)Format, SizeX, SizeY, SizeZ, NumMips, Flags, CreateInfo.BulkData, CreateInfo.ClearValueBinding);
-	}
-}
-
 void FAGXDynamicRHI::RHIGenerateMips(FRHITexture* SourceSurfaceRHI)
 {
 	@autoreleasepool {
@@ -1682,131 +1459,98 @@ void FAGXDynamicRHI::RHIGenerateMips(FRHITexture* SourceSurfaceRHI)
 	}
 }
 
-// deferred unlock commands
-struct FAGXRHICommandAsyncReallocateTexture2D final : public FRHICommand<FAGXRHICommandAsyncReallocateTexture2D>
-{
-	FAGXContext& Context;
-	FAGXTexture2D* OldTexture;
-	FAGXTexture2D* NewTexture;
-	int32 NewMipCount;
-	int32 NewSizeX;
-	int32 NewSizeY;
-	FThreadSafeCounter* RequestStatus;
-	
-	FORCEINLINE_DEBUGGABLE FAGXRHICommandAsyncReallocateTexture2D(FAGXContext& InContext, FAGXTexture2D* InOldTexture, FAGXTexture2D* InNewTexture, int32 InNewMipCount, int32 InNewSizeX, int32 InNewSizeY, FThreadSafeCounter* InRequestStatus)
-	: Context(InContext)
-	, OldTexture(InOldTexture)
-	, NewTexture(InNewTexture)
-	, NewMipCount(InNewMipCount)
-	, NewSizeX(InNewSizeX)
-	, NewSizeY(InNewSizeY)
-	, RequestStatus(InRequestStatus)
-	{
-	}
-	
-	void Execute(FRHICommandListBase& CmdList)
-	{
-		CopyMips(Context, OldTexture, NewTexture, NewMipCount, NewSizeX, NewSizeY, RequestStatus);
-	}
-	
-	static void CopyMips(FAGXContext& Context, FAGXTexture2D* OldTexture, FAGXTexture2D* NewTexture, int32 NewMipCount, int32 NewSizeX, int32 NewSizeY, FThreadSafeCounter* RequestStatus)
-	{
-		// figure out what mips to schedule
-		const uint32 NumSharedMips = FMath::Min(OldTexture->GetNumMips(), NewTexture->GetNumMips());
-		const uint32 SourceMipOffset = OldTexture->GetNumMips() - NumSharedMips;
-		const uint32 DestMipOffset = NewTexture->GetNumMips() - NumSharedMips;
-		
-		const uint32 BlockSizeX = GPixelFormats[OldTexture->GetFormat()].BlockSizeX;
-		const uint32 BlockSizeY = GPixelFormats[OldTexture->GetFormat()].BlockSizeY;
-		
-		// only handling straight 2D textures here
-		uint32 SliceIndex = 0;
-		mtlpp::Origin Origin(0,0,0);
-		
-		FAGXTexture Tex = OldTexture->Surface.Texture;
-		
-		// DXT/BC formats on Mac actually do have mip-tails that are smaller than the block size, they end up being uncompressed.
-		bool const bPixelFormatASTC = IsPixelFormatASTCCompressed(OldTexture->GetFormat());
-		
-		bool bAsync = true;
-		for (uint32 MipIndex = 0; MipIndex < NumSharedMips; ++MipIndex)
-		{
-			const uint32 UnalignedMipSizeX = FMath::Max<uint32>(1, NewSizeX >> (MipIndex + DestMipOffset));
-			const uint32 UnalignedMipSizeY = FMath::Max<uint32>(1, NewSizeY >> (MipIndex + DestMipOffset));
-			const uint32 MipSizeX = FMath::Max<uint32>(1, NewSizeX >> (MipIndex + DestMipOffset));
-			const uint32 MipSizeY = FMath::Max<uint32>(1, NewSizeY >> (MipIndex + DestMipOffset));
-			
-			bAsync &= Context.AsyncCopyFromTextureToTexture(OldTexture->Surface.Texture, SliceIndex, MipIndex + SourceMipOffset, Origin, mtlpp::Size(MipSizeX, MipSizeY, 1), NewTexture->Surface.Texture, SliceIndex, MipIndex + DestMipOffset, Origin);
-		}
-		
-		// when done, decrement the counter to indicate it's safe
-		mtlpp::CommandBufferHandler CompletionHandler = [Tex](mtlpp::CommandBuffer const&)
-		{
-		};
-		
-		if(bAsync)
-		{
-	    	// kck it off!
-			Context.SubmitAsyncCommands(nil, CompletionHandler, false);
-		}
-		else
-		{
-			Context.GetCurrentRenderPass().AddCompletionHandler(CompletionHandler);
-		}
-		
-		// Like D3D mark this as complete immediately.
-		RequestStatus->Decrement();
-		
-		FAGXSurface* Source = AGXGetMetalSurfaceFromRHITexture(OldTexture);
-		Source->MakeAliasable();
-	}
-};
-
 FTexture2DRHIRef FAGXDynamicRHI::AsyncReallocateTexture2D_RenderThread(class FRHICommandListImmediate& RHICmdList, FRHITexture2D* Texture2D, int32 NewMipCount, int32 NewSizeX, int32 NewSizeY, FThreadSafeCounter* RequestStatus)
 {
-	@autoreleasepool {
-		FTexture2DRHIRef Result;
-		
-		if (RHICmdList.Bypass() || !IsRunningRHIInSeparateThread())
-		{
-			Result = GDynamicRHI->RHIAsyncReallocateTexture2D(Texture2D, NewMipCount, NewSizeX, NewSizeY, RequestStatus);
-		}
-		else
-		{
-			FAGXTexture2D* OldTexture = ResourceCast(Texture2D);
-			
-			FAGXTexture2D* NewTexture = new FAGXTexture2D(OldTexture->GetFormat(), NewSizeX, NewSizeY, NewMipCount, OldTexture->GetNumSamples(), OldTexture->GetFlags(), NULL, Texture2D->GetClearBinding());
-			
-			Result = NewTexture;
-			
-			new (RHICmdList.AllocCommand<FAGXRHICommandAsyncReallocateTexture2D>()) FAGXRHICommandAsyncReallocateTexture2D(ImmediateContext.GetInternalContext(), OldTexture, NewTexture, NewMipCount, NewSizeX, NewSizeY, RequestStatus);
-		}
-		
-		return Result;
-	}
+	return this->RHIAsyncReallocateTexture2D(Texture2D, NewMipCount, NewSizeX, NewSizeY, RequestStatus);
 }
 
 ETextureReallocationStatus FAGXDynamicRHI::FinalizeAsyncReallocateTexture2D_RenderThread(class FRHICommandListImmediate& RHICmdList, FRHITexture2D* Texture2D, bool bBlockUntilCompleted)
 {
 	// No need to flush - does nothing
-	return GDynamicRHI->RHIFinalizeAsyncReallocateTexture2D(Texture2D, bBlockUntilCompleted);
+	return this->RHIFinalizeAsyncReallocateTexture2D(Texture2D, bBlockUntilCompleted);
 }
 
 ETextureReallocationStatus FAGXDynamicRHI::CancelAsyncReallocateTexture2D_RenderThread(class FRHICommandListImmediate& RHICmdList, FRHITexture2D* Texture2D, bool bBlockUntilCompleted)
 {
 	// No need to flush - does nothing
-	return GDynamicRHI->RHICancelAsyncReallocateTexture2D(Texture2D, bBlockUntilCompleted);
+	return this->RHICancelAsyncReallocateTexture2D(Texture2D, bBlockUntilCompleted);
 }
 
-
 FTexture2DRHIRef FAGXDynamicRHI::RHIAsyncReallocateTexture2D(FRHITexture2D* OldTextureRHI, int32 NewMipCount, int32 NewSizeX, int32 NewSizeY, FThreadSafeCounter* RequestStatus)
 {
 	@autoreleasepool {
-		FAGXTexture2D* OldTexture = ResourceCast(OldTextureRHI);
-		
-		FAGXTexture2D* NewTexture = new FAGXTexture2D(OldTexture->GetFormat(), NewSizeX, NewSizeY, NewMipCount, OldTexture->GetNumSamples(), OldTexture->GetFlags(), NULL, OldTextureRHI->GetClearBinding());
+
+		check(IsInRenderingThread());
+		FRHICommandListImmediate& RHICmdList = FRHICommandListExecutor::GetImmediateCommandList();
+
+		FAGXSurface* OldTexture = ResourceCast(OldTextureRHI);
+
+		FRHITextureDesc Desc = OldTexture->GetDesc();
+		Desc.Extent = FIntPoint(NewSizeX, NewSizeY);
+		Desc.NumMips = NewMipCount;
+
+		FRHITextureCreateDesc CreateDesc(
+			Desc,
+			RHIGetDefaultResourceState(Desc.Flags, false),
+			TEXT("RHIAsyncReallocateTexture2D")
+		);
 		
-		FAGXRHICommandAsyncReallocateTexture2D::CopyMips(ImmediateContext.GetInternalContext(), OldTexture, NewTexture, NewMipCount, NewSizeX, NewSizeY, RequestStatus);
+		FAGXSurface* NewTexture = new FAGXSurface(CreateDesc);
+
+		// Copy shared mips
+		RHICmdList.EnqueueLambda([this, OldTexture, NewSizeX, NewSizeY, NewTexture, RequestStatus](FRHICommandListImmediate& RHICmdList)
+		{
+			FAGXContext& Context = ImmediateContext.GetInternalContext();
+
+			// figure out what mips to schedule
+			const uint32 NumSharedMips = FMath::Min(OldTexture->GetNumMips(), NewTexture->GetNumMips());
+			const uint32 SourceMipOffset = OldTexture->GetNumMips() - NumSharedMips;
+			const uint32 DestMipOffset = NewTexture->GetNumMips() - NumSharedMips;
+
+			const uint32 BlockSizeX = GPixelFormats[OldTexture->GetFormat()].BlockSizeX;
+			const uint32 BlockSizeY = GPixelFormats[OldTexture->GetFormat()].BlockSizeY;
+
+			// only handling straight 2D textures here
+			uint32 SliceIndex = 0;
+			mtlpp::Origin Origin(0, 0, 0);
+
+			FAGXTexture Tex = OldTexture->Texture;
+
+			// DXT/BC formats on Mac actually do have mip-tails that are smaller than the block size, they end up being uncompressed.
+			bool const bPixelFormatASTC = IsPixelFormatASTCCompressed(OldTexture->GetFormat());
+
+			bool bAsync = true;
+			for (uint32 MipIndex = 0; MipIndex < NumSharedMips; ++MipIndex)
+			{
+				const uint32 UnalignedMipSizeX = FMath::Max<uint32>(1, NewSizeX >> (MipIndex + DestMipOffset));
+				const uint32 UnalignedMipSizeY = FMath::Max<uint32>(1, NewSizeY >> (MipIndex + DestMipOffset));
+				const uint32 MipSizeX = FMath::Max<uint32>(1, NewSizeX >> (MipIndex + DestMipOffset));
+				const uint32 MipSizeY = FMath::Max<uint32>(1, NewSizeY >> (MipIndex + DestMipOffset));
+
+				bAsync &= Context.AsyncCopyFromTextureToTexture(OldTexture->Texture, SliceIndex, MipIndex + SourceMipOffset, Origin, mtlpp::Size(MipSizeX, MipSizeY, 1), NewTexture->Texture, SliceIndex, MipIndex + DestMipOffset, Origin);
+			}
+
+			// when done, decrement the counter to indicate it's safe
+			mtlpp::CommandBufferHandler CompletionHandler = [Tex](mtlpp::CommandBuffer const&)
+			{
+			};
+
+			if (bAsync)
+			{
+				// kck it off!
+				Context.SubmitAsyncCommands(nil, CompletionHandler, false);
+			}
+			else
+			{
+				Context.GetCurrentRenderPass().AddCompletionHandler(CompletionHandler);
+			}
+
+			// Like D3D mark this as complete immediately.
+			RequestStatus->Decrement();
+
+			FAGXSurface* Source = AGXGetMetalSurfaceFromRHITexture(OldTexture);
+			Source->MakeAliasable();
+		});
 		
 		return NewTexture;
 	}
@@ -1827,9 +1571,9 @@ void* FAGXDynamicRHI::LockTexture2D_RenderThread(class FRHICommandListImmediate&
 	@autoreleasepool {
 		check(IsInRenderingThread());
 		
-		FAGXTexture2D* TextureMTL = ResourceCast(Texture);
+		FAGXSurface* TextureMTL = ResourceCast(Texture);
 		
-		void* BufferData = TextureMTL->Surface.AsyncLock(RHICmdList, MipIndex, 0, LockMode, DestStride, bNeedsDefaultRHIFlush);
+		void* BufferData = TextureMTL->AsyncLock(RHICmdList, MipIndex, 0, LockMode, DestStride, bNeedsDefaultRHIFlush);
 		
 		return BufferData;
 	}
@@ -1841,8 +1585,8 @@ void FAGXDynamicRHI::UnlockTexture2D_RenderThread(class FRHICommandListImmediate
 	{
 		check(IsInRenderingThread());
 		
-		FAGXTexture2D* TextureMTL = ResourceCast(Texture);
-		TextureMTL->Surface.Unlock(MipIndex, 0, true);
+		FAGXSurface* TextureMTL = ResourceCast(Texture);
+		TextureMTL->Unlock(MipIndex, 0, true);
 	}
 }
 
@@ -1850,32 +1594,32 @@ void FAGXDynamicRHI::UnlockTexture2D_RenderThread(class FRHICommandListImmediate
 void* FAGXDynamicRHI::RHILockTexture2D(FRHITexture2D* TextureRHI,uint32 MipIndex,EResourceLockMode LockMode,uint32& DestStride,bool bLockWithinMiptail)
 {
 	@autoreleasepool {
-		FAGXTexture2D* Texture = ResourceCast(TextureRHI);
-		return Texture->Surface.Lock(MipIndex, 0, LockMode, DestStride);
+		FAGXSurface* Texture = ResourceCast(TextureRHI);
+		return Texture->Lock(MipIndex, 0, LockMode, DestStride);
 	}
 }
 
 void FAGXDynamicRHI::RHIUnlockTexture2D(FRHITexture2D* TextureRHI,uint32 MipIndex,bool bLockWithinMiptail)
 {
 	@autoreleasepool {
-		FAGXTexture2D* Texture = ResourceCast(TextureRHI);
-		Texture->Surface.Unlock(MipIndex, 0, false);
+		FAGXSurface* Texture = ResourceCast(TextureRHI);
+		Texture->Unlock(MipIndex, 0, false);
 	}
 }
 
 void* FAGXDynamicRHI::RHILockTexture2DArray(FRHITexture2DArray* TextureRHI,uint32 TextureIndex,uint32 MipIndex,EResourceLockMode LockMode,uint32& DestStride,bool bLockWithinMiptail)
 {
 	@autoreleasepool {
-		FAGXTexture2DArray* Texture = ResourceCast(TextureRHI);
-		return Texture->Surface.Lock(MipIndex, TextureIndex, LockMode, DestStride);
+		FAGXSurface* Texture = ResourceCast(TextureRHI);
+		return Texture->Lock(MipIndex, TextureIndex, LockMode, DestStride);
 	}
 }
 
 void FAGXDynamicRHI::RHIUnlockTexture2DArray(FRHITexture2DArray* TextureRHI,uint32 TextureIndex,uint32 MipIndex,bool bLockWithinMiptail)
 {
 	@autoreleasepool {
-		FAGXTexture2DArray* Texture = ResourceCast(TextureRHI);
-		Texture->Surface.Unlock(MipIndex, TextureIndex, false);
+		FAGXSurface* Texture = ResourceCast(TextureRHI);
+		Texture->Unlock(MipIndex, TextureIndex, false);
 	}
 }
 
@@ -1903,13 +1647,15 @@ static FAGXBuffer InternalCopyTexture2DUpdateRegion(FRHITexture2D* TextureRHI, c
 
 	FAGXBuffer Buffer;
 
-	FAGXTexture2D* Texture = ResourceCast(TextureRHI);	
+	FAGXSurface* Texture = ResourceCast(TextureRHI);	
 
 #if PLATFORM_MAC
 	// Expand R8_sRGB into RGBA8_sRGB for Mac.
-	if(Texture->GetFormat() == PF_G8 && EnumHasAnyFlags(Texture->GetFlags(), TexCreate_SRGB))
+	if (   Texture->GetFormat() == PF_G8
+		&& EnumHasAnyFlags(Texture->GetFlags(), TexCreate_SRGB)
+		&& Texture->Texture.GetPixelFormat() == mtlpp::PixelFormat::RGBA8Unorm_sRGB)
 	{
-		const uint32 BufferSize = UpdateRegion.Height * UpdateRegion.Width * sizeof(uint32);		
+		const uint32 BufferSize = UpdateRegion.Height * UpdateRegion.Width * sizeof(uint32);
 		Buffer = GetAGXDeviceContext().CreatePooledBuffer(FAGXPooledBufferArgs(BufferSize, BUF_Dynamic, mtlpp::StorageMode::Shared));
 		InternalExpandR8ToStandardRGBA((uint32*)Buffer.GetContents(), UpdateRegion, InOutSourcePitch, SourceData);
 	}
@@ -1939,8 +1685,8 @@ static FAGXBuffer InternalCopyTexture2DUpdateRegion(FRHITexture2D* TextureRHI, c
 
 static void InternalUpdateTexture2D(FAGXContext& Context, FRHITexture2D* TextureRHI, uint32 MipIndex, FUpdateTextureRegion2D const& UpdateRegion, uint32 SourcePitch, FAGXBuffer Buffer)
 {
-	FAGXTexture2D* Texture = ResourceCast(TextureRHI);
-	FAGXTexture Tex = Texture->Surface.Texture;
+	FAGXSurface* Texture = ResourceCast(TextureRHI);
+	FAGXTexture Tex = Texture->Texture;
 	
 	mtlpp::Region Region = mtlpp::Region(UpdateRegion.DestX, UpdateRegion.DestY, UpdateRegion.Width, UpdateRegion.Height);
 	
@@ -1967,7 +1713,7 @@ static void InternalUpdateTexture2D(FAGXContext& Context, FRHITexture2D* Texture
 		MTLPP_VALIDATE(mtlpp::Texture, Tex, AGXSafeGetRuntimeDebuggingLevel() >= EAGXDebugLevelValidation, Replace(Region, MipIndex, 0, MTLPP_VALIDATE(mtlpp::Buffer, Buffer, AGXSafeGetRuntimeDebuggingLevel() >= EAGXDebugLevelValidation, GetContents()), SourcePitch, 0));
 	}
 
-	FPlatformAtomics::InterlockedExchange(&Texture->Surface.Written, 1);
+	FPlatformAtomics::InterlockedExchange(&Texture->Written, 1);
 }
 
 struct FAGXRHICommandUpdateTexture2D final : public FRHICommand<FAGXRHICommandUpdateTexture2D>
@@ -2016,8 +1762,8 @@ void FAGXDynamicRHI::RHIUpdateTexture2D(FRHITexture2D* TextureRHI, uint32 MipInd
 {
 	@autoreleasepool 
 	{
-		FAGXTexture2D* Texture = ResourceCast(TextureRHI);
-		FAGXTexture Tex = Texture->Surface.Texture;
+		FAGXSurface* Texture = ResourceCast(TextureRHI);
+		FAGXTexture Tex = Texture->Texture;
 		bool const bUseIntermediateMetalBuffer = Tex.GetStorageMode() == mtlpp::StorageMode::Private;
 		
 		if(bUseIntermediateMetalBuffer)
@@ -2030,7 +1776,7 @@ void FAGXDynamicRHI::RHIUpdateTexture2D(FRHITexture2D* TextureRHI, uint32 MipInd
 		{
 #if PLATFORM_MAC
 			TArray<uint32> ExpandedData;
-			if(Texture->GetFormat() == PF_G8 && EnumHasAnyFlags(Texture->GetFlags(), TexCreate_SRGB))
+			if (Texture->GetFormat() == PF_G8 && EnumHasAnyFlags(Texture->GetFlags(), TexCreate_SRGB) && Tex.GetPixelFormat() == mtlpp::PixelFormat::RGBA8Unorm_sRGB)
 			{
 				ExpandedData.AddZeroed(UpdateRegion.Height * UpdateRegion.Width);
 				InternalExpandR8ToStandardRGBA((uint32*)ExpandedData.GetData(), UpdateRegion, SourcePitch, SourceData);
@@ -2040,7 +1786,7 @@ void FAGXDynamicRHI::RHIUpdateTexture2D(FRHITexture2D* TextureRHI, uint32 MipInd
 			mtlpp::Region Region = mtlpp::Region(UpdateRegion.DestX, UpdateRegion.DestY, UpdateRegion.Width, UpdateRegion.Height);
 			MTLPP_VALIDATE(mtlpp::Texture, Tex, AGXSafeGetRuntimeDebuggingLevel() >= EAGXDebugLevelValidation, Replace(Region, MipIndex, 0, SourceData, SourcePitch, 0));
 
-			FPlatformAtomics::InterlockedExchange(&Texture->Surface.Written, 1);
+			FPlatformAtomics::InterlockedExchange(&Texture->Written, 1);
 		}
 		
 		INC_DWORD_STAT_BY(STAT_AGXTextureMemUpdate, UpdateRegion.Height*SourcePitch);
@@ -2050,7 +1796,7 @@ void FAGXDynamicRHI::RHIUpdateTexture2D(FRHITexture2D* TextureRHI, uint32 MipInd
 static void InternalCopyTexture3DUpdateRegionData(FRHITexture3D* TextureRHI, const struct FUpdateTextureRegion3D& UpdateRegion, uint32 SourceRowPitch, uint32 SourceDepthPitch, const uint8* SourceData, uint8* DestData)
 {
 	// Perform safe line copy
-	FAGXTexture3D* Texture = ResourceCast(TextureRHI);	
+	FAGXSurface* Texture = ResourceCast(TextureRHI);	
 	
 	uint8* pDest = DestData;
 	uint32 CopyPitch = FMath::DivideAndRoundUp(UpdateRegion.Width, (uint32)GPixelFormats[Texture->GetFormat()].BlockSizeX) * GPixelFormats[Texture->GetFormat()].BlockBytes;
@@ -2072,8 +1818,8 @@ static void InternalCopyTexture3DUpdateRegionData(FRHITexture3D* TextureRHI, con
 
 static void InternalUpdateTexture3D(FAGXContext& Context, FRHITexture3D* TextureRHI, uint32 MipIndex, const FUpdateTextureRegion3D& UpdateRegion, uint32 SourceRowPitch, uint32 SourceDepthPitch, FAGXBuffer Buffer)
 {
-	FAGXTexture3D* Texture = ResourceCast(TextureRHI);
-	FAGXTexture Tex = Texture->Surface.Texture;
+	FAGXSurface* Texture = ResourceCast(TextureRHI);
+	FAGXTexture Tex = Texture->Texture;
 	
 	mtlpp::Region Region = mtlpp::Region(UpdateRegion.DestX, UpdateRegion.DestY, UpdateRegion.DestZ, UpdateRegion.Width, UpdateRegion.Height, UpdateRegion.Depth);
 	
@@ -2097,7 +1843,7 @@ static void InternalUpdateTexture3D(FAGXContext& Context, FRHITexture3D* Texture
 		MTLPP_VALIDATE(mtlpp::Texture, Tex, AGXSafeGetRuntimeDebuggingLevel() >= EAGXDebugLevelValidation, Replace(Region, MipIndex, 0, (uint8*)Buffer.GetContents(), SourceRowPitch, SourceDepthPitch));
 	}
 
-	FPlatformAtomics::InterlockedExchange(&Texture->Surface.Written, 1);
+	FPlatformAtomics::InterlockedExchange(&Texture->Written, 1);
 }
 
 struct FAGXDynamicRHIUpdateTexture3DCommand final : public FRHICommand<FAGXDynamicRHIUpdateTexture3DCommand>
@@ -2118,8 +1864,8 @@ struct FAGXDynamicRHIUpdateTexture3DCommand final : public FRHICommand<FAGXDynam
 	, SourceRowPitch(InSourceRowPitch)
 	, SourceDepthPitch(InSourceDepthPitch)
 	{
-		FAGXTexture3D* Texture = ResourceCast(TextureRHI);
-		FAGXTexture Tex = Texture->Surface.Texture;
+		FAGXSurface* Texture = ResourceCast(TextureRHI);
+		FAGXTexture Tex = Texture->Texture;
 		const uint32 BufferSize = UpdateRegion.Height * UpdateRegion.Depth* SourceRowPitch;
 		
 		Buffer = GetAGXDeviceContext().CreatePooledBuffer(FAGXPooledBufferArgs(BufferSize, BUF_Dynamic, mtlpp::StorageMode::Shared));
@@ -2182,11 +1928,11 @@ void FAGXDynamicRHI::RHIUpdateTexture3D(FRHITexture3D* TextureRHI,uint32 MipInde
 {
 	@autoreleasepool {
 	
-		FAGXTexture3D* Texture = ResourceCast(TextureRHI);
-		FAGXTexture Tex = Texture->Surface.Texture;
+		FAGXSurface* Texture = ResourceCast(TextureRHI);
+		FAGXTexture Tex = Texture->Texture;
 		
 #if PLATFORM_MAC
-		checkf(!(Texture->GetFormat() == PF_G8 && EnumHasAnyFlags(Texture->GetFlags(), TexCreate_SRGB)), TEXT("AGXRHI does not support PF_G8_sRGB on 3D, array or cube textures as it requires manual, CPU-side expansion to RGBA8_sRGB which is expensive!"));
+		checkf(!(Texture->GetFormat() == PF_G8 && EnumHasAnyFlags(Texture->GetFlags(), TexCreate_SRGB) && Tex.GetPixelFormat() == mtlpp::PixelFormat::RGBA8Unorm_sRGB), TEXT("AGXRHI does not support PF_G8_sRGB on 3D, array or cube textures as it requires manual, CPU-side expansion to RGBA8_sRGB which is expensive!"));
 #endif
 		if(Tex.GetStorageMode() == mtlpp::StorageMode::Private)
 		{
@@ -2202,7 +1948,7 @@ void FAGXDynamicRHI::RHIUpdateTexture3D(FRHITexture3D* TextureRHI,uint32 MipInde
 		{
 			mtlpp::Region Region = mtlpp::Region(UpdateRegion.DestX, UpdateRegion.DestY, UpdateRegion.DestZ, UpdateRegion.Width, UpdateRegion.Height, UpdateRegion.Depth);
 			MTLPP_VALIDATE(mtlpp::Texture, Tex, AGXSafeGetRuntimeDebuggingLevel() >= EAGXDebugLevelValidation, Replace(Region, MipIndex, 0, SourceData, SourceRowPitch, SourceDepthPitch));
-			FPlatformAtomics::InterlockedExchange(&Texture->Surface.Written, 1);
+			FPlatformAtomics::InterlockedExchange(&Texture->Written, 1);
 		}
 		
 		INC_DWORD_STAT_BY(STAT_AGXTextureMemUpdate, UpdateRegion.Height * UpdateRegion.Width * SourceDepthPitch);
@@ -2212,139 +1958,24 @@ void FAGXDynamicRHI::RHIUpdateTexture3D(FRHITexture3D* TextureRHI,uint32 MipInde
 /*-----------------------------------------------------------------------------
  Cubemap texture support.
  -----------------------------------------------------------------------------*/
-FTextureCubeRHIRef FAGXDynamicRHI::RHICreateTextureCube(uint32 Size, uint8 Format, uint32 NumMips, ETextureCreateFlags Flags, ERHIAccess InResourceState, FRHIResourceCreateInfo& CreateInfo)
-{
-	@autoreleasepool {
-		return new FAGXTextureCube((EPixelFormat)Format, Size, false, 1, NumMips, Flags, CreateInfo.BulkData, CreateInfo.ClearValueBinding);
-	}
-}
-
-FTextureCubeRHIRef FAGXDynamicRHI::RHICreateTextureCubeArray(uint32 Size, uint32 ArraySize, uint8 Format, uint32 NumMips, ETextureCreateFlags Flags, ERHIAccess InResourceState, FRHIResourceCreateInfo& CreateInfo)
-{
-	@autoreleasepool {
-		return new FAGXTextureCube((EPixelFormat)Format, Size, true, ArraySize, NumMips, Flags, CreateInfo.BulkData, CreateInfo.ClearValueBinding);
-	}
-}
-
 void* FAGXDynamicRHI::RHILockTextureCubeFace(FRHITextureCube* TextureCubeRHI,uint32 FaceIndex,uint32 ArrayIndex,uint32 MipIndex,EResourceLockMode LockMode,uint32& DestStride,bool bLockWithinMiptail)
 {
 	@autoreleasepool {
-		FAGXTextureCube* TextureCube = ResourceCast(TextureCubeRHI);
+		FAGXSurface* TextureCube = ResourceCast(TextureCubeRHI);
 		uint32 MetalFace = GetMetalCubeFace((ECubeFace)FaceIndex);
-		return TextureCube->Surface.Lock(MipIndex, MetalFace + (6 * ArrayIndex), LockMode, DestStride, true);
+		return TextureCube->Lock(MipIndex, MetalFace + (6 * ArrayIndex), LockMode, DestStride, true);
 	}
 }
 
 void FAGXDynamicRHI::RHIUnlockTextureCubeFace(FRHITextureCube* TextureCubeRHI,uint32 FaceIndex,uint32 ArrayIndex,uint32 MipIndex,bool bLockWithinMiptail)
 {
 	@autoreleasepool {
-		FAGXTextureCube* TextureCube = ResourceCast(TextureCubeRHI);
+		FAGXSurface* TextureCube = ResourceCast(TextureCubeRHI);
 		uint32 MetalFace = GetMetalCubeFace((ECubeFace)FaceIndex);
-		TextureCube->Surface.Unlock(MipIndex, MetalFace + (ArrayIndex * 6), false);
-	}
-}
-
-
-FTexture2DRHIRef FAGXDynamicRHI::RHICreateTexture2D_RenderThread(class FRHICommandListImmediate& RHICmdList, uint32 SizeX, uint32 SizeY, uint8 Format, uint32 NumMips, uint32 NumSamples, ETextureCreateFlags Flags, ERHIAccess InResourceState, FRHIResourceCreateInfo& CreateInfo)
-{
-	@autoreleasepool {
-		FRHIResourceCreateInfo Info = CreateInfo;
-		bool const bIOSurfaceData = (CreateInfo.BulkData && CreateInfo.BulkData->GetResourceType() != FResourceBulkDataInterface::EBulkDataType::Default);
-		Info.BulkData = bIOSurfaceData ? CreateInfo.BulkData : nullptr;
-		FTexture2DRHIRef Result = GDynamicRHI->RHICreateTexture2D(SizeX, SizeY, Format, NumMips, NumSamples, Flags, InResourceState, Info);
-		if (!bIOSurfaceData && CreateInfo.BulkData)
-		{
-			// upload existing bulkdata
-			UE_LOG(LogAGX, Display, TEXT("Got a bulk data texture, with %d mips"), NumMips);
-			checkf(NumMips == 1, TEXT("Only handling bulk data with 1 mip and 1 array length"));
-			uint32 Stride;
-			
-			// lock, copy, unlock
-			void* LockedData = LockTexture2D_RenderThread(RHICmdList, Result.GetReference(), 0, RLM_WriteOnly, Stride, false);
-			check(LockedData);
-			FMemory::Memcpy(LockedData, CreateInfo.BulkData->GetResourceBulkData(), CreateInfo.BulkData->GetResourceBulkDataSize());
-			UnlockTexture2D_RenderThread(RHICmdList, Result.GetReference(), 0, false);
-			
-			// bulk data can be unloaded now
-			CreateInfo.BulkData->Discard();
-		}
-		return Result;
+		TextureCube->Unlock(MipIndex, MetalFace + (ArrayIndex * 6), false);
 	}
 }
 
-FTexture2DArrayRHIRef FAGXDynamicRHI::RHICreateTexture2DArray_RenderThread(class FRHICommandListImmediate& RHICmdList, uint32 SizeX, uint32 SizeY, uint32 SizeZ, uint8 Format, uint32 NumMips, uint32 NumSamples, ETextureCreateFlags Flags, ERHIAccess InResourceState, FRHIResourceCreateInfo& CreateInfo)
-{
-	@autoreleasepool {
-		FRHIResourceCreateInfo Info = CreateInfo;
-		bool const bIOSurfaceData = (CreateInfo.BulkData && CreateInfo.BulkData->GetResourceType() != FResourceBulkDataInterface::EBulkDataType::Default);
-		Info.BulkData = bIOSurfaceData ? CreateInfo.BulkData : nullptr;
-		FTexture2DArrayRHIRef Result = GDynamicRHI->RHICreateTexture2DArray(SizeX, SizeY, SizeZ, Format, NumMips, NumSamples, Flags, InResourceState, Info);
-		if (CreateInfo.BulkData)
-		{
-			// upload existing bulkdata
-			UE_LOG(LogAGX, Display, TEXT("Got a bulk data texture, with %d mips"), NumMips);
-			checkf(NumMips == 1 && SizeZ == 1, TEXT("Only handling bulk data with 1 mip and 1 array length"));
-			uint32 Stride;
-			
-			// lock, copy, unlock
-			FAGXTexture2DArray* TextureMTL = ResourceCast(Result.GetReference());
-			void* LockedData = TextureMTL->Surface.AsyncLock(RHICmdList, 0, 0, RLM_WriteOnly, Stride, false);
-			FMemory::Memcpy(LockedData, CreateInfo.BulkData->GetResourceBulkData(), CreateInfo.BulkData->GetResourceBulkDataSize());
-			TextureMTL->Surface.Unlock(0, 0, true);
-			
-			// bulk data can be unloaded now
-			CreateInfo.BulkData->Discard();
-		}
-		return Result;
-	}
-}
-
-FTexture3DRHIRef FAGXDynamicRHI::RHICreateTexture3D_RenderThread(class FRHICommandListImmediate& RHICmdList, uint32 SizeX, uint32 SizeY, uint32 SizeZ, uint8 Format, uint32 NumMips, ETextureCreateFlags Flags, ERHIAccess InResourceState, FRHIResourceCreateInfo& CreateInfo)
-{
-	@autoreleasepool {
-		FRHIResourceCreateInfo Info = CreateInfo;
-		bool const bIOSurfaceData = (CreateInfo.BulkData && CreateInfo.BulkData->GetResourceType() != FResourceBulkDataInterface::EBulkDataType::Default);
-		Info.BulkData = bIOSurfaceData ? CreateInfo.BulkData : nullptr;
-		FTexture3DRHIRef Result = GDynamicRHI->RHICreateTexture3D(SizeX, SizeY, SizeZ, Format, NumMips, Flags, InResourceState, Info);
-		if (CreateInfo.BulkData)
-		{
-			// upload existing bulkdata
-			UE_LOG(LogAGX, Display, TEXT("Got a bulk data texture, with %d mips"), NumMips);
-			checkf(NumMips == 1, TEXT("Only handling bulk data with 1 mip and 1 array length"));
-			uint32 Stride;
-			
-			// lock, copy, unlock
-			FAGXTexture3D* TextureMTL = ResourceCast(Result.GetReference());
-			void* LockedData = TextureMTL->Surface.AsyncLock(RHICmdList, 0, 0, RLM_WriteOnly, Stride, false);
-			FMemory::Memcpy(LockedData, CreateInfo.BulkData->GetResourceBulkData(), CreateInfo.BulkData->GetResourceBulkDataSize());
-			TextureMTL->Surface.Unlock(0, 0, true);
-			
-			// bulk data can be unloaded now
-			CreateInfo.BulkData->Discard();
-		}
-		return Result;
-	}
-}
-
-FTextureCubeRHIRef FAGXDynamicRHI::RHICreateTextureCube_RenderThread(class FRHICommandListImmediate& RHICmdList, uint32 Size, uint8 Format, uint32 NumMips, ETextureCreateFlags Flags, ERHIAccess InResourceState, FRHIResourceCreateInfo& CreateInfo)
-{
-	@autoreleasepool {
-		check(!CreateInfo.BulkData);
-		FTextureCubeRHIRef Result = GDynamicRHI->RHICreateTextureCube(Size, Format, NumMips, Flags, InResourceState, CreateInfo);
-		return Result;
-	}
-}
-
-FTextureCubeRHIRef FAGXDynamicRHI::RHICreateTextureCubeArray_RenderThread(class FRHICommandListImmediate& RHICmdList, uint32 Size, uint32 ArraySize, uint8 Format, uint32 NumMips, ETextureCreateFlags Flags, ERHIAccess InResourceState, FRHIResourceCreateInfo& CreateInfo)
-{
-	@autoreleasepool {
-		check(!CreateInfo.BulkData);
-		FTextureCubeRHIRef Result = GDynamicRHI->RHICreateTextureCubeArray(Size, ArraySize, Format, NumMips, Flags, InResourceState, CreateInfo);
-		return Result;
-	}
-}
-
-
 void FAGXDynamicRHI::RHIBindDebugLabelName(FRHITexture* TextureRHI, const TCHAR* Name)
 {
 	@autoreleasepool {
@@ -2403,7 +2034,7 @@ void FAGXRHICommandContext::RHICopyTexture(FRHITexture* SourceTextureRHI, FRHITe
 		
 		if (TextureFormatExactMatch || TextureFormatCompatible)
 		{
-			FIntVector Size = (CopyInfo.Size != FIntVector::ZeroValue) ? CopyInfo.Size : FIntVector(MetalSrcTexture->SizeX, MetalSrcTexture->SizeY, MetalSrcTexture->SizeZ);
+			FIntVector Size = (CopyInfo.Size != FIntVector::ZeroValue) ? CopyInfo.Size : FIntVector(MetalSrcTexture->GetDesc().Extent.X, MetalSrcTexture->GetDesc().Extent.Y, MetalSrcTexture->GetDesc().Depth);
 			
 			mtlpp::Origin SourceOrigin(CopyInfo.SourcePosition.X, CopyInfo.SourcePosition.Y, CopyInfo.SourcePosition.Z);
 			mtlpp::Origin DestinationOrigin(CopyInfo.DestPosition.X, CopyInfo.DestPosition.Y, CopyInfo.DestPosition.Z);
@@ -2415,7 +2046,7 @@ void FAGXRHICommandContext::RHICopyTexture(FRHITexture* SourceTextureRHI, FRHITe
 				mtlpp::TextureUsage Usage = MetalSrcTexture->Texture.GetUsage();
 				if (Usage & mtlpp::TextureUsage::PixelFormatView)
 				{
-					ns::Range Slices(0, MetalSrcTexture->Texture.GetArrayLength() * (MetalSrcTexture->bIsCubemap ? 6 : 1));
+					ns::Range Slices(0, MetalSrcTexture->Texture.GetArrayLength() * (MetalSrcTexture->GetDesc().IsTextureCube() ? 6 : 1));
 					if (MetalSrcTexture->Texture.GetPixelFormat() != MetalDestTexture->Texture.GetPixelFormat())
 					{
 						SrcTexture = MetalSrcTexture->Texture.NewTextureView(MetalDestTexture->Texture.GetPixelFormat(), MetalSrcTexture->Texture.GetTextureType(), ns::Range(0, MetalSrcTexture->Texture.GetMipmapLevelCount()), Slices);
@@ -2441,8 +2072,8 @@ void FAGXRHICommandContext::RHICopyTexture(FRHITexture* SourceTextureRHI, FRHITe
 
 					if (TextureFormatCompatible)
 					{
-						DestSize.width  *= GPixelFormats[MetalDestTexture->PixelFormat].BlockSizeX;
-						DestSize.height *= GPixelFormats[MetalDestTexture->PixelFormat].BlockSizeY;
+						DestSize.width  *= GPixelFormats[MetalDestTexture->GetDesc().Format].BlockSizeX;
+						DestSize.height *= GPixelFormats[MetalDestTexture->GetDesc().Format].BlockSizeY;
 					}
 
 					// Account for create with TexCreate_SRGB flag which could make these different
@@ -2457,8 +2088,8 @@ void FAGXRHICommandContext::RHICopyTexture(FRHITexture* SourceTextureRHI, FRHITe
 						// format mismatch (like linear vs. sRGB), then we must
 						// achieve the copy by going through a buffer object.
 						//
-						const bool BlockSizeMatch = (GPixelFormats[MetalSrcTexture->PixelFormat].BlockSizeX == GPixelFormats[MetalDestTexture->PixelFormat].BlockSizeX);
-						const uint32 BytesPerPixel = (MetalSrcTexture->PixelFormat != PF_DepthStencil) ? GPixelFormats[MetalSrcTexture->PixelFormat].BlockBytes : 1;
+						const bool BlockSizeMatch = (GPixelFormats[MetalSrcTexture->GetDesc().Format].BlockSizeX == GPixelFormats[MetalDestTexture->GetDesc().Format].BlockSizeX);
+						const uint32 BytesPerPixel = (MetalSrcTexture->GetDesc().Format != PF_DepthStencil) ? GPixelFormats[MetalSrcTexture->GetDesc().Format].BlockBytes : 1;
 						const uint32 Stride = BytesPerPixel * SourceSize.width;
 #if PLATFORM_MAC
 						const uint32 Alignment = 1u;
