Copying //UE5/Dev-ParallelRendering to //UE5/Main

diff --git a/Engine/Source/Runtime/Apple/AGXRHI/Private/AGXUAV.cpp b/Engine/Source/Runtime/Apple/AGXRHI/Private/AGXUAV.cpp
--- a/Engine/Source/Runtime/Apple/AGXRHI/Private/AGXUAV.cpp
+++ b/Engine/Source/Runtime/Apple/AGXRHI/Private/AGXUAV.cpp
@@ -7,437 +7,430 @@
 #include "ClearReplacementShaders.h"
 #include "AGXTransitionData.h"
 
-FAGXShaderResourceView::FAGXShaderResourceView()
-	: TextureView(nullptr)
-	, Offset(0)
-	, MipLevel(0)
-	, bSRGBForceDisable(0)
-	, Reserved(0)
-	, NumMips(0)
-	, Format(0)
-	, Stride(0)
-	, LinearTextureDesc(nullptr)
-{
-	// void
-}
+// Constructor for buffers
+FAGXResourceViewBase::FAGXResourceViewBase(FRHIBuffer* InBuffer, uint32 InStartOffsetBytes, uint32 InNumElements, EPixelFormat InFormat)
+	: SourceBuffer     (ResourceCast(InBuffer))
+	, bTexture         (false)
+	, bSRGBForceDisable(false)
+	, MipLevel         (0)
+	, Reserved         (0)
+	, NumMips          (0)
+	, Format           (InFormat)
+	, Stride           (0)
+	, Offset           (InBuffer ? InStartOffsetBytes : 0)
+{
+	check(!bTexture);
+
+	if (SourceBuffer)
+	{
+		SourceBuffer->AddRef();
+	}
 
-FAGXShaderResourceView::~FAGXShaderResourceView()
-{
-	if (LinearTextureDesc)
+	EBufferUsageFlags Usage = SourceBuffer->GetUsage();
+	if (EnumHasAnyFlags(Usage, BUF_VertexBuffer))
 	{
-		delete LinearTextureDesc;
-		LinearTextureDesc = nullptr;
+		if (!SourceBuffer)
+		{
+			Stride = 0;
+		}
+		else
+		{
+			check(SourceBuffer->GetUsage() & BUF_ShaderResource);
+			Stride = GPixelFormats[Format].BlockBytes;
+
+			LinearTextureDesc = MakeUnique<FAGXLinearTextureDescriptor>(InStartOffsetBytes, InNumElements, Stride);
+			SourceBuffer->CreateLinearTexture((EPixelFormat)Format, SourceBuffer, LinearTextureDesc.Get());
+		}
 	}
-	
-	if (TextureView)
+	else if (EnumHasAnyFlags(Usage, BUF_IndexBuffer))
 	{
-		FAGXSurface* Surface = AGXGetMetalSurfaceFromRHITexture(SourceTexture);
-		if (Surface)
+		if (!SourceBuffer)
 		{
-			Surface->SRVs.Remove(this);
+			Format = PF_R16_UINT;
+			Stride = 0;
 		}
+		else
+		{
+			Format = (SourceBuffer->IndexType == mtlpp::IndexType::UInt16) ? PF_R16_UINT : PF_R32_UINT;
+			Stride = SourceBuffer->GetStride();
+
+			check(Stride == ((Format == PF_R16_UINT) ? 2 : 4));
 
-		delete TextureView;
-		TextureView = nullptr;
+			LinearTextureDesc = MakeUnique<FAGXLinearTextureDescriptor>(InStartOffsetBytes, InNumElements, Stride);
+			SourceBuffer->CreateLinearTexture((EPixelFormat)Format, SourceBuffer, LinearTextureDesc.Get());
+		}
 	}
-	
-	SourceVertexBuffer = NULL;
-	SourceTexture = NULL;
-}
+	else
+	{
+		check(EnumHasAnyFlags(Usage, BUF_StructuredBuffer));
+
+		Format = PF_Unknown;
+		Stride = SourceBuffer->GetStride();
+	}
+}
+
+// Constructor for textures
+FAGXResourceViewBase::FAGXResourceViewBase(
+	  FRHITexture* InTexture
+	, EPixelFormat InFormat
+	, uint8 InMipLevel
+	, uint8 InNumMipLevels
+	, ERHITextureSRVOverrideSRGBType InSRGBOverride
+	, uint32 InFirstArraySlice
+	, uint32 InNumArraySlices
+	)
+	: SourceTexture    (ResourceCast(InTexture))
+	, bTexture         (true)
+	, bSRGBForceDisable(InSRGBOverride == SRGBO_ForceDisable)
+	, MipLevel         (InMipLevel)
+	, Reserved         (0)
+	, NumMips          (InNumMipLevels)
+	, Format           ((InTexture && InFormat == PF_Unknown) ? InTexture->GetDesc().Format : InFormat)
+	, Stride           (0)
+	, Offset           (0)
+{
+	// Slice selection of a texture array still need to be implemented on Metal
+	check(InFirstArraySlice == 0 && InNumArraySlices == 0);
+
+	if (SourceTexture)
+	{
+		SourceTexture->AddRef();
 
-void FAGXShaderResourceView::InitLinearTextureDescriptor(const FAGXLinearTextureDescriptor& InLinearTextureDescriptor)
-{
-	check(!LinearTextureDesc);
-	LinearTextureDesc = new FAGXLinearTextureDescriptor(InLinearTextureDescriptor);
-	check(LinearTextureDesc);
+		// TODO: Apple Silicon supports memoryless
+#if PLATFORM_IOS
+		// Memoryless targets can't have texture views (SRVs or UAVs)
+		if (Source.Texture.GetStorageMode() != mtlpp::StorageMode::Memoryless)
+#endif
+		{
+			// Determine the appropriate metal format for the view.
+			// This format will be non-sRGB. We convert to sRGB below if required.
+			mtlpp::PixelFormat MetalFormat = (mtlpp::PixelFormat)GPixelFormats[Format].PlatformFormat;
+
+			if (Format == PF_X24_G8)
+			{
+				// Stencil buffer view of a depth texture
+				check(SourceTexture->GetDesc().Format == PF_DepthStencil);
+				switch (SourceTexture->Texture.GetPixelFormat())
+				{
+					default:
+						checkNoEntry();
+						break;
+
+					case mtlpp::PixelFormat::Depth24Unorm_Stencil8:
+						MetalFormat = mtlpp::PixelFormat::X24_Stencil8;
+						break;
+
+					case mtlpp::PixelFormat::Depth32Float_Stencil8:
+						MetalFormat = mtlpp::PixelFormat::X32_Stencil8;
+						break;
+				}
+			}
+			else
+			{
+				// Override the format's sRGB setting if appropriate
+				if (EnumHasAnyFlags(SourceTexture->GetDesc().Flags, TexCreate_SRGB))
+				{
+					if (bSRGBForceDisable)
+					{
+#if PLATFORM_MAC
+						// R8Unorm has been expanded in the source surface for sRGBA support - we need to expand to RGBA to enable compatible texture format view for non apple silicon macs
+						if (Format == PF_G8 && SourceTexture->Texture.GetPixelFormat() == mtlpp::PixelFormat::RGBA8Unorm_sRGB)
+						{
+							MetalFormat = mtlpp::PixelFormat::RGBA8Unorm;
+						}
+#endif
+					}
+					else
+					{
+						// Ensure we have the correct sRGB target format if we create a new texture view rather than using the source texture
+						MetalFormat = AGXToSRGBFormat(MetalFormat);
+					}
+				}
+			}
+
+			// We can use the source texture directly if the view's format / mip count etc matches.
+			bool bUseSourceTex =
+				   MipLevel == 0
+				&& NumMips == SourceTexture->Texture.GetMipmapLevelCount()
+				&& MetalFormat == SourceTexture->Texture.GetPixelFormat();
+
+			if (bUseSourceTex)
+			{
+				// SRV is exactly compatible with the original texture.
+				TextureView = SourceTexture->Texture;
+			}
+			else
+			{
+				// Recreate the texture to enable MTLTextureUsagePixelFormatView which must be off unless we definitely use this feature or we are throwing ~4% performance vs. Windows on the floor.
+				// @todo recreating resources like this will likely prevent us from making view creation multi-threaded.
+				if (!(SourceTexture->Texture.GetUsage() & mtlpp::TextureUsage::PixelFormatView))
+					SourceTexture->PrepareTextureView();
+
+				ns::Range Slices(0, SourceTexture->Texture.GetArrayLength() * (SourceTexture->GetDesc().IsTextureCube() ? 6 : 1));
+
+				TextureView = SourceTexture->Texture.NewTextureView(
+					MetalFormat,
+					SourceTexture->Texture.GetTextureType(),
+					ns::Range(MipLevel, NumMips),
+					Slices);
+			}
+		}
+	}
+	else
+	{
+		TextureView = {};
+	}
 }
 
-ns::AutoReleased<FAGXTexture> FAGXShaderResourceView::GetLinearTexture(bool const bUAV)
+FAGXResourceViewBase::~FAGXResourceViewBase()
 {
-	ns::AutoReleased<FAGXTexture> NewLinearTexture;
+	if (bTexture)
 	{
-		if (IsValidRef(SourceVertexBuffer))
+		if (SourceTexture)
 		{
-			NewLinearTexture = SourceVertexBuffer->GetLinearTexture((EPixelFormat)Format, LinearTextureDesc);
+			SourceTexture->Release();
 		}
-		else if (IsValidRef(SourceIndexBuffer))
+	}
+	else
+	{
+		if (SourceBuffer)
 		{
-			NewLinearTexture = SourceIndexBuffer->GetLinearTexture((EPixelFormat)Format, LinearTextureDesc);
+			SourceBuffer->Release();
 		}
 	}
+}
+
+ns::AutoReleased<FAGXTexture> FAGXResourceViewBase::GetLinearTexture()
+{
+	ns::AutoReleased<FAGXTexture> NewLinearTexture;
+	if (SourceBuffer)
+	{
+		NewLinearTexture = SourceBuffer->GetLinearTexture((EPixelFormat)Format, LinearTextureDesc.Get());
+	}
 	return NewLinearTexture;
 }
 
 FUnorderedAccessViewRHIRef FAGXDynamicRHI::RHICreateUnorderedAccessView_RenderThread(class FRHICommandListImmediate& RHICmdList, FRHIBuffer* Buffer, bool bUseUAVCounter, bool bAppendBuffer)
 {
-	return GDynamicRHI->RHICreateUnorderedAccessView(Buffer, bUseUAVCounter, bAppendBuffer);
+	return this->RHICreateUnorderedAccessView(Buffer, bUseUAVCounter, bAppendBuffer);
 }
 
 FUnorderedAccessViewRHIRef FAGXDynamicRHI::RHICreateUnorderedAccessView_RenderThread(class FRHICommandListImmediate& RHICmdList, FRHITexture* Texture, uint32 MipLevel, uint16 FirstArraySlice, uint16 NumArraySlices)
 {
-	FAGXSurface* Surface = (FAGXSurface*)Texture->GetTextureBaseRHI();
-	FAGXTexture Tex = Surface->Texture;
-	if (!(Tex.GetUsage() & mtlpp::TextureUsage::PixelFormatView))
+	FAGXSurface* Surface = AGXGetMetalSurfaceFromRHITexture(Texture);
+
+	//
+	// The FAGXResourceViewBase constructor for textures currently modifies the
+	// underlying texture object via FMetalSurface::PrepareTextureView() to add
+	// PixelFormatView support if it was not already created with it.
+	//
+	// Because of this, the following RHI thread stall is necessary. We will need
+	// to clean this up in future before RHI functions can be completely thread-
+	// safe.
+	//
+
+	if (!(Surface->Texture.GetUsage() & mtlpp::TextureUsage::PixelFormatView))
 	{
 		FScopedRHIThreadStaller StallRHIThread(RHICmdList);
-		return GDynamicRHI->RHICreateUnorderedAccessView(Texture, MipLevel, FirstArraySlice, NumArraySlices);
+		return this->RHICreateUnorderedAccessView(Texture, MipLevel, FirstArraySlice, NumArraySlices);
 	}
 	else
 	{
-		return GDynamicRHI->RHICreateUnorderedAccessView(Texture, MipLevel, FirstArraySlice, NumArraySlices);
+		return this->RHICreateUnorderedAccessView(Texture, MipLevel, FirstArraySlice, NumArraySlices);
 	}
 }
 
 FUnorderedAccessViewRHIRef FAGXDynamicRHI::RHICreateUnorderedAccessView_RenderThread(class FRHICommandListImmediate& RHICmdList, FRHIBuffer* Buffer, uint8 Format)
 {
-	FUnorderedAccessViewRHIRef Result = GDynamicRHI->RHICreateUnorderedAccessView(Buffer, Format);
-	if (IsRunningRHIInSeparateThread() && !RHICmdList.Bypass()) 	{ 		RHICmdList.RHIThreadFence(true); 	}
+	FUnorderedAccessViewRHIRef Result = this->RHICreateUnorderedAccessView(Buffer, Format);
+	if (IsRunningRHIInSeparateThread() && !RHICmdList.Bypass())
+	{
+		RHICmdList.RHIThreadFence(true);
+	}
 	return Result;
 }
 
 FUnorderedAccessViewRHIRef FAGXDynamicRHI::RHICreateUnorderedAccessView(FRHIBuffer* BufferRHI, bool bUseUAVCounter, bool bAppendBuffer)
 {
 	@autoreleasepool {
-	FAGXStructuredBuffer* StructuredBuffer = ResourceCast(BufferRHI);
-	
-	FAGXShaderResourceView* SRV = new FAGXShaderResourceView;
-	SRV->SourceVertexBuffer = nullptr;
-	SRV->SourceIndexBuffer = nullptr;
-	SRV->TextureView = nullptr;
-	SRV->SourceStructuredBuffer = StructuredBuffer;
-
-	// create the UAV buffer to point to the structured buffer's memory
-	FAGXUnorderedAccessView* UAV = new FAGXUnorderedAccessView;
-	UAV->SourceView = SRV;
-	return UAV;
+		return new FAGXUnorderedAccessView(BufferRHI, bUseUAVCounter, bAppendBuffer);
 	}
 }
 
 FUnorderedAccessViewRHIRef FAGXDynamicRHI::RHICreateUnorderedAccessView(FRHITexture* TextureRHI, uint32 MipLevel, uint16 FirstArraySlice, uint16 NumArraySlices)
 {
-	// Slice selection of a texture array still need to be implemented on AGX
-	check(FirstArraySlice == 0 && NumArraySlices == 0);
 	@autoreleasepool {
-	FAGXShaderResourceView* SRV = new FAGXShaderResourceView;
-	SRV->SourceTexture = TextureRHI;
-	
-	FAGXSurface* Surface = AGXGetMetalSurfaceFromRHITexture(TextureRHI);
-	SRV->TextureView = Surface ? new FAGXSurface(*Surface, NSMakeRange(MipLevel, 1)) : nullptr;
-	
-	SRV->SourceVertexBuffer = nullptr;
-	SRV->SourceIndexBuffer = nullptr;
-	SRV->SourceStructuredBuffer = nullptr;
-	
-	SRV->MipLevel = MipLevel;
-	SRV->NumMips = 1;
-	SRV->Format = PF_Unknown;
-		
-	if (Surface)
-	{
-		Surface->SRVs.Add(SRV);
-	}
-		
-	// create the UAV buffer to point to the structured buffer's memory
-	FAGXUnorderedAccessView* UAV = new FAGXUnorderedAccessView;
-	UAV->SourceView = SRV;
-
-	return UAV;
+		return new FAGXUnorderedAccessView(TextureRHI, MipLevel, FirstArraySlice, NumArraySlices);
 	}
 }
 
 FUnorderedAccessViewRHIRef FAGXDynamicRHI::RHICreateUnorderedAccessView(FRHIBuffer* BufferRHI, uint8 Format)
 {
 	@autoreleasepool {
-	FAGXResourceMultiBuffer* Buffer = ResourceCast(BufferRHI);
-	
-	FAGXShaderResourceView* SRV = new FAGXShaderResourceView;
-	SRV->SourceVertexBuffer = EnumHasAnyFlags(BufferRHI->GetUsage(), BUF_VertexBuffer) ? Buffer : nullptr;
-	SRV->TextureView = nullptr;
-	SRV->SourceIndexBuffer = EnumHasAnyFlags(BufferRHI->GetUsage(), BUF_IndexBuffer) ? Buffer : nullptr;
-	SRV->SourceStructuredBuffer = nullptr;
-	SRV->Format = Format;
-	{
-		check(Buffer->GetUsage() & BUF_UnorderedAccess);
-		Buffer->CreateLinearTexture((EPixelFormat)Format, Buffer);
-	}
-		
-	// create the UAV buffer to point to the structured buffer's memory
-	FAGXUnorderedAccessView* UAV = new FAGXUnorderedAccessView;
-	UAV->SourceView = SRV;
-
-	return UAV;
+		return new FAGXUnorderedAccessView(BufferRHI, (EPixelFormat)Format);
 	}
 }
 
 FShaderResourceViewRHIRef FAGXDynamicRHI::CreateShaderResourceView_RenderThread(class FRHICommandListImmediate& RHICmdList, FRHIBuffer* Buffer, uint32 Stride, uint8 Format)
 {
-	FShaderResourceViewRHIRef Result = GDynamicRHI->RHICreateShaderResourceView(Buffer, Stride, Format);
-	if (IsRunningRHIInSeparateThread() && !RHICmdList.Bypass()) 	{ 		RHICmdList.RHIThreadFence(true); 	}
+	FShaderResourceViewRHIRef Result = this->RHICreateShaderResourceView(Buffer, Stride, Format);
+	if (IsRunningRHIInSeparateThread() && !RHICmdList.Bypass())
+	{
+		RHICmdList.RHIThreadFence(true);
+	}
 	return Result;
 }
 
 FShaderResourceViewRHIRef FAGXDynamicRHI::CreateShaderResourceView_RenderThread(class FRHICommandListImmediate& RHICmdList, const FShaderResourceViewInitializer& Initializer)
 {
-	FShaderResourceViewRHIRef Result = GDynamicRHI->RHICreateShaderResourceView(Initializer);
-	if (IsRunningRHIInSeparateThread() && !RHICmdList.Bypass()) { RHICmdList.RHIThreadFence(true); }
+	FShaderResourceViewRHIRef Result = this->RHICreateShaderResourceView(Initializer);
+	if (IsRunningRHIInSeparateThread() && !RHICmdList.Bypass())
+	{
+		RHICmdList.RHIThreadFence(true);
+	}
 	return Result;
 }
 
 FShaderResourceViewRHIRef FAGXDynamicRHI::CreateShaderResourceView_RenderThread(class FRHICommandListImmediate& RHICmdList, FRHIBuffer* Buffer)
 {
-	FShaderResourceViewRHIRef Result = GDynamicRHI->RHICreateShaderResourceView(Buffer);
-	if (IsRunningRHIInSeparateThread() && !RHICmdList.Bypass()) 	{ 		RHICmdList.RHIThreadFence(true); 	}
+	FShaderResourceViewRHIRef Result = this->RHICreateShaderResourceView(Buffer);
+	if (IsRunningRHIInSeparateThread() && !RHICmdList.Bypass())
+	{
+		RHICmdList.RHIThreadFence(true);
+	}
 	return Result;
 }
 
 FShaderResourceViewRHIRef FAGXDynamicRHI::RHICreateShaderResourceView_RenderThread(class FRHICommandListImmediate& RHICmdList, FRHIBuffer* Buffer, uint32 Stride, uint8 Format)
 {
-	FShaderResourceViewRHIRef Result = GDynamicRHI->RHICreateShaderResourceView(Buffer, Stride, Format);
-	if (IsRunningRHIInSeparateThread() && !RHICmdList.Bypass()) 	{ 		RHICmdList.RHIThreadFence(true); 	}
+	FShaderResourceViewRHIRef Result = this->RHICreateShaderResourceView(Buffer, Stride, Format);
+	if (IsRunningRHIInSeparateThread() && !RHICmdList.Bypass())
+	{
+		RHICmdList.RHIThreadFence(true);
+	}
 	return Result;
 }
 
 FShaderResourceViewRHIRef FAGXDynamicRHI::RHICreateShaderResourceView_RenderThread(class FRHICommandListImmediate& RHICmdList, const FShaderResourceViewInitializer& Initializer)
 {
-	FShaderResourceViewRHIRef Result = GDynamicRHI->RHICreateShaderResourceView(Initializer);
-	if (IsRunningRHIInSeparateThread() && !RHICmdList.Bypass()) 	{ 		RHICmdList.RHIThreadFence(true); 	}
+	FShaderResourceViewRHIRef Result = this->RHICreateShaderResourceView(Initializer);
+	if (IsRunningRHIInSeparateThread() && !RHICmdList.Bypass())
+	{
+		RHICmdList.RHIThreadFence(true);
+	}
 	return Result;
 }
 
 FShaderResourceViewRHIRef FAGXDynamicRHI::RHICreateShaderResourceView_RenderThread(class FRHICommandListImmediate& RHICmdList, FRHIBuffer* Buffer)
 {
-	FShaderResourceViewRHIRef Result = GDynamicRHI->RHICreateShaderResourceView(Buffer);
-	if (IsRunningRHIInSeparateThread() && !RHICmdList.Bypass()) 	{ 		RHICmdList.RHIThreadFence(true); 	}
+	FShaderResourceViewRHIRef Result = this->RHICreateShaderResourceView(Buffer);
+	if (IsRunningRHIInSeparateThread() && !RHICmdList.Bypass())
+	{
+		RHICmdList.RHIThreadFence(true);
+	}
 	return Result;
 }
 
 FShaderResourceViewRHIRef FAGXDynamicRHI::RHICreateShaderResourceView_RenderThread(class FRHICommandListImmediate& RHICmdList, FRHITexture* Texture2DRHI, const FRHITextureSRVCreateInfo& CreateInfo)
 {
-	FAGXSurface* Surface = (FAGXSurface*)Texture2DRHI->GetTextureBaseRHI();
-	FAGXTexture Tex = Surface->Texture;
-	if (!(Tex.GetUsage() & mtlpp::TextureUsage::PixelFormatView))
+	FAGXSurface* Surface = AGXGetMetalSurfaceFromRHITexture(Texture2DRHI);
+
+	//
+	// The FMetalResourceViewBase constructor for textures currently modifies the
+	// underlying texture object via FMetalSurface::PrepareTextureView() to add
+	// PixelFormatView support if it was not already created with it.
+	//
+	// Because of this, the following RHI thread stall is necessary. We will need
+	// to clean this up in future before RHI functions can be completely thread-
+	// safe.
+	//
+	
+	if (!(Surface->Texture.GetUsage() & mtlpp::TextureUsage::PixelFormatView))
 	{
 		FScopedRHIThreadStaller StallRHIThread(RHICmdList);
-		return GDynamicRHI->RHICreateShaderResourceView(Texture2DRHI, CreateInfo);
+		return this->RHICreateShaderResourceView(Texture2DRHI, CreateInfo);
 	}
 	else
 	{
-		return GDynamicRHI->RHICreateShaderResourceView(Texture2DRHI, CreateInfo);
+		return this->RHICreateShaderResourceView(Texture2DRHI, CreateInfo);
 	}
 }
 
 FShaderResourceViewRHIRef FAGXDynamicRHI::RHICreateShaderResourceView(FRHITexture* Texture2DRHI, const FRHITextureSRVCreateInfo& CreateInfo)
 {
 	@autoreleasepool {
-		FAGXShaderResourceView* SRV = new FAGXShaderResourceView;
-		SRV->SourceTexture = (FRHITexture*)Texture2DRHI;
-		
-		FAGXSurface* Surface = AGXGetMetalSurfaceFromRHITexture(Texture2DRHI);
-		
-		// Asking to make a SRV with PF_Unknown means to use the same format.
-		// This matches the behavior of the DX11 RHI.
-		EPixelFormat Format = (EPixelFormat) CreateInfo.Format;
-		if(Surface && Format == PF_Unknown)
-		{
-			Format = Surface->PixelFormat;
-		}
-		
-		const bool bSRGBForceDisable = (CreateInfo.SRGBOverride == SRGBO_ForceDisable);
-		
-		SRV->TextureView = Surface ? new FAGXSurface(*Surface, NSMakeRange(CreateInfo.MipLevel, CreateInfo.NumMipLevels), Format, bSRGBForceDisable) : nullptr;
-		
-		SRV->SourceVertexBuffer = nullptr;
-		SRV->SourceIndexBuffer = nullptr;
-		SRV->SourceStructuredBuffer = nullptr;
-		
-		SRV->MipLevel = CreateInfo.MipLevel;
-		SRV->bSRGBForceDisable = bSRGBForceDisable;
-		SRV->NumMips = CreateInfo.NumMipLevels;
-		SRV->Format = CreateInfo.Format;
-		
-		if (Surface)
-		{
-			Surface->SRVs.Add(SRV);
-		}
-		
-		return SRV;
+		return new FAGXShaderResourceView(Texture2DRHI, CreateInfo);
 	}
 }
 
 FShaderResourceViewRHIRef FAGXDynamicRHI::RHICreateShaderResourceView(FRHIBuffer* BufferRHI)
 {
 	@autoreleasepool {
-	return RHICreateShaderResourceView(FShaderResourceViewInitializer(BufferRHI));
+		return this->RHICreateShaderResourceView(FShaderResourceViewInitializer(BufferRHI));
 	}
 }
 
 FShaderResourceViewRHIRef FAGXDynamicRHI::RHICreateShaderResourceView(FRHIBuffer* BufferRHI, uint32 Stride, uint8 Format)
 {
-	check(GPixelFormats[Format].BlockBytes == Stride);
-
-	@autoreleasepool
-	{
-		return RHICreateShaderResourceView(FShaderResourceViewInitializer(BufferRHI, EPixelFormat(Format)));
+	@autoreleasepool {
+		check(GPixelFormats[Format].BlockBytes == Stride);
+		return this->RHICreateShaderResourceView(FShaderResourceViewInitializer(BufferRHI, EPixelFormat(Format)));
 	}
 }
 
 FShaderResourceViewRHIRef FAGXDynamicRHI::RHICreateShaderResourceView(const FShaderResourceViewInitializer& Initializer)
 {
-	@autoreleasepool
-	{
-		switch(Initializer.GetType())
-		{
-			case FShaderResourceViewInitializer::EType::VertexBufferSRV:
-			{
-				FShaderResourceViewInitializer::FVertexBufferShaderResourceViewInitializer Desc = Initializer.AsVertexBufferSRV();
-				
-				FAGXVertexBuffer* VertexBuffer = ResourceCast(Desc.Buffer);
-				
-				FAGXShaderResourceView* SRV = new FAGXShaderResourceView;
-				SRV->SourceVertexBuffer = VertexBuffer;
-				SRV->TextureView = nullptr;
-				SRV->SourceIndexBuffer = nullptr;
-				SRV->SourceStructuredBuffer = nullptr;
-				
-				SRV->Format = Desc.Format;
-				
-				if(!VertexBuffer)
-				{
-					SRV->Offset = 0;
-					SRV->Stride = 0;
-				}
-				else
-				{
-					check(VertexBuffer->GetUsage() & BUF_ShaderResource);
-					uint32 Stride = GPixelFormats[Desc.Format].BlockBytes;
-					
-					SRV->Stride = Stride;
-					SRV->Offset = Desc.StartOffsetBytes;
-					
-					FAGXLinearTextureDescriptor LinearTextureDesc(Desc.StartOffsetBytes, Desc.NumElements, Stride);
-					SRV->InitLinearTextureDescriptor(LinearTextureDesc);
-					
-					VertexBuffer->CreateLinearTexture((EPixelFormat)Desc.Format, VertexBuffer, &LinearTextureDesc);
-				}
-				
-				return SRV;
-			} // VertexBufferSRV
-				
-			case FShaderResourceViewInitializer::EType::StructuredBufferSRV:
-			{
-				FShaderResourceViewInitializer::FStructuredBufferShaderResourceViewInitializer Desc = Initializer.AsStructuredBufferSRV();
-				
-				FAGXStructuredBuffer* StructuredBuffer = ResourceCast(Desc.Buffer);
-
-				FAGXShaderResourceView* SRV = new FAGXShaderResourceView;
-				SRV->SourceVertexBuffer = nullptr;
-				SRV->SourceIndexBuffer = nullptr;
-				SRV->TextureView = nullptr;
-				SRV->SourceStructuredBuffer = StructuredBuffer;
-				
-				SRV->Offset = Desc.StartOffsetBytes;
-				SRV->Format = 0;
-				SRV->Stride = StructuredBuffer->GetStride();
-				
-				return SRV;
-			} // StructuredBufferSRV
-				
-			case FShaderResourceViewInitializer::EType::IndexBufferSRV:
-			{
-				FShaderResourceViewInitializer::FIndexBufferShaderResourceViewInitializer Desc = Initializer.AsIndexBufferSRV();
-				
-				FAGXIndexBuffer* IndexBuffer = ResourceCast(Desc.Buffer);
-				
-				FAGXShaderResourceView* SRV = new FAGXShaderResourceView;
-				SRV->SourceVertexBuffer = nullptr;
-				SRV->SourceIndexBuffer = IndexBuffer;
-				SRV->TextureView = nullptr;
-				
-				if(!IndexBuffer)
-				{
-					SRV->Format = PF_R16_UINT;
-					SRV->Stride = 0;
-					SRV->Offset = 0;
-				}
-				else
-				{
-					SRV->Format = (IndexBuffer->IndexType == mtlpp::IndexType::UInt16) ? PF_R16_UINT : PF_R32_UINT;
-					
-					const uint32 Stride = Desc.Buffer->GetStride();
-					check(Stride == ((SRV->Format == PF_R16_UINT) ? 2 : 4));
-					
-					SRV->Offset = Desc.StartOffsetBytes;
-					SRV->Stride = Stride;
-					
-					FAGXLinearTextureDescriptor LinearTextureDesc(Desc.StartOffsetBytes, Desc.NumElements, Stride);
-					SRV->InitLinearTextureDescriptor(LinearTextureDesc);
-					
-					IndexBuffer->CreateLinearTexture((EPixelFormat)SRV->Format, IndexBuffer, &LinearTextureDesc);
-				}
-					  
-				return SRV;
-			} // IndexBufferSRV
-					  
-			default:
-			{
-				checkNoEntry();
-				return nullptr;
-			}
-		}
+	@autoreleasepool {
+		return new FAGXShaderResourceView(Initializer);
 	}
 }
 
-void FAGXDynamicRHI::RHIUpdateShaderResourceView(FRHIShaderResourceView* SRVRHI, FRHIBuffer* VertexBufferRHI, uint32 Stride, uint8 Format)
+void FAGXDynamicRHI::RHIUpdateShaderResourceView(FRHIShaderResourceView* SRVRHI, FRHIBuffer* BufferRHI, uint32 Stride, uint8 Format)
 {
 	check(SRVRHI);
 	FAGXShaderResourceView* SRV = ResourceCast(SRVRHI);
-	if (!VertexBufferRHI)
+	check(!SRV->bTexture);
+
+	FAGXResourceMultiBuffer* OldBuffer = SRV->SourceBuffer;
+
+	SRV->SourceBuffer = ResourceCast(BufferRHI);
+	SRV->Stride = Stride;
+	SRV->Format = Format;
+
+	if (SRV->SourceBuffer)
 	{
-		SRV->SourceVertexBuffer = nullptr;
-		SRV->TextureView = nullptr;
-		SRV->SourceIndexBuffer = nullptr;
-		SRV->SourceStructuredBuffer = nullptr;
-		SRV->Offset = 0;
-		SRV->Format = Format;
-		SRV->Stride = Stride;
+		SRV->SourceBuffer->AddRef();
 	}
-	else if (SRV->SourceVertexBuffer != VertexBufferRHI)
+
+	if (OldBuffer)
 	{
-		FAGXVertexBuffer* VertexBuffer = ResourceCast(VertexBufferRHI);
-		SRV->SourceVertexBuffer = VertexBuffer;
-		SRV->TextureView = nullptr;
-		SRV->SourceIndexBuffer = nullptr;
-		SRV->SourceStructuredBuffer = nullptr;
-		SRV->Offset = 0;
-		SRV->Format = Format;
-		SRV->Stride = Stride;
+		OldBuffer->Release();
 	}
 }
 
-void FAGXDynamicRHI::RHIUpdateShaderResourceView(FRHIShaderResourceView* SRVRHI, FRHIBuffer* IndexBufferRHI)
+void FAGXDynamicRHI::RHIUpdateShaderResourceView(FRHIShaderResourceView* SRVRHI, FRHIBuffer* BufferRHI)
 {
 	check(SRVRHI);
 	FAGXShaderResourceView* SRV = ResourceCast(SRVRHI);
-	if (!IndexBufferRHI)
+	check(!SRV->bTexture);
+
+	FAGXResourceMultiBuffer* OldBuffer = SRV->SourceBuffer;
+
+	SRV->SourceBuffer = ResourceCast(BufferRHI);
+	SRV->Stride = 0;
+
+	SRV->Format = SRV->SourceBuffer && SRV->SourceBuffer->IndexType != mtlpp::IndexType::UInt16
+		? PF_R32_UINT
+		: PF_R16_UINT;
+
+	if (SRV->SourceBuffer)
 	{
-		SRV->SourceVertexBuffer = nullptr;
-		SRV->TextureView = nullptr;
-		SRV->SourceIndexBuffer = nullptr;
-		SRV->SourceStructuredBuffer = nullptr;
-		SRV->Offset = 0;
-		SRV->Format = PF_R16_UINT;
-		SRV->Stride = 0;
+		SRV->SourceBuffer->AddRef();
 	}
-	else if (SRV->SourceIndexBuffer != IndexBufferRHI)
+
+	if (OldBuffer)
 	{
-		FAGXIndexBuffer* IndexBuffer = ResourceCast(IndexBufferRHI);
-		SRV->SourceVertexBuffer = nullptr;
-		SRV->TextureView = nullptr;
-		SRV->SourceIndexBuffer = IndexBuffer;
-		SRV->SourceStructuredBuffer = nullptr;
-		SRV->Offset = 0;
-		SRV->Format = (IndexBuffer->IndexType == mtlpp::IndexType::UInt16) ? PF_R16_UINT : PF_R32_UINT;
-		SRV->Stride = 0;
+		OldBuffer->Release();
 	}
 }
 
@@ -446,24 +439,11 @@ void FAGXRHICommandContext::ClearUAVWithBlitEncoder(FRHIUnorderedAccessView* Uno
 {
 	SCOPED_AUTORELEASE_POOL;
 
-	FRHICommandListImmediate& RHICmdList = FRHICommandListExecutor::GetImmediateCommandList();
-	FAGXUnorderedAccessView* UnorderedAccessView = ResourceCast(UnorderedAccessViewRHI);
-	FAGXBuffer Buffer;
-	uint32 Size = 0;
-
-	switch (Type)
-	{
-		case EAGXRHIClearUAVType::VertexBuffer:
-			check(EnumHasAnyFlags(UnorderedAccessView->SourceView->SourceVertexBuffer->GetUsage(), BUF_ByteAddressBuffer));
-			Buffer = UnorderedAccessView->SourceView->SourceVertexBuffer->GetCurrentBuffer();
-			Size = UnorderedAccessView->SourceView->SourceVertexBuffer->GetSize();
-			break;
+	FAGXResourceMultiBuffer* SourceBuffer = ResourceCast(UnorderedAccessViewRHI)->GetSourceBuffer();
+	FAGXBuffer Buffer = SourceBuffer->GetCurrentBuffer();
+	uint32 Size = SourceBuffer->GetSize();
 
-		case EAGXRHIClearUAVType::StructuredBuffer:
-			Buffer = UnorderedAccessView->SourceView->SourceStructuredBuffer->GetCurrentBuffer();
-			Size = UnorderedAccessView->SourceView->SourceStructuredBuffer->GetSize();
-			break;
-	};
+	check(Type != EAGXRHIClearUAVType::VertexBuffer || EnumHasAnyFlags(SourceBuffer->GetUsage(), BUF_ByteAddressBuffer));
 
 	uint32 AlignedSize = Align(Size, BufferOffsetAlignment);
 	FAGXPooledBufferArgs Args(AlignedSize, BUF_Dynamic, mtlpp::StorageMode::Shared);
@@ -481,8 +461,8 @@ void FAGXRHICommandContext::ClearUAVWithBlitEncoder(FRHIUnorderedAccessView* Uno
 void FAGXRHICommandContext::RHIClearUAVFloat(FRHIUnorderedAccessView* UnorderedAccessViewRHI, const FVector4f& Values)
 {
 #if UE_METAL_RHI_SUPPORT_CLEAR_UAV_WITH_BLIT_ENCODER
-	FAGXUnorderedAccessView* UnorderedAccessView = ResourceCast(UnorderedAccessViewRHI);
-	if (UnorderedAccessView->SourceView->SourceStructuredBuffer)
+	FAGXUnorderedAccessView* UAV = ResourceCast(UnorderedAccessViewRHI);
+	if (!UAV->bTexture && EnumHasAnyFlags(UAV->GetSourceBuffer()->GetUsage(), BUF_StructuredBuffer))
 	{
 		ClearUAVWithBlitEncoder(UnorderedAccessViewRHI, EAGXRHIClearUAVType::StructuredBuffer, *(uint32*)&Values.X);
 	}
@@ -497,8 +477,8 @@ void FAGXRHICommandContext::RHIClearUAVFloat(FRHIUnorderedAccessView* UnorderedA
 void FAGXRHICommandContext::RHIClearUAVUint(FRHIUnorderedAccessView* UnorderedAccessViewRHI, const FUintVector4& Values)
 {
 #if UE_METAL_RHI_SUPPORT_CLEAR_UAV_WITH_BLIT_ENCODER
-	FAGXUnorderedAccessView* UnorderedAccessView = ResourceCast(UnorderedAccessViewRHI);
-	if (UnorderedAccessView->SourceView->SourceStructuredBuffer)
+	FAGXUnorderedAccessView* UAV = ResourceCast(UnorderedAccessViewRHI);
+	if (!UAV->bTexture && EnumHasAnyFlags(UAV->GetSourceBuffer()->GetUsage(), BUF_StructuredBuffer))
 	{
 		ClearUAVWithBlitEncoder(UnorderedAccessViewRHI, EAGXRHIClearUAVType::StructuredBuffer, *(uint32*)&Values.X);
 	}
@@ -518,7 +498,7 @@ void FAGXRHICommandContext::ClearUAV(TRHICommandList_RecursiveHazardous<FAGXRHIC
 		// The Metal validation layer will complain about resources with a
 		// signed format bound against an unsigned data format type as the
 		// shader parameter.
-		switch (GPixelFormats[UnorderedAccessView->SourceView->Format].UnrealFormat)
+		switch (GPixelFormats[UnorderedAccessView->Format].UnrealFormat)
 		{
 			case PF_R32_SINT:
 			case PF_R16_SINT:
@@ -530,37 +510,24 @@ void FAGXRHICommandContext::ClearUAV(TRHICommandList_RecursiveHazardous<FAGXRHIC
 				break;
 		}
 
-		if (UnorderedAccessView->SourceView->SourceVertexBuffer)
+		if (UnorderedAccessView->bTexture)
 		{
-#if UE_METAL_RHI_SUPPORT_CLEAR_UAV_WITH_BLIT_ENCODER
-			if (EnumHasAnyFlags(UnorderedAccessView->SourceView->SourceVertexBuffer->GetUsage(), BUF_ByteAddressBuffer))
-			{
-				ClearUAVWithBlitEncoder(UnorderedAccessView, EAGXRHIClearUAVType::VertexBuffer, *(const uint32*)ClearValue);
-			}
-			else
-#endif // UE_METAL_RHI_SUPPORT_CLEAR_UAV_WITH_BLIT_ENCODER
-			{
-				uint32 NumElements = UnorderedAccessView->SourceView->SourceVertexBuffer->GetSize() / GPixelFormats[UnorderedAccessView->SourceView->Format].BlockBytes;
-				ClearUAVShader_T<EClearReplacementResourceType::Buffer, 4, false>(RHICmdList, UnorderedAccessView, NumElements, 1, 1, ClearValue, ValueType);
-			}
-		}
-		else if (UnorderedAccessView->SourceView->SourceTexture)
-		{
-			FIntVector SizeXYZ = UnorderedAccessView->SourceView->SourceTexture->GetSizeXYZ();
+			FAGXSurface* Texture = UnorderedAccessView->GetSourceTexture();
+			FIntVector SizeXYZ = Texture->GetSizeXYZ();
 
-			if (FRHITexture2D* Texture2D = UnorderedAccessView->SourceView->SourceTexture->GetTexture2D())
+			if (FRHITexture2D* Texture2D = Texture->GetTexture2D())
 			{
 				ClearUAVShader_T<EClearReplacementResourceType::Texture2D, 4, false>(RHICmdList, UnorderedAccessView, SizeXYZ.X, SizeXYZ.Y, SizeXYZ.Z, ClearValue, ValueType);
 			}
-			else if (FRHITexture2DArray* Texture2DArray = UnorderedAccessView->SourceView->SourceTexture->GetTexture2DArray())
+			else if (FRHITexture2DArray* Texture2DArray = Texture->GetTexture2DArray())
 			{
 				ClearUAVShader_T<EClearReplacementResourceType::Texture2DArray, 4, false>(RHICmdList, UnorderedAccessView, SizeXYZ.X, SizeXYZ.Y, SizeXYZ.Z, ClearValue, ValueType);
 			}
-			else if (FRHITexture3D* Texture3D = UnorderedAccessView->SourceView->SourceTexture->GetTexture3D())
+			else if (FRHITexture3D* Texture3D = Texture->GetTexture3D())
 			{
 				ClearUAVShader_T<EClearReplacementResourceType::Texture3D, 4, false>(RHICmdList, UnorderedAccessView, SizeXYZ.X, SizeXYZ.Y, SizeXYZ.Z, ClearValue, ValueType);
 			}
-			else if (FRHITextureCube* TextureCube = UnorderedAccessView->SourceView->SourceTexture->GetTextureCube())
+			else if (FRHITextureCube* TextureCube = Texture->GetTextureCube())
 			{
 				ClearUAVShader_T<EClearReplacementResourceType::Texture2DArray, 4, false>(RHICmdList, UnorderedAccessView, SizeXYZ.X, SizeXYZ.Y, SizeXYZ.Z, ClearValue, ValueType);
 			}
@@ -571,8 +538,19 @@ void FAGXRHICommandContext::ClearUAV(TRHICommandList_RecursiveHazardous<FAGXRHIC
 		}
 		else
 		{
-			// TODO: ensure(0);
-			UE_LOG(LogRHI, Warning, TEXT("AGX RHI ClearUAV does not yet support clearing of a UAV without a SourceView."));
+			FAGXResourceMultiBuffer* SourceBuffer = UnorderedAccessView->GetSourceBuffer();
+
+#if UE_METAL_RHI_SUPPORT_CLEAR_UAV_WITH_BLIT_ENCODER
+			if (EnumHasAnyFlags(SourceBuffer->GetUsage(), BUF_ByteAddressBuffer))
+			{
+				ClearUAVWithBlitEncoder(UnorderedAccessView, EAGXRHIClearUAVType::VertexBuffer, *(const uint32*)ClearValue);
+			}
+			else
+#endif // UE_METAL_RHI_SUPPORT_CLEAR_UAV_WITH_BLIT_ENCODER
+			{
+				uint32 NumElements = SourceBuffer->GetSize() / GPixelFormats[UnorderedAccessView->Format].BlockBytes;
+				ClearUAVShader_T<EClearReplacementResourceType::Buffer, 4, false>(RHICmdList, UnorderedAccessView, NumElements, 1, 1, ClearValue, ValueType);
+			}
 		}
 	} // @autoreleasepool
 }
diff --git a/Engine/Source/Runtime/Apple/AGXRHI/Private/AGXViewport.cpp b/Engine/Source/Runtime/Apple/AGXRHI/Private/AGXViewport.cpp
--- a/Engine/Source/Runtime/Apple/AGXRHI/Private/AGXViewport.cpp
+++ b/Engine/Source/Runtime/Apple/AGXRHI/Private/AGXViewport.cpp
@@ -238,29 +238,35 @@ void FAGXViewport::Resize(uint32 InSizeX, uint32 InSizeY, bool bInIsFullscreen,E
 
     {
         FScopeLock Lock(&Mutex);
-        FRHIResourceCreateInfo CreateInfo(TEXT("BackBuffer"));
-        FTexture2DRHIRef NewBackBuffer;
-        FTexture2DRHIRef DoubleBuffer;
-        if (GAGXSupportsIntermediateBackBuffer)
-        {
-            NewBackBuffer = (FAGXTexture2D*)(FRHITexture2D*)GDynamicRHI->RHICreateTexture2D(InSizeX, InSizeY, Format, 1, 1, TexCreate_RenderTargetable, ERHIAccess::Unknown, CreateInfo);
-            
-            if (GAGXSeparatePresentThread)
-            {
-                DoubleBuffer = GDynamicRHI->RHICreateTexture2D(InSizeX, InSizeY, Format, 1, 1, TexCreate_RenderTargetable, ERHIAccess::Unknown, CreateInfo);
-                ((FAGXTexture2D*)DoubleBuffer.GetReference())->Surface.Viewport = this;
-            }
-        }
-        else
+
+		TRefCountPtr<FAGXSurface> NewBackBuffer;
+		TRefCountPtr<FAGXSurface> DoubleBuffer;
+
+		ETextureCreateFlags Flags = GAGXSupportsIntermediateBackBuffer
+			? TexCreate_RenderTargetable
+			: TexCreate_RenderTargetable | TexCreate_Presentable;
+
+		FRHITextureCreateDesc CreateDesc = FRHITextureCreateDesc::Create2D(
+			  TEXT("BackBuffer")
+			, { (int32)InSizeX, (int32)InSizeY }
+			, Format
+			, FClearValueBinding::None
+			, Flags
+		);
+
+		NewBackBuffer = new FAGXSurface(CreateDesc);
+		NewBackBuffer->Viewport = this;
+
+        if (GAGXSupportsIntermediateBackBuffer && GAGXSeparatePresentThread)
         {
-            NewBackBuffer = (FAGXTexture2D*)(FRHITexture2D*)GDynamicRHI->RHICreateTexture2D(InSizeX, InSizeY, Format, 1, 1, TexCreate_RenderTargetable | TexCreate_Presentable, ERHIAccess::Unknown, CreateInfo);
+            DoubleBuffer = new FAGXSurface(CreateDesc);
+            DoubleBuffer->Viewport = this;
         }
-        ((FAGXTexture2D*)NewBackBuffer.GetReference())->Surface.Viewport = this;
-        
-        BackBuffer[Index] = (FAGXTexture2D*)NewBackBuffer.GetReference();
+
+        BackBuffer[Index] = NewBackBuffer;
         if (GAGXSeparatePresentThread)
         {
-            BackBuffer[EAGXViewportAccessRHI] = (FAGXTexture2D*)DoubleBuffer.GetReference();
+            BackBuffer[EAGXViewportAccessRHI] = DoubleBuffer;
         }
         else
         {
@@ -269,7 +275,7 @@ void FAGXViewport::Resize(uint32 InSizeX, uint32 InSizeY, bool bInIsFullscreen,E
     }
 }
 
-TRefCountPtr<FAGXTexture2D> FAGXViewport::GetBackBuffer(EAGXViewportAccessFlag Accessor) const
+TRefCountPtr<FAGXSurface> FAGXViewport::GetBackBuffer(EAGXViewportAccessFlag Accessor) const
 {
 	FScopeLock Lock(&Mutex);
 	
@@ -405,7 +411,7 @@ void FAGXViewport::ReleaseDrawable()
 
 		if (!GAGXSupportsIntermediateBackBuffer && IsValidRef(BackBuffer[GetViewportIndex(EAGXViewportAccessRHI)]))
 		{
-			BackBuffer[GetViewportIndex(EAGXViewportAccessRHI)]->Surface.Texture = nil;
+			BackBuffer[GetViewportIndex(EAGXViewportAccessRHI)]->Texture = nil;
 		}
 	}
 }
@@ -470,10 +476,10 @@ void FAGXViewport::Present(FAGXCommandQueue& CommandQueue, bool bLockToVsync)
 						
 						if (GAGXSupportsIntermediateBackBuffer)
 						{
-							TRefCountPtr<FAGXTexture2D> Texture = LastCompleteFrame;
+							TRefCountPtr<FAGXSurface> Texture = LastCompleteFrame;
 							check(IsValidRef(Texture));
 							
-							FAGXTexture Src = Texture->Surface.Texture;
+							FAGXTexture Src = Texture->Texture;
 							FAGXTexture Dst = LocalDrawable.texture;
 							
 							NSUInteger Width = FMath::Min(Src.GetWidth(), Dst.GetWidth());
@@ -578,8 +584,8 @@ void FAGXViewport::Swap()
 		check(IsValidRef(BackBuffer[0]));
 		check(IsValidRef(BackBuffer[1]));
 		
-		TRefCountPtr<FAGXTexture2D> BB0 = BackBuffer[0];
-		TRefCountPtr<FAGXTexture2D> BB1 = BackBuffer[1];
+		TRefCountPtr<FAGXSurface> BB0 = BackBuffer[0];
+		TRefCountPtr<FAGXSurface> BB1 = BackBuffer[1];
 		
 		BackBuffer[0] = BB1;
 		BackBuffer[1] = BB0;
diff --git a/Engine/Source/Runtime/Apple/AGXRHI/Public/AGXResources.h b/Engine/Source/Runtime/Apple/AGXRHI/Public/AGXResources.h
--- a/Engine/Source/Runtime/Apple/AGXRHI/Public/AGXResources.h
+++ b/Engine/Source/Runtime/Apple/AGXRHI/Public/AGXResources.h
@@ -87,27 +87,25 @@ private:
 class FAGXTexture : public mtlpp::Texture
 {
 public:
-	FAGXTexture(ns::Ownership retain = ns::Ownership::Retain) : mtlpp::Texture(retain) { }
+	FAGXTexture(ns::Ownership retain = ns::Ownership::Retain)
+		: mtlpp::Texture(retain)
+	{}
+
 	FAGXTexture(ns::Protocol<id<MTLTexture>>::type handle, ns::Ownership retain = ns::Ownership::Retain)
-	: mtlpp::Texture(handle, nullptr, retain) {}
+		: mtlpp::Texture(handle, nullptr, retain)
+	{}
 	
 	FAGXTexture(mtlpp::Texture&& rhs)
-	: mtlpp::Texture((mtlpp::Texture&&)rhs)
-	{
-		
-	}
+		: mtlpp::Texture((mtlpp::Texture&&)rhs)
+	{}
 	
 	FAGXTexture(const FAGXTexture& rhs)
-	: mtlpp::Texture(rhs)
-	{
-		
-	}
+		: mtlpp::Texture(rhs)
+	{}
 	
 	FAGXTexture(FAGXTexture&& rhs)
-	: mtlpp::Texture((mtlpp::Texture&&)rhs)
-	{
-		
-	}
+		: mtlpp::Texture((mtlpp::Texture&&)rhs)
+	{}
 	
 	FAGXTexture& operator=(const FAGXTexture& rhs)
 	{
@@ -135,24 +133,43 @@ public:
 	}
 };
 
-/** Texture/RT wrapper. */
-class AGXRHI_API FAGXSurface
+struct FAGXTextureDesc
+{
+	FAGXTextureDesc(FRHITextureDesc const& InDesc);
+	
+	mtlpp::TextureDescriptor Desc;
+	mtlpp::PixelFormat MTLFormat;
+	bool bMemoryless = false;
+	bool bIsRenderTarget = false;
+	uint8 FormatKey = 0;
+};
+
+struct FAGXTextureCreateDesc : public FRHITextureCreateDesc, public FAGXTextureDesc
+{
+	FAGXTextureCreateDesc(FRHITextureCreateDesc const& CreateDesc)
+		: FRHITextureCreateDesc(CreateDesc)
+		, FAGXTextureDesc(CreateDesc)
+	{
+		// @todo: texture type unification - Metal can override NumSamples based on command line options.
+		// We should instead require the renderer to do this.
+		NumSamples = Desc.GetSampleCount();
+	}
+};
+
+// AGX RHI texture resource
+class AGXRHI_API FAGXSurface : public FRHITexture
 {
 public:
 
 	/** 
 	 * Constructor that will create Texture and Color/DepthBuffers as needed
 	 */
-	FAGXSurface(ERHIResourceType ResourceType, EPixelFormat Format, uint32 SizeX, uint32 SizeY, uint32 SizeZ, uint32 NumSamples, bool bArray, uint32 ArraySize, uint32 NumMips, ETextureCreateFlags Flags, FResourceBulkDataInterface* BulkData);
-
-	FAGXSurface(FAGXSurface& Source, NSRange MipRange);
-	
-	FAGXSurface(FAGXSurface& Source, NSRange MipRange, EPixelFormat Format, bool bSRGBForceDisable);
+	FAGXSurface(FAGXTextureCreateDesc const& CreateDesc);
 	
 	/**
 	 * Destructor
 	 */
-	~FAGXSurface();
+	virtual ~FAGXSurface();
 
 	/** Prepare for texture-view support - need only call this once on the source texture which is to be viewed. */
 	void PrepareTextureView();
@@ -207,13 +224,9 @@ public:
 	ns::AutoReleased<FAGXTexture> GetCurrentTexture();
 
 	FAGXTexture Reallocate(FAGXTexture Texture, mtlpp::TextureUsage UsageModifier);
-	void ReplaceTexture(FAGXContext& Context, FAGXTexture OldTexture, FAGXTexture NewTexture);
 	void MakeAliasable(void);
-	void MakeUnAliasable(void);
-	
-	ERHIResourceType Type;
-	EPixelFormat PixelFormat;
-	uint8 FormatKey;
+
+	uint8 const FormatKey;
 	//texture used for store actions and binding to shader params
 	FAGXTexture Texture;
 	//if surface is MSAA, texture used to bind for RT
@@ -226,8 +239,7 @@ public:
 	// iOS A9+ where depth resolve is available
 	// iOS < A9 where depth resolve is unavailable.
 	FAGXTexture MSAAResolveTexture;
-	uint32 SizeX, SizeY, SizeZ;
-	bool bIsCubemap;
+
 	int16 volatile Written;
 	int16 GPUReadback = 0;
 	enum EAGXGPUReadbackFlags : int16
@@ -238,8 +250,6 @@ public:
 		ReadbackFenceComplete 			= 1 << ReadbackFenceCompleteShift,
 		ReadbackRequestedAndComplete 	= ReadbackRequested | ReadbackFenceComplete
 	};
-	
-	ETextureCreateFlags Flags;
 
 	uint32 BufferLocks;
 
@@ -248,125 +258,23 @@ public:
 	
 	// For back-buffers, the owning viewport.
 	class FAGXViewport* Viewport;
-	
-	TSet<class FAGXShaderResourceView*> SRVs;
-
-private:
-	void Init(FAGXSurface& Source, NSRange MipRange);
-	
-	void Init(FAGXSurface& Source, NSRange MipRange, EPixelFormat Format, bool bSRGBForceDisable);
-	
-private:
-	// The movie playback IOSurface/CVTexture wrapper to avoid page-off
-	CFTypeRef ImageSurfaceRef;
-	
-	// Texture view surfaces don't own their resources, only reference
-	bool bTextureView;
-	
-	// Count of outstanding async. texture uploads
-	static volatile int64 ActiveUploads;
-};
 
-class FAGXTexture2D : public FRHITexture2D
-{
-public:
-	/** The surface info */
-	FAGXSurface Surface;
-
-	// Constructor, just calls base and Surface constructor
-	FAGXTexture2D(EPixelFormat Format, uint32 SizeX, uint32 SizeY, uint32 NumMips, uint32 NumSamples, ETextureCreateFlags Flags, FResourceBulkDataInterface* BulkData, const FClearValueBinding& InClearValue)
-		: FRHITexture2D(SizeX, SizeY, NumMips, NumSamples, Format, Flags, InClearValue)
-		, Surface(RRT_Texture2D, Format, SizeX, SizeY, 1, NumSamples, /*bArray=*/ false, 1, NumMips, Flags, BulkData)
-	{
-	}
-	
-	virtual ~FAGXTexture2D()
-	{
-	}
-	
 	virtual void* GetTextureBaseRHI() override final
 	{
-		return &Surface;
-	}
-	
-	virtual void* GetNativeResource() const override final
-	{
-		return Surface.Texture;
+		return this;
 	}
-};
-
-class FAGXTexture2DArray : public FRHITexture2DArray
-{
-public:
-	/** The surface info */
-	FAGXSurface Surface;
 
-	// Constructor, just calls base and Surface constructor
-	FAGXTexture2DArray(EPixelFormat Format, uint32 SizeX, uint32 SizeY, uint32 ArraySize, uint32 NumMips, ETextureCreateFlags Flags, FResourceBulkDataInterface* BulkData, const FClearValueBinding& InClearValue)
-		: FRHITexture2DArray(SizeX, SizeY, ArraySize, NumMips, 1, Format, Flags, InClearValue)
-		, Surface(RRT_Texture2DArray, Format, SizeX, SizeY, 1, /*NumSamples=*/1, /*bArray=*/ true, ArraySize, NumMips, Flags, BulkData)
-	{
-	}
-	
-	virtual ~FAGXTexture2DArray()
-	{
-	}
-	
-	virtual void* GetTextureBaseRHI() override final
-	{
-		return &Surface;
-	}
-};
-
-class FAGXTexture3D : public FRHITexture3D
-{
-public:
-	/** The surface info */
-	FAGXSurface Surface;
-
-	// Constructor, just calls base and Surface constructor
-	FAGXTexture3D(EPixelFormat Format, uint32 SizeX, uint32 SizeY, uint32 SizeZ, uint32 NumMips, ETextureCreateFlags Flags, FResourceBulkDataInterface* BulkData, const FClearValueBinding& InClearValue)
-		: FRHITexture3D(SizeX, SizeY, SizeZ, NumMips, Format, Flags, InClearValue)
-		, Surface(RRT_Texture3D, Format, SizeX, SizeY, SizeZ, /*NumSamples=*/1, /*bArray=*/ false, 1, NumMips, Flags, BulkData)
-	{
-	}
-	
-	virtual ~FAGXTexture3D()
-	{
-	}
-	
-	virtual void* GetTextureBaseRHI() override final
-	{
-		return &Surface;
-	}
-};
-
-class FAGXTextureCube : public FRHITextureCube
-{
-public:
-	/** The surface info */
-	FAGXSurface Surface;
-
-	// Constructor, just calls base and Surface constructor
-	FAGXTextureCube(EPixelFormat Format, uint32 Size, bool bArray, uint32 ArraySize, uint32 NumMips, ETextureCreateFlags Flags, FResourceBulkDataInterface* BulkData, const FClearValueBinding& InClearValue)
-		: FRHITextureCube(Size, NumMips, Format, Flags, InClearValue)
-		, Surface(RRT_TextureCube, Format, Size, Size, 6, /*NumSamples=*/1, bArray, ArraySize, NumMips, Flags, BulkData)
-	{
-	}
-	
-	virtual ~FAGXTextureCube()
+	virtual void* GetNativeResource() const override final
 	{
+		return Texture;
 	}
 	
-	virtual void* GetTextureBaseRHI() override final
-	{
-		return &Surface;
-	}
+private:
+	// The movie playback IOSurface/CVTexture wrapper to avoid page-off
+	CFTypeRef ImageSurfaceRef;
 
-	virtual void* GetNativeResource() const override final
-	{
-		return Surface.Texture;
-	}
+	// Count of outstanding async. texture uploads
+	static volatile int64 ActiveUploads;
 };
 
 @interface FAGXBufferData : FApplePlatformObject<NSObject>
@@ -390,34 +298,13 @@ ENUM_CLASS_FLAGS(EAGXBufferUsage);
 class FAGXLinearTextureDescriptor
 {
 public:
-	FAGXLinearTextureDescriptor()
-		: StartOffsetBytes(0)
-		, NumElements(UINT_MAX)
-		, BytesPerElement(0)
-	{
-		// void
-	}
+	FAGXLinearTextureDescriptor() = default;
 
 	FAGXLinearTextureDescriptor(uint32 InStartOffsetBytes, uint32 InNumElements, uint32 InBytesPerElement)
 		: StartOffsetBytes(InStartOffsetBytes)
-		, NumElements(InNumElements)
-		, BytesPerElement(InBytesPerElement)
-	{
-		// void
-	}
-	
-	FAGXLinearTextureDescriptor(const FAGXLinearTextureDescriptor& Other)
-		: StartOffsetBytes(Other.StartOffsetBytes)
-		, NumElements(Other.NumElements)
-		, BytesPerElement(Other.BytesPerElement)
-	{
-		// void
-	}
-	
-	~FAGXLinearTextureDescriptor()
-	{
-		// void
-	}
+		, NumElements     (InNumElements)
+		, BytesPerElement (InBytesPerElement)
+	{}
 
 	friend uint32 GetTypeHash(FAGXLinearTextureDescriptor const& Key)
 	{
@@ -434,9 +321,9 @@ public:
 		       && BytesPerElement  == Other.BytesPerElement;
 	}
 
-	uint32 StartOffsetBytes;
-	uint32 NumElements;
-	uint32 BytesPerElement;
+	uint32 StartOffsetBytes = 0;
+	uint32 NumElements      = UINT_MAX;
+	uint32 BytesPerElement  = 0;
 };
 
 class FAGXRHIBuffer
@@ -599,50 +486,118 @@ typedef FAGXResourceMultiBuffer FAGXIndexBuffer;
 typedef FAGXResourceMultiBuffer FAGXVertexBuffer;
 typedef FAGXResourceMultiBuffer FAGXStructuredBuffer;
 
-class FAGXShaderResourceView : public FRHIShaderResourceView
-{
+class FAGXResourceViewBase
+{
+protected:
+	// Constructor for buffers
+	FAGXResourceViewBase(
+		  FRHIBuffer* InBuffer
+		, uint32 InStartOffsetBytes
+		, uint32 InNumElements
+		, EPixelFormat InFormat
+	);
+
+	// Constructor for textures
+	FAGXResourceViewBase(
+		  FRHITexture* InTexture
+		, EPixelFormat InFormat
+		, uint8 InMipLevel
+		, uint8 InNumMipLevels
+		, ERHITextureSRVOverrideSRGBType InSRGBOverride
+		, uint32 InFirstArraySlice
+		, uint32 InNumArraySlices
+	);
+
 public:
+	~FAGXResourceViewBase();
 
-	// The vertex buffer this SRV comes from (can be null)
-	TRefCountPtr<FAGXVertexBuffer> SourceVertexBuffer;
-	
-	// The index buffer this SRV comes from (can be null)
-	TRefCountPtr<FAGXIndexBuffer> SourceIndexBuffer;
+	inline FAGXResourceMultiBuffer* GetSourceBuffer () const { check(!bTexture); return SourceBuffer; }
 
-	// The texture that this SRV come from
-	TRefCountPtr<FRHITexture> SourceTexture;
-	
-	// The source structured buffer (can be null)
-	TRefCountPtr<FAGXStructuredBuffer> SourceStructuredBuffer;
-	
-	FAGXSurface* TextureView;
-	uint32 Offset;
-	uint8 MipLevel          : 4;
-	uint8 bSRGBForceDisable : 1;
-	uint8 Reserved          : 3;
-	uint8 NumMips;
-	uint8 Format;
-	uint8 Stride;
+	inline FAGXSurface*             GetSourceTexture() const { check(bTexture); return SourceTexture; }
+	inline FAGXTexture const&       GetTextureView  () const { check(bTexture); return TextureView;   }
 
-	void InitLinearTextureDescriptor(const FAGXLinearTextureDescriptor& InLinearTextureDescriptor);
+private:
+	// Needed for RHIUpdateShaderResourceView
+	friend class FAGXDynamicRHI;
 
-	FAGXShaderResourceView();
-	~FAGXShaderResourceView();
-	
-	ns::AutoReleased<FAGXTexture> GetLinearTexture(bool const bUAV);
+	union
+	{
+		FAGXResourceMultiBuffer* SourceBuffer;
+		FAGXSurface* SourceTexture;
+	};
 
-private:
-	FAGXLinearTextureDescriptor* LinearTextureDesc;
-};
+	TUniquePtr<FAGXLinearTextureDescriptor> LinearTextureDesc = nullptr;
+	FAGXTexture TextureView = nullptr;
 
+public:
+	uint8 const bTexture : 1;
+	uint8       bSRGBForceDisable : 1;
+	uint8       MipLevel : 4;
+	uint8       Reserved : 2;
+	uint8       NumMips;
+	uint8       Format;
+	uint8       Stride;
+	uint32      Offset;
 
+	ns::AutoReleased<FAGXTexture> GetLinearTexture();
+};
 
-class FAGXUnorderedAccessView : public FRHIUnorderedAccessView
+class FAGXShaderResourceView final : public FRHIShaderResourceView, public FAGXResourceViewBase
 {
 public:
-	
-	// the potential resources to refer to with the UAV object
-	TRefCountPtr<FAGXShaderResourceView> SourceView;
+	explicit FAGXShaderResourceView(const FShaderResourceViewInitializer& Initializer)
+		: FAGXResourceViewBase(
+			  Initializer.AsBufferSRV().Buffer
+			, Initializer.AsBufferSRV().StartOffsetBytes
+			, Initializer.AsBufferSRV().NumElements
+			, Initializer.AsBufferSRV().Format
+		)
+	{}
+
+	explicit FAGXShaderResourceView(FRHITexture* Texture, const FRHITextureSRVCreateInfo& CreateInfo)
+		: FAGXResourceViewBase(
+			  Texture
+			, CreateInfo.Format
+			, CreateInfo.MipLevel
+			, CreateInfo.NumMipLevels
+			, CreateInfo.SRGBOverride
+			, CreateInfo.FirstArraySlice
+			, CreateInfo.NumArraySlices
+		)
+	{}
+
+	virtual ~FAGXShaderResourceView()
+	{}
+};
+
+class FAGXUnorderedAccessView final : public FRHIUnorderedAccessView, public FAGXResourceViewBase
+{
+public:
+	explicit FAGXUnorderedAccessView(FRHIBuffer* Buffer, EPixelFormat Format)
+		: FAGXResourceViewBase(Buffer, 0, UINT_MAX, Format)
+	{}
+
+	explicit FAGXUnorderedAccessView(FRHIBuffer* Buffer, bool bUseUAVCounter, bool bAppendBuffer)
+		: FAGXResourceViewBase(Buffer, 0, UINT_MAX, PF_Unknown)
+	{
+		checkf(!bUseUAVCounter, TEXT("UAV counters not implemented."));
+		checkf(!bAppendBuffer, TEXT("UAV append buffers not implemented."));
+	}
+
+	explicit FAGXUnorderedAccessView(FRHITexture* Texture, uint32 MipLevel, uint16 FirstArraySlice, uint16 NumArraySlices)
+		: FAGXResourceViewBase(
+			  Texture
+			, PF_Unknown
+			, MipLevel
+			, 1 // NumMipLevels
+			, ERHITextureSRVOverrideSRGBType::SRGBO_ForceDisable
+			, FirstArraySlice
+			, NumArraySlices
+		)
+	{}
+
+	virtual ~FAGXUnorderedAccessView()
+	{}
 };
 
 class FAGXGPUFence final : public FRHIGPUFence
@@ -714,26 +669,6 @@ struct TAGXResourceTraits<FRHIComputeShader>
 	typedef FAGXComputeShader TConcreteType;
 };
 template<>
-struct TAGXResourceTraits<FRHITexture3D>
-{
-	typedef FAGXTexture3D TConcreteType;
-};
-template<>
-struct TAGXResourceTraits<FRHITexture2D>
-{
-	typedef FAGXTexture2D TConcreteType;
-};
-template<>
-struct TAGXResourceTraits<FRHITexture2DArray>
-{
-	typedef FAGXTexture2DArray TConcreteType;
-};
-template<>
-struct TAGXResourceTraits<FRHITextureCube>
-{
-	typedef FAGXTextureCube TConcreteType;
-};
-template<>
 struct TAGXResourceTraits<FRHIRenderQuery>
 {
 	typedef FAGXRHIRenderQuery TConcreteType;
diff --git a/Engine/Source/Runtime/Apple/AGXRHI/Public/AGXViewport.h b/Engine/Source/Runtime/Apple/AGXRHI/Public/AGXViewport.h
--- a/Engine/Source/Runtime/Apple/AGXRHI/Public/AGXViewport.h
+++ b/Engine/Source/Runtime/Apple/AGXRHI/Public/AGXViewport.h
@@ -36,7 +36,7 @@ public:
 
 	void Resize(uint32 InSizeX, uint32 InSizeY, bool bInIsFullscreen,EPixelFormat Format);
 	
-	TRefCountPtr<FAGXTexture2D> GetBackBuffer(EAGXViewportAccessFlag Accessor) const;
+	TRefCountPtr<FAGXSurface> GetBackBuffer(EAGXViewportAccessFlag Accessor) const;
 	id<CAMetalDrawable> GetDrawable(EAGXViewportAccessFlag Accessor);
 	FAGXTexture GetDrawableTexture(EAGXViewportAccessFlag Accessor);
 	ns::AutoReleased<FAGXTexture> GetCurrentTexture(EAGXViewportAccessFlag Accessor);
@@ -65,7 +65,7 @@ private:
 
 private:
 	id<CAMetalDrawable> Drawable;
-	TRefCountPtr<FAGXTexture2D> BackBuffer[2];
+	TRefCountPtr<FAGXSurface> BackBuffer[2];
 	mutable FCriticalSection Mutex;
 	
 	ns::AutoReleased<FAGXTexture> DrawableTextures[2];
@@ -73,7 +73,7 @@ private:
 	uint32 DisplayID;
 	FAGXViewportPresentHandler Block;
 	volatile int32 FrameAvailable;
-	TRefCountPtr<FAGXTexture2D> LastCompleteFrame;
+	TRefCountPtr<FAGXSurface> LastCompleteFrame;
 	bool bIsFullScreen;
 
 #if PLATFORM_MAC
diff --git a/Engine/Source/Runtime/Apple/MetalRHI/Private/MetalCommandEncoder.cpp b/Engine/Source/Runtime/Apple/MetalRHI/Private/MetalCommandEncoder.cpp
--- a/Engine/Source/Runtime/Apple/MetalRHI/Private/MetalCommandEncoder.cpp
+++ b/Engine/Source/Runtime/Apple/MetalRHI/Private/MetalCommandEncoder.cpp
@@ -1798,21 +1798,23 @@ void FMetalCommandEncoder::UseResource(mtlpp::Resource const& Resource, mtlpp::R
 
 void FMetalCommandEncoder::SetShaderBufferInternal(mtlpp::FunctionType Function, uint32 Index)
 {
-	NSUInteger Offset = ShaderBuffers[uint32(Function)].Offsets[Index];
-	mtlpp::ResourceUsage Usage = ShaderBuffers[uint32(Function)].Usage[Index];
-	bool bBufferHasBytes = ShaderBuffers[uint32(Function)].Bytes[Index] != nil;
-	if (!ShaderBuffers[uint32(Function)].Buffers[Index] && bBufferHasBytes && !bSupportsMetalFeaturesSetBytes)
+	FMetalBufferBindings& Binding = ShaderBuffers[uint32(Function)];
+
+	NSUInteger Offset = Binding.Offsets[Index];
+	
+	bool bBufferHasBytes = Binding.Bytes[Index] != nil;
+	if (!Binding.Buffers[Index] && bBufferHasBytes && !bSupportsMetalFeaturesSetBytes)
 	{
-		uint8 const* Bytes = (((uint8 const*)ShaderBuffers[uint32(Function)].Bytes[Index]->Data) + ShaderBuffers[uint32(Function)].Offsets[Index]);
-		uint32 Len = ShaderBuffers[uint32(Function)].Bytes[Index]->Len - ShaderBuffers[uint32(Function)].Offsets[Index];
+		uint8 const* Bytes = (((uint8 const*)Binding.Bytes[Index]->Data) + Binding.Offsets[Index]);
+		uint32 Len = Binding.Bytes[Index]->Len - Binding.Offsets[Index];
 		
 		Offset = 0;
-		ShaderBuffers[uint32(Function)].Buffers[Index] = RingBuffer.NewBuffer(Len, BufferOffsetAlignment);
+		Binding.Buffers[Index] = RingBuffer.NewBuffer(Len, BufferOffsetAlignment);
 		
-		FMemory::Memcpy(((uint8*)ShaderBuffers[uint32(Function)].Buffers[Index].GetContents()) + Offset, Bytes, Len);
+		FMemory::Memcpy(((uint8*)Binding.Buffers[Index].GetContents()) + Offset, Bytes, Len);
 	}
 	
-	ns::AutoReleased<FMetalBuffer>& Buffer = ShaderBuffers[uint32(Function)].Buffers[Index];
+	ns::AutoReleased<FMetalBuffer>& Buffer = Binding.Buffers[Index];
 	if (Buffer)
 	{
 #if METAL_DEBUG_OPTIONS
@@ -1825,28 +1827,31 @@ void FMetalCommandEncoder::SetShaderBufferInternal(mtlpp::FunctionType Function,
 		switch (Function)
 		{
 			case mtlpp::FunctionType::Vertex:
-				ShaderBuffers[uint32(Function)].Bound |= (1 << Index);
+				Binding.Bound |= (1 << Index);
 				check(RenderCommandEncoder);
 				FenceResource(Buffer);
 				// MTLPP_VALIDATE(mtlpp::RenderCommandEncoder, RenderCommandEncoder, SafeGetRuntimeDebuggingLevel() >= EMetalDebugLevelValidation, UseResource(Buffer, mtlpp::ResourceUsage::Read));
 				METAL_DEBUG_LAYER(EMetalDebugLevelFastValidation, RenderEncoderDebug.SetBuffer(EMetalShaderVertex, Buffer, Offset, Index));
 				RenderCommandEncoder.SetVertexBuffer(Buffer, Offset, Index);
 				break;
+
 			case mtlpp::FunctionType::Fragment:
-				ShaderBuffers[uint32(Function)].Bound |= (1 << Index);
+				Binding.Bound |= (1 << Index);
 				check(RenderCommandEncoder);
 				// MTLPP_VALIDATE(mtlpp::RenderCommandEncoder, RenderCommandEncoder, SafeGetRuntimeDebuggingLevel() >= EMetalDebugLevelValidation, UseResource(Buffer, mtlpp::ResourceUsage::Read));
 				METAL_DEBUG_LAYER(EMetalDebugLevelFastValidation, RenderEncoderDebug.SetBuffer(EMetalShaderFragment, Buffer, Offset, Index));
 				RenderCommandEncoder.SetFragmentBuffer(Buffer, Offset, Index);
 				break;
+
 			case mtlpp::FunctionType::Kernel:
-				ShaderBuffers[uint32(Function)].Bound |= (1 << Index);
+				Binding.Bound |= (1 << Index);
 				check(ComputeCommandEncoder);
 				FenceResource(Buffer);
 				// MTLPP_VALIDATE(mtlpp::ComputeCommandEncoder, ComputeCommandEncoder, SafeGetRuntimeDebuggingLevel() >= EMetalDebugLevelValidation, UseResource(Buffer, mtlpp::ResourceUsage::Read));
 				METAL_DEBUG_LAYER(EMetalDebugLevelFastValidation, ComputeEncoderDebug.SetBuffer(Buffer, Offset, Index));
 				ComputeCommandEncoder.SetBuffer(Buffer, Offset, Index);
 				break;
+
 			default:
 				check(false);
 				break;
@@ -1854,37 +1859,40 @@ void FMetalCommandEncoder::SetShaderBufferInternal(mtlpp::FunctionType Function,
 		
 		if (Buffer.IsSingleUse())
 		{
-			ShaderBuffers[uint32(Function)].Usage[Index] = mtlpp::ResourceUsage(0);
-			ShaderBuffers[uint32(Function)].Offsets[Index] = 0;
-			ShaderBuffers[uint32(Function)].Buffers[Index] = nil;
-			ShaderBuffers[uint32(Function)].Bound &= ~(1 << Index);
+			Binding.Usage[Index] = mtlpp::ResourceUsage(0);
+			Binding.Offsets[Index] = 0;
+			Binding.Buffers[Index] = nil;
+			Binding.Bound &= ~(1 << Index);
 		}
 	}
 	else if (bBufferHasBytes && bSupportsMetalFeaturesSetBytes)
 	{
-		uint8 const* Bytes = (((uint8 const*)ShaderBuffers[uint32(Function)].Bytes[Index]->Data) + ShaderBuffers[uint32(Function)].Offsets[Index]);
-		uint32 Len = ShaderBuffers[uint32(Function)].Bytes[Index]->Len - ShaderBuffers[uint32(Function)].Offsets[Index];
+		uint8 const* Bytes = (((uint8 const*)Binding.Bytes[Index]->Data) + Binding.Offsets[Index]);
+		uint32 Len = Binding.Bytes[Index]->Len - Binding.Offsets[Index];
 		
 		switch (Function)
 		{
 			case mtlpp::FunctionType::Vertex:
-				ShaderBuffers[uint32(Function)].Bound |= (1 << Index);
+				Binding.Bound |= (1 << Index);
 				check(RenderCommandEncoder);
 				METAL_DEBUG_LAYER(EMetalDebugLevelFastValidation, RenderEncoderDebug.SetBytes(EMetalShaderVertex, Bytes, Len, Index));
 				RenderCommandEncoder.SetVertexData(Bytes, Len, Index);
 				break;
+
 			case mtlpp::FunctionType::Fragment:
-				ShaderBuffers[uint32(Function)].Bound |= (1 << Index);
+				Binding.Bound |= (1 << Index);
 				check(RenderCommandEncoder);
 				METAL_DEBUG_LAYER(EMetalDebugLevelFastValidation, RenderEncoderDebug.SetBytes(EMetalShaderFragment, Bytes, Len, Index));
 				RenderCommandEncoder.SetFragmentData(Bytes, Len, Index);
 				break;
+
 			case mtlpp::FunctionType::Kernel:
-				ShaderBuffers[uint32(Function)].Bound |= (1 << Index);
+				Binding.Bound |= (1 << Index);
 				check(ComputeCommandEncoder);
 				METAL_DEBUG_LAYER(EMetalDebugLevelFastValidation, ComputeEncoderDebug.SetBytes(Bytes, Len, Index));
 				ComputeCommandEncoder.SetBytes(Bytes, Len, Index);
 				break;
+
 			default:
 				check(false);
 				break;
diff --git a/Engine/Source/Runtime/Apple/MetalRHI/Private/MetalCommands.cpp b/Engine/Source/Runtime/Apple/MetalRHI/Private/MetalCommands.cpp
--- a/Engine/Source/Runtime/Apple/MetalRHI/Private/MetalCommands.cpp
+++ b/Engine/Source/Runtime/Apple/MetalRHI/Private/MetalCommands.cpp
@@ -308,7 +308,7 @@ void FMetalRHICommandContext::RHISetShaderTexture(FRHIGraphicsShader* ShaderRHI,
 	EMetalShaderStages Stage = GetShaderStage(ShaderRHI);
 	if (Surface != nullptr)
 	{
-        if (Surface->Texture || !(Surface->Flags & TexCreate_Presentable))
+        if (Surface->Texture || !(Surface->GetDesc().Flags & TexCreate_Presentable))
         {
             Context->GetCurrentState().SetShaderTexture(Stage, Surface->Texture, TextureIndex, (mtlpp::ResourceUsage)(mtlpp::ResourceUsage::Read|mtlpp::ResourceUsage::Sample));
         }
@@ -331,7 +331,7 @@ void FMetalRHICommandContext::RHISetShaderTexture(FRHIComputeShader* ComputeShad
 	FMetalSurface* Surface = GetMetalSurfaceFromRHITexture(NewTextureRHI);
 	if (Surface != nullptr)
     {
-        if (Surface->Texture || !(Surface->Flags & TexCreate_Presentable))
+        if (Surface->Texture || !(Surface->GetDesc().Flags & TexCreate_Presentable))
         {
             Context->GetCurrentState().SetShaderTexture(EMetalShaderStages::Compute, Surface->Texture, TextureIndex, (mtlpp::ResourceUsage)(mtlpp::ResourceUsage::Read|mtlpp::ResourceUsage::Sample));
         }
diff --git a/Engine/Source/Runtime/Apple/MetalRHI/Private/MetalContext.cpp b/Engine/Source/Runtime/Apple/MetalRHI/Private/MetalContext.cpp
--- a/Engine/Source/Runtime/Apple/MetalRHI/Private/MetalContext.cpp
+++ b/Engine/Source/Runtime/Apple/MetalRHI/Private/MetalContext.cpp
@@ -1177,36 +1177,10 @@ void FMetalContext::TransitionResource(FRHIUnorderedAccessView* InResource)
 {
 	FMetalUnorderedAccessView* UAV = ResourceCast(InResource);
 
-	// figure out which one of the resources we need to set
-	FMetalStructuredBuffer* StructuredBuffer = UAV->SourceView->SourceStructuredBuffer.GetReference();
-	FMetalVertexBuffer*     VertexBuffer     = UAV->SourceView->SourceVertexBuffer.GetReference();
-	FMetalIndexBuffer*      IndexBuffer      = UAV->SourceView->SourceIndexBuffer.GetReference();
-	FRHITexture*            Texture          = UAV->SourceView->SourceTexture.GetReference();
-	FMetalSurface*          Surface          = UAV->SourceView->TextureView;
-
-	if (StructuredBuffer)
-	{
-		RenderPass.TransitionResources(StructuredBuffer->GetCurrentBuffer());
-	}
-	else if (VertexBuffer && VertexBuffer->GetCurrentBufferOrNil())
-	{
-		RenderPass.TransitionResources(VertexBuffer->GetCurrentBuffer());
-	}
-	else if (IndexBuffer)
-	{
-		RenderPass.TransitionResources(IndexBuffer->GetCurrentBuffer());
-	}
-	else if (Surface)
+	if (UAV->bTexture)
 	{
-		RenderPass.TransitionResources(Surface->Texture.GetParentTexture());
-	}
-	else if (Texture)
-	{
-		if (!Surface)
-		{
-			Surface = GetMetalSurfaceFromRHITexture(Texture);
-		}
-		if ((Surface != nullptr) && Surface->Texture)
+		FMetalSurface* Surface = UAV->GetSourceTexture();
+		if (Surface->Texture)
 		{
 			RenderPass.TransitionResources(Surface->Texture);
 			if (Surface->MSAATexture)
@@ -1215,6 +1189,11 @@ void FMetalContext::TransitionResource(FRHIUnorderedAccessView* InResource)
 			}
 		}
 	}
+	else
+	{
+		FMetalResourceMultiBuffer* Buffer = UAV->GetSourceBuffer();
+		RenderPass.TransitionResources(Buffer->GetCurrentBuffer());
+	}
 }
 
 void FMetalContext::TransitionResource(FRHITexture* InResource)
diff --git a/Engine/Source/Runtime/Apple/MetalRHI/Private/MetalDynamicRHI.h b/Engine/Source/Runtime/Apple/MetalRHI/Private/MetalDynamicRHI.h
--- a/Engine/Source/Runtime/Apple/MetalRHI/Private/MetalDynamicRHI.h
+++ b/Engine/Source/Runtime/Apple/MetalRHI/Private/MetalDynamicRHI.h
@@ -23,7 +23,9 @@ class FMetalDeviceContext;
 /** The interface which is implemented by the dynamically bound RHI. */
 class FMetalDynamicRHI : public FDynamicRHI
 {
+	static inline FMetalDynamicRHI* Singleton = nullptr;
 public:
+	static inline FMetalDynamicRHI& Get() { return *Singleton; }
 
 	/** Initialization constructor. */
 	FMetalDynamicRHI(ERHIFeatureLevel::Type RequestedFeatureLevel);
@@ -68,18 +70,13 @@ public:
 	virtual FShaderResourceViewRHIRef RHICreateShaderResourceView(const FShaderResourceViewInitializer& Initializer) final override;
 	virtual void RHIUpdateShaderResourceView(FRHIShaderResourceView* SRV, FRHIBuffer* Buffer, uint32 Stride, uint8 Format) final override;
 	virtual void RHIUpdateShaderResourceView(FRHIShaderResourceView* SRV, FRHIBuffer* Buffer) final override;
-	virtual uint64 RHICalcTexture2DPlatformSize(uint32 SizeX, uint32 SizeY, uint8 Format, uint32 NumMips, uint32 NumSamples, ETextureCreateFlags Flags, const FRHIResourceCreateInfo& CreateInfo, uint32& OutAlign) final override;
-	virtual uint64 RHICalcTexture3DPlatformSize(uint32 SizeX, uint32 SizeY, uint32 SizeZ, uint8 Format, uint32 NumMips, ETextureCreateFlags Flags, const FRHIResourceCreateInfo& CreateInfo, uint32& OutAlign) final override;
-	virtual uint64 RHICalcTextureCubePlatformSize(uint32 Size, uint8 Format, uint32 NumMips, ETextureCreateFlags Flags, const FRHIResourceCreateInfo& CreateInfo, uint32& OutAlign) final override;
+	virtual FRHICalcTextureSizeResult RHICalcTexturePlatformSize(FRHITextureDesc const& Desc, uint32 FirstMipIndex) final override;
 	virtual uint64 RHIGetMinimumAlignmentForBufferBackedSRV(EPixelFormat Format) final override;
 	virtual void RHIGetTextureMemoryStats(FTextureMemoryStats& OutStats) final override;
 	virtual bool RHIGetTextureMemoryVisualizeData(FColor* TextureData,int32 SizeX,int32 SizeY,int32 Pitch,int32 PixelSize) final override;
-	virtual FTexture2DRHIRef RHICreateTexture2D(uint32 SizeX, uint32 SizeY, uint8 Format, uint32 NumMips, uint32 NumSamples, ETextureCreateFlags Flags, ERHIAccess InResourceState, FRHIResourceCreateInfo& CreateInfo) final override;
-	virtual FTexture2DRHIRef RHICreateTextureExternal2D(uint32 SizeX, uint32 SizeY, uint8 Format, uint32 NumMips, uint32 NumSamples, ETextureCreateFlags Flags, ERHIAccess InResourceState, FRHIResourceCreateInfo& Info) final override;
+	virtual FTextureRHIRef RHICreateTexture(const FRHITextureCreateDesc& CreateDesc) final override;
 	virtual FTexture2DRHIRef RHIAsyncCreateTexture2D(uint32 SizeX, uint32 SizeY, uint8 Format, uint32 NumMips, ETextureCreateFlags Flags, ERHIAccess InResourceState, void** InitialMipData, uint32 NumInitialMips) final override;
 	virtual void RHICopySharedMips(FRHITexture2D* DestTexture2D, FRHITexture2D* SrcTexture2D) final override;
-	virtual FTexture2DArrayRHIRef RHICreateTexture2DArray(uint32 SizeX, uint32 SizeY, uint32 SizeZ, uint8 Format, uint32 NumMips, uint32 NumSamples, ETextureCreateFlags Flags, ERHIAccess InResourceState, FRHIResourceCreateInfo& CreateInfo) final override;
-	virtual FTexture3DRHIRef RHICreateTexture3D(uint32 SizeX, uint32 SizeY, uint32 SizeZ, uint8 Format, uint32 NumMips, ETextureCreateFlags Flags, ERHIAccess InResourceState, FRHIResourceCreateInfo& CreateInfo) final override;
 	virtual void RHIGenerateMips(FRHITexture* Texture) final override;
 	virtual uint32 RHIComputeMemorySize(FRHITexture* TextureRHI) final override;
 	virtual FTexture2DRHIRef RHIAsyncReallocateTexture2D(FRHITexture2D* Texture2D, int32 NewMipCount, int32 NewSizeX, int32 NewSizeY, FThreadSafeCounter* RequestStatus) final override;
@@ -91,8 +88,6 @@ public:
 	virtual void RHIUnlockTexture2DArray(FRHITexture2DArray* Texture, uint32 TextureIndex, uint32 MipIndex, bool bLockWithinMiptail) final override;
 	virtual void RHIUpdateTexture2D(FRHITexture2D* Texture, uint32 MipIndex, const struct FUpdateTextureRegion2D& UpdateRegion, uint32 SourcePitch, const uint8* SourceData) final override;
 	virtual void RHIUpdateTexture3D(FRHITexture3D* Texture, uint32 MipIndex, const struct FUpdateTextureRegion3D& UpdateRegion, uint32 SourceRowPitch, uint32 SourceDepthPitch, const uint8* SourceData) final override;
-	virtual FTextureCubeRHIRef RHICreateTextureCube(uint32 Size, uint8 Format, uint32 NumMips, ETextureCreateFlags Flags, ERHIAccess InResourceState, FRHIResourceCreateInfo& CreateInfo) final override;
-	virtual FTextureCubeRHIRef RHICreateTextureCubeArray(uint32 Size, uint32 ArraySize, uint8 Format, uint32 NumMips, ETextureCreateFlags Flags, ERHIAccess InResourceState, FRHIResourceCreateInfo& CreateInfo) final override;
 	virtual void* RHILockTextureCubeFace(FRHITextureCube* Texture, uint32 FaceIndex, uint32 ArrayIndex, uint32 MipIndex, EResourceLockMode LockMode, uint32& DestStride, bool bLockWithinMiptail) final override;
 	virtual void RHIUnlockTextureCubeFace(FRHITextureCube* Texture, uint32 FaceIndex, uint32 ArrayIndex, uint32 MipIndex, bool bLockWithinMiptail) final override;
 	virtual void RHIBindDebugLabelName(FRHITexture* Texture, const TCHAR* Name) final override;
@@ -142,11 +137,7 @@ public:
 	virtual void EndUpdateTexture3D_RenderThread(class FRHICommandListImmediate& RHICmdList, FUpdateTexture3DData& UpdateData) final override;
 	virtual void UpdateTexture3D_RenderThread(class FRHICommandListImmediate& RHICmdList, FRHITexture3D* Texture, uint32 MipIndex, const struct FUpdateTextureRegion3D& UpdateRegion, uint32 SourceRowPitch, uint32 SourceDepthPitch, const uint8* SourceData) final override;
 	
-	virtual FTexture2DRHIRef RHICreateTexture2D_RenderThread(class FRHICommandListImmediate& RHICmdList, uint32 SizeX, uint32 SizeY, uint8 Format, uint32 NumMips, uint32 NumSamples, ETextureCreateFlags Flags, ERHIAccess InResourceState, FRHIResourceCreateInfo& CreateInfo) final override;
-	virtual FTexture2DArrayRHIRef RHICreateTexture2DArray_RenderThread(class FRHICommandListImmediate& RHICmdList, uint32 SizeX, uint32 SizeY, uint32 SizeZ, uint8 Format, uint32 NumMips, uint32 NumSamples, ETextureCreateFlags Flags, ERHIAccess InResourceState, FRHIResourceCreateInfo& CreateInfo) final override;
-	virtual FTexture3DRHIRef RHICreateTexture3D_RenderThread(class FRHICommandListImmediate& RHICmdList, uint32 SizeX, uint32 SizeY, uint32 SizeZ, uint8 Format, uint32 NumMips, ETextureCreateFlags Flags, ERHIAccess InResourceState, FRHIResourceCreateInfo& CreateInfo) final override;
-	virtual FTextureCubeRHIRef RHICreateTextureCube_RenderThread(class FRHICommandListImmediate& RHICmdList, uint32 Size, uint8 Format, uint32 NumMips, ETextureCreateFlags Flags, ERHIAccess InResourceState, FRHIResourceCreateInfo& CreateInfo) final override;
-	virtual FTextureCubeRHIRef RHICreateTextureCubeArray_RenderThread(class FRHICommandListImmediate& RHICmdList, uint32 Size, uint32 ArraySize, uint8 Format, uint32 NumMips, ETextureCreateFlags Flags, ERHIAccess InResourceState, FRHIResourceCreateInfo& CreateInfo) final override;
+	virtual FTextureRHIRef RHICreateTexture_RenderThread(class FRHICommandListImmediate& RHICmdList, const FRHITextureCreateDesc& CreateDesc) final override;
 	
 	virtual FRenderQueryRHIRef RHICreateRenderQuery_RenderThread(class FRHICommandListImmediate& RHICmdList, ERenderQueryType QueryType) final override;
 	
diff --git a/Engine/Source/Runtime/Apple/MetalRHI/Private/MetalLLM.cpp b/Engine/Source/Runtime/Apple/MetalRHI/Private/MetalLLM.cpp
--- a/Engine/Source/Runtime/Apple/MetalRHI/Private/MetalLLM.cpp
+++ b/Engine/Source/Runtime/Apple/MetalRHI/Private/MetalLLM.cpp
@@ -143,47 +143,6 @@ static EPixelFormat MetalToRHIPixelFormat(mtlpp::PixelFormat Format)
 	return PF_MAX;
 }
 
-static mtlpp::SizeAndAlign TextureSizeAndAlign(mtlpp::TextureType TextureType, uint32 Width, uint32 Height, uint32 Depth, mtlpp::PixelFormat Format, uint32 MipCount, uint32 SampleCount, uint32 ArrayCount)
-{
-	mtlpp::SizeAndAlign SizeAlign;
-	SizeAlign.Size = 0;
-	SizeAlign.Align = 0;
-	
-	uint32 Align = 0;
-	FRHIResourceCreateInfo CreateInfo(TEXT(""));
-	switch (TextureType)
-	{
-		case mtlpp::TextureType::Texture2D:
-		case mtlpp::TextureType::Texture2DMultisample:
-			SizeAlign.Size = RHICalcTexture2DPlatformSize(Width, Height, MetalToRHIPixelFormat(Format), MipCount, SampleCount, TexCreate_None, CreateInfo, Align);
-			SizeAlign.Align = Align;
-			break;
-		case mtlpp::TextureType::Texture2DArray:
-			SizeAlign.Size = RHICalcTexture2DPlatformSize(Width, Height, MetalToRHIPixelFormat(Format), MipCount, SampleCount, TexCreate_None, CreateInfo, Align) * ArrayCount;
-			SizeAlign.Align = Align;
-			break;
-		case mtlpp::TextureType::TextureCube:
-			SizeAlign.Size = RHICalcTextureCubePlatformSize(Width, MetalToRHIPixelFormat(Format), MipCount, TexCreate_None, CreateInfo, Align);
-			SizeAlign.Align = Align;
-			break;
-		case mtlpp::TextureType::TextureCubeArray:
-			SizeAlign.Size = RHICalcTextureCubePlatformSize(Width, MetalToRHIPixelFormat(Format), MipCount, TexCreate_None, CreateInfo, Align) * ArrayCount;
-			SizeAlign.Align = Align;
-			break;
-		case mtlpp::TextureType::Texture3D:
-			SizeAlign.Size = RHICalcTexture3DPlatformSize(Width, Height, Depth, MetalToRHIPixelFormat(Format), MipCount, TexCreate_None, CreateInfo, Align);
-			SizeAlign.Align = Align;
-			break;
-		case mtlpp::TextureType::Texture1D:
-		case mtlpp::TextureType::Texture1DArray:
-		default:
-			check(false);
-			break;
-	}
-	
-	return SizeAlign;
-}
-
 void MetalLLM::LogAllocTexture(mtlpp::Device& Device, mtlpp::TextureDescriptor const& Desc, mtlpp::Texture const& Texture)
 {
 	mtlpp::SizeAndAlign SizeAlign;
diff --git a/Engine/Source/Runtime/Apple/MetalRHI/Private/MetalRHI.cpp b/Engine/Source/Runtime/Apple/MetalRHI/Private/MetalRHI.cpp
--- a/Engine/Source/Runtime/Apple/MetalRHI/Private/MetalRHI.cpp
+++ b/Engine/Source/Runtime/Apple/MetalRHI/Private/MetalRHI.cpp
@@ -154,6 +154,9 @@ FMetalDynamicRHI::FMetalDynamicRHI(ERHIFeatureLevel::Type RequestedFeatureLevel)
 : ImmediateContext(nullptr, FMetalDeviceContext::CreateDeviceContext())
 , AsyncComputeContext(nullptr)
 {
+	check(Singleton == nullptr);
+	Singleton = this;
+
 	@autoreleasepool {
 	// This should be called once at the start 
 	check( IsInGameThread() );
@@ -1055,28 +1058,12 @@ FMetalDynamicRHI::~FMetalDynamicRHI()
 #endif
 }
 
-uint64 FMetalDynamicRHI::RHICalcTexture2DPlatformSize(uint32 SizeX, uint32 SizeY, uint8 Format, uint32 NumMips, uint32 NumSamples, ETextureCreateFlags Flags, const FRHIResourceCreateInfo& CreateInfo, uint32& OutAlign)
+FDynamicRHI::FRHICalcTextureSizeResult FMetalDynamicRHI::RHICalcTexturePlatformSize(FRHITextureDesc const& Desc, uint32 FirstMipIndex)
 {
-	@autoreleasepool {
-	OutAlign = 0;
-	return CalcTextureSize(SizeX, SizeY, (EPixelFormat)Format, NumMips);
-	}
-}
-
-uint64 FMetalDynamicRHI::RHICalcTexture3DPlatformSize(uint32 SizeX, uint32 SizeY, uint32 SizeZ, uint8 Format, uint32 NumMips, ETextureCreateFlags Flags, const FRHIResourceCreateInfo& CreateInfo, uint32& OutAlign)
-{
-	@autoreleasepool {
-	OutAlign = 0;
-	return CalcTextureSize3D(SizeX, SizeY, SizeZ, (EPixelFormat)Format, NumMips);
-	}
-}
-
-uint64 FMetalDynamicRHI::RHICalcTextureCubePlatformSize(uint32 Size, uint8 Format, uint32 NumMips, ETextureCreateFlags Flags, const FRHIResourceCreateInfo& CreateInfo, uint32& OutAlign)
-{
-	@autoreleasepool {
-	OutAlign = 0;
-	return CalcTextureSize(Size, Size, (EPixelFormat)Format, NumMips) * 6;
-	}
+	FDynamicRHI::FRHICalcTextureSizeResult Result;
+	Result.Size = Desc.CalcMemorySizeEstimate(FirstMipIndex);
+	Result.Align = 0;
+	return Result;
 }
 
 uint64 FMetalDynamicRHI::RHIGetMinimumAlignmentForBufferBackedSRV(EPixelFormat Format)
diff --git a/Engine/Source/Runtime/Apple/MetalRHI/Private/MetalRHIPrivate.h b/Engine/Source/Runtime/Apple/MetalRHI/Private/MetalRHIPrivate.h
--- a/Engine/Source/Runtime/Apple/MetalRHI/Private/MetalRHIPrivate.h
+++ b/Engine/Source/Runtime/Apple/MetalRHI/Private/MetalRHIPrivate.h
@@ -225,6 +225,11 @@ static FORCEINLINE typename TMetalResourceTraits<TRHIType>::TConcreteType* Resou
 	return static_cast<typename TMetalResourceTraits<TRHIType>::TConcreteType*>(Resource);
 }
 
+static FORCEINLINE FMetalSurface* ResourceCast(FRHITexture* Texture)
+{
+	return GetMetalSurfaceFromRHITexture(Texture);
+}
+
 uint32 SafeGetRuntimeDebuggingLevel();
 
 extern int32 GMetalBufferZeroFill;
diff --git a/Engine/Source/Runtime/Apple/MetalRHI/Private/MetalRenderTarget.cpp b/Engine/Source/Runtime/Apple/MetalRHI/Private/MetalRenderTarget.cpp
--- a/Engine/Source/Runtime/Apple/MetalRHI/Private/MetalRenderTarget.cpp
+++ b/Engine/Source/Runtime/Apple/MetalRHI/Private/MetalRenderTarget.cpp
@@ -41,22 +41,25 @@ void FMetalRHICommandContext::RHICopyToResolveTarget(FRHITexture* SourceTextureR
 		// nothing to do if one of the textures is null!
 		return;
 	}
-	if(SourceTextureRHI != DestTextureRHI)
+	if (SourceTextureRHI != DestTextureRHI)
 	{
 		FMetalSurface* Source = GetMetalSurfaceFromRHITexture(SourceTextureRHI);
 		FMetalSurface* Destination = GetMetalSurfaceFromRHITexture(DestTextureRHI);
 		
+		const FRHITextureDesc& SourceDesc = Source->GetDesc();
+		const FRHITextureDesc& DestinationDesc = Destination->GetDesc();
+
 		// Only valid to have nil Metal textures when they are TexCreate_Presentable
-		if(!Source->Texture)
+		if (!Source->Texture)
 		{
 			// Source RHI texture is valid with no Presentable Metal texture - there is nothing to copy from
-			check(EnumHasAnyFlags(Source->Flags, TexCreate_Presentable));
+			check(EnumHasAnyFlags(SourceDesc.Flags, TexCreate_Presentable));
 			return;
 		}
-		if(!Destination->Texture)
+		if (!Destination->Texture)
 		{
 			// Destination RHI texture is valid with no Presentable Metal texture - force fetch it now so we can complete the copy
-			check(EnumHasAnyFlags(Destination->Flags, TexCreate_Presentable));
+			check(EnumHasAnyFlags(DestinationDesc.Flags, TexCreate_Presentable));
 			Destination->GetDrawableTexture();
 			if(!Destination->Texture)
 			{
@@ -65,29 +68,12 @@ void FMetalRHICommandContext::RHICopyToResolveTarget(FRHITexture* SourceTextureR
 			}
 		}
 
-		switch (Source->Type)
-		{
-			case RRT_Texture2D:
-				break;
-			case RRT_TextureCube:
-				check(Source->SizeZ == 6); // Arrays might not work yet.
-				break;
-			default:
-				check(false); // Only Tex2D & Cube are tested to work so far!
-				break;
-		}
-		switch (Destination->Type)
-		{
-			case RRT_Texture2D:
-				break;
-			case RRT_TextureCube:
-				check(Destination->SizeZ == 6); // Arrays might not work yet.
-				break;
-			default:
-				check(false); // Only Tex2D & Cube are tested to work so far!
-				break;
-		}
-		
+		checkf( SourceDesc.IsTexture2D()   || SourceDesc.IsTextureCube(), TEXT("Only Tex2D & Cube are tested to work so far!"));
+		checkf(!SourceDesc.IsTextureCube() || SourceDesc.ArraySize == 1,  TEXT("Cube arrays might not work yet."));
+
+		checkf( DestinationDesc.IsTexture2D()   || DestinationDesc.IsTextureCube(), TEXT("Only Tex2D & Cube are tested to work so far!"));
+		checkf(!DestinationDesc.IsTextureCube() || DestinationDesc.ArraySize == 1,  TEXT("Cube arrays might not work yet."));
+
 		mtlpp::Origin Origin(0, 0, 0);
 		mtlpp::Size Size(0, 0, 1);
 		if (ResolveParams.Rect.IsValid())
@@ -104,15 +90,15 @@ void FMetalRHICommandContext::RHICopyToResolveTarget(FRHITexture* SourceTextureR
 			Origin.x = 0;
 			Origin.y = 0;
 			
-			Size.width = FMath::Max<uint32>(1, Source->SizeX >> ResolveParams.MipIndex);
-			Size.height = FMath::Max<uint32>(1, Source->SizeY >> ResolveParams.MipIndex);
+			Size.width = FMath::Max<uint32>(1, SourceDesc.Extent.X >> ResolveParams.MipIndex);
+			Size.height = FMath::Max<uint32>(1, SourceDesc.Extent.Y >> ResolveParams.MipIndex);
 			// clamp to a destination size
-			Size.width = FMath::Min<uint32>(Size.width, Destination->SizeX >> ResolveParams.MipIndex);
-			Size.height = FMath::Min<uint32>(Size.height, Destination->SizeY >> ResolveParams.MipIndex);
+			Size.width = FMath::Min<uint32>(Size.width, DestinationDesc.Extent.X >> ResolveParams.MipIndex);
+			Size.height = FMath::Min<uint32>(Size.height, DestinationDesc.Extent.Y >> ResolveParams.MipIndex);
 		}
 		
-		const bool bSrcCubemap  = Source->bIsCubemap;
-		const bool bDestCubemap = Destination->bIsCubemap;
+		const bool bSrcCubemap  = SourceDesc.IsTextureCube();
+		const bool bDestCubemap = DestinationDesc.IsTextureCube();
 		
 		uint32 DestIndex = ResolveParams.DestArrayIndex * (bDestCubemap ? 6 : 1) + (bDestCubemap ? uint32(ResolveParams.CubeFace) : 0);
 		uint32 SrcIndex  = ResolveParams.SourceArrayIndex * (bSrcCubemap ? 6 : 1) + (bSrcCubemap ? uint32(ResolveParams.CubeFace) : 0);
@@ -124,7 +110,7 @@ void FMetalRHICommandContext::RHICopyToResolveTarget(FRHITexture* SourceTextureR
 
 		const bool bMSAASource = Source->MSAATexture;
         const bool bMSAADest = Destination->MSAATexture;
-        const bool bDepthStencil = Source->PixelFormat == PF_DepthStencil;
+        const bool bDepthStencil = SourceDesc.Format == PF_DepthStencil;
 		if (bMSAASource && !bMSAADest)
 		{
 			// Resolve required - Device must support this - Using Shader for resolve not supported amd NumSamples should be 1
@@ -141,8 +127,8 @@ void FMetalRHICommandContext::RHICopyToResolveTarget(FRHITexture* SourceTextureR
 		}
 		else
 		{
-			const FPixelFormatInfo& SourceFormatInfo = GPixelFormats[Source->PixelFormat];
-			const FPixelFormatInfo& DestFormatInfo = GPixelFormats[Destination->PixelFormat];
+			const FPixelFormatInfo& SourceFormatInfo = GPixelFormats[SourceDesc.Format];
+			const FPixelFormatInfo& DestFormatInfo = GPixelFormats[DestinationDesc.Format];
 			bool bUsingPixelFormatView = (Source->Texture.GetUsage() & mtlpp::TextureUsage::PixelFormatView) != 0;
 			
 			// Attempt to Resolve with a texture view - source Texture doesn't have to be created with MTLTextureUsagePixelFormatView for these cases e.g:
@@ -283,7 +269,7 @@ void FMetalDynamicRHI::RHIReadSurfaceData(FRHITexture* TextureRHI, FIntRect Rect
 	mtlpp::Region Region(Rect.Min.X, Rect.Min.Y, SizeX, SizeY);
     
 	FMetalTexture Texture = Surface->Texture;
-    if(!Texture && EnumHasAnyFlags(Surface->Flags, TexCreate_Presentable))
+    if(!Texture && EnumHasAnyFlags(Surface->GetDesc().Flags, TexCreate_Presentable))
     {
         Texture = Surface->GetCurrentTexture();
     }
@@ -293,7 +279,7 @@ void FMetalDynamicRHI::RHIReadSurfaceData(FRHITexture* TextureRHI, FIntRect Rect
         return;
     }
 
-	if (GMetalUseTexGetBytes && Surface->PixelFormat != PF_DepthStencil && Surface->PixelFormat != PF_ShadowDepth)
+	if (GMetalUseTexGetBytes && Surface->GetDesc().Format != PF_DepthStencil && Surface->GetDesc().Format != PF_ShadowDepth)
 	{
 		SCOPE_CYCLE_COUNTER(STAT_MetalTexturePageOffTime);
 		
@@ -305,7 +291,7 @@ void FMetalDynamicRHI::RHIReadSurfaceData(FRHITexture* TextureRHI, FIntRect Rect
 #else
 			mtlpp::StorageMode StorageMode = mtlpp::StorageMode::Shared;
 #endif
-			mtlpp::PixelFormat MetalFormat = (mtlpp::PixelFormat)GPixelFormats[Surface->PixelFormat].PlatformFormat;
+			mtlpp::PixelFormat MetalFormat = (mtlpp::PixelFormat)GPixelFormats[Surface->GetDesc().Format].PlatformFormat;
 			mtlpp::TextureDescriptor Desc;
 			Desc.SetTextureType(Texture.GetTextureType());
 			Desc.SetPixelFormat(Texture.GetPixelFormat());
@@ -341,7 +327,7 @@ void FMetalDynamicRHI::RHIReadSurfaceData(FRHITexture* TextureRHI, FIntRect Rect
 		//kick the current command buffer.
 		ImmediateContext.Context->SubmitCommandBufferAndWait();
 		
-		const uint32 Stride = GPixelFormats[Surface->PixelFormat].BlockBytes * SizeX;
+		const uint32 Stride = GPixelFormats[Surface->GetDesc().Format].BlockBytes * SizeX;
 		const uint32 BytesPerImage = Stride * SizeY;
 
 		TArray<uint8> Data;
@@ -349,7 +335,7 @@ void FMetalDynamicRHI::RHIReadSurfaceData(FRHITexture* TextureRHI, FIntRect Rect
 		
 		Texture.GetBytes(Data.GetData(), Stride, BytesPerImage, Region, 0, 0);
 		
-		ConvertSurfaceDataToFColor(Surface->PixelFormat, SizeX, SizeY, (uint8*)Data.GetData(), Stride, OutDataPtr, InFlags);
+		ConvertSurfaceDataToFColor(Surface->GetDesc().Format, SizeX, SizeY, (uint8*)Data.GetData(), Stride, OutDataPtr, InFlags);
 		
 		if (TempTexture)
 		{
@@ -358,7 +344,7 @@ void FMetalDynamicRHI::RHIReadSurfaceData(FRHITexture* TextureRHI, FIntRect Rect
 	}
 	else
 	{
-		uint32 BytesPerPixel = (Surface->PixelFormat != PF_DepthStencil || !InFlags.GetOutputStencil()) ? GPixelFormats[Surface->PixelFormat].BlockBytes : 1;
+		uint32 BytesPerPixel = (Surface->GetDesc().Format != PF_DepthStencil || !InFlags.GetOutputStencil()) ? GPixelFormats[Surface->GetDesc().Format].BlockBytes : 1;
 		const uint32 Stride = BytesPerPixel * SizeX;
 		const uint32 Alignment = PLATFORM_MAC ? 1u : 64u; // Mac permits natural row alignment (tightly-packed) but iOS does not.
 		const uint32 AlignedStride = ((Stride - 1) & ~(Alignment - 1)) + Alignment;
@@ -368,7 +354,7 @@ void FMetalDynamicRHI::RHIReadSurfaceData(FRHITexture* TextureRHI, FIntRect Rect
 			// Synchronise the texture with the CPU
 			SCOPE_CYCLE_COUNTER(STAT_MetalTexturePageOffTime);
 			
-			if (Surface->PixelFormat != PF_DepthStencil)
+			if (Surface->GetDesc().Format != PF_DepthStencil)
 			{
 				ImmediateContext.Context->CopyFromTextureToBuffer(Texture, 0, InFlags.GetMip(), Region.origin, Region.size, Buffer, 0, AlignedStride, BytesPerImage, mtlpp::BlitOption::None);
 			}
@@ -387,7 +373,7 @@ void FMetalDynamicRHI::RHIReadSurfaceData(FRHITexture* TextureRHI, FIntRect Rect
 			//kick the current command buffer.
 			ImmediateContext.Context->SubmitCommandBufferAndWait();
 			
-			ConvertSurfaceDataToFColor(Surface->PixelFormat, SizeX, SizeY, (uint8*)Buffer.GetContents(), AlignedStride, OutDataPtr, InFlags);
+			ConvertSurfaceDataToFColor(Surface->GetDesc().Format, SizeX, SizeY, (uint8*)Buffer.GetContents(), AlignedStride, OutDataPtr, InFlags);
 		}
 		((FMetalDeviceContext*)ImmediateContext.Context)->ReleaseBuffer(Buffer);
 	}
@@ -398,7 +384,6 @@ void FMetalDynamicRHI::RHIMapStagingSurface(FRHITexture* TextureRHI, FRHIGPUFenc
 {
 	@autoreleasepool {
     FMetalSurface* Surface = GetMetalSurfaceFromRHITexture(TextureRHI);
-    FMetalTexture2D* Texture = (FMetalTexture2D*)TextureRHI->GetTexture2D();
     
 #if PLATFORM_MAC
 	uint16 FencePoll = FenceRHI && FenceRHI->Poll() ? 1 : 0;
@@ -406,8 +391,8 @@ void FMetalDynamicRHI::RHIMapStagingSurface(FRHITexture* TextureRHI, FRHIGPUFenc
 #endif
     
     uint32 Stride = 0;
-    OutWidth = Texture->GetSizeX();
-    OutHeight = Texture->GetSizeY();
+    OutWidth = Surface->GetSizeX();
+    OutHeight = Surface->GetSizeY();
     OutData = Surface->Lock(0, 0, RLM_ReadOnly, Stride);
     
 #if PLATFORM_MAC
@@ -431,7 +416,7 @@ void FMetalDynamicRHI::RHIReadSurfaceFloatData(FRHITexture* TextureRHI, FIntRect
 	FMetalSurface* Surface = GetMetalSurfaceFromRHITexture(TextureRHI);
 	
     FMetalTexture Texture = Surface->Texture;
-    if(!Texture && EnumHasAnyFlags(Surface->Flags, TexCreate_Presentable))
+    if(!Texture && EnumHasAnyFlags(Surface->GetDesc().Flags, TexCreate_Presentable))
     {
 		Texture = Surface->GetCurrentTexture();
     }
@@ -442,7 +427,7 @@ void FMetalDynamicRHI::RHIReadSurfaceFloatData(FRHITexture* TextureRHI, FIntRect
     }
     
 	// verify the input image format (but don't crash)
-	if (Surface->PixelFormat != PF_FloatRGBA)
+	if (Surface->GetDesc().Format != PF_FloatRGBA)
 	{
 		UE_LOG(LogRHI, Log, TEXT("Trying to read non-FloatRGBA surface."));
 	}
@@ -463,7 +448,7 @@ void FMetalDynamicRHI::RHIReadSurfaceFloatData(FRHITexture* TextureRHI, FIntRect
 	mtlpp::Region Region = mtlpp::Region(Rect.Min.X, Rect.Min.Y, SizeX, SizeY);
 	
 	// function wants details about the destination, not the source
-	const uint32 Stride = GPixelFormats[Surface->PixelFormat].BlockBytes * SizeX;
+	const uint32 Stride = GPixelFormats[Surface->GetDesc().Format].BlockBytes * SizeX;
 	const uint32 Alignment = PLATFORM_MAC ? 1u : 64u; // Mac permits natural row alignment (tightly-packed) but iOS does not.
 	const uint32 AlignedStride = ((Stride - 1) & ~(Alignment - 1)) + Alignment;
 	const uint32 BytesPerImage = AlignedStride  * SizeY;
@@ -514,7 +499,7 @@ void FMetalDynamicRHI::RHIRead3DSurfaceFloatData(FRHITexture* TextureRHI,FIntRec
 	}
 	
 	// verify the input image format (but don't crash)
-	if (Surface->PixelFormat != PF_FloatRGBA)
+	if (Surface->GetDesc().Format != PF_FloatRGBA)
 	{
 		UE_LOG(LogRHI, Log, TEXT("Trying to read non-FloatRGBA surface."));
 	}
@@ -529,7 +514,7 @@ void FMetalDynamicRHI::RHIRead3DSurfaceFloatData(FRHITexture* TextureRHI,FIntRec
 	mtlpp::Region Region = mtlpp::Region(InRect.Min.X, InRect.Min.Y, ZMinMax.X, SizeX, SizeY, SizeZ);
 	
 	// function wants details about the destination, not the source
-	const uint32 Stride = GPixelFormats[Surface->PixelFormat].BlockBytes * SizeX;
+	const uint32 Stride = GPixelFormats[Surface->GetDesc().Format].BlockBytes * SizeX;
 	const uint32 Alignment = PLATFORM_MAC ? 1u : 64u; // Mac permits natural row alignment (tightly-packed) but iOS does not.
 	const uint32 AlignedStride = ((Stride - 1) & ~(Alignment - 1)) + Alignment;
 	const uint32 BytesPerImage = AlignedStride  * SizeY;
diff --git a/Engine/Source/Runtime/Apple/MetalRHI/Private/MetalStateCache.cpp b/Engine/Source/Runtime/Apple/MetalRHI/Private/MetalStateCache.cpp
--- a/Engine/Source/Runtime/Apple/MetalRHI/Private/MetalStateCache.cpp
+++ b/Engine/Source/Runtime/Apple/MetalRHI/Private/MetalStateCache.cpp
@@ -501,8 +501,8 @@ bool FMetalStateCache::SetRenderPassInfo(FRHIRenderPassInfo const& InRenderTarge
 				FMetalSurface& Surface = *GetMetalSurfaceFromRHITexture(RenderTargetView.RenderTarget);
 				FormatKey = Surface.FormatKey;
 				
-				uint32 Width = FMath::Max((uint32)(Surface.SizeX >> RenderTargetView.MipIndex), (uint32)1);
-				uint32 Height = FMath::Max((uint32)(Surface.SizeY >> RenderTargetView.MipIndex), (uint32)1);
+				uint32 Width = FMath::Max((uint32)(Surface.GetDesc().Extent.X >> RenderTargetView.MipIndex), (uint32)1);
+				uint32 Height = FMath::Max((uint32)(Surface.GetDesc().Extent.Y >> RenderTargetView.MipIndex), (uint32)1);
 				if(!bFramebufferSizeSet)
 				{
 					bFramebufferSizeSet = true;
@@ -539,16 +539,17 @@ bool FMetalStateCache::SetRenderPassInfo(FRHIRenderPassInfo const& InRenderTarge
 	
 				// user code generally passes -1 as a default, but we need 0
 				uint32 ArraySliceIndex = RenderTargetView.ArraySlice == 0xFFFFFFFF ? 0 : RenderTargetView.ArraySlice;
-				if (Surface.bIsCubemap)
+				if (Surface.GetDesc().IsTextureCube())
 				{
 					ArraySliceIndex = GetMetalCubeFace((ECubeFace)ArraySliceIndex);
 				}
 				
-				switch(Surface.Type)
+				switch(Surface.GetDesc().Dimension)
 				{
-					case RRT_Texture2DArray:
-					case RRT_Texture3D:
-					case RRT_TextureCube:
+					case ETextureDimension::Texture2DArray:
+					case ETextureDimension::Texture3D:
+					case ETextureDimension::TextureCube:
+					case ETextureDimension::TextureCubeArray:
 						if(RenderTargetView.ArraySlice == 0xFFFFFFFF)
 						{
 							ArrayTargets |= (1 << RenderTargetIndex);
@@ -614,7 +615,7 @@ bool FMetalStateCache::SetRenderPassInfo(FRHIRenderPassInfo const& InRenderTarge
 				}
 				
 				ColorAttachment.SetLevel(RenderTargetView.MipIndex);
-				if(Surface.Type == RRT_Texture3D)
+				if(Surface.GetDesc().IsTexture3D())
 				{
 					ColorAttachment.SetSlice(0);
 					ColorAttachment.SetDepthPlane(ArraySliceIndex);
@@ -677,11 +678,12 @@ bool FMetalStateCache::SetRenderPassInfo(FRHIRenderPassInfo const& InRenderTarge
 		{
 			FMetalSurface& Surface = *GetMetalSurfaceFromRHITexture(RenderPassInfo.DepthStencilRenderTarget.DepthStencilTarget);
 			
-			switch(Surface.Type)
+			switch(Surface.GetDesc().Dimension)
 			{
-				case RRT_Texture2DArray:
-				case RRT_Texture3D:
-				case RRT_TextureCube:
+				case ETextureDimension::Texture2DArray:
+				case ETextureDimension::Texture3D:
+				case ETextureDimension::TextureCube:
+				case ETextureDimension::TextureCubeArray:
 					ArrayRenderLayers = Surface.GetNumFaces();
 					break;
 				default:
@@ -700,13 +702,13 @@ bool FMetalStateCache::SetRenderPassInfo(FRHIRenderPassInfo const& InRenderTarge
 			if(!bFramebufferSizeSet)
 			{
 				bFramebufferSizeSet = true;
-				FrameBufferSize.width = Surface.SizeX;
-				FrameBufferSize.height = Surface.SizeY;
+				FrameBufferSize.width  = Surface.GetDesc().Extent.X;
+				FrameBufferSize.height = Surface.GetDesc().Extent.Y;
 			}
 			else
 			{
-				FrameBufferSize.width = FMath::Min(FrameBufferSize.width, (CGFloat)Surface.SizeX);
-				FrameBufferSize.height = FMath::Min(FrameBufferSize.height, (CGFloat)Surface.SizeY);
+				FrameBufferSize.width = FMath::Min(FrameBufferSize.width, (CGFloat)Surface.GetDesc().Extent.X);
+				FrameBufferSize.height = FMath::Min(FrameBufferSize.height, (CGFloat)Surface.GetDesc().Extent.Y);
 			}
 			
 			EPixelFormat DepthStencilPixelFormat = RenderPassInfo.DepthStencilRenderTarget.DepthStencilTarget->GetFormat();
@@ -1195,7 +1197,7 @@ void FMetalStateCache::SetVisibilityResultMode(mtlpp::VisibilityResultMode const
 void FMetalStateCache::ConditionalUpdateBackBuffer(FMetalSurface& Surface)
 {
 	// are we setting the back buffer? if so, make sure we have the drawable
-	if (EnumHasAnyFlags(Surface.Flags, TexCreate_Presentable))
+	if (EnumHasAnyFlags(Surface.GetDesc().Flags, TexCreate_Presentable))
 	{
 		// update the back buffer texture the first time used this frame
 		if (Surface.Texture.GetPtr() == nil)
@@ -1453,42 +1455,33 @@ void FMetalStateCache::SetShaderResourceView(FMetalContext* Context, EMetalShade
 {
 	if (SRV)
 	{
-		FRHITexture* Texture = SRV->SourceTexture.GetReference();
-		FMetalVertexBuffer* VB = SRV->SourceVertexBuffer.GetReference();
-		FMetalIndexBuffer* IB = SRV->SourceIndexBuffer.GetReference();
-		FMetalStructuredBuffer* SB = SRV->SourceStructuredBuffer.GetReference();
-		if (Texture)
+		if (SRV->bTexture)
 		{
-			FMetalSurface* Surface = SRV->TextureView;
-			if (Surface != nullptr)
+			FMetalTexture const& View = SRV->GetTextureView();
+			if (View)
 			{
-				SetShaderTexture(ShaderStage, Surface->Texture, BindIndex, mtlpp::ResourceUsage(mtlpp::ResourceUsage::Read|mtlpp::ResourceUsage::Sample));
+				SetShaderTexture(ShaderStage, View, BindIndex, mtlpp::ResourceUsage(mtlpp::ResourceUsage::Read | mtlpp::ResourceUsage::Sample));
 			}
 			else
 			{
 				SetShaderTexture(ShaderStage, nil, BindIndex, mtlpp::ResourceUsage(0));
 			}
 		}
-		else if (IsLinearBuffer(ShaderStage, BindIndex) && SRV->GetLinearTexture(false))
-		{
-			ns::AutoReleased<FMetalTexture> Tex;
-			Tex = SRV->GetLinearTexture(false);
-			
-			SetShaderTexture(ShaderStage, Tex, BindIndex, mtlpp::ResourceUsage(mtlpp::ResourceUsage::Read|mtlpp::ResourceUsage::Sample));
-			
-			// Removed the binding of the underlying buffer for SRVs as we are now using texture buffers here
-		}
-		else if (VB)
-		{
-			SetShaderBuffer(ShaderStage, VB->GetCurrentBufferOrNil(), VB->Data, SRV->Offset, VB->GetSize(), BindIndex, mtlpp::ResourceUsage::Read, (EPixelFormat)SRV->Format);
-		}
-		else if (IB)
-		{
-			SetShaderBuffer(ShaderStage, IB->GetCurrentBufferOrNil(), nil, SRV->Offset, IB->GetSize(), BindIndex, mtlpp::ResourceUsage::Read, (EPixelFormat)SRV->Format);
-		}
-		else if (SB)
+		else
 		{
-			SetShaderBuffer(ShaderStage, SB->GetCurrentBufferOrNil(), nil, SRV->Offset, SB->GetSize(), BindIndex, mtlpp::ResourceUsage::Read);
+			FMetalResourceMultiBuffer* Buffer = SRV->GetSourceBuffer();
+
+			if (IsLinearBuffer(ShaderStage, BindIndex) && SRV->GetLinearTexture())
+			{
+				ns::AutoReleased<FMetalTexture> Tex;
+				Tex = SRV->GetLinearTexture();
+
+				SetShaderTexture(ShaderStage, Tex, BindIndex, mtlpp::ResourceUsage(mtlpp::ResourceUsage::Read | mtlpp::ResourceUsage::Sample));
+			}
+			else 
+			{
+				SetShaderBuffer(ShaderStage, Buffer->GetCurrentBufferOrNil(), Buffer->Data, SRV->Offset, Buffer->GetSize(), BindIndex, mtlpp::ResourceUsage::Read, (EPixelFormat)SRV->Format);
+			}
 		}
 	}
 }
@@ -1523,69 +1516,22 @@ void FMetalStateCache::SetShaderUnorderedAccessView(EMetalShaderStages ShaderSta
 {
 	if (UAV)
 	{
-		// figure out which one of the resources we need to set
-		FMetalStructuredBuffer* StructuredBuffer = UAV->SourceView->SourceStructuredBuffer.GetReference();
-		FMetalVertexBuffer* VertexBuffer = UAV->SourceView->SourceVertexBuffer.GetReference();
-		FMetalIndexBuffer* IndexBuffer = UAV->SourceView->SourceIndexBuffer.GetReference();
-		FRHITexture* Texture = UAV->SourceView->SourceTexture.GetReference();
-		FMetalSurface* Surface = UAV->SourceView->TextureView;
-		if (StructuredBuffer)
-		{
-			SetShaderBuffer(ShaderStage, StructuredBuffer->GetCurrentBufferOrNil(), nil, 0, StructuredBuffer->GetSize(), BindIndex, mtlpp::ResourceUsage(mtlpp::ResourceUsage::Read | mtlpp::ResourceUsage::Write));
-		}
-		else if (VertexBuffer)
-		{
-			check(!VertexBuffer->Data && VertexBuffer->GetCurrentBufferOrNil());
-			if (IsLinearBuffer(ShaderStage, BindIndex) && UAV->SourceView->GetLinearTexture(true))
-			{
-				ns::AutoReleased<FMetalTexture> Tex;
-				Tex = UAV->SourceView->GetLinearTexture(true);
-				SetShaderTexture(ShaderStage, Tex, BindIndex, mtlpp::ResourceUsage(mtlpp::ResourceUsage::Read | mtlpp::ResourceUsage::Write));
-                
-                SetShaderBuffer(ShaderStage, VertexBuffer->GetCurrentBufferOrNil(), VertexBuffer->Data, 0, VertexBuffer->GetSize(), BindIndex, mtlpp::ResourceUsage(mtlpp::ResourceUsage::Read | mtlpp::ResourceUsage::Write), (EPixelFormat)UAV->SourceView->Format);
-			}
-			else
-			{
-				SetShaderBuffer(ShaderStage, VertexBuffer->GetCurrentBufferOrNil(), VertexBuffer->Data, 0, VertexBuffer->GetSize(), BindIndex, mtlpp::ResourceUsage(mtlpp::ResourceUsage::Read | mtlpp::ResourceUsage::Write), (EPixelFormat)UAV->SourceView->Format);
-			}
-		}
-		else if (IndexBuffer)
+		if (UAV->bTexture)
 		{
-			check(IndexBuffer->GetCurrentBufferOrNil());
-			if (IsLinearBuffer(ShaderStage, BindIndex) && UAV->SourceView->GetLinearTexture(true))
-			{
-				ns::AutoReleased<FMetalTexture> Tex;
-				Tex = UAV->SourceView->GetLinearTexture(true);
-				SetShaderTexture(ShaderStage, Tex, BindIndex, mtlpp::ResourceUsage(mtlpp::ResourceUsage::Read | mtlpp::ResourceUsage::Write));
-				
-				SetShaderBuffer(ShaderStage, IndexBuffer->GetCurrentBufferOrNil(), nullptr, 0, IndexBuffer->GetSize(), BindIndex, mtlpp::ResourceUsage(mtlpp::ResourceUsage::Read | mtlpp::ResourceUsage::Write), (EPixelFormat)UAV->SourceView->Format);
-			}
-			else
-			{
-				SetShaderBuffer(ShaderStage, IndexBuffer->GetCurrentBufferOrNil(), nullptr, 0, IndexBuffer->GetSize(), BindIndex, mtlpp::ResourceUsage(mtlpp::ResourceUsage::Read | mtlpp::ResourceUsage::Write), (EPixelFormat)UAV->SourceView->Format);
-			}
-		}
-		else if (Texture)
-		{
-			if (!Surface)
-			{
-				Surface = GetMetalSurfaceFromRHITexture(Texture);
-			}
-			if (Surface != nullptr)
+			FMetalSurface* Surface = UAV->GetSourceTexture();
+			FMetalTexture const& View = UAV->GetTextureView();
+
+			if (View)
 			{
-				FMetalSurface* Source = GetMetalSurfaceFromRHITexture(Texture);
-				
 				FPlatformAtomics::InterlockedExchange(&Surface->Written, 1);
-				FPlatformAtomics::InterlockedExchange(&Source->Written, 1);
-				
-				SetShaderTexture(ShaderStage, Surface->Texture, BindIndex, mtlpp::ResourceUsage(mtlpp::ResourceUsage::Read | mtlpp::ResourceUsage::Write));
-				
-				if (Surface->Texture.GetBuffer() &&
-					(EnumHasAllFlags(Source->Flags, TexCreate_UAV | TexCreate_NoTiling) || EnumHasAllFlags(Source->Flags, TexCreate_AtomicCompatible)))
+
+				SetShaderTexture(ShaderStage, View, BindIndex, mtlpp::ResourceUsage(mtlpp::ResourceUsage::Read | mtlpp::ResourceUsage::Write));
+
+				if (Surface->Texture.GetBuffer() && (EnumHasAllFlags(Surface->GetDesc().Flags, TexCreate_UAV | TexCreate_NoTiling) || EnumHasAllFlags(Surface->GetDesc().Flags, TexCreate_AtomicCompatible)))
 				{
 					uint32 BytesPerRow = Surface->Texture.GetBufferBytesPerRow();
-					uint32 ElementsPerRow = BytesPerRow / GPixelFormats[(EPixelFormat)Texture->GetFormat()].BlockBytes;
-				
+					uint32 ElementsPerRow = BytesPerRow / GPixelFormats[(EPixelFormat)Surface->GetFormat()].BlockBytes;
+
 					FMetalBuffer Buffer(Surface->Texture.GetBuffer(), false);
 					const uint32 BufferOffset = Surface->Texture.GetBufferOffset();
 					const uint32 BufferSize = Surface->Texture.GetBuffer().GetLength();
@@ -1597,6 +1543,20 @@ void FMetalStateCache::SetShaderUnorderedAccessView(EMetalShaderStages ShaderSta
 				SetShaderTexture(ShaderStage, nil, BindIndex, mtlpp::ResourceUsage(0));
 			}
 		}
+		else
+		{
+			FMetalResourceMultiBuffer* Buffer = UAV->GetSourceBuffer();
+			check(!Buffer->Data && Buffer->GetCurrentBufferOrNil());
+
+			if (IsLinearBuffer(ShaderStage, BindIndex) && UAV->GetLinearTexture())
+			{
+				ns::AutoReleased<FMetalTexture> Tex;
+				Tex = UAV->GetLinearTexture();
+				SetShaderTexture(ShaderStage, Tex, BindIndex, mtlpp::ResourceUsage(mtlpp::ResourceUsage::Read | mtlpp::ResourceUsage::Write));
+			}
+
+			SetShaderBuffer(ShaderStage, Buffer->GetCurrentBufferOrNil(), Buffer->Data, 0, Buffer->GetSize(), BindIndex, mtlpp::ResourceUsage(mtlpp::ResourceUsage::Read | mtlpp::ResourceUsage::Write), (EPixelFormat)UAV->Format);
+		}
 	}
 }
 
