Copying //UE5/Dev-ParallelRendering to //UE5/Main

diff --git a/Engine/Source/Runtime/Apple/MetalRHI/Private/MetalTexture.cpp b/Engine/Source/Runtime/Apple/MetalRHI/Private/MetalTexture.cpp
--- a/Engine/Source/Runtime/Apple/MetalRHI/Private/MetalTexture.cpp
+++ b/Engine/Source/Runtime/Apple/MetalRHI/Private/MetalTexture.cpp
@@ -32,15 +32,9 @@ FAutoConsoleVariableRef CVarMetalForceIOSTexturesShared(
 /** Given a pointer to a RHI texture that was created by the Metal RHI, returns a pointer to the FMetalTextureBase it encapsulates. */
 FMetalSurface* GetMetalSurfaceFromRHITexture(FRHITexture* Texture)
 {
-	if (!Texture)
-	{
-		return NULL;
-	}
-	else
-	{
-		FMetalSurface* Surface = (FMetalSurface*)Texture->GetTextureBaseRHI();
-		return Surface;
-	}
+	return Texture
+		? static_cast<FMetalSurface*>(Texture->GetTextureBaseRHI())
+		: nullptr;
 }
 
 static bool IsRenderTarget(ETextureCreateFlags Flags)
@@ -137,11 +131,11 @@ static bool IsPixelFormatPVRTCCompressed(EPixelFormat Format)
 	}
 }
 
-void SafeReleaseMetalTexture(FMetalSurface* Surface, FMetalTexture& Texture, bool bTextureView, bool bAVFoundationTexture)
+void SafeReleaseMetalTexture(FMetalSurface* Surface, FMetalTexture& Texture, bool bAVFoundationTexture)
 {
 	if(GIsMetalInitialized && GDynamicRHI)
 	{
-		if (!bAVFoundationTexture && !bTextureView)
+		if (!bAVFoundationTexture)
 		{
 			GetMetalDeviceContext().ReleaseTexture(Surface, Texture);
 		}
@@ -277,7 +271,6 @@ void FMetalSurface::PrepareTextureView()
 #endif
 	if(!(Usage & mtlpp::TextureUsage::PixelFormatView) && !bMemoryLess)
 	{
-		check(!bTextureView);
 		check(ImageSurfaceRef == nullptr);
 		
 		check(Texture);
@@ -287,18 +280,18 @@ void FMetalSurface::PrepareTextureView()
 		{
 			FMetalTexture OldMSAATexture = MSAATexture;
 			MSAATexture = Reallocate(MSAATexture, mtlpp::TextureUsage::PixelFormatView);
-			SafeReleaseMetalTexture(this, OldMSAATexture, bTextureView, ImageSurfaceRef != nullptr);
+			SafeReleaseMetalTexture(this, OldMSAATexture, ImageSurfaceRef != nullptr);
 		}
 		if (MSAAResolveTexture && !bMSAAResolveTextureIsTexture)
 		{
 			FMetalTexture OldMSAAResolveTexture = MSAAResolveTexture;
 			MSAAResolveTexture = Reallocate(MSAAResolveTexture, mtlpp::TextureUsage::PixelFormatView);
-			SafeReleaseMetalTexture(this, OldMSAAResolveTexture, bTextureView, ImageSurfaceRef != nullptr);
+			SafeReleaseMetalTexture(this, OldMSAAResolveTexture, ImageSurfaceRef != nullptr);
 		}
 		
 		FMetalTexture OldTexture = Texture;
 		Texture = Reallocate(Texture, mtlpp::TextureUsage::PixelFormatView);
-		SafeReleaseMetalTexture(this, OldTexture, bTextureView, ImageSurfaceRef != nullptr);
+		SafeReleaseMetalTexture(this, OldTexture, ImageSurfaceRef != nullptr);
 		
 		if (bMSAATextureIsTexture)
 		{
@@ -335,95 +328,8 @@ FMetalTexture FMetalSurface::Reallocate(FMetalTexture InTexture, mtlpp::TextureU
 	return NewTex;
 }
 
-void FMetalSurface::ReplaceTexture(FMetalContext& Context, FMetalTexture CurrentTexture, FMetalTexture NewTexture)
-{
-	check(NewTexture);
-	{
-		// Copy texture content...
-		const uint32 BlockSizeX = GPixelFormats[PixelFormat].BlockSizeX;
-		const uint32 BlockSizeY = GPixelFormats[PixelFormat].BlockSizeY;
-		const uint32 BlockSizeZ = GPixelFormats[PixelFormat].BlockSizeZ;
-		
-		mtlpp::Origin Origin(0,0,0);
-		
-		// DXT/BC formats on Mac actually do have mip-tails that are smaller than the block size, they end up being uncompressed.
-		bool const bPixelFormatASTC = IsPixelFormatASTCCompressed(PixelFormat);
-		
-		for (uint32 SliceIndex = 0; SliceIndex < CurrentTexture.GetArrayLength(); SliceIndex++)
-		{
-			for (uint32 MipIndex = 0; MipIndex < CurrentTexture.GetMipmapLevelCount(); ++MipIndex)
-			{
-				const uint32 UnalignedMipSizeX = FMath::Max<uint32>(1, CurrentTexture.GetWidth() >> (MipIndex));
-				const uint32 UnalignedMipSizeY = FMath::Max<uint32>(1, CurrentTexture.GetHeight() >> (MipIndex));
-				const uint32 UnalignedMipSizeZ = FMath::Max<uint32>(1, CurrentTexture.GetDepth() >> (MipIndex));
-				const uint32 MipSizeX = (bPixelFormatASTC) ? AlignArbitrary(UnalignedMipSizeX, BlockSizeX) : UnalignedMipSizeX;
-				const uint32 MipSizeY = (bPixelFormatASTC) ? AlignArbitrary(UnalignedMipSizeY, BlockSizeY) : UnalignedMipSizeY;
-				const uint32 MipSizeZ = (bPixelFormatASTC) ? AlignArbitrary(UnalignedMipSizeZ, BlockSizeZ) : UnalignedMipSizeZ;
-				
-				Context.CopyFromTextureToTexture(CurrentTexture, SliceIndex, MipIndex, Origin, mtlpp::Size(MipSizeX, MipSizeY, MipSizeZ), NewTexture, SliceIndex, MipIndex, Origin);
-			}
-		}
-		
-	}
-	
-	FMetalTexture OldTexture = Texture;
-	if (Texture && Texture == CurrentTexture)
-	{
-		SafeReleaseMetalTexture(OldTexture);
-		Texture = NewTexture;
-	}
-	if (MSAATexture && MSAATexture == CurrentTexture)
-	{
-		if (MSAATexture && (MSAATexture != OldTexture))
-		{
-			FMetalTexture OldMSAATexture = MSAATexture;
-			SafeReleaseMetalTexture(OldMSAATexture);
-			MSAATexture = NewTexture;
-		}
-		else if(MSAATexture)
-		{
-			MSAATexture = NewTexture;
-		}
-		
-		if (MSAAResolveTexture && MSAAResolveTexture == CurrentTexture)
-		{
-			if (MSAAResolveTexture && (MSAAResolveTexture != OldTexture))
-			{
-				FMetalTexture OldMSAAResolveTexture = MSAAResolveTexture;
-				SafeReleaseMetalTexture(OldMSAAResolveTexture);
-				MSAAResolveTexture = NewTexture;
-			}
-			else if (MSAAResolveTexture)
-			{
-				MSAAResolveTexture = NewTexture;
-			}
-		}
-	}
-	
-	for (FMetalShaderResourceView* SRV : SRVs)
-	{
-		if(OldTexture != SRV->TextureView->Texture)
-		{
-			SafeReleaseMetalTexture(nullptr, SRV->TextureView->Texture, true, false);
-		}
-		SRV->TextureView->Texture = nil;
-		
-		SRV->TextureView->MSAATexture = nil;
-		
-		if (SRV->Format == PF_Unknown)
-		{
-			SRV->TextureView->Init(*this, NSMakeRange(SRV->MipLevel, SRV->NumMips));
-		}
-		else
-		{
-			SRV->TextureView->Init(*this, NSMakeRange(SRV->MipLevel, SRV->NumMips), (EPixelFormat)SRV->Format, SRV->bSRGBForceDisable);
-		}
-	}
-}
-
 void FMetalSurface::MakeAliasable(void)
 {
-	check(!bTextureView);
 	check(ImageSurfaceRef == nullptr);
 	
 	static bool bSupportsHeaps = GetMetalDeviceContext().SupportsFeature(EMetalFeaturesHeaps);
@@ -446,313 +352,70 @@ void FMetalSurface::MakeAliasable(void)
 	}
 }
 
-void FMetalSurface::MakeUnAliasable(void)
-{
-	check(!bTextureView);
-	check(ImageSurfaceRef == nullptr);
-	
-	static bool bSupportsHeaps = GetMetalDeviceContext().SupportsFeature(EMetalFeaturesHeaps);
-	if (bSupportsHeaps && Texture.GetStorageMode() == mtlpp::StorageMode::Private && Texture.GetHeap() && Texture.IsAliasable())
-	{
-		FMetalTexture OldTexture = Texture;
-		Texture = Reallocate(Texture, mtlpp::TextureUsage::Unknown);
-		SafeReleaseMetalTexture(this, OldTexture);
-		if (MSAATexture && (MSAATexture != OldTexture))
-		{
-			FMetalTexture OldMSAATexture = MSAATexture;
-			MSAATexture = Reallocate(MSAATexture, mtlpp::TextureUsage::Unknown);
-			SafeReleaseMetalTexture(this, OldMSAATexture);
-		}
-		else if(MSAATexture)
-		{
-			MSAATexture = Texture;
-		}
-		
-		for (FMetalShaderResourceView* SRV : SRVs)
-		{
-			if(OldTexture != SRV->TextureView->Texture)
-			{
-				SafeReleaseMetalTexture(this, SRV->TextureView->Texture, true, false);
-			}
-			SRV->TextureView->Texture = nil;
-			
-			SRV->TextureView->MSAATexture = nil;
-			
-			if (SRV->Format == PF_Unknown)
-			{
-				SRV->TextureView->Init(*this, NSMakeRange(SRV->MipLevel, SRV->NumMips));
-			}
-			else
-			{
-				SRV->TextureView->Init(*this, NSMakeRange(SRV->MipLevel, SRV->NumMips), (EPixelFormat)SRV->Format, SRV->bSRGBForceDisable);
-			}
-		}
-	}
-}
-
-void FMetalSurface::Init(FMetalSurface& Source, NSRange MipRange)
+uint8 GetMetalPixelFormatKey(mtlpp::PixelFormat Format)
 {
-#if PLATFORM_IOS
-	// Mmeory;ess targets can't have texture views (SRVs or UAVs)
-	if (Source.Texture.GetStorageMode() == mtlpp::StorageMode::Memoryless)
-	{
-		return;
-	}
-#endif
-	auto TextureViewType = [&Source]()
-	{
-		return Source.bIsCubemap ? mtlpp::TextureType::Texture2DArray : Source.Texture.GetTextureType();
-	};
-
-	mtlpp::PixelFormat MetalFormat = (mtlpp::PixelFormat)GPixelFormats[PixelFormat].PlatformFormat;
-	if (EnumHasAnyFlags(Flags, TexCreate_SRGB))
-	{
-		// Ensure we have the correct sRGB target format if we create a new texture view rather than using the source texture
-		MetalFormat = ToSRGBFormat(MetalFormat);
-	}
-	
-	bool const bUseSourceTex = (Source.PixelFormat != PF_DepthStencil) && MipRange.location == 0 && MipRange.length == Source.Texture.GetMipmapLevelCount();
-	
-	// Recreate the texture to enable MTLTextureUsagePixelFormatView which must be off unless we definitely use this feature or we are throwing ~4% performance vs. Windows on the floor.
-	mtlpp::TextureUsage Usage = Source.Texture.GetUsage();
-	if(!(Usage & mtlpp::TextureUsage::PixelFormatView) && (Source.PixelFormat != PF_DepthStencil) && !bUseSourceTex)
-	{
-		Source.PrepareTextureView();
-	}
-	
-	ns::Range Slices(0, Source.Texture.GetArrayLength() * (bIsCubemap ? 6 : 1));
-	// Stencil requires a format conversion, so this will access depth only, consequently there are no mip levels other than 0, so this path requires no real special casing.
-	if(Source.PixelFormat != PF_DepthStencil && !bUseSourceTex)
+	struct FMetalPixelFormatKeyMap
 	{
-		Texture = Source.Texture.NewTextureView(MetalFormat, TextureViewType(), ns::Range(MipRange.location, MipRange.length), Slices);
-	}
-	else
-	{
-		Texture = Source.Texture;
-	}
-	
-	const uint32 BlockSizeX = GPixelFormats[PixelFormat].BlockSizeX;
-	const uint32 BlockSizeY = GPixelFormats[PixelFormat].BlockSizeY;
-	const uint32 BlockBytes = GPixelFormats[PixelFormat].BlockBytes;
-	SizeX = FMath::Max(SizeX >> MipRange.location, BlockSizeX);
-	SizeY = FMath::Max(SizeY >> MipRange.location, BlockSizeY);
-	SizeZ = (Type != RRT_Texture3D) ? SizeZ : FMath::Max(SizeZ >> MipRange.location, 1u);
-}
+		FRWLock Mutex;
+		uint8 NextKey = 1; // 0 is reserved for mtlpp::PixelFormat::Invalid
+		TMap<uint64, uint8> Map;
 
-void FMetalSurface::Init(FMetalSurface& Source, NSRange MipRange, EPixelFormat Format, bool bSRGBForceDisable)
-{
-#if PLATFORM_IOS
-	// Mmeory;ess targets can't have texture views (SRVs or UAVs)
-	if (Source.Texture.GetStorageMode() == mtlpp::StorageMode::Memoryless)
-	{
-		return;
-	}
-#endif
-	
-	mtlpp::PixelFormat MetalFormat = (mtlpp::PixelFormat)GPixelFormats[PixelFormat].PlatformFormat;
-	
-	bool bUseSourceTex = (Source.PixelFormat != PF_DepthStencil) && Source.PixelFormat == Format && MipRange.location == 0 && MipRange.length == Source.Texture.GetMipmapLevelCount();
-	
-	if (EnumHasAnyFlags(Flags, TexCreate_SRGB))
-	{
-		if(bSRGBForceDisable)
+		uint8 Get(mtlpp::PixelFormat Format)
 		{
-			// Force no sRGB, should create new view without sRGB format
-			bUseSourceTex = false;
-#if PLATFORM_MAC
-			// R8Unorm has been expanded in the source surface for sRGBA support - we need to expand to RGBA to enable compatible texture format view for non apple silicon macs
-			if (Format == PF_G8 && Source.Texture.GetPixelFormat() == mtlpp::PixelFormat::RGBA8Unorm_sRGB)
+			FRWScopeLock Lock(Mutex, SLT_ReadOnly);
+			uint8* Key = Map.Find((uint64)Format);
+			if (Key == nullptr)
 			{
-				MetalFormat = mtlpp::PixelFormat::RGBA8Unorm;
+				Lock.ReleaseReadOnlyLockAndAcquireWriteLock_USE_WITH_CAUTION();
+				Key = Map.Find((uint64)Format);
+				if (Key == nullptr)
+				{
+					Key = &Map.Add((uint64)Format, NextKey++);
+					// only giving 6 bits to the key
+					checkf(NextKey < 64, TEXT("Too many unique pixel formats to fit into the PipelineStateHash"));
+				}
 			}
-#endif
+			return *Key;
 		}
-		else
-		{
-			// Ensure we have the correct sRGB target format if we create a new texture view rather than using the source texture
-			MetalFormat = ToSRGBFormat(MetalFormat);
-		}
-	}
-	
-	// Recreate the texture to enable MTLTextureUsagePixelFormatView which must be off unless we definitely use this feature or we are throwing ~4% performance vs. Windows on the floor.
-	mtlpp::TextureUsage Usage = Source.Texture.GetUsage();
-	if(!(Usage & mtlpp::TextureUsage::PixelFormatView) && (Source.PixelFormat == PF_DepthStencil && Format == PF_X24_G8) && !bUseSourceTex)
-	{
-		Source.PrepareTextureView();
-	}
-	
-	ns::Range Slices(0, Source.Texture.GetArrayLength() * (bIsCubemap ? 6 : 1));
-	if(Source.PixelFormat != PF_DepthStencil && !bUseSourceTex)
-	{
-		Texture = Source.Texture.NewTextureView(MetalFormat, Source.Texture.GetTextureType(), ns::Range(MipRange.location, MipRange.length), Slices);
-	}
-	else if (Source.PixelFormat == PF_DepthStencil && Format == PF_X24_G8)
-	{
-		switch (Source.Texture.GetPixelFormat())
+
+		FMetalPixelFormatKeyMap()
 		{
-#if PLATFORM_MAC // Currently only required on Mac as iOS uses separate textures
-			case mtlpp::PixelFormat::Depth24Unorm_Stencil8:
-				MetalFormat = mtlpp::PixelFormat::X24_Stencil8;
-				break;
-#endif
-			case mtlpp::PixelFormat::Depth32Float_Stencil8:
-				MetalFormat = mtlpp::PixelFormat::X32_Stencil8;
-				break;
-			default:
-				check(false);
-				break;
+			// Add depth stencil formats first, so we don't have to use 6 bits for them in the pipeline hash
+			Get(mtlpp::PixelFormat::Depth32Float);
+			Get(mtlpp::PixelFormat::Stencil8);
+			Get(mtlpp::PixelFormat::Depth32Float_Stencil8);
+	#if PLATFORM_MAC
+			Get(mtlpp::PixelFormat::Depth24Unorm_Stencil8);
+			Get(mtlpp::PixelFormat::Depth16Unorm);
+	#endif
 		}
-		
-		Texture = Source.Texture.NewTextureView(MetalFormat, Source.Texture.GetTextureType(), ns::Range(MipRange.location, MipRange.length), Slices);
-	}
-	else
-	{
-		Texture = Source.Texture;
-	}
-	
-	const uint32 BlockSizeX = GPixelFormats[PixelFormat].BlockSizeX;
-	const uint32 BlockSizeY = GPixelFormats[PixelFormat].BlockSizeY;
-	const uint32 BlockBytes = GPixelFormats[PixelFormat].BlockBytes;
-	SizeX = FMath::Max(SizeX >> MipRange.location, BlockSizeX);
-	SizeY = FMath::Max(SizeY >> MipRange.location, BlockSizeY);
-	SizeZ = (Type != RRT_Texture3D) ? SizeZ : FMath::Max(SizeZ >> MipRange.location, 1u);
-}
-
-FMetalSurface::FMetalSurface(FMetalSurface& Source, NSRange MipRange)
-: Type(Source.Type)
-, PixelFormat(Source.PixelFormat)
-, Texture(nil)
-, MSAATexture(nil)
-, MSAAResolveTexture(nil)
-, SizeX(Source.SizeX)
-, SizeY(Source.SizeY)
-, SizeZ(Source.SizeZ)
-, bIsCubemap(Source.bIsCubemap)
-, Flags(Source.Flags)
-, BufferLocks(0)
-, TotalTextureSize(0)
-, Viewport(nullptr)
-, ImageSurfaceRef(nullptr)
-, bTextureView(true)
-{
-	FPlatformAtomics::InterlockedExchange(&Written, Source.Written);
-	
-	Init(Source, MipRange);
-}
-
-FMetalSurface::FMetalSurface(FMetalSurface& Source, NSRange const MipRange, EPixelFormat Format, bool bSRGBForceDisable)
-: Type(Source.Type)
-, PixelFormat(Format)
-, Texture(nil)
-, MSAATexture(nil)
-, MSAAResolveTexture(nil)
-, SizeX(Source.SizeX)
-, SizeY(Source.SizeY)
-, SizeZ(Source.SizeZ)
-, bIsCubemap(Source.bIsCubemap)
-, Flags(Source.Flags)
-, BufferLocks(0)
-, TotalTextureSize(0)
-, Viewport(nullptr)
-, ImageSurfaceRef(nullptr)
-, bTextureView(true)
-{
-	FPlatformAtomics::InterlockedExchange(&Written, Source.Written);
-	
-	Init(Source, MipRange, Format, bSRGBForceDisable);
-}
+	} static PixelFormatKeyMap;
 
-static TMap<uint64, uint8>& GetMetalPixelFormatKeyMap()
-{
-	static TMap<uint64, uint8> PixelFormatKeyMap;
-	return PixelFormatKeyMap;
+	return PixelFormatKeyMap.Get(Format);
 }
 
-uint8 GetMetalPixelFormatKey(mtlpp::PixelFormat Format)
+FMetalTextureDesc::FMetalTextureDesc(FRHITextureDesc const& InDesc)
+	: bIsRenderTarget(IsRenderTarget(InDesc.Flags))
 {
-	static FRWLock Mutex;
-	static uint8 NextKey = 1; // 0 is reserved for mtlpp::PixelFormat::Invalid
-	FRWScopeLock Lock(Mutex, SLT_ReadOnly);
-	uint8* Key = GetMetalPixelFormatKeyMap().Find((uint64)Format);
-	if (Key == NULL)
+	MTLFormat = (mtlpp::PixelFormat)GPixelFormats[InDesc.Format].PlatformFormat;
+	if (EnumHasAnyFlags(InDesc.Flags, TexCreate_SRGB))
 	{
-		Lock.ReleaseReadOnlyLockAndAcquireWriteLock_USE_WITH_CAUTION();
-		Key = GetMetalPixelFormatKeyMap().Find((uint64)Format);
-		if (Key == NULL)
-		{
-			Key = &GetMetalPixelFormatKeyMap().Add((uint64)Format, NextKey++);
-			// only giving 6 bits to the key
-			checkf(NextKey < 64, TEXT("Too many unique pixel formats to fit into the PipelineStateHash"));
-		}
+		MTLFormat = ToSRGBFormat(MTLFormat);
 	}
-	return *Key;
-}
 
-FMetalSurface::FMetalSurface(ERHIResourceType ResourceType, EPixelFormat Format, uint32 InSizeX, uint32 InSizeY, uint32 InSizeZ, uint32 NumSamples, bool bArray, uint32 ArraySize, uint32 NumMips, ETextureCreateFlags InFlags, FResourceBulkDataInterface* BulkData)
-: Type(ResourceType)
-, PixelFormat(Format)
-, Texture(nil)
-, MSAATexture(nil)
-, MSAAResolveTexture(nil)
-, SizeX(InSizeX)
-, SizeY(InSizeY)
-, SizeZ(InSizeZ)
-, bIsCubemap(false)
-, Flags(InFlags)
-, BufferLocks(0)
-, TotalTextureSize(0)
-, Viewport(nullptr)
-, ImageSurfaceRef(nullptr)
-, bTextureView(false)
-{
-	check(SizeX > 0 && SizeY > 0 && NumMips > 0);
-	
 	// get a unique key for this surface's format
-	TMap<uint64, uint8>& PixelFormatKeyMap = GetMetalPixelFormatKeyMap();
-	if (PixelFormatKeyMap.Num() == 0)
-	{
-		// Add depth stencil formats first, so we don't have to use 6 bits for them in the pipeline hash
-		GetMetalPixelFormatKey(mtlpp::PixelFormat::Depth32Float);
-		GetMetalPixelFormatKey(mtlpp::PixelFormat::Stencil8);
-		GetMetalPixelFormatKey(mtlpp::PixelFormat::Depth32Float_Stencil8);
-#if PLATFORM_MAC
-		GetMetalPixelFormatKey(mtlpp::PixelFormat::Depth24Unorm_Stencil8);
-		GetMetalPixelFormatKey(mtlpp::PixelFormat::Depth16Unorm);
-#endif
-	}
-	
-	FPlatformAtomics::InterlockedExchange(&Written, 0);
-	mtlpp::PixelFormat MTLFormat = (mtlpp::PixelFormat)GPixelFormats[Format].PlatformFormat;
-	
-	if (EnumHasAnyFlags(Flags, TexCreate_SRGB))
-	{
-		MTLFormat = ToSRGBFormat(MTLFormat);
-	}
-	
-	// set the key
 	FormatKey = GetMetalPixelFormatKey(MTLFormat);
-	
-	
-	// the special back buffer surface will be updated in GetMetalDeviceContext().BeginDrawingViewport - no need to set the texture here
-	if (EnumHasAnyFlags(Flags, TexCreate_Presentable))
-	{
-		return;
-	}
-	
-	bool bIsRenderTarget = IsRenderTarget(Flags);
-	mtlpp::TextureDescriptor Desc;
-	
-	if (ResourceType == RRT_TextureCube)
+
+	if (InDesc.IsTextureCube())
 	{
-		Desc = mtlpp::TextureDescriptor::TextureCubeDescriptor(MTLFormat, SizeX, (NumMips > 1));
-		bIsCubemap = true;
+		Desc = mtlpp::TextureDescriptor::TextureCubeDescriptor(MTLFormat, InDesc.Extent.X, (InDesc.NumMips > 1));
 	}
-	else if (ResourceType == RRT_Texture3D)
+	else if (InDesc.IsTexture3D())
 	{
 		Desc.SetTextureType(mtlpp::TextureType::Texture3D);
-		Desc.SetWidth(SizeX);
-		Desc.SetHeight(SizeY);
-		Desc.SetDepth(SizeZ);
+		Desc.SetWidth(InDesc.Extent.X);
+		Desc.SetHeight(InDesc.Extent.Y);
+		Desc.SetDepth(InDesc.Depth);
 		Desc.SetPixelFormat(MTLFormat);
 		Desc.SetArrayLength(1);
 		Desc.SetMipmapLevelCount(1);
@@ -760,15 +423,15 @@ FMetalSurface::FMetalSurface(ERHIResourceType ResourceType, EPixelFormat Format,
 	}
 	else
 	{
-		Desc = mtlpp::TextureDescriptor::Texture2DDescriptor(MTLFormat, SizeX, SizeY, (NumMips > 1));
-		Desc.SetDepth(SizeZ);
+		Desc = mtlpp::TextureDescriptor::Texture2DDescriptor(MTLFormat, InDesc.Extent.X, InDesc.Extent.Y, (InDesc.NumMips > 1));
+		Desc.SetArrayLength(InDesc.ArraySize);
 	}
-	
+
 	// flesh out the descriptor
-	if (bArray)
+	if (InDesc.IsTextureArray())
 	{
-		Desc.SetArrayLength(ArraySize);
-		if (bIsCubemap)
+		Desc.SetArrayLength(InDesc.ArraySize);
+		if (InDesc.IsTextureCube())
 		{
 			if (FMetalCommandQueue::SupportsFeature(EMetalFeaturesCubemapArrays))
 			{
@@ -777,7 +440,7 @@ FMetalSurface::FMetalSurface(ERHIResourceType ResourceType, EPixelFormat Format,
 			else
 			{
 				Desc.SetTextureType(mtlpp::TextureType::Texture2DArray);
-				Desc.SetArrayLength(ArraySize * 6);
+				Desc.SetArrayLength(InDesc.ArraySize * 6);
 			}
 		}
 		else
@@ -785,90 +448,90 @@ FMetalSurface::FMetalSurface(ERHIResourceType ResourceType, EPixelFormat Format,
 			Desc.SetTextureType(mtlpp::TextureType::Texture2DArray);
 		}
 	}
-	Desc.SetMipmapLevelCount(NumMips);
-	
+	Desc.SetMipmapLevelCount(InDesc.NumMips);
+
 	{
-		Desc.SetUsage(ConvertFlagsToUsage(Flags));
-		
-		if(EnumHasAnyFlags(Flags, TexCreate_CPUReadback) && !EnumHasAnyFlags(Flags, TexCreate_RenderTargetable|TexCreate_DepthStencilTargetable|TexCreate_FastVRAM))
+		Desc.SetUsage(ConvertFlagsToUsage(InDesc.Flags));
+
+		if (EnumHasAnyFlags(InDesc.Flags, TexCreate_CPUReadback) && !EnumHasAnyFlags(InDesc.Flags, TexCreate_RenderTargetable | TexCreate_DepthStencilTargetable | TexCreate_FastVRAM))
 		{
 			Desc.SetCpuCacheMode(mtlpp::CpuCacheMode::DefaultCache);
 #if PLATFORM_MAC
 			Desc.SetStorageMode(mtlpp::StorageMode::Managed);
-			Desc.SetResourceOptions((mtlpp::ResourceOptions)(mtlpp::ResourceOptions::CpuCacheModeDefaultCache|mtlpp::ResourceOptions::StorageModeManaged));
+			Desc.SetResourceOptions((mtlpp::ResourceOptions)(mtlpp::ResourceOptions::CpuCacheModeDefaultCache | mtlpp::ResourceOptions::StorageModeManaged));
 #else
 			Desc.SetStorageMode(mtlpp::StorageMode::Shared);
-			Desc.SetResourceOptions((mtlpp::ResourceOptions)(mtlpp::ResourceOptions::CpuCacheModeDefaultCache|mtlpp::ResourceOptions::StorageModeShared));
+			Desc.SetResourceOptions((mtlpp::ResourceOptions)(mtlpp::ResourceOptions::CpuCacheModeDefaultCache | mtlpp::ResourceOptions::StorageModeShared));
 #endif
 		}
-		else if (EnumHasAnyFlags(Flags, TexCreate_NoTiling) && !EnumHasAnyFlags(Flags, TexCreate_FastVRAM|TexCreate_DepthStencilTargetable|TexCreate_RenderTargetable|TexCreate_UAV))
+		else if (EnumHasAnyFlags(InDesc.Flags, TexCreate_NoTiling) && !EnumHasAnyFlags(InDesc.Flags, TexCreate_FastVRAM | TexCreate_DepthStencilTargetable | TexCreate_RenderTargetable | TexCreate_UAV))
 		{
 			Desc.SetCpuCacheMode(mtlpp::CpuCacheMode::DefaultCache);
 #if PLATFORM_MAC
 			Desc.SetStorageMode(mtlpp::StorageMode::Managed);
-			Desc.SetResourceOptions((mtlpp::ResourceOptions)(mtlpp::ResourceOptions::CpuCacheModeDefaultCache|mtlpp::ResourceOptions::StorageModeManaged));
+			Desc.SetResourceOptions((mtlpp::ResourceOptions)(mtlpp::ResourceOptions::CpuCacheModeDefaultCache | mtlpp::ResourceOptions::StorageModeManaged));
 #else
 			Desc.SetStorageMode(mtlpp::StorageMode::Shared);
-			Desc.SetResourceOptions((mtlpp::ResourceOptions)(mtlpp::ResourceOptions::CpuCacheModeDefaultCache|mtlpp::ResourceOptions::StorageModeShared));
+			Desc.SetResourceOptions((mtlpp::ResourceOptions)(mtlpp::ResourceOptions::CpuCacheModeDefaultCache | mtlpp::ResourceOptions::StorageModeShared));
 #endif
 		}
-		else if (EnumHasAnyFlags(Flags, TexCreate_RenderTargetable|TexCreate_DepthStencilTargetable|TexCreate_ResolveTargetable|TexCreate_DepthStencilResolveTarget))
+		else if (EnumHasAnyFlags(InDesc.Flags, TexCreate_RenderTargetable | TexCreate_DepthStencilTargetable | TexCreate_ResolveTargetable | TexCreate_DepthStencilResolveTarget))
 		{
-			check(!(Flags & TexCreate_CPUReadback));
+			check(!(InDesc.Flags & TexCreate_CPUReadback));
 			Desc.SetCpuCacheMode(mtlpp::CpuCacheMode::DefaultCache);
 #if PLATFORM_MAC
 			Desc.SetStorageMode(mtlpp::StorageMode::Private);
-			Desc.SetResourceOptions((mtlpp::ResourceOptions)(mtlpp::ResourceOptions::CpuCacheModeDefaultCache|mtlpp::ResourceOptions::StorageModePrivate));
+			Desc.SetResourceOptions((mtlpp::ResourceOptions)(mtlpp::ResourceOptions::CpuCacheModeDefaultCache | mtlpp::ResourceOptions::StorageModePrivate));
 #else
-			if(GMetalForceIOSTexturesShared)
+			if (GMetalForceIOSTexturesShared)
 			{
 				Desc.SetStorageMode(mtlpp::StorageMode::Shared);
-				Desc.SetResourceOptions((mtlpp::ResourceOptions)(mtlpp::ResourceOptions::CpuCacheModeDefaultCache|mtlpp::ResourceOptions::StorageModeShared));
+				Desc.SetResourceOptions((mtlpp::ResourceOptions)(mtlpp::ResourceOptions::CpuCacheModeDefaultCache | mtlpp::ResourceOptions::StorageModeShared));
 			}
 			else
 			{
 				Desc.SetStorageMode(mtlpp::StorageMode::Private);
-				Desc.SetResourceOptions((mtlpp::ResourceOptions)(mtlpp::ResourceOptions::CpuCacheModeDefaultCache|mtlpp::ResourceOptions::StorageModePrivate));
+				Desc.SetResourceOptions((mtlpp::ResourceOptions)(mtlpp::ResourceOptions::CpuCacheModeDefaultCache | mtlpp::ResourceOptions::StorageModePrivate));
 			}
 #endif
 		}
 		else
 		{
-			check(!(Flags & TexCreate_CPUReadback));
+			check(!(InDesc.Flags & TexCreate_CPUReadback));
 			Desc.SetCpuCacheMode(mtlpp::CpuCacheMode::DefaultCache);
 #if PLATFORM_MAC
 			Desc.SetStorageMode(mtlpp::StorageMode::Private);
-			Desc.SetResourceOptions((mtlpp::ResourceOptions)(mtlpp::ResourceOptions::CpuCacheModeDefaultCache|mtlpp::ResourceOptions::StorageModePrivate));
+			Desc.SetResourceOptions((mtlpp::ResourceOptions)(mtlpp::ResourceOptions::CpuCacheModeDefaultCache | mtlpp::ResourceOptions::StorageModePrivate));
 #else
-			if(GMetalForceIOSTexturesShared)
+			if (GMetalForceIOSTexturesShared)
 			{
 				Desc.SetStorageMode(mtlpp::StorageMode::Shared);
-				Desc.SetResourceOptions((mtlpp::ResourceOptions)(mtlpp::ResourceOptions::CpuCacheModeDefaultCache|mtlpp::ResourceOptions::StorageModeShared));
+				Desc.SetResourceOptions((mtlpp::ResourceOptions)(mtlpp::ResourceOptions::CpuCacheModeDefaultCache | mtlpp::ResourceOptions::StorageModeShared));
 			}
 			// No private storage for PVRTC as it messes up the blit-encoder usage.
 			// note: this is set to always be on and will be re-addressed in a future release
 			else
 			{
-				if (IsPixelFormatPVRTCCompressed(Format))
+				if (IsPixelFormatPVRTCCompressed(InDesc.Format))
 				{
 					Desc.SetStorageMode(mtlpp::StorageMode::Shared);
-					Desc.SetResourceOptions((mtlpp::ResourceOptions)(mtlpp::ResourceOptions::CpuCacheModeDefaultCache|mtlpp::ResourceOptions::StorageModeShared));
+					Desc.SetResourceOptions((mtlpp::ResourceOptions)(mtlpp::ResourceOptions::CpuCacheModeDefaultCache | mtlpp::ResourceOptions::StorageModeShared));
 				}
 				else
 				{
 					Desc.SetStorageMode(mtlpp::StorageMode::Private);
-					Desc.SetResourceOptions((mtlpp::ResourceOptions)(mtlpp::ResourceOptions::CpuCacheModeDefaultCache|mtlpp::ResourceOptions::StorageModePrivate));
+					Desc.SetResourceOptions((mtlpp::ResourceOptions)(mtlpp::ResourceOptions::CpuCacheModeDefaultCache | mtlpp::ResourceOptions::StorageModePrivate));
 				}
 			}
 #endif
 		}
-		
+
 #if PLATFORM_IOS
-		if (EnumHasAnyFlags(Flags, TexCreate_Memoryless))
+		if (EnumHasAnyFlags(InDesc.Flags, TexCreate_Memoryless))
 		{
-			ensure(EnumHasAnyFlags(Flags, (TexCreate_RenderTargetable | TexCreate_DepthStencilTargetable)));
-			ensure(!EnumHasAnyFlags(Flags, (TexCreate_CPUReadback | TexCreate_CPUWritable)));
-			ensure(!EnumHasAnyFlags(Flags, TexCreate_UAV));
+			ensure(EnumHasAnyFlags(InDesc.Flags, (TexCreate_RenderTargetable | TexCreate_DepthStencilTargetable)));
+			ensure(!EnumHasAnyFlags(InDesc.Flags, (TexCreate_CPUReadback | TexCreate_CPUWritable)));
+			ensure(!EnumHasAnyFlags(InDesc.Flags, TexCreate_UAV));
 			Desc.SetStorageMode(mtlpp::StorageMode::Memoryless);
 			Desc.SetResourceOptions((mtlpp::ResourceOptions)(mtlpp::ResourceOptions::CpuCacheModeDefaultCache | mtlpp::ResourceOptions::StorageModeMemoryless));
 		}
@@ -877,163 +540,217 @@ FMetalSurface::FMetalSurface(ERHIResourceType ResourceType, EPixelFormat Format,
 		static mtlpp::ResourceOptions GeneralResourceOption = FMetalCommandQueue::GetCompatibleResourceOptions(mtlpp::ResourceOptions::HazardTrackingModeUntracked);
 		Desc.SetResourceOptions((mtlpp::ResourceOptions)(Desc.GetResourceOptions() | GeneralResourceOption));
 	}
+
 	
+	if (!FParse::Param(FCommandLine::Get(), TEXT("nomsaa")))
+	{
+		if (InDesc.NumSamples > 1)
+		{
+			check(bIsRenderTarget);
+			Desc.SetTextureType(mtlpp::TextureType::Texture2DMultisample);
+
+			// allow commandline to override
+			uint32 NewNumSamples;
+			if (FParse::Value(FCommandLine::Get(), TEXT("msaa="), NewNumSamples))
+			{
+				Desc.SetSampleCount(NewNumSamples);
+			}
+			else
+			{
+				Desc.SetSampleCount(InDesc.NumSamples);
+			}
+
+#if PLATFORM_IOS
+			if (GMaxRHIFeatureLevel < ERHIFeatureLevel::SM5)
+			{
+				bMemoryless = true;
+				Desc.SetStorageMode(mtlpp::StorageMode::Memoryless);
+				Desc.SetResourceOptions(mtlpp::ResourceOptions::StorageModeMemoryless);
+			}
+#endif
+		}
+	}
+}
+
+FMetalSurface::FMetalSurface(FMetalTextureCreateDesc const& CreateDesc)
+	: FRHITexture       (CreateDesc)
+	, FormatKey         (CreateDesc.FormatKey)
+	, Texture           (nil)
+	, MSAATexture       (nil)
+	, MSAAResolveTexture(nil)
+	, BufferLocks       (0)
+	, TotalTextureSize  (0)
+	, Viewport          (nullptr)
+	, ImageSurfaceRef   (nullptr)
+{
+	FPlatformAtomics::InterlockedExchange(&Written, 0);
+
+	check(CreateDesc.Extent.X > 0 && CreateDesc.Extent.Y > 0 && CreateDesc.NumMips > 0);
+
+	// the special back buffer surface will be updated in GetMetalDeviceContext().BeginDrawingViewport - no need to set the texture here
+	if (EnumHasAnyFlags(CreateDesc.Flags, TexCreate_Presentable))
+	{
+		return;
+	}
+
+	FResourceBulkDataInterface* BulkData = CreateDesc.BulkData;
+
+	// The bulk data interface can be used to create external textures for VR and media player.
+	// Handle these first.
 	if (BulkData != nullptr)
 	{
-		if (BulkData->GetResourceType() == FResourceBulkDataInterface::EBulkDataType::MediaTexture)
+		switch (BulkData->GetResourceType())
 		{
-			checkf(NumMips == 1 && ArraySize == 1, TEXT("Only handling bulk data with 1 mip and 1 array length"));
-			ImageSurfaceRef = (CFTypeRef)BulkData->GetResourceBulkData();
-			CFRetain(ImageSurfaceRef);
+		case FResourceBulkDataInterface::EBulkDataType::MediaTexture:
 			{
+				checkf(CreateDesc.NumMips == 1 && CreateDesc.ArraySize == 1, TEXT("Only handling bulk data with 1 mip and 1 array length"));
+				ImageSurfaceRef = (CFTypeRef)BulkData->GetResourceBulkData();
+				CFRetain(ImageSurfaceRef);
+				
 #if !COREVIDEO_SUPPORTS_METAL
-				Texture = MTLPP_VALIDATE(mtlpp::Device, GetMetalDeviceContext().GetDevice(), SafeGetRuntimeDebuggingLevel() >= EMetalDebugLevelValidation, NewTexture(Desc, CVPixelBufferGetIOSurface((CVPixelBufferRef)ImageSurfaceRef), 0));
+				Texture = MTLPP_VALIDATE(mtlpp::Device, GetMetalDeviceContext().GetDevice(), SafeGetRuntimeDebuggingLevel() >= EMetalDebugLevelValidation, NewTexture(CreateDesc.Desc, CVPixelBufferGetIOSurface((CVPixelBufferRef)ImageSurfaceRef), 0));
 #else
 				Texture = CVMetalTextureGetTexture((CVMetalTextureRef)ImageSurfaceRef);
 #endif
-				METAL_FATAL_ASSERT(Texture, TEXT("Failed to create texture, desc %s"), *FString([Desc description]));
+				METAL_FATAL_ASSERT(Texture, TEXT("Failed to create texture, desc %s"), *FString([CreateDesc.Desc description]));
+
+				BulkData->Discard();
+				BulkData = nullptr;
 			}
-			
-			BulkData->Discard();
-		}
+			break;
+
 #if PLATFORM_MAC
-		else if (BulkData->GetResourceType() == FResourceBulkDataInterface::EBulkDataType::VREyeBuffer)
-		{
-			ImageSurfaceRef = (CFTypeRef)BulkData->GetResourceBulkData();
-			CFRetain(ImageSurfaceRef);
-			
-			Desc.SetStorageMode(mtlpp::StorageMode::Managed);
-			Desc.SetResourceOptions((mtlpp::ResourceOptions)((Desc.GetResourceOptions() & ~(mtlpp::ResourceStorageModeMask)) | mtlpp::ResourceOptions::StorageModeManaged));
-			
-			Texture = [GetMetalDeviceContext().GetDevice() newTextureWithDescriptor:Desc iosurface : (IOSurfaceRef)ImageSurfaceRef plane : 0];
-			
-			METAL_FATAL_ASSERT(Texture, TEXT("Failed to create texture, desc %s"), *FString([Desc description]));
-			
-			BulkData->Discard();
-		}
+		case FResourceBulkDataInterface::EBulkDataType::VREyeBuffer:
+			{
+				ImageSurfaceRef = (CFTypeRef)BulkData->GetResourceBulkData();
+				CFRetain(ImageSurfaceRef);
+
+				mtlpp::TextureDescriptor DescCopy = CreateDesc.Desc;
+				DescCopy.SetStorageMode(mtlpp::StorageMode::Managed);
+				DescCopy.SetResourceOptions((mtlpp::ResourceOptions)((DescCopy.GetResourceOptions() & ~(mtlpp::ResourceStorageModeMask)) | mtlpp::ResourceOptions::StorageModeManaged));
+
+				Texture = [GetMetalDeviceContext().GetDevice() newTextureWithDescriptor:DescCopy iosurface : (IOSurfaceRef)ImageSurfaceRef plane : 0];
+
+				METAL_FATAL_ASSERT(Texture, TEXT("Failed to create texture, desc %s"), *FString([DescCopy description]));
+
+				BulkData->Discard();
+				BulkData = nullptr;
+			}
+			break;
 #endif
+		}
 	}
-	else
+
+	if (!Texture)
 	{
-		const bool bBufferCompatibleOption = (Desc.GetTextureType() == mtlpp::TextureType::Texture2D || Desc.GetTextureType() == mtlpp::TextureType::TextureBuffer) && NumMips == 1;
-		if (!bBufferCompatibleOption || (!EnumHasAllFlags(Flags, TexCreate_UAV|TexCreate_NoTiling) && !EnumHasAllFlags(Flags, TexCreate_AtomicCompatible)))
+		// Non VR/media texture case (i.e. a regular texture).
+		// Create the actual texture resource.
+
+		const bool bBufferCompatibleOption = (CreateDesc.Desc.GetTextureType() == mtlpp::TextureType::Texture2D || CreateDesc.Desc.GetTextureType() == mtlpp::TextureType::TextureBuffer) && CreateDesc.NumMips == 1;
+		if (!bBufferCompatibleOption || (!EnumHasAllFlags(CreateDesc.Flags, TexCreate_UAV | TexCreate_NoTiling) && !EnumHasAllFlags(CreateDesc.Flags, TexCreate_AtomicCompatible)))
 		{
-			Texture = GetMetalDeviceContext().CreateTexture(this, Desc);
+			Texture = GetMetalDeviceContext().CreateTexture(this, CreateDesc.Desc);
 		}
 		else
 		{
 			mtlpp::Device Device = GetMetalDeviceContext().GetDevice();
-			
-			const uint32 MinimumByteAlignment = GetMetalDeviceContext().GetDevice().GetMinimumLinearTextureAlignmentForPixelFormat(MTLFormat);
-			const NSUInteger BytesPerRow = Align(Desc.GetWidth() * GPixelFormats[Format].BlockBytes, MinimumByteAlignment);
-			
+
+			const uint32 MinimumByteAlignment = GetMetalDeviceContext().GetDevice().GetMinimumLinearTextureAlignmentForPixelFormat(CreateDesc.MTLFormat);
+			const NSUInteger BytesPerRow = Align(CreateDesc.Desc.GetWidth() * GPixelFormats[CreateDesc.Format].BlockBytes, MinimumByteAlignment);
+
 			// Backing buffer resource options must match the texture we are going to create from it
-			FMetalPooledBufferArgs Args(Device, BytesPerRow * Desc.GetHeight(), BUF_Dynamic, mtlpp::StorageMode::Private, Desc.GetCpuCacheMode());
+			FMetalPooledBufferArgs Args(Device, BytesPerRow * CreateDesc.Desc.GetHeight(), BUF_Dynamic, mtlpp::StorageMode::Private, CreateDesc.Desc.GetCpuCacheMode());
 			FMetalBuffer Buffer = GetMetalDeviceContext().CreatePooledBuffer(Args);
 
-			Texture = Buffer.NewTexture(Desc, 0, BytesPerRow);
-		}
-		
-		METAL_FATAL_ASSERT(Texture, TEXT("Failed to create texture, desc %s"), *FString([Desc description]));
-		
-		// upload existing bulkdata
-		if (BulkData)
-		{
-			UE_LOG(LogMetal, Display, TEXT("Got a bulk data texture, with %d mips"), NumMips);
-			checkf(NumMips == 1&& ArraySize == 1, TEXT("Only handling bulk data with 1 mip and 1 array length"));
-			uint32 Stride;
-			
-			// lock, copy, unlock
-			void* LockedData = Lock(0, 0, RLM_WriteOnly, Stride);
-			FMemory::Memcpy(LockedData, BulkData->GetResourceBulkData(), BulkData->GetResourceBulkDataSize());
-			Unlock(0, 0, false);
-			
-			// bulk data can be unloaded now
-			BulkData->Discard();
+			Texture = Buffer.NewTexture(CreateDesc.Desc, 0, BytesPerRow);
 		}
+
+		METAL_FATAL_ASSERT(Texture, TEXT("Failed to create texture, desc %s"), *FString([CreateDesc.Desc description]));
 	}
-	
+
+	if (BulkData)
+	{
+		// Regular texture has some bulk data to handle
+		UE_LOG(LogMetal, Display, TEXT("Got a bulk data texture, with %d mips"), CreateDesc.NumMips);
+		checkf(CreateDesc.NumMips == 1, TEXT("Only handling bulk data with 1 mip and 1 array length"));
+
+		check(IsInRenderingThread());
+		FRHICommandListImmediate& RHICmdList = FRHICommandListExecutor::GetImmediateCommandList();
+
+		// lock, copy, unlock
+		uint32 Stride;
+		void* LockedData = FMetalDynamicRHI::Get().LockTexture2D_RenderThread(RHICmdList, this, 0, RLM_WriteOnly, Stride, false);
+		check(LockedData);
+		FMemory::Memcpy(LockedData, BulkData->GetResourceBulkData(), BulkData->GetResourceBulkDataSize());
+		FMetalDynamicRHI::Get().UnlockTexture2D_RenderThread(RHICmdList, this, 0, false);
+
+		// bulk data can be unloaded now
+		BulkData->Discard();
+		BulkData = nullptr;
+	}
+
 	// calculate size of the texture
 	TotalTextureSize = GetMemorySize();
-	
-	const bool bSupportsMSAADepthResolve = GetMetalDeviceContext().SupportsFeature(EMetalFeaturesMSAADepthResolve);
-	if (!FParse::Param(FCommandLine::Get(), TEXT("nomsaa")))
+
+	if (CreateDesc.NumSamples > 1 && !FParse::Param(FCommandLine::Get(), TEXT("nomsaa")))
 	{
-		if (NumSamples > 1)
-		{
-			check(bIsRenderTarget);
-			Desc.SetTextureType(mtlpp::TextureType::Texture2DMultisample);
+		MSAATexture = GetMetalDeviceContext().CreateTexture(this, CreateDesc.Desc);
 			
-			// allow commandline to override
-			FParse::Value(FCommandLine::Get(), TEXT("msaa="), NumSamples);
-			Desc.SetSampleCount(NumSamples);
-			
-			bool bMemoryless = false;
-#if PLATFORM_IOS
-			if (GMaxRHIFeatureLevel < ERHIFeatureLevel::SM5)
-			{
-				bMemoryless = true;
-				Desc.SetStorageMode(mtlpp::StorageMode::Memoryless);
-				Desc.SetResourceOptions(mtlpp::ResourceOptions::StorageModeMemoryless);
-			}
-#endif
-			
-			MSAATexture = GetMetalDeviceContext().CreateTexture(this, Desc);
+		//device doesn't support HW depth resolve.  This case only valid on mobile renderer or
+		//on Mac where RHISupportsSeparateMSAAAndResolveTextures is true.
+		const bool bSupportsMSAADepthResolve = GetMetalDeviceContext().SupportsFeature(EMetalFeaturesMSAADepthResolve);
+		const bool bDepthButNoResolveSupported = CreateDesc.Format == PF_DepthStencil && !bSupportsMSAADepthResolve;
+		if (bDepthButNoResolveSupported)
+		{
+			Texture = MSAATexture;
+				
+			// we don't have the resolve texture, so we just update the memory size with the MSAA size
+			TotalTextureSize = TotalTextureSize * CreateDesc.NumSamples;
+		}
+		else if (!CreateDesc.bMemoryless)
+		{
+			// an MSAA render target takes NumSamples more space, in addition to the resolve texture
+			TotalTextureSize += TotalTextureSize * CreateDesc.NumSamples;
+		}
 			
-			//device doesn't support HW depth resolve.  This case only valid on mobile renderer or
-			//on Mac where RHISupportsSeparateMSAAAndResolveTextures is true.
-			const bool bDepthButNoResolveSupported = Format == PF_DepthStencil && !bSupportsMSAADepthResolve;
-			if (bDepthButNoResolveSupported)
-			{
-				Texture = MSAATexture;
+		if (MSAATexture != Texture)
+		{
+			check(!MSAAResolveTexture);
 				
-				// we don't have the resolve texture, so we just update the memory size with the MSAA size
-				TotalTextureSize = TotalTextureSize * NumSamples;
-			}
-			else if (!bMemoryless)
+			//if bSupportsSeparateMSAAAndResolve then the high level expect to binds the MSAA when binding shader params.
+			const bool bSupportsSeparateMSAAAndResolve = FMetalCommandQueue::SupportsSeparateMSAAAndResolveTarget();
+			if (bSupportsSeparateMSAAAndResolve)
 			{
-				// an MSAA render target takes NumSamples more space, in addition to the resolve texture
-				TotalTextureSize += TotalTextureSize * NumSamples;
+				MSAAResolveTexture = Texture;
+					
+				Texture = MSAATexture;
 			}
-			
-			if (MSAATexture != Texture)
+			else
 			{
-				check(!MSAAResolveTexture);
-				
-				//if bSupportsSeparateMSAAAndResolve then the high level expect to binds the MSAA when binding shader params.
-				const bool bSupportsSeparateMSAAAndResolve = FMetalCommandQueue::SupportsSeparateMSAAAndResolveTarget();
-				if (bSupportsSeparateMSAAAndResolve)
-				{
-					MSAAResolveTexture = Texture;
-					
-					Texture = MSAATexture;
-				}
-				else
-				{
-					MSAAResolveTexture = Texture;
-				}
+				MSAAResolveTexture = Texture;
 			}
+		}
 			
-			//we always require an MSAAResolveTexture if MSAATexture is active.
-			check(!MSAATexture || MSAAResolveTexture || bDepthButNoResolveSupported);
+		//we always require an MSAAResolveTexture if MSAATexture is active.
+		check(!MSAATexture || MSAAResolveTexture || bDepthButNoResolveSupported);
 			
-			NSLog(@"Creating %dx MSAA %d x %d %s surface", (int32)Desc.GetSampleCount(), SizeX, SizeY, EnumHasAnyFlags(Flags, TexCreate_RenderTargetable) ? "Color" : "Depth");
-			if (MSAATexture.GetPtr() == nil)
-			{
-				NSLog(@"Failed to create texture, desc  %@", Desc.GetPtr());
-			}
+		NSLog(@"Creating %dx MSAA %d x %d %s surface", (int32)CreateDesc.Desc.GetSampleCount(), CreateDesc.Extent.X, CreateDesc.Extent.Y, EnumHasAnyFlags(CreateDesc.Flags, TexCreate_RenderTargetable) ? "Color" : "Depth");
+		if (MSAATexture.GetPtr() == nil)
+		{
+			NSLog(@"Failed to create texture, desc  %@", CreateDesc.Desc.GetPtr());
 		}
 	}
 	
 	// create a stencil buffer if needed
-	if (Format == PF_DepthStencil)
+	if (CreateDesc.Format == PF_DepthStencil)
 	{
 		// 1 byte per texel
-		TotalTextureSize += SizeX * SizeY;
+		TotalTextureSize += CreateDesc.Extent.X * CreateDesc.Extent.Y;
 	}
 	
 	// track memory usage
-	if (bIsRenderTarget)
+	
+	if (CreateDesc.bIsRenderTarget)
 	{
 		GCurrentRendertargetMemorySize += Align(TotalTextureSize, 1024) / 1024;
 	}
@@ -1043,9 +760,9 @@ FMetalSurface::FMetalSurface(ERHIResourceType ResourceType, EPixelFormat Format,
 	}
 	
 #if STATS
-	if (ResourceType == RRT_TextureCube)
+	if (CreateDesc.IsTextureCube())
 	{
-		if (bIsRenderTarget)
+		if (CreateDesc.bIsRenderTarget)
 		{
 			INC_MEMORY_STAT_BY(STAT_RenderTargetMemoryCube, TotalTextureSize);
 		}
@@ -1054,9 +771,9 @@ FMetalSurface::FMetalSurface(ERHIResourceType ResourceType, EPixelFormat Format,
 			INC_MEMORY_STAT_BY(STAT_TextureMemoryCube, TotalTextureSize);
 		}
 	}
-	else if (ResourceType == RRT_Texture3D)
+	else if (CreateDesc.IsTexture3D())
 	{
-		if (bIsRenderTarget)
+		if (CreateDesc.bIsRenderTarget)
 		{
 			INC_MEMORY_STAT_BY(STAT_RenderTargetMemory3D, TotalTextureSize);
 		}
@@ -1067,7 +784,7 @@ FMetalSurface::FMetalSurface(ERHIResourceType ResourceType, EPixelFormat Format,
 	}
 	else
 	{
-		if (bIsRenderTarget)
+		if (CreateDesc.bIsRenderTarget)
 		{
 			INC_MEMORY_STAT_BY(STAT_RenderTargetMemory2D, TotalTextureSize);
 		}
@@ -1083,7 +800,7 @@ FMetalSurface::FMetalSurface(ERHIResourceType ResourceType, EPixelFormat Format,
 {
 @public
 	uint64 TextureSize;
-	ERHIResourceType TexType;
+	ETextureDimension Dimension;
 	bool bIsRenderTarget;
 }
 @end
@@ -1093,7 +810,7 @@ APPLE_PLATFORM_OBJECT_ALLOC_OVERRIDES(FMetalDeferredStats)
 -(void)dealloc
 {
 #if STATS
-	if (TexType == RRT_TextureCube)
+	if (Dimension == ETextureDimension::TextureCube || Dimension == ETextureDimension::TextureCubeArray)
 	{
 		if (bIsRenderTarget)
 		{
@@ -1104,7 +821,7 @@ APPLE_PLATFORM_OBJECT_ALLOC_OVERRIDES(FMetalDeferredStats)
 			DEC_MEMORY_STAT_BY(STAT_TextureMemoryCube, TextureSize);
 		}
 	}
-	else if (TexType == RRT_Texture3D)
+	else if (Dimension == ETextureDimension::Texture3D)
 	{
 		if (bIsRenderTarget)
 		{
@@ -1141,13 +858,13 @@ APPLE_PLATFORM_OBJECT_ALLOC_OVERRIDES(FMetalDeferredStats)
 
 FMetalSurface::~FMetalSurface()
 {
-	bool const bIsRenderTarget = IsRenderTarget(Flags);
+	bool const bIsRenderTarget = IsRenderTarget(GetDesc().Flags);
 	
 	if (MSAATexture.GetPtr())
 	{
 		if (Texture.GetPtr() != MSAATexture.GetPtr())
 		{
-			SafeReleaseMetalTexture(this, MSAATexture, bTextureView, false);
+			SafeReleaseMetalTexture(this, MSAATexture, false);
 		}
 	}
 	
@@ -1157,13 +874,13 @@ FMetalSurface::~FMetalSurface()
 	{
 		if (Texture.GetPtr() != MSAAResolveTexture.GetPtr())
 		{
-			SafeReleaseMetalTexture(this, MSAAResolveTexture, bTextureView, false);
+			SafeReleaseMetalTexture(this, MSAAResolveTexture, false);
 		}
 	}
 	
-	if (!(Flags & TexCreate_Presentable) && Texture.GetPtr())
+	if (!(GetDesc().Flags & TexCreate_Presentable) && Texture.GetPtr())
 	{
-		SafeReleaseMetalTexture(this, Texture, bTextureView, (ImageSurfaceRef != nullptr));
+		SafeReleaseMetalTexture(this, Texture, (ImageSurfaceRef != nullptr));
 	}
 	
 	MSAATexture = nil;
@@ -1172,7 +889,7 @@ FMetalSurface::~FMetalSurface()
 	
 	// track memory usage
 	FMetalDeferredStats* Block = [FMetalDeferredStats new];
-	Block->TexType = Type;
+	Block->Dimension = GetDesc().Dimension;
 	Block->TextureSize = TotalTextureSize;
 	Block->bIsRenderTarget = bIsRenderTarget;
 	SafeReleaseMetalObject(Block);
@@ -1206,9 +923,13 @@ id <MTLBuffer> FMetalSurface::AllocSurface(uint32 MipIndex, uint32 ArrayIndex, E
 	
 #if PLATFORM_MAC
 	// Expand R8_sRGB into RGBA8_sRGB for non Apple Silicon Mac.
-	if (PixelFormat == PF_G8 && EnumHasAnyFlags(Flags, TexCreate_SRGB) && Type == RRT_Texture2D && LockMode == RLM_WriteOnly && Texture.GetPixelFormat() == mtlpp::PixelFormat::RGBA8Unorm_sRGB)
+	if (   GetDesc().Format == PF_G8 
+		&& GetDesc().Dimension == ETextureDimension::Texture2D
+		&& EnumHasAnyFlags(GetDesc().Flags, TexCreate_SRGB) 
+		&& LockMode == RLM_WriteOnly 
+		&& Texture.GetPixelFormat() == mtlpp::PixelFormat::RGBA8Unorm_sRGB)
 	{
-		DestStride = FMath::Max<uint32>(SizeX >> MipIndex, 1);
+		DestStride = FMath::Max<uint32>(GetDesc().Extent.X >> MipIndex, 1);
 	}
 #endif
 	
@@ -1228,29 +949,42 @@ void FMetalSurface::UpdateSurfaceAndDestroySourceBuffer(id <MTLBuffer> SourceBuf
 	uint32 BytesPerImage = GetMipSize(MipIndex, &Stride, true);
 	
 	mtlpp::Region Region;
-	if (SizeZ <= 1 || bIsCubemap)
+	if (GetDesc().IsTexture3D())
 	{
 		// upload the texture to the texture slice
-		Region = mtlpp::Region(0, 0, FMath::Max<uint32>(SizeX >> MipIndex, 1), FMath::Max<uint32>(SizeY >> MipIndex, 1));
+		Region = mtlpp::Region(
+			0, 0, 0,
+			FMath::Max<uint32>(GetDesc().Extent.X >> MipIndex, 1),
+			FMath::Max<uint32>(GetDesc().Extent.Y >> MipIndex, 1),
+			FMath::Max<uint32>(GetDesc().Depth    >> MipIndex, 1)
+		);
 	}
 	else
 	{
 		// upload the texture to the texture slice
-		Region = mtlpp::Region(0, 0, 0, FMath::Max<uint32>(SizeX >> MipIndex, 1), FMath::Max<uint32>(SizeY >> MipIndex, 1), FMath::Max<uint32>(SizeZ >> MipIndex, 1));
+		Region = mtlpp::Region(
+			0, 0,
+			FMath::Max<uint32>(GetDesc().Extent.X >> MipIndex, 1),
+			FMath::Max<uint32>(GetDesc().Extent.Y >> MipIndex, 1)
+		);
 	}
+
 #if PLATFORM_MAC
 	// Expand R8_sRGB into RGBA8_sRGB for non Apple Silicon Mac.
-	if (PixelFormat == PF_G8 && EnumHasAnyFlags(Flags, TexCreate_SRGB) && Type == RRT_Texture2D && Texture.GetPixelFormat() == mtlpp::PixelFormat::RGBA8Unorm_sRGB)
+	if (   GetDesc().Format == PF_G8 
+		&& GetDesc().Dimension == ETextureDimension::Texture2D
+		&& EnumHasAnyFlags(GetDesc().Flags, TexCreate_SRGB) 
+		&& Texture.GetPixelFormat() == mtlpp::PixelFormat::RGBA8Unorm_sRGB)
 	{
 		TArray<uint8> Data;
 		uint8* ExpandedMem = (uint8*) SourceBuffer.contents;
 		check(ExpandedMem);
 		Data.Append(ExpandedMem, BytesPerImage);
-		uint32 SrcStride = FMath::Max<uint32>(SizeX >> MipIndex, 1);
-		for(uint y = 0; y < FMath::Max<uint32>(SizeY >> MipIndex, 1); y++)
+		uint32 SrcStride = FMath::Max<uint32>(GetDesc().Extent.X >> MipIndex, 1);
+		for(uint y = 0; y < FMath::Max<uint32>(GetDesc().Extent.Y >> MipIndex, 1); y++)
 		{
 			uint8* RowDest = ExpandedMem;
-			for(uint x = 0; x < FMath::Max<uint32>(SizeX >> MipIndex, 1); x++)
+			for(uint x = 0; x < FMath::Max<uint32>(GetDesc().Extent.X >> MipIndex, 1); x++)
 			{
 				*(RowDest++) = Data[(y * SrcStride) + x];
 				*(RowDest++) = Data[(y * SrcStride) + x];
@@ -1367,15 +1101,23 @@ void* FMetalSurface::Lock(uint32 MipIndex, uint32 ArrayIndex, EResourceLockMode
 			const bool bIssueImmediateCommands = RHICmdList.Bypass() || IsInRHIThread();
 			
 			mtlpp::Region Region;
-			if (SizeZ <= 1 || bIsCubemap)
+			if (GetDesc().IsTexture3D())
 			{
 				// upload the texture to the texture slice
-				Region = mtlpp::Region(0, 0, FMath::Max<uint32>(SizeX >> MipIndex, 1), FMath::Max<uint32>(SizeY >> MipIndex, 1));
+				Region = mtlpp::Region(
+					0, 0, 0,
+					FMath::Max<uint32>(GetDesc().Extent.X >> MipIndex, 1),
+					FMath::Max<uint32>(GetDesc().Extent.Y >> MipIndex, 1),
+					FMath::Max<uint32>(GetDesc().Depth    >> MipIndex, 1));
 			}
 			else
 			{
 				// upload the texture to the texture slice
-				Region = mtlpp::Region(0, 0, 0, FMath::Max<uint32>(SizeX >> MipIndex, 1), FMath::Max<uint32>(SizeY >> MipIndex, 1), FMath::Max<uint32>(SizeZ >> MipIndex, 1));
+				Region = mtlpp::Region(
+					0, 0,
+					FMath::Max<uint32>(GetDesc().Extent.X >> MipIndex, 1),
+					FMath::Max<uint32>(GetDesc().Extent.Y >> MipIndex, 1)
+				);
 			}
 			
 			if (Texture.GetStorageMode() == mtlpp::StorageMode::Private)
@@ -1428,7 +1170,7 @@ void* FMetalSurface::Lock(uint32 MipIndex, uint32 ArrayIndex, EResourceLockMode
 				
 				// This block breaks the texture atlas system in Ocean, which depends on nonzero strides coming back from compressed textures. Turning off.
 #if 0
-				if (PixelFormat == PF_PVRTC2 || PixelFormat == PF_PVRTC4)
+				if (GetDesc().Format == PF_PVRTC2 || GetDesc().Format == PF_PVRTC4)
 				{
 					// for compressed textures metal debug RT expects 0 for rowBytes and imageBytes.
 					DestStride = 0;
@@ -1436,7 +1178,7 @@ void* FMetalSurface::Lock(uint32 MipIndex, uint32 ArrayIndex, EResourceLockMode
 				}
 #endif
 				uint32 BytesPerRow = DestStride;
-				if (PixelFormat == PF_PVRTC2 || PixelFormat == PF_PVRTC4)
+				if (GetDesc().Format == PF_PVRTC2 || GetDesc().Format == PF_PVRTC4)
 				{
 					// for compressed textures metal debug RT expects 0 for rowBytes and imageBytes.
 					BytesPerRow = 0;
@@ -1447,17 +1189,20 @@ void* FMetalSurface::Lock(uint32 MipIndex, uint32 ArrayIndex, EResourceLockMode
 			
 #if PLATFORM_MAC
 			// Pack RGBA8_sRGB into R8_sRGB for non Apple Silicon Mac.
-			if (PixelFormat == PF_G8 && EnumHasAnyFlags(Flags, TexCreate_SRGB) && Type == RRT_Texture2D && Texture.GetPixelFormat() == mtlpp::PixelFormat::RGBA8Unorm_sRGB)
+			if (   GetDesc().Format == PF_G8 
+				&& GetDesc().Dimension == ETextureDimension::Texture2D
+				&& EnumHasAnyFlags(GetDesc().Flags, TexCreate_SRGB) 
+				&& Texture.GetPixelFormat() == mtlpp::PixelFormat::RGBA8Unorm_sRGB)
 			{
 				TArray<uint8> Data;
 				uint8* ExpandedMem = (uint8*)SourceData.GetContents();
 				Data.Append(ExpandedMem, MipBytes);
 				uint32 SrcStride = DestStride;
-				DestStride = FMath::Max<uint32>(SizeX >> MipIndex, 1);
-				for(uint y = 0; y < FMath::Max<uint32>(SizeY >> MipIndex, 1); y++)
+				DestStride = FMath::Max<uint32>(GetDesc().Extent.X >> MipIndex, 1);
+				for(uint y = 0; y < FMath::Max<uint32>(GetDesc().Extent.Y >> MipIndex, 1); y++)
 				{
 					uint8* RowDest = ExpandedMem;
-					for(uint x = 0; x < FMath::Max<uint32>(SizeX >> MipIndex, 1); x++)
+					for(uint x = 0; x < FMath::Max<uint32>(GetDesc().Extent.X >> MipIndex, 1); x++)
 					{
 						*(RowDest++) = Data[(y * SrcStride) + (x * 4)];
 					}
@@ -1569,6 +1314,8 @@ void FMetalSurface::AsyncUnlock(id <MTLBuffer> SourceData, uint32 MipIndex, uint
 
 uint32 FMetalSurface::GetMipSize(uint32 MipIndex, uint32* Stride, bool bSingleLayer)
 {
+	EPixelFormat PixelFormat = GetDesc().Format;
+
 	// DXT/BC formats on Mac actually do have mip-tails that are smaller than the block size, they end up being uncompressed.
 	bool const bPixelFormatASTC = IsPixelFormatASTCCompressed(PixelFormat);
 	
@@ -1577,12 +1324,12 @@ uint32 FMetalSurface::GetMipSize(uint32 MipIndex, uint32* Stride, bool bSingleLa
 	const uint32 BlockSizeY = GPixelFormats[PixelFormat].BlockSizeY;
 	const uint32 BlockBytes = GPixelFormats[PixelFormat].BlockBytes;
 	const uint32 Alignment = 1u; // Apparently we always want natural row alignment (tightly-packed) even though the docs say iOS doesn't support it - this may be because we don't upload texture data from one contiguous buffer.
-	const uint32 UnalignedMipSizeX = FMath::Max(SizeX >> MipIndex, BlockSizeX);
-	const uint32 UnalignedMipSizeY = FMath::Max(SizeY >> MipIndex, BlockSizeY);
+	const uint32 UnalignedMipSizeX = FMath::Max<uint32>(GetDesc().Extent.X >> MipIndex, BlockSizeX);
+	const uint32 UnalignedMipSizeY = FMath::Max<uint32>(GetDesc().Extent.Y >> MipIndex, BlockSizeY);
 	const uint32 MipSizeX = (bPixelFormatASTC) ? AlignArbitrary(UnalignedMipSizeX, BlockSizeX) : UnalignedMipSizeX;
 	const uint32 MipSizeY = (bPixelFormatASTC) ? AlignArbitrary(UnalignedMipSizeY, BlockSizeY) : UnalignedMipSizeY;
 	
-	const uint32 MipSizeZ = bSingleLayer ? 1 : FMath::Max(SizeZ >> MipIndex, 1u);
+	const uint32 MipSizeZ = bSingleLayer ? 1 : FMath::Max<uint32>(GetDesc().Depth >> MipIndex, 1u);
 	uint32 NumBlocksX = (MipSizeX + BlockSizeX - 1) / BlockSizeX;
 	uint32 NumBlocksY = (MipSizeY + BlockSizeY - 1) / BlockSizeY;
 	if (PixelFormat == PF_PVRTC2 || PixelFormat == PF_PVRTC4)
@@ -1592,7 +1339,7 @@ uint32 FMetalSurface::GetMipSize(uint32 MipIndex, uint32* Stride, bool bSingleLa
 		NumBlocksY = FMath::Max<uint32>(NumBlocksY, 2);
 	}
 #if PLATFORM_MAC
-	else if (PixelFormat == PF_G8 && EnumHasAnyFlags(Flags, TexCreate_SRGB) && Texture.GetPixelFormat() == mtlpp::PixelFormat::RGBA8Unorm_sRGB)
+	else if (PixelFormat == PF_G8 && EnumHasAnyFlags(GetDesc().Flags, TexCreate_SRGB) && Texture.GetPixelFormat() == mtlpp::PixelFormat::RGBA8Unorm_sRGB)
 	{
 		// RGBA_sRGB is the closest match - so expand the data.
 		NumBlocksX *= 4;
@@ -1636,23 +1383,12 @@ uint32 FMetalSurface::GetMemorySize()
 
 uint32 FMetalSurface::GetNumFaces()
 {
-	switch (Type)
-	{
-		case RRT_Texture2DArray:
-		case RRT_Texture3D:
-		case RRT_TextureCube:
-			return SizeZ * Texture.GetArrayLength();
-			
-		case RRT_Texture:
-		case RRT_Texture2D:
-		default:
-			return 1;
-	}
+	return GetDesc().Depth * GetDesc().ArraySize;
 }
 
 FMetalTexture FMetalSurface::GetDrawableTexture()
 {
-	if (!Texture && EnumHasAnyFlags(Flags, TexCreate_Presentable))
+	if (!Texture && EnumHasAnyFlags(GetDesc().Flags, TexCreate_Presentable))
 	{
 		check(Viewport);
 		Texture = Viewport->GetDrawableTexture(EMetalViewportAccessRHI);
@@ -1663,7 +1399,7 @@ FMetalTexture FMetalSurface::GetDrawableTexture()
 ns::AutoReleased<FMetalTexture> FMetalSurface::GetCurrentTexture()
 {
 	ns::AutoReleased<FMetalTexture> Tex;
-	if (Viewport && EnumHasAnyFlags(Flags, TexCreate_Presentable))
+	if (Viewport && EnumHasAnyFlags(GetDesc().Flags, TexCreate_Presentable))
 	{
 		check(Viewport);
 		Tex = Viewport->GetCurrentTexture(EMetalViewportAccessRHI);
@@ -1721,17 +1457,17 @@ uint32 FMetalDynamicRHI::RHIComputeMemorySize(FRHITexture* TextureRHI)
  2D texture support.
  -----------------------------------------------------------------------------*/
 
-FTexture2DRHIRef FMetalDynamicRHI::RHICreateTexture2D(uint32 SizeX, uint32 SizeY, uint8 Format, uint32 NumMips, uint32 NumSamples, ETextureCreateFlags Flags, ERHIAccess InResourceState, FRHIResourceCreateInfo& CreateInfo)
+FTextureRHIRef FMetalDynamicRHI::RHICreateTexture_RenderThread(class FRHICommandListImmediate& RHICmdList, const FRHITextureCreateDesc& CreateDesc)
 {
-	@autoreleasepool {
-		return new FMetalTexture2D((EPixelFormat)Format, SizeX, SizeY, NumMips, NumSamples, Flags, CreateInfo.BulkData, CreateInfo.ClearValueBinding);
+	@autoreleasepool{
+		return this->RHICreateTexture(CreateDesc);
 	}
 }
 
-FTexture2DRHIRef FMetalDynamicRHI::RHICreateTextureExternal2D(uint32 SizeX, uint32 SizeY, uint8 Format, uint32 NumMips, uint32 NumSamples, ETextureCreateFlags Flags, ERHIAccess InResourceState, FRHIResourceCreateInfo& CreateInfo)
+FTextureRHIRef FMetalDynamicRHI::RHICreateTexture(const FRHITextureCreateDesc& CreateDesc)
 {
-	@autoreleasepool {
-		return new FMetalTexture2D((EPixelFormat)Format, SizeX, SizeY, NumMips, NumSamples, Flags, CreateInfo.BulkData, CreateInfo.ClearValueBinding);
+	@autoreleasepool{
+		return new FMetalSurface(CreateDesc);
 	}
 }
 
@@ -1746,20 +1482,6 @@ void FMetalDynamicRHI::RHICopySharedMips(FRHITexture2D* DestTexture2D, FRHITextu
 	NOT_SUPPORTED("RHICopySharedMips");
 }
 
-FTexture2DArrayRHIRef FMetalDynamicRHI::RHICreateTexture2DArray(uint32 SizeX, uint32 SizeY, uint32 SizeZ, uint8 Format, uint32 NumMips, uint32 NumSamples, ETextureCreateFlags Flags, ERHIAccess InResourceState, FRHIResourceCreateInfo& CreateInfo)
-{
-	@autoreleasepool {
-		return new FMetalTexture2DArray((EPixelFormat)Format, SizeX, SizeY, SizeZ, NumMips, Flags, CreateInfo.BulkData, CreateInfo.ClearValueBinding);
-	}
-}
-
-FTexture3DRHIRef FMetalDynamicRHI::RHICreateTexture3D(uint32 SizeX, uint32 SizeY, uint32 SizeZ, uint8 Format, uint32 NumMips, ETextureCreateFlags Flags, ERHIAccess InResourceState, FRHIResourceCreateInfo& CreateInfo)
-{
-	@autoreleasepool {
-		return new FMetalTexture3D((EPixelFormat)Format, SizeX, SizeY, SizeZ, NumMips, Flags, CreateInfo.BulkData, CreateInfo.ClearValueBinding);
-	}
-}
-
 void FMetalDynamicRHI::RHIGenerateMips(FRHITexture* SourceSurfaceRHI)
 {
 	@autoreleasepool {
@@ -1771,131 +1493,98 @@ void FMetalDynamicRHI::RHIGenerateMips(FRHITexture* SourceSurfaceRHI)
 	}
 }
 
-// deferred unlock commands
-struct FMetalRHICommandAsyncReallocateTexture2D final : public FRHICommand<FMetalRHICommandAsyncReallocateTexture2D>
-{
-	FMetalContext& Context;
-	FMetalTexture2D* OldTexture;
-	FMetalTexture2D* NewTexture;
-	int32 NewMipCount;
-	int32 NewSizeX;
-	int32 NewSizeY;
-	FThreadSafeCounter* RequestStatus;
-	
-	FORCEINLINE_DEBUGGABLE FMetalRHICommandAsyncReallocateTexture2D(FMetalContext& InContext, FMetalTexture2D* InOldTexture, FMetalTexture2D* InNewTexture, int32 InNewMipCount, int32 InNewSizeX, int32 InNewSizeY, FThreadSafeCounter* InRequestStatus)
-	: Context(InContext)
-	, OldTexture(InOldTexture)
-	, NewTexture(InNewTexture)
-	, NewMipCount(InNewMipCount)
-	, NewSizeX(InNewSizeX)
-	, NewSizeY(InNewSizeY)
-	, RequestStatus(InRequestStatus)
-	{
-	}
-	
-	void Execute(FRHICommandListBase& CmdList)
-	{
-		CopyMips(Context, OldTexture, NewTexture, NewMipCount, NewSizeX, NewSizeY, RequestStatus);
-	}
-	
-	static void CopyMips(FMetalContext& Context, FMetalTexture2D* OldTexture, FMetalTexture2D* NewTexture, int32 NewMipCount, int32 NewSizeX, int32 NewSizeY, FThreadSafeCounter* RequestStatus)
-	{
-		// figure out what mips to schedule
-		const uint32 NumSharedMips = FMath::Min(OldTexture->GetNumMips(), NewTexture->GetNumMips());
-		const uint32 SourceMipOffset = OldTexture->GetNumMips() - NumSharedMips;
-		const uint32 DestMipOffset = NewTexture->GetNumMips() - NumSharedMips;
-		
-		const uint32 BlockSizeX = GPixelFormats[OldTexture->GetFormat()].BlockSizeX;
-		const uint32 BlockSizeY = GPixelFormats[OldTexture->GetFormat()].BlockSizeY;
-		
-		// only handling straight 2D textures here
-		uint32 SliceIndex = 0;
-		mtlpp::Origin Origin(0,0,0);
-		
-		FMetalTexture Tex = OldTexture->Surface.Texture;
-		
-		// DXT/BC formats on Mac actually do have mip-tails that are smaller than the block size, they end up being uncompressed.
-		bool const bPixelFormatASTC = IsPixelFormatASTCCompressed(OldTexture->GetFormat());
-		
-		bool bAsync = true;
-		for (uint32 MipIndex = 0; MipIndex < NumSharedMips; ++MipIndex)
-		{
-			const uint32 UnalignedMipSizeX = FMath::Max<uint32>(1, NewSizeX >> (MipIndex + DestMipOffset));
-			const uint32 UnalignedMipSizeY = FMath::Max<uint32>(1, NewSizeY >> (MipIndex + DestMipOffset));
-			const uint32 MipSizeX = FMath::Max<uint32>(1, NewSizeX >> (MipIndex + DestMipOffset));
-			const uint32 MipSizeY = FMath::Max<uint32>(1, NewSizeY >> (MipIndex + DestMipOffset));
-			
-			bAsync &= Context.AsyncCopyFromTextureToTexture(OldTexture->Surface.Texture, SliceIndex, MipIndex + SourceMipOffset, Origin, mtlpp::Size(MipSizeX, MipSizeY, 1), NewTexture->Surface.Texture, SliceIndex, MipIndex + DestMipOffset, Origin);
-		}
-		
-		// when done, decrement the counter to indicate it's safe
-		mtlpp::CommandBufferHandler CompletionHandler = [Tex](mtlpp::CommandBuffer const&)
-		{
-		};
-		
-		if(bAsync)
-		{
-	    	// kck it off!
-			Context.SubmitAsyncCommands(nil, CompletionHandler, false);
-		}
-		else
-		{
-			Context.GetCurrentRenderPass().AddCompletionHandler(CompletionHandler);
-		}
-		
-		// Like D3D mark this as complete immediately.
-		RequestStatus->Decrement();
-		
-		FMetalSurface* Source = GetMetalSurfaceFromRHITexture(OldTexture);
-		Source->MakeAliasable();
-	}
-};
-
 FTexture2DRHIRef FMetalDynamicRHI::AsyncReallocateTexture2D_RenderThread(class FRHICommandListImmediate& RHICmdList, FRHITexture2D* Texture2D, int32 NewMipCount, int32 NewSizeX, int32 NewSizeY, FThreadSafeCounter* RequestStatus)
 {
-	@autoreleasepool {
-		FTexture2DRHIRef Result;
-		
-		if (RHICmdList.Bypass() || !IsRunningRHIInSeparateThread())
-		{
-			Result = GDynamicRHI->RHIAsyncReallocateTexture2D(Texture2D, NewMipCount, NewSizeX, NewSizeY, RequestStatus);
-		}
-		else
-		{
-			FMetalTexture2D* OldTexture = ResourceCast(Texture2D);
-			
-			FMetalTexture2D* NewTexture = new FMetalTexture2D(OldTexture->GetFormat(), NewSizeX, NewSizeY, NewMipCount, OldTexture->GetNumSamples(), OldTexture->GetFlags(), NULL, Texture2D->GetClearBinding());
-			
-			Result = NewTexture;
-			
-			new (RHICmdList.AllocCommand<FMetalRHICommandAsyncReallocateTexture2D>()) FMetalRHICommandAsyncReallocateTexture2D(ImmediateContext.GetInternalContext(), OldTexture, NewTexture, NewMipCount, NewSizeX, NewSizeY, RequestStatus);
-		}
-		
-		return Result;
-	}
+	return this->RHIAsyncReallocateTexture2D(Texture2D, NewMipCount, NewSizeX, NewSizeY, RequestStatus);
 }
 
 ETextureReallocationStatus FMetalDynamicRHI::FinalizeAsyncReallocateTexture2D_RenderThread(class FRHICommandListImmediate& RHICmdList, FRHITexture2D* Texture2D, bool bBlockUntilCompleted)
 {
 	// No need to flush - does nothing
-	return GDynamicRHI->RHIFinalizeAsyncReallocateTexture2D(Texture2D, bBlockUntilCompleted);
+	return this->RHIFinalizeAsyncReallocateTexture2D(Texture2D, bBlockUntilCompleted);
 }
 
 ETextureReallocationStatus FMetalDynamicRHI::CancelAsyncReallocateTexture2D_RenderThread(class FRHICommandListImmediate& RHICmdList, FRHITexture2D* Texture2D, bool bBlockUntilCompleted)
 {
 	// No need to flush - does nothing
-	return GDynamicRHI->RHICancelAsyncReallocateTexture2D(Texture2D, bBlockUntilCompleted);
+	return this->RHICancelAsyncReallocateTexture2D(Texture2D, bBlockUntilCompleted);
 }
 
-
 FTexture2DRHIRef FMetalDynamicRHI::RHIAsyncReallocateTexture2D(FRHITexture2D* OldTextureRHI, int32 NewMipCount, int32 NewSizeX, int32 NewSizeY, FThreadSafeCounter* RequestStatus)
 {
 	@autoreleasepool {
-		FMetalTexture2D* OldTexture = ResourceCast(OldTextureRHI);
-		
-		FMetalTexture2D* NewTexture = new FMetalTexture2D(OldTexture->GetFormat(), NewSizeX, NewSizeY, NewMipCount, OldTexture->GetNumSamples(), OldTexture->GetFlags(), NULL, OldTextureRHI->GetClearBinding());
+
+		check(IsInRenderingThread());
+		FRHICommandListImmediate& RHICmdList = FRHICommandListExecutor::GetImmediateCommandList();
+
+		FMetalSurface* OldTexture = ResourceCast(OldTextureRHI);
+
+		FRHITextureDesc Desc = OldTexture->GetDesc();
+		Desc.Extent = FIntPoint(NewSizeX, NewSizeY);
+		Desc.NumMips = NewMipCount;
+
+		FRHITextureCreateDesc CreateDesc(
+			Desc,
+			RHIGetDefaultResourceState(Desc.Flags, false),
+			TEXT("RHIAsyncReallocateTexture2D")
+		);
 		
-		FMetalRHICommandAsyncReallocateTexture2D::CopyMips(ImmediateContext.GetInternalContext(), OldTexture, NewTexture, NewMipCount, NewSizeX, NewSizeY, RequestStatus);
+		FMetalSurface* NewTexture = new FMetalSurface(CreateDesc);
+
+		// Copy shared mips
+		RHICmdList.EnqueueLambda([this, OldTexture, NewSizeX, NewSizeY, NewTexture, RequestStatus](FRHICommandListImmediate& RHICmdList)
+		{
+			FMetalContext& Context = ImmediateContext.GetInternalContext();
+
+			// figure out what mips to schedule
+			const uint32 NumSharedMips = FMath::Min(OldTexture->GetNumMips(), NewTexture->GetNumMips());
+			const uint32 SourceMipOffset = OldTexture->GetNumMips() - NumSharedMips;
+			const uint32 DestMipOffset = NewTexture->GetNumMips() - NumSharedMips;
+
+			const uint32 BlockSizeX = GPixelFormats[OldTexture->GetFormat()].BlockSizeX;
+			const uint32 BlockSizeY = GPixelFormats[OldTexture->GetFormat()].BlockSizeY;
+
+			// only handling straight 2D textures here
+			uint32 SliceIndex = 0;
+			mtlpp::Origin Origin(0, 0, 0);
+
+			FMetalTexture Tex = OldTexture->Texture;
+
+			// DXT/BC formats on Mac actually do have mip-tails that are smaller than the block size, they end up being uncompressed.
+			bool const bPixelFormatASTC = IsPixelFormatASTCCompressed(OldTexture->GetFormat());
+
+			bool bAsync = true;
+			for (uint32 MipIndex = 0; MipIndex < NumSharedMips; ++MipIndex)
+			{
+				const uint32 UnalignedMipSizeX = FMath::Max<uint32>(1, NewSizeX >> (MipIndex + DestMipOffset));
+				const uint32 UnalignedMipSizeY = FMath::Max<uint32>(1, NewSizeY >> (MipIndex + DestMipOffset));
+				const uint32 MipSizeX = FMath::Max<uint32>(1, NewSizeX >> (MipIndex + DestMipOffset));
+				const uint32 MipSizeY = FMath::Max<uint32>(1, NewSizeY >> (MipIndex + DestMipOffset));
+
+				bAsync &= Context.AsyncCopyFromTextureToTexture(OldTexture->Texture, SliceIndex, MipIndex + SourceMipOffset, Origin, mtlpp::Size(MipSizeX, MipSizeY, 1), NewTexture->Texture, SliceIndex, MipIndex + DestMipOffset, Origin);
+			}
+
+			// when done, decrement the counter to indicate it's safe
+			mtlpp::CommandBufferHandler CompletionHandler = [Tex](mtlpp::CommandBuffer const&)
+			{
+			};
+
+			if (bAsync)
+			{
+				// kck it off!
+				Context.SubmitAsyncCommands(nil, CompletionHandler, false);
+			}
+			else
+			{
+				Context.GetCurrentRenderPass().AddCompletionHandler(CompletionHandler);
+			}
+
+			// Like D3D mark this as complete immediately.
+			RequestStatus->Decrement();
+
+			FMetalSurface* Source = GetMetalSurfaceFromRHITexture(OldTexture);
+			Source->MakeAliasable();
+		});
 		
 		return NewTexture;
 	}
@@ -1916,9 +1605,9 @@ void* FMetalDynamicRHI::LockTexture2D_RenderThread(class FRHICommandListImmediat
 	@autoreleasepool {
 		check(IsInRenderingThread());
 		
-		FMetalTexture2D* TextureMTL = ResourceCast(Texture);
+		FMetalSurface* TextureMTL = ResourceCast(Texture);
 		
-		void* BufferData = TextureMTL->Surface.AsyncLock(RHICmdList, MipIndex, 0, LockMode, DestStride, bNeedsDefaultRHIFlush);
+		void* BufferData = TextureMTL->AsyncLock(RHICmdList, MipIndex, 0, LockMode, DestStride, bNeedsDefaultRHIFlush);
 		
 		return BufferData;
 	}
@@ -1930,8 +1619,8 @@ void FMetalDynamicRHI::UnlockTexture2D_RenderThread(class FRHICommandListImmedia
 	{
 		check(IsInRenderingThread());
 		
-		FMetalTexture2D* TextureMTL = ResourceCast(Texture);
-		TextureMTL->Surface.Unlock(MipIndex, 0, true);
+		FMetalSurface* TextureMTL = ResourceCast(Texture);
+		TextureMTL->Unlock(MipIndex, 0, true);
 	}
 }
 
@@ -1939,32 +1628,32 @@ void FMetalDynamicRHI::UnlockTexture2D_RenderThread(class FRHICommandListImmedia
 void* FMetalDynamicRHI::RHILockTexture2D(FRHITexture2D* TextureRHI,uint32 MipIndex,EResourceLockMode LockMode,uint32& DestStride,bool bLockWithinMiptail)
 {
 	@autoreleasepool {
-		FMetalTexture2D* Texture = ResourceCast(TextureRHI);
-		return Texture->Surface.Lock(MipIndex, 0, LockMode, DestStride);
+		FMetalSurface* Texture = ResourceCast(TextureRHI);
+		return Texture->Lock(MipIndex, 0, LockMode, DestStride);
 	}
 }
 
 void FMetalDynamicRHI::RHIUnlockTexture2D(FRHITexture2D* TextureRHI,uint32 MipIndex,bool bLockWithinMiptail)
 {
 	@autoreleasepool {
-		FMetalTexture2D* Texture = ResourceCast(TextureRHI);
-		Texture->Surface.Unlock(MipIndex, 0, false);
+		FMetalSurface* Texture = ResourceCast(TextureRHI);
+		Texture->Unlock(MipIndex, 0, false);
 	}
 }
 
 void* FMetalDynamicRHI::RHILockTexture2DArray(FRHITexture2DArray* TextureRHI,uint32 TextureIndex,uint32 MipIndex,EResourceLockMode LockMode,uint32& DestStride,bool bLockWithinMiptail)
 {
 	@autoreleasepool {
-		FMetalTexture2DArray* Texture = ResourceCast(TextureRHI);
-		return Texture->Surface.Lock(MipIndex, TextureIndex, LockMode, DestStride);
+		FMetalSurface* Texture = ResourceCast(TextureRHI);
+		return Texture->Lock(MipIndex, TextureIndex, LockMode, DestStride);
 	}
 }
 
 void FMetalDynamicRHI::RHIUnlockTexture2DArray(FRHITexture2DArray* TextureRHI,uint32 TextureIndex,uint32 MipIndex,bool bLockWithinMiptail)
 {
 	@autoreleasepool {
-		FMetalTexture2DArray* Texture = ResourceCast(TextureRHI);
-		Texture->Surface.Unlock(MipIndex, TextureIndex, false);
+		FMetalSurface* Texture = ResourceCast(TextureRHI);
+		Texture->Unlock(MipIndex, TextureIndex, false);
 	}
 }
 
@@ -1993,12 +1682,13 @@ static FMetalBuffer InternalCopyTexture2DUpdateRegion(FRHITexture2D* TextureRHI,
 
 	FMetalBuffer Buffer;
 
-	FMetalTexture2D* Texture = ResourceCast(TextureRHI);	
+	FMetalSurface* Texture = ResourceCast(TextureRHI);
 
 #if PLATFORM_MAC
 	// Expand R8_sRGB into RGBA8_sRGB for non Apple Silicon Mac.
-	FMetalTexture Tex = Texture->Surface.Texture;
-	if(Texture->GetFormat() == PF_G8 && EnumHasAnyFlags(Texture->GetFlags(), TexCreate_SRGB) && Tex.GetPixelFormat() == mtlpp::PixelFormat::RGBA8Unorm_sRGB)
+	if (   Texture->GetFormat() == PF_G8 
+		&& EnumHasAnyFlags(Texture->GetFlags(), TexCreate_SRGB) 
+		&& Texture->Texture.GetPixelFormat() == mtlpp::PixelFormat::RGBA8Unorm_sRGB)
 	{
 		const uint32 BufferSize = UpdateRegion.Height * UpdateRegion.Width * sizeof(uint32);		
 		Buffer = GetMetalDeviceContext().CreatePooledBuffer(FMetalPooledBufferArgs(GetMetalDeviceContext().GetDevice(), BufferSize, BUF_Dynamic, mtlpp::StorageMode::Shared));
@@ -2030,8 +1720,8 @@ static FMetalBuffer InternalCopyTexture2DUpdateRegion(FRHITexture2D* TextureRHI,
 
 static void InternalUpdateTexture2D(FMetalContext& Context, FRHITexture2D* TextureRHI, uint32 MipIndex, FUpdateTextureRegion2D const& UpdateRegion, uint32 SourcePitch, FMetalBuffer Buffer)
 {
-	FMetalTexture2D* Texture = ResourceCast(TextureRHI);
-	FMetalTexture Tex = Texture->Surface.Texture;
+	FMetalSurface* Texture = ResourceCast(TextureRHI);
+	FMetalTexture Tex = Texture->Texture;
 	
 	mtlpp::Region Region = mtlpp::Region(UpdateRegion.DestX, UpdateRegion.DestY, UpdateRegion.Width, UpdateRegion.Height);
 	
@@ -2058,7 +1748,7 @@ static void InternalUpdateTexture2D(FMetalContext& Context, FRHITexture2D* Textu
 		MTLPP_VALIDATE(mtlpp::Texture, Tex, SafeGetRuntimeDebuggingLevel() >= EMetalDebugLevelValidation, Replace(Region, MipIndex, 0, MTLPP_VALIDATE(mtlpp::Buffer, Buffer, SafeGetRuntimeDebuggingLevel() >= EMetalDebugLevelValidation, GetContents()), SourcePitch, 0));
 	}
 
-	FPlatformAtomics::InterlockedExchange(&Texture->Surface.Written, 1);
+	FPlatformAtomics::InterlockedExchange(&Texture->Written, 1);
 }
 
 struct FMetalRHICommandUpdateTexture2D final : public FRHICommand<FMetalRHICommandUpdateTexture2D>
@@ -2107,8 +1797,8 @@ void FMetalDynamicRHI::RHIUpdateTexture2D(FRHITexture2D* TextureRHI, uint32 MipI
 {
 	@autoreleasepool 
 	{
-		FMetalTexture2D* Texture = ResourceCast(TextureRHI);
-		FMetalTexture Tex = Texture->Surface.Texture;
+		FMetalSurface* Texture = ResourceCast(TextureRHI);
+		FMetalTexture Tex = Texture->Texture;
 		bool const bUseIntermediateMetalBuffer = Tex.GetStorageMode() == mtlpp::StorageMode::Private;
 		
 		if(bUseIntermediateMetalBuffer)
@@ -2132,7 +1822,7 @@ void FMetalDynamicRHI::RHIUpdateTexture2D(FRHITexture2D* TextureRHI, uint32 MipI
 			mtlpp::Region Region = mtlpp::Region(UpdateRegion.DestX, UpdateRegion.DestY, UpdateRegion.Width, UpdateRegion.Height);
 			MTLPP_VALIDATE(mtlpp::Texture, Tex, SafeGetRuntimeDebuggingLevel() >= EMetalDebugLevelValidation, Replace(Region, MipIndex, 0, SourceData, SourcePitch, 0));
 
-			FPlatformAtomics::InterlockedExchange(&Texture->Surface.Written, 1);
+			FPlatformAtomics::InterlockedExchange(&Texture->Written, 1);
 		}
 		
 		INC_DWORD_STAT_BY(STAT_MetalTextureMemUpdate, UpdateRegion.Height*SourcePitch);
@@ -2142,7 +1832,7 @@ void FMetalDynamicRHI::RHIUpdateTexture2D(FRHITexture2D* TextureRHI, uint32 MipI
 static void InternalCopyTexture3DUpdateRegionData(FRHITexture3D* TextureRHI, const struct FUpdateTextureRegion3D& UpdateRegion, uint32 SourceRowPitch, uint32 SourceDepthPitch, const uint8* SourceData, uint8* DestData)
 {
 	// Perform safe line copy
-	FMetalTexture3D* Texture = ResourceCast(TextureRHI);	
+	FMetalSurface* Texture = ResourceCast(TextureRHI);	
 	
 	uint8* pDest = DestData;
 	uint32 CopyPitch = FMath::DivideAndRoundUp(UpdateRegion.Width, (uint32)GPixelFormats[Texture->GetFormat()].BlockSizeX) * GPixelFormats[Texture->GetFormat()].BlockBytes;
@@ -2164,8 +1854,8 @@ static void InternalCopyTexture3DUpdateRegionData(FRHITexture3D* TextureRHI, con
 
 static void InternalUpdateTexture3D(FMetalContext& Context, FRHITexture3D* TextureRHI, uint32 MipIndex, const FUpdateTextureRegion3D& UpdateRegion, uint32 SourceRowPitch, uint32 SourceDepthPitch, FMetalBuffer Buffer)
 {
-	FMetalTexture3D* Texture = ResourceCast(TextureRHI);
-	FMetalTexture Tex = Texture->Surface.Texture;
+	FMetalSurface* Texture = ResourceCast(TextureRHI);
+	FMetalTexture Tex = Texture->Texture;
 	
 	mtlpp::Region Region = mtlpp::Region(UpdateRegion.DestX, UpdateRegion.DestY, UpdateRegion.DestZ, UpdateRegion.Width, UpdateRegion.Height, UpdateRegion.Depth);
 	
@@ -2189,7 +1879,7 @@ static void InternalUpdateTexture3D(FMetalContext& Context, FRHITexture3D* Textu
 		MTLPP_VALIDATE(mtlpp::Texture, Tex, SafeGetRuntimeDebuggingLevel() >= EMetalDebugLevelValidation, Replace(Region, MipIndex, 0, (uint8*)Buffer.GetContents(), SourceRowPitch, SourceDepthPitch));
 	}
 
-	FPlatformAtomics::InterlockedExchange(&Texture->Surface.Written, 1);
+	FPlatformAtomics::InterlockedExchange(&Texture->Written, 1);
 }
 
 struct FMetalDynamicRHIUpdateTexture3DCommand final : public FRHICommand<FMetalDynamicRHIUpdateTexture3DCommand>
@@ -2210,8 +1900,8 @@ struct FMetalDynamicRHIUpdateTexture3DCommand final : public FRHICommand<FMetalD
 	, SourceRowPitch(InSourceRowPitch)
 	, SourceDepthPitch(InSourceDepthPitch)
 	{
-		FMetalTexture3D* Texture = ResourceCast(TextureRHI);
-		FMetalTexture Tex = Texture->Surface.Texture;
+		FMetalSurface* Texture = ResourceCast(TextureRHI);
+		FMetalTexture Tex = Texture->Texture;
 		const uint32 BufferSize = UpdateRegion.Height * UpdateRegion.Depth* SourceRowPitch;
 		
 		Buffer = GetMetalDeviceContext().CreatePooledBuffer(FMetalPooledBufferArgs(GetMetalDeviceContext().GetDevice(), BufferSize, BUF_Dynamic, mtlpp::StorageMode::Shared));
@@ -2274,8 +1964,8 @@ void FMetalDynamicRHI::RHIUpdateTexture3D(FRHITexture3D* TextureRHI,uint32 MipIn
 {
 	@autoreleasepool {
 	
-		FMetalTexture3D* Texture = ResourceCast(TextureRHI);
-		FMetalTexture Tex = Texture->Surface.Texture;
+		FMetalSurface* Texture = ResourceCast(TextureRHI);
+		FMetalTexture Tex = Texture->Texture;
 		
 #if PLATFORM_MAC
 		checkf(!(Texture->GetFormat() == PF_G8 && EnumHasAnyFlags(Texture->GetFlags(), TexCreate_SRGB) && Tex.GetPixelFormat() == mtlpp::PixelFormat::RGBA8Unorm_sRGB), TEXT("MetalRHI on non Apple Silicon does not support PF_G8_sRGB on 3D, array or cube textures as it requires manual, CPU-side expansion to RGBA8_sRGB which is expensive!"));
@@ -2294,7 +1984,7 @@ void FMetalDynamicRHI::RHIUpdateTexture3D(FRHITexture3D* TextureRHI,uint32 MipIn
 		{
 			mtlpp::Region Region = mtlpp::Region(UpdateRegion.DestX, UpdateRegion.DestY, UpdateRegion.DestZ, UpdateRegion.Width, UpdateRegion.Height, UpdateRegion.Depth);
 			MTLPP_VALIDATE(mtlpp::Texture, Tex, SafeGetRuntimeDebuggingLevel() >= EMetalDebugLevelValidation, Replace(Region, MipIndex, 0, SourceData, SourceRowPitch, SourceDepthPitch));
-			FPlatformAtomics::InterlockedExchange(&Texture->Surface.Written, 1);
+			FPlatformAtomics::InterlockedExchange(&Texture->Written, 1);
 		}
 		
 		INC_DWORD_STAT_BY(STAT_MetalTextureMemUpdate, UpdateRegion.Height * UpdateRegion.Width * SourceDepthPitch);
@@ -2304,139 +1994,24 @@ void FMetalDynamicRHI::RHIUpdateTexture3D(FRHITexture3D* TextureRHI,uint32 MipIn
 /*-----------------------------------------------------------------------------
  Cubemap texture support.
  -----------------------------------------------------------------------------*/
-FTextureCubeRHIRef FMetalDynamicRHI::RHICreateTextureCube(uint32 Size, uint8 Format, uint32 NumMips, ETextureCreateFlags Flags, ERHIAccess InResourceState, FRHIResourceCreateInfo& CreateInfo)
-{
-	@autoreleasepool {
-		return new FMetalTextureCube((EPixelFormat)Format, Size, false, 1, NumMips, Flags, CreateInfo.BulkData, CreateInfo.ClearValueBinding);
-	}
-}
-
-FTextureCubeRHIRef FMetalDynamicRHI::RHICreateTextureCubeArray(uint32 Size, uint32 ArraySize, uint8 Format, uint32 NumMips, ETextureCreateFlags Flags, ERHIAccess InResourceState, FRHIResourceCreateInfo& CreateInfo)
-{
-	@autoreleasepool {
-		return new FMetalTextureCube((EPixelFormat)Format, Size, true, ArraySize, NumMips, Flags, CreateInfo.BulkData, CreateInfo.ClearValueBinding);
-	}
-}
-
 void* FMetalDynamicRHI::RHILockTextureCubeFace(FRHITextureCube* TextureCubeRHI,uint32 FaceIndex,uint32 ArrayIndex,uint32 MipIndex,EResourceLockMode LockMode,uint32& DestStride,bool bLockWithinMiptail)
 {
 	@autoreleasepool {
-		FMetalTextureCube* TextureCube = ResourceCast(TextureCubeRHI);
+		FMetalSurface* TextureCube = ResourceCast(TextureCubeRHI);
 		uint32 MetalFace = GetMetalCubeFace((ECubeFace)FaceIndex);
-		return TextureCube->Surface.Lock(MipIndex, MetalFace + (6 * ArrayIndex), LockMode, DestStride, true);
+		return TextureCube->Lock(MipIndex, MetalFace + (6 * ArrayIndex), LockMode, DestStride, true);
 	}
 }
 
 void FMetalDynamicRHI::RHIUnlockTextureCubeFace(FRHITextureCube* TextureCubeRHI,uint32 FaceIndex,uint32 ArrayIndex,uint32 MipIndex,bool bLockWithinMiptail)
 {
 	@autoreleasepool {
-		FMetalTextureCube* TextureCube = ResourceCast(TextureCubeRHI);
+		FMetalSurface* TextureCube = ResourceCast(TextureCubeRHI);
 		uint32 MetalFace = GetMetalCubeFace((ECubeFace)FaceIndex);
-		TextureCube->Surface.Unlock(MipIndex, MetalFace + (ArrayIndex * 6), false);
-	}
-}
-
-
-FTexture2DRHIRef FMetalDynamicRHI::RHICreateTexture2D_RenderThread(class FRHICommandListImmediate& RHICmdList, uint32 SizeX, uint32 SizeY, uint8 Format, uint32 NumMips, uint32 NumSamples, ETextureCreateFlags Flags, ERHIAccess InResourceState, FRHIResourceCreateInfo& CreateInfo)
-{
-	@autoreleasepool {
-		FRHIResourceCreateInfo Info = CreateInfo;
-		bool const bIOSurfaceData = (CreateInfo.BulkData && CreateInfo.BulkData->GetResourceType() != FResourceBulkDataInterface::EBulkDataType::Default);
-		Info.BulkData = bIOSurfaceData ? CreateInfo.BulkData : nullptr;
-		FTexture2DRHIRef Result = GDynamicRHI->RHICreateTexture2D(SizeX, SizeY, Format, NumMips, NumSamples, Flags, InResourceState, Info);
-		if (!bIOSurfaceData && CreateInfo.BulkData)
-		{
-			// upload existing bulkdata
-			UE_LOG(LogMetal, Display, TEXT("Got a bulk data texture, with %d mips"), NumMips);
-			checkf(NumMips == 1, TEXT("Only handling bulk data with 1 mip and 1 array length"));
-			uint32 Stride;
-			
-			// lock, copy, unlock
-			void* LockedData = LockTexture2D_RenderThread(RHICmdList, Result.GetReference(), 0, RLM_WriteOnly, Stride, false);
-			check(LockedData);
-			FMemory::Memcpy(LockedData, CreateInfo.BulkData->GetResourceBulkData(), CreateInfo.BulkData->GetResourceBulkDataSize());
-			UnlockTexture2D_RenderThread(RHICmdList, Result.GetReference(), 0, false);
-			
-			// bulk data can be unloaded now
-			CreateInfo.BulkData->Discard();
-		}
-		return Result;
-	}
-}
-
-FTexture2DArrayRHIRef FMetalDynamicRHI::RHICreateTexture2DArray_RenderThread(class FRHICommandListImmediate& RHICmdList, uint32 SizeX, uint32 SizeY, uint32 SizeZ, uint8 Format, uint32 NumMips, uint32 NumSamples, ETextureCreateFlags Flags, ERHIAccess InResourceState, FRHIResourceCreateInfo& CreateInfo)
-{
-	@autoreleasepool {
-		FRHIResourceCreateInfo Info = CreateInfo;
-		bool const bIOSurfaceData = (CreateInfo.BulkData && CreateInfo.BulkData->GetResourceType() != FResourceBulkDataInterface::EBulkDataType::Default);
-		Info.BulkData = bIOSurfaceData ? CreateInfo.BulkData : nullptr;
-		FTexture2DArrayRHIRef Result = GDynamicRHI->RHICreateTexture2DArray(SizeX, SizeY, SizeZ, Format, NumMips, NumSamples, Flags, InResourceState, Info);
-		if (CreateInfo.BulkData)
-		{
-			// upload existing bulkdata
-			UE_LOG(LogMetal, Display, TEXT("Got a bulk data texture, with %d mips"), NumMips);
-			checkf(NumMips == 1 && SizeZ == 1, TEXT("Only handling bulk data with 1 mip and 1 array length"));
-			uint32 Stride;
-			
-			// lock, copy, unlock
-			FMetalTexture2DArray* TextureMTL = ResourceCast(Result.GetReference());
-			void* LockedData = TextureMTL->Surface.AsyncLock(RHICmdList, 0, 0, RLM_WriteOnly, Stride, false);
-			FMemory::Memcpy(LockedData, CreateInfo.BulkData->GetResourceBulkData(), CreateInfo.BulkData->GetResourceBulkDataSize());
-			TextureMTL->Surface.Unlock(0, 0, true);
-			
-			// bulk data can be unloaded now
-			CreateInfo.BulkData->Discard();
-		}
-		return Result;
-	}
-}
-
-FTexture3DRHIRef FMetalDynamicRHI::RHICreateTexture3D_RenderThread(class FRHICommandListImmediate& RHICmdList, uint32 SizeX, uint32 SizeY, uint32 SizeZ, uint8 Format, uint32 NumMips, ETextureCreateFlags Flags, ERHIAccess InResourceState, FRHIResourceCreateInfo& CreateInfo)
-{
-	@autoreleasepool {
-		FRHIResourceCreateInfo Info = CreateInfo;
-		bool const bIOSurfaceData = (CreateInfo.BulkData && CreateInfo.BulkData->GetResourceType() != FResourceBulkDataInterface::EBulkDataType::Default);
-		Info.BulkData = bIOSurfaceData ? CreateInfo.BulkData : nullptr;
-		FTexture3DRHIRef Result = GDynamicRHI->RHICreateTexture3D(SizeX, SizeY, SizeZ, Format, NumMips, Flags, InResourceState, Info);
-		if (CreateInfo.BulkData)
-		{
-			// upload existing bulkdata
-			UE_LOG(LogMetal, Display, TEXT("Got a bulk data texture, with %d mips"), NumMips);
-			checkf(NumMips == 1, TEXT("Only handling bulk data with 1 mip and 1 array length"));
-			uint32 Stride;
-			
-			// lock, copy, unlock
-			FMetalTexture3D* TextureMTL = ResourceCast(Result.GetReference());
-			void* LockedData = TextureMTL->Surface.AsyncLock(RHICmdList, 0, 0, RLM_WriteOnly, Stride, false);
-			FMemory::Memcpy(LockedData, CreateInfo.BulkData->GetResourceBulkData(), CreateInfo.BulkData->GetResourceBulkDataSize());
-			TextureMTL->Surface.Unlock(0, 0, true);
-			
-			// bulk data can be unloaded now
-			CreateInfo.BulkData->Discard();
-		}
-		return Result;
-	}
-}
-
-FTextureCubeRHIRef FMetalDynamicRHI::RHICreateTextureCube_RenderThread(class FRHICommandListImmediate& RHICmdList, uint32 Size, uint8 Format, uint32 NumMips, ETextureCreateFlags Flags, ERHIAccess InResourceState, FRHIResourceCreateInfo& CreateInfo)
-{
-	@autoreleasepool {
-		check(!CreateInfo.BulkData);
-		FTextureCubeRHIRef Result = GDynamicRHI->RHICreateTextureCube(Size, Format, NumMips, Flags, InResourceState, CreateInfo);
-		return Result;
+		TextureCube->Unlock(MipIndex, MetalFace + (ArrayIndex * 6), false);
 	}
 }
 
-FTextureCubeRHIRef FMetalDynamicRHI::RHICreateTextureCubeArray_RenderThread(class FRHICommandListImmediate& RHICmdList, uint32 Size, uint32 ArraySize, uint8 Format, uint32 NumMips, ETextureCreateFlags Flags, ERHIAccess InResourceState, FRHIResourceCreateInfo& CreateInfo)
-{
-	@autoreleasepool {
-		check(!CreateInfo.BulkData);
-		FTextureCubeRHIRef Result = GDynamicRHI->RHICreateTextureCubeArray(Size, ArraySize, Format, NumMips, Flags, InResourceState, CreateInfo);
-		return Result;
-	}
-}
-
-
 void FMetalDynamicRHI::RHIBindDebugLabelName(FRHITexture* TextureRHI, const TCHAR* Name)
 {
 	@autoreleasepool {
@@ -2495,7 +2070,7 @@ void FMetalRHICommandContext::RHICopyTexture(FRHITexture* SourceTextureRHI, FRHI
 		
 		if (TextureFormatExactMatch || TextureFormatCompatible)
 		{
-			FIntVector Size = (CopyInfo.Size != FIntVector::ZeroValue) ? CopyInfo.Size : FIntVector(MetalSrcTexture->SizeX, MetalSrcTexture->SizeY, MetalSrcTexture->SizeZ);
+			FIntVector Size = (CopyInfo.Size != FIntVector::ZeroValue) ? CopyInfo.Size : FIntVector(MetalSrcTexture->GetDesc().Extent.X, MetalSrcTexture->GetDesc().Extent.Y, MetalSrcTexture->GetDesc().Depth);
 			
 			mtlpp::Origin SourceOrigin(CopyInfo.SourcePosition.X, CopyInfo.SourcePosition.Y, CopyInfo.SourcePosition.Z);
 			mtlpp::Origin DestinationOrigin(CopyInfo.DestPosition.X, CopyInfo.DestPosition.Y, CopyInfo.DestPosition.Z);
@@ -2507,7 +2082,7 @@ void FMetalRHICommandContext::RHICopyTexture(FRHITexture* SourceTextureRHI, FRHI
 				mtlpp::TextureUsage Usage = MetalSrcTexture->Texture.GetUsage();
 				if (Usage & mtlpp::TextureUsage::PixelFormatView)
 				{
-					ns::Range Slices(0, MetalSrcTexture->Texture.GetArrayLength() * (MetalSrcTexture->bIsCubemap ? 6 : 1));
+					ns::Range Slices(0, MetalSrcTexture->Texture.GetArrayLength() * (MetalSrcTexture->GetDesc().IsTextureCube() ? 6 : 1));
 					if (MetalSrcTexture->Texture.GetPixelFormat() != MetalDestTexture->Texture.GetPixelFormat())
 					{
 						SrcTexture = MetalSrcTexture->Texture.NewTextureView(MetalDestTexture->Texture.GetPixelFormat(), MetalSrcTexture->Texture.GetTextureType(), ns::Range(0, MetalSrcTexture->Texture.GetMipmapLevelCount()), Slices);
@@ -2533,8 +2108,8 @@ void FMetalRHICommandContext::RHICopyTexture(FRHITexture* SourceTextureRHI, FRHI
 
 					if (TextureFormatCompatible)
 					{
-						DestSize.width  *= GPixelFormats[MetalDestTexture->PixelFormat].BlockSizeX;
-						DestSize.height *= GPixelFormats[MetalDestTexture->PixelFormat].BlockSizeY;
+						DestSize.width  *= GPixelFormats[MetalDestTexture->GetDesc().Format].BlockSizeX;
+						DestSize.height *= GPixelFormats[MetalDestTexture->GetDesc().Format].BlockSizeY;
 					}
 
 					// Account for create with TexCreate_SRGB flag which could make these different
@@ -2549,8 +2124,8 @@ void FMetalRHICommandContext::RHICopyTexture(FRHITexture* SourceTextureRHI, FRHI
 						// format mismatch (like linear vs. sRGB), then we must
 						// achieve the copy by going through a buffer object.
 						//
-						const bool BlockSizeMatch = (GPixelFormats[MetalSrcTexture->PixelFormat].BlockSizeX == GPixelFormats[MetalDestTexture->PixelFormat].BlockSizeX);
-						const uint32 BytesPerPixel = (MetalSrcTexture->PixelFormat != PF_DepthStencil) ? GPixelFormats[MetalSrcTexture->PixelFormat].BlockBytes : 1;
+						const bool BlockSizeMatch = (GPixelFormats[MetalSrcTexture->GetDesc().Format].BlockSizeX == GPixelFormats[MetalDestTexture->GetDesc().Format].BlockSizeX);
+						const uint32 BytesPerPixel = (MetalSrcTexture->GetDesc().Format != PF_DepthStencil) ? GPixelFormats[MetalSrcTexture->GetDesc().Format].BlockBytes : 1;
 						const uint32 Stride = BytesPerPixel * SourceSize.width;
 #if PLATFORM_MAC
 						const uint32 Alignment = 1u;
