Copying //UE5/Dev-ParallelRendering to //UE5/Main

diff --git a/Engine/Source/Runtime/Apple/MetalRHI/Private/MetalUAV.cpp b/Engine/Source/Runtime/Apple/MetalRHI/Private/MetalUAV.cpp
--- a/Engine/Source/Runtime/Apple/MetalRHI/Private/MetalUAV.cpp
+++ b/Engine/Source/Runtime/Apple/MetalRHI/Private/MetalUAV.cpp
@@ -8,438 +8,392 @@
 #include "ClearReplacementShaders.h"
 #include "MetalTransitionData.h"
 
-FMetalShaderResourceView::FMetalShaderResourceView()
-	: TextureView(nullptr)
-	, Offset(0)
-	, MipLevel(0)
-	, bSRGBForceDisable(0)
+// Constructor for buffers
+FMetalResourceViewBase::FMetalResourceViewBase(FRHIBuffer* InBuffer, uint32 InStartOffsetBytes, uint32 InNumElements, EPixelFormat InFormat)
+	: SourceBuffer(ResourceCast(InBuffer))
+	, bTexture         (false)
+	, bSRGBForceDisable(false)
+	, MipLevel         (0)
+	, Reserved         (0)
+	, NumMips          (0)
+	, Format           (InFormat)
+	, Stride           (0)
+	, Offset           (InBuffer ? InStartOffsetBytes : 0)
+{
+	check(!bTexture);
+
+	if (SourceBuffer)
+		SourceBuffer->AddRef();
+
+	EBufferUsageFlags Usage = SourceBuffer->GetUsage();
+	if (EnumHasAnyFlags(Usage, BUF_VertexBuffer))
+	{
+		if (!SourceBuffer)
+		{
+			Stride = 0;
+		}
+		else
+		{
+			check(SourceBuffer->GetUsage() & BUF_ShaderResource);
+			Stride = GPixelFormats[Format].BlockBytes;
+
+			LinearTextureDesc = MakeUnique<FMetalLinearTextureDescriptor>(InStartOffsetBytes, InNumElements, Stride);
+			SourceBuffer->CreateLinearTexture((EPixelFormat)Format, SourceBuffer, LinearTextureDesc.Get());
+		}
+	}
+	else if (EnumHasAnyFlags(Usage, BUF_IndexBuffer))
+	{
+		if (!SourceBuffer)
+		{
+			Format = PF_R16_UINT;
+			Stride = 0;
+		}
+		else
+		{
+			Format = (SourceBuffer->IndexType == mtlpp::IndexType::UInt16) ? PF_R16_UINT : PF_R32_UINT;
+			Stride = SourceBuffer->GetStride();
+
+			check(Stride == ((Format == PF_R16_UINT) ? 2 : 4));
+
+			LinearTextureDesc = MakeUnique<FMetalLinearTextureDescriptor>(InStartOffsetBytes, InNumElements, Stride);
+			SourceBuffer->CreateLinearTexture((EPixelFormat)Format, SourceBuffer, LinearTextureDesc.Get());
+		}
+	}
+	else
+	{
+		check(EnumHasAnyFlags(Usage, BUF_StructuredBuffer));
+
+		Format = PF_Unknown;
+		Stride = SourceBuffer->GetStride();
+	}
+}
+
+// Constructor for textures
+FMetalResourceViewBase::FMetalResourceViewBase(
+	  FRHITexture* InTexture
+	, EPixelFormat InFormat
+	, uint8 InMipLevel
+	, uint8 InNumMipLevels
+	, ERHITextureSRVOverrideSRGBType InSRGBOverride
+	, uint32 InFirstArraySlice
+	, uint32 InNumArraySlices
+	)
+	: SourceTexture(ResourceCast(InTexture))
+	, bTexture(true)
+	, bSRGBForceDisable(InSRGBOverride == SRGBO_ForceDisable)
+	, MipLevel(InMipLevel)
 	, Reserved(0)
-	, NumMips(0)
-	, Format(0)
+	, NumMips(InNumMipLevels)
+	, Format((InTexture && InFormat == PF_Unknown) ? InTexture->GetDesc().Format : InFormat)
 	, Stride(0)
-	, LinearTextureDesc(nullptr)
+	, Offset(0)
 {
-	// void
-}
+	// Slice selection of a texture array still need to be implemented on Metal
+	check(InFirstArraySlice == 0 && InNumArraySlices == 0);
 
-FMetalShaderResourceView::~FMetalShaderResourceView()
-{
-	if (LinearTextureDesc)
-	{
-		delete LinearTextureDesc;
-		LinearTextureDesc = nullptr;
-	}
-	
-	if (TextureView)
+	if (SourceTexture)
 	{
-		FMetalSurface* Surface = GetMetalSurfaceFromRHITexture(SourceTexture);
-		if (Surface)
+		SourceTexture->AddRef();
+
+#if PLATFORM_IOS
+		// Memoryless targets can't have texture views (SRVs or UAVs)
+		if (Source.Texture.GetStorageMode() != mtlpp::StorageMode::Memoryless)
+#endif
 		{
-			Surface->SRVs.Remove(this);
-		}
+			// Determine the appropriate metal format for the view.
+			// This format will be non-sRGB. We convert to sRGB below if required.
+			mtlpp::PixelFormat MetalFormat = (mtlpp::PixelFormat)GPixelFormats[Format].PlatformFormat;
+
+			if (Format == PF_X24_G8)
+			{
+				// Stencil buffer view of a depth texture
+				check(SourceTexture->GetDesc().Format == PF_DepthStencil);
+				switch (SourceTexture->Texture.GetPixelFormat())
+				{
+				default: checkNoEntry(); break;
+				case mtlpp::PixelFormat::Depth24Unorm_Stencil8: MetalFormat = mtlpp::PixelFormat::X24_Stencil8; break;
+				case mtlpp::PixelFormat::Depth32Float_Stencil8: MetalFormat = mtlpp::PixelFormat::X32_Stencil8; break;
+				}
+			}
+			else
+			{
+				// Override the format's sRGB setting if appropriate
+				if (EnumHasAnyFlags(SourceTexture->GetDesc().Flags, TexCreate_SRGB))
+				{
+					if (bSRGBForceDisable)
+					{
+#if PLATFORM_MAC
+						// R8Unorm has been expanded in the source surface for sRGBA support - we need to expand to RGBA to enable compatible texture format view for non apple silicon macs
+						if (Format == PF_G8 && SourceTexture->Texture.GetPixelFormat() == mtlpp::PixelFormat::RGBA8Unorm_sRGB)
+						{
+							MetalFormat = mtlpp::PixelFormat::RGBA8Unorm;
+						}
+#endif
+					}
+					else
+					{
+						// Ensure we have the correct sRGB target format if we create a new texture view rather than using the source texture
+						MetalFormat = ToSRGBFormat(MetalFormat);
+					}
+				}
+			}
 
-		delete TextureView;
-		TextureView = nullptr;
+			// We can use the source texture directly if the view's format / mip count etc matches.
+			bool bUseSourceTex =
+				MipLevel == 0
+				&& NumMips == SourceTexture->Texture.GetMipmapLevelCount()
+				&& MetalFormat == SourceTexture->Texture.GetPixelFormat();
+
+			if (bUseSourceTex)
+			{
+				// SRV is exactly compatible with the original texture.
+				TextureView = SourceTexture->Texture;
+			}
+			else
+			{
+				// Recreate the texture to enable MTLTextureUsagePixelFormatView which must be off unless we definitely use this feature or we are throwing ~4% performance vs. Windows on the floor.
+				// @todo recreating resources like this will likely prevent us from making view creation multi-threaded.
+				if (!(SourceTexture->Texture.GetUsage() & mtlpp::TextureUsage::PixelFormatView))
+					SourceTexture->PrepareTextureView();
+
+				ns::Range Slices(0, SourceTexture->Texture.GetArrayLength() * (SourceTexture->GetDesc().IsTextureCube() ? 6 : 1));
+
+				TextureView = SourceTexture->Texture.NewTextureView(
+					MetalFormat,
+					SourceTexture->Texture.GetTextureType(),
+					ns::Range(MipLevel, NumMips),
+					Slices);
+			}
+		}
+	}
+	else
+	{
+		TextureView = {};
 	}
-	
-	SourceVertexBuffer = NULL;
-	SourceTexture = NULL;
 }
 
-void FMetalShaderResourceView::InitLinearTextureDescriptor(const FMetalLinearTextureDescriptor& InLinearTextureDescriptor)
+FMetalResourceViewBase::~FMetalResourceViewBase()
 {
-	check(!LinearTextureDesc);
-	LinearTextureDesc = new FMetalLinearTextureDescriptor(InLinearTextureDescriptor);
-	check(LinearTextureDesc);
+	if (bTexture)
+	{
+		if (SourceTexture)
+			SourceTexture->Release();
+	}
+	else
+	{
+		if (SourceBuffer)
+			SourceBuffer->Release();
+	}
 }
 
-ns::AutoReleased<FMetalTexture> FMetalShaderResourceView::GetLinearTexture(bool const bUAV)
+ns::AutoReleased<FMetalTexture> FMetalResourceViewBase::GetLinearTexture()
 {
 	ns::AutoReleased<FMetalTexture> NewLinearTexture;
+	if (SourceBuffer)
 	{
-		if (IsValidRef(SourceVertexBuffer))
-		{
-			NewLinearTexture = SourceVertexBuffer->GetLinearTexture((EPixelFormat)Format, LinearTextureDesc);
-		}
-		else if (IsValidRef(SourceIndexBuffer))
-		{
-			NewLinearTexture = SourceIndexBuffer->GetLinearTexture((EPixelFormat)Format, LinearTextureDesc);
-		}
+		NewLinearTexture = SourceBuffer->GetLinearTexture((EPixelFormat)Format, LinearTextureDesc.Get());
 	}
 	return NewLinearTexture;
 }
 
 FUnorderedAccessViewRHIRef FMetalDynamicRHI::RHICreateUnorderedAccessView_RenderThread(class FRHICommandListImmediate& RHICmdList, FRHIBuffer* Buffer, bool bUseUAVCounter, bool bAppendBuffer)
 {
-	return GDynamicRHI->RHICreateUnorderedAccessView(Buffer, bUseUAVCounter, bAppendBuffer);
+	return this->RHICreateUnorderedAccessView(Buffer, bUseUAVCounter, bAppendBuffer);
 }
 
 FUnorderedAccessViewRHIRef FMetalDynamicRHI::RHICreateUnorderedAccessView_RenderThread(class FRHICommandListImmediate& RHICmdList, FRHITexture* Texture, uint32 MipLevel, uint16 FirstArraySlice, uint16 NumArraySlices)
 {
-	FMetalSurface* Surface = (FMetalSurface*)Texture->GetTextureBaseRHI();
-	FMetalTexture Tex = Surface->Texture;
-	if (!(Tex.GetUsage() & mtlpp::TextureUsage::PixelFormatView))
+	FMetalSurface* Surface = GetMetalSurfaceFromRHITexture(Texture);
+
+	// The FMetalResourceViewBase constructor for textures currently modifies the underlying texture object via FMetalSurface::PrepareTextureView() to add PixelFormatView support if it was not already created with it.
+	// Because of this, the following RHI thread stall is necessary. We will need to clean this up in future before RHI functions can be completely thread safe.
+	if (!(Surface->Texture.GetUsage() & mtlpp::TextureUsage::PixelFormatView))
 	{
 		FScopedRHIThreadStaller StallRHIThread(RHICmdList);
-		return GDynamicRHI->RHICreateUnorderedAccessView(Texture, MipLevel, FirstArraySlice, NumArraySlices);
+		return this->RHICreateUnorderedAccessView(Texture, MipLevel, FirstArraySlice, NumArraySlices);
 	}
 	else
 	{
-		return GDynamicRHI->RHICreateUnorderedAccessView(Texture, MipLevel, FirstArraySlice, NumArraySlices);
+		return this->RHICreateUnorderedAccessView(Texture, MipLevel, FirstArraySlice, NumArraySlices);
 	}
 }
 
 FUnorderedAccessViewRHIRef FMetalDynamicRHI::RHICreateUnorderedAccessView_RenderThread(class FRHICommandListImmediate& RHICmdList, FRHIBuffer* Buffer, uint8 Format)
 {
-	FUnorderedAccessViewRHIRef Result = GDynamicRHI->RHICreateUnorderedAccessView(Buffer, Format);
-	if (IsRunningRHIInSeparateThread() && !RHICmdList.Bypass()) 	{ 		RHICmdList.RHIThreadFence(true); 	}
+	FUnorderedAccessViewRHIRef Result = this->RHICreateUnorderedAccessView(Buffer, Format);
+	if (IsRunningRHIInSeparateThread() && !RHICmdList.Bypass())
+	{
+		RHICmdList.RHIThreadFence(true);
+	}
 	return Result;
 }
 
 FUnorderedAccessViewRHIRef FMetalDynamicRHI::RHICreateUnorderedAccessView(FRHIBuffer* BufferRHI, bool bUseUAVCounter, bool bAppendBuffer)
 {
 	@autoreleasepool {
-	FMetalStructuredBuffer* StructuredBuffer = ResourceCast(BufferRHI);
-	
-	FMetalShaderResourceView* SRV = new FMetalShaderResourceView;
-	SRV->SourceVertexBuffer = nullptr;
-	SRV->SourceIndexBuffer = nullptr;
-	SRV->TextureView = nullptr;
-	SRV->SourceStructuredBuffer = StructuredBuffer;
-
-	// create the UAV buffer to point to the structured buffer's memory
-	FMetalUnorderedAccessView* UAV = new FMetalUnorderedAccessView;
-	UAV->SourceView = SRV;
-	return UAV;
+		return new FMetalUnorderedAccessView(BufferRHI, bUseUAVCounter, bAppendBuffer);
 	}
 }
 
 FUnorderedAccessViewRHIRef FMetalDynamicRHI::RHICreateUnorderedAccessView(FRHITexture* TextureRHI, uint32 MipLevel, uint16 FirstArraySlice, uint16 NumArraySlices)
 {
-	// Slice selection of a texture array still need to be implemented on Metal
-	check(FirstArraySlice == 0 && NumArraySlices == 0);
 	@autoreleasepool {
-	FMetalShaderResourceView* SRV = new FMetalShaderResourceView;
-	SRV->SourceTexture = TextureRHI;
-	
-	FMetalSurface* Surface = GetMetalSurfaceFromRHITexture(TextureRHI);
-	SRV->TextureView = Surface ? new FMetalSurface(*Surface, NSMakeRange(MipLevel, 1)) : nullptr;
-	
-	SRV->SourceVertexBuffer = nullptr;
-	SRV->SourceIndexBuffer = nullptr;
-	SRV->SourceStructuredBuffer = nullptr;
-	
-	SRV->MipLevel = MipLevel;
-	SRV->NumMips = 1;
-	SRV->Format = PF_Unknown;
-		
-	if (Surface)
-	{
-		Surface->SRVs.Add(SRV);
-	}
-		
-	// create the UAV buffer to point to the structured buffer's memory
-	FMetalUnorderedAccessView* UAV = new FMetalUnorderedAccessView;
-	UAV->SourceView = SRV;
-
-	return UAV;
+		return new FMetalUnorderedAccessView(TextureRHI, MipLevel, FirstArraySlice, NumArraySlices);
 	}
 }
 
 FUnorderedAccessViewRHIRef FMetalDynamicRHI::RHICreateUnorderedAccessView(FRHIBuffer* BufferRHI, uint8 Format)
 {
 	@autoreleasepool {
-	FMetalResourceMultiBuffer* Buffer = ResourceCast(BufferRHI);
-	
-	FMetalShaderResourceView* SRV = new FMetalShaderResourceView;
-	SRV->SourceVertexBuffer = EnumHasAnyFlags(BufferRHI->GetUsage(), BUF_VertexBuffer) ? Buffer : nullptr;
-	SRV->TextureView = nullptr;
-	SRV->SourceIndexBuffer = EnumHasAnyFlags(BufferRHI->GetUsage(), BUF_IndexBuffer) ? Buffer : nullptr;
-	SRV->SourceStructuredBuffer = nullptr;
-	SRV->Format = Format;
-	{
-		check(Buffer->GetUsage() & BUF_UnorderedAccess);
-		Buffer->CreateLinearTexture((EPixelFormat)Format, Buffer);
-	}
-		
-	// create the UAV buffer to point to the structured buffer's memory
-	FMetalUnorderedAccessView* UAV = new FMetalUnorderedAccessView;
-	UAV->SourceView = SRV;
-
-	return UAV;
+		return new FMetalUnorderedAccessView(BufferRHI, (EPixelFormat)Format);
 	}
 }
 
 FShaderResourceViewRHIRef FMetalDynamicRHI::CreateShaderResourceView_RenderThread(class FRHICommandListImmediate& RHICmdList, FRHIBuffer* Buffer, uint32 Stride, uint8 Format)
 {
-	FShaderResourceViewRHIRef Result = GDynamicRHI->RHICreateShaderResourceView(Buffer, Stride, Format);
-	if (IsRunningRHIInSeparateThread() && !RHICmdList.Bypass()) 	{ 		RHICmdList.RHIThreadFence(true); 	}
+	FShaderResourceViewRHIRef Result = this->RHICreateShaderResourceView(Buffer, Stride, (EPixelFormat)Format);
+	if (IsRunningRHIInSeparateThread() && !RHICmdList.Bypass())
+	{
+		RHICmdList.RHIThreadFence(true);
+	}
 	return Result;
 }
 
 FShaderResourceViewRHIRef FMetalDynamicRHI::CreateShaderResourceView_RenderThread(class FRHICommandListImmediate& RHICmdList, const FShaderResourceViewInitializer& Initializer)
 {
-	FShaderResourceViewRHIRef Result = GDynamicRHI->RHICreateShaderResourceView(Initializer);
-	if (IsRunningRHIInSeparateThread() && !RHICmdList.Bypass()) { RHICmdList.RHIThreadFence(true); }
+	FShaderResourceViewRHIRef Result = this->RHICreateShaderResourceView(Initializer);
+	if (IsRunningRHIInSeparateThread() && !RHICmdList.Bypass())
+	{
+		RHICmdList.RHIThreadFence(true);
+	}
 	return Result;
 }
 
 FShaderResourceViewRHIRef FMetalDynamicRHI::CreateShaderResourceView_RenderThread(class FRHICommandListImmediate& RHICmdList, FRHIBuffer* Buffer)
 {
-	FShaderResourceViewRHIRef Result = GDynamicRHI->RHICreateShaderResourceView(Buffer);
-	if (IsRunningRHIInSeparateThread() && !RHICmdList.Bypass()) 	{ 		RHICmdList.RHIThreadFence(true); 	}
+	FShaderResourceViewRHIRef Result = this->RHICreateShaderResourceView(Buffer);
+	if (IsRunningRHIInSeparateThread() && !RHICmdList.Bypass())
+	{
+		RHICmdList.RHIThreadFence(true);
+	}
 	return Result;
 }
 
 FShaderResourceViewRHIRef FMetalDynamicRHI::RHICreateShaderResourceView_RenderThread(class FRHICommandListImmediate& RHICmdList, FRHIBuffer* Buffer, uint32 Stride, uint8 Format)
 {
-	FShaderResourceViewRHIRef Result = GDynamicRHI->RHICreateShaderResourceView(Buffer, Stride, Format);
-	if (IsRunningRHIInSeparateThread() && !RHICmdList.Bypass()) 	{ 		RHICmdList.RHIThreadFence(true); 	}
+	FShaderResourceViewRHIRef Result = this->RHICreateShaderResourceView(Buffer, Stride, Format);
+	if (IsRunningRHIInSeparateThread() && !RHICmdList.Bypass())
+	{
+		RHICmdList.RHIThreadFence(true);
+	}
 	return Result;
 }
 
 FShaderResourceViewRHIRef FMetalDynamicRHI::RHICreateShaderResourceView_RenderThread(class FRHICommandListImmediate& RHICmdList, const FShaderResourceViewInitializer& Initializer)
 {
-	FShaderResourceViewRHIRef Result = GDynamicRHI->RHICreateShaderResourceView(Initializer);
-	if (IsRunningRHIInSeparateThread() && !RHICmdList.Bypass()) 	{ 		RHICmdList.RHIThreadFence(true); 	}
+	FShaderResourceViewRHIRef Result = this->RHICreateShaderResourceView(Initializer);
+	if (IsRunningRHIInSeparateThread() && !RHICmdList.Bypass())
+	{
+		RHICmdList.RHIThreadFence(true);
+	}
 	return Result;
 }
 
 FShaderResourceViewRHIRef FMetalDynamicRHI::RHICreateShaderResourceView_RenderThread(class FRHICommandListImmediate& RHICmdList, FRHIBuffer* Buffer)
 {
-	FShaderResourceViewRHIRef Result = GDynamicRHI->RHICreateShaderResourceView(Buffer);
-	if (IsRunningRHIInSeparateThread() && !RHICmdList.Bypass()) 	{ 		RHICmdList.RHIThreadFence(true); 	}
+	FShaderResourceViewRHIRef Result = this->RHICreateShaderResourceView(Buffer);
+	if (IsRunningRHIInSeparateThread() && !RHICmdList.Bypass())
+	{
+		RHICmdList.RHIThreadFence(true);
+	}
 	return Result;
 }
 
 FShaderResourceViewRHIRef FMetalDynamicRHI::RHICreateShaderResourceView_RenderThread(class FRHICommandListImmediate& RHICmdList, FRHITexture* Texture2DRHI, const FRHITextureSRVCreateInfo& CreateInfo)
 {
-	FMetalSurface* Surface = (FMetalSurface*)Texture2DRHI->GetTextureBaseRHI();
-	FMetalTexture Tex = Surface->Texture;
-	if (!(Tex.GetUsage() & mtlpp::TextureUsage::PixelFormatView))
+	FMetalSurface* Surface = GetMetalSurfaceFromRHITexture(Texture2DRHI);
+
+	// The FMetalResourceViewBase constructor for textures currently modifies the underlying texture object via FMetalSurface::PrepareTextureView() to add PixelFormatView support if it was not already created with it.
+	// Because of this, the following RHI thread stall is necessary. We will need to clean this up in future before RHI functions can be completely thread safe.
+	if (!(Surface->Texture.GetUsage() & mtlpp::TextureUsage::PixelFormatView))
 	{
 		FScopedRHIThreadStaller StallRHIThread(RHICmdList);
-		return GDynamicRHI->RHICreateShaderResourceView(Texture2DRHI, CreateInfo);
+		return this->RHICreateShaderResourceView(Texture2DRHI, CreateInfo);
 	}
 	else
 	{
-		return GDynamicRHI->RHICreateShaderResourceView(Texture2DRHI, CreateInfo);
+		return this->RHICreateShaderResourceView(Texture2DRHI, CreateInfo);
 	}
 }
 
 FShaderResourceViewRHIRef FMetalDynamicRHI::RHICreateShaderResourceView(FRHITexture* Texture2DRHI, const FRHITextureSRVCreateInfo& CreateInfo)
 {
 	@autoreleasepool {
-		FMetalShaderResourceView* SRV = new FMetalShaderResourceView;
-		SRV->SourceTexture = (FRHITexture*)Texture2DRHI;
-		
-		FMetalSurface* Surface = GetMetalSurfaceFromRHITexture(Texture2DRHI);
-		
-		// Asking to make a SRV with PF_Unknown means to use the same format.
-		// This matches the behavior of the DX11 RHI.
-		EPixelFormat Format = (EPixelFormat) CreateInfo.Format;
-		if(Surface && Format == PF_Unknown)
-		{
-			Format = Surface->PixelFormat;
-		}
-		
-		const bool bSRGBForceDisable = (CreateInfo.SRGBOverride == SRGBO_ForceDisable);
-		
-		SRV->TextureView = Surface ? new FMetalSurface(*Surface, NSMakeRange(CreateInfo.MipLevel, CreateInfo.NumMipLevels), Format, bSRGBForceDisable) : nullptr;
-		
-		SRV->SourceVertexBuffer = nullptr;
-		SRV->SourceIndexBuffer = nullptr;
-		SRV->SourceStructuredBuffer = nullptr;
-		
-		SRV->MipLevel = CreateInfo.MipLevel;
-		SRV->bSRGBForceDisable = bSRGBForceDisable;
-		SRV->NumMips = CreateInfo.NumMipLevels;
-		SRV->Format = CreateInfo.Format;
-		
-		if (Surface)
-		{
-			Surface->SRVs.Add(SRV);
-		}
-		
-		return SRV;
+		return new FMetalShaderResourceView(Texture2DRHI, CreateInfo);
 	}
 }
 
 FShaderResourceViewRHIRef FMetalDynamicRHI::RHICreateShaderResourceView(FRHIBuffer* BufferRHI)
 {
 	@autoreleasepool {
-	return RHICreateShaderResourceView(FShaderResourceViewInitializer(BufferRHI));
+		return this->RHICreateShaderResourceView(FShaderResourceViewInitializer(BufferRHI));
 	}
 }
 
 FShaderResourceViewRHIRef FMetalDynamicRHI::RHICreateShaderResourceView(FRHIBuffer* BufferRHI, uint32 Stride, uint8 Format)
 {
-	check(GPixelFormats[Format].BlockBytes == Stride);
-
-	@autoreleasepool
-	{
-		return RHICreateShaderResourceView(FShaderResourceViewInitializer(BufferRHI, EPixelFormat(Format)));
+	@autoreleasepool {
+		check(GPixelFormats[Format].BlockBytes == Stride);
+		return this->RHICreateShaderResourceView(FShaderResourceViewInitializer(BufferRHI, EPixelFormat(Format)));
 	}
 }
 
 FShaderResourceViewRHIRef FMetalDynamicRHI::RHICreateShaderResourceView(const FShaderResourceViewInitializer& Initializer)
 {
-	@autoreleasepool
-	{
-		switch(Initializer.GetType())
-		{
-			case FShaderResourceViewInitializer::EType::VertexBufferSRV:
-			{
-				FShaderResourceViewInitializer::FVertexBufferShaderResourceViewInitializer Desc = Initializer.AsVertexBufferSRV();
-				
-				FMetalVertexBuffer* VertexBuffer = ResourceCast(Desc.Buffer);
-				
-				FMetalShaderResourceView* SRV = new FMetalShaderResourceView;
-				SRV->SourceVertexBuffer = VertexBuffer;
-				SRV->TextureView = nullptr;
-				SRV->SourceIndexBuffer = nullptr;
-				SRV->SourceStructuredBuffer = nullptr;
-				
-				SRV->Format = Desc.Format;
-				
-				if(!VertexBuffer)
-				{
-					SRV->Offset = 0;
-					SRV->Stride = 0;
-				}
-				else
-				{
-					check(VertexBuffer->GetUsage() & BUF_ShaderResource);
-					uint32 Stride = GPixelFormats[Desc.Format].BlockBytes;
-					
-					SRV->Stride = Stride;
-					SRV->Offset = Desc.StartOffsetBytes;
-					
-					FMetalLinearTextureDescriptor LinearTextureDesc(Desc.StartOffsetBytes, Desc.NumElements, Stride);
-					SRV->InitLinearTextureDescriptor(LinearTextureDesc);
-					
-					VertexBuffer->CreateLinearTexture((EPixelFormat)Desc.Format, VertexBuffer, &LinearTextureDesc);
-				}
-				
-				return SRV;
-			} // VertexBufferSRV
-				
-			case FShaderResourceViewInitializer::EType::StructuredBufferSRV:
-			{
-				FShaderResourceViewInitializer::FStructuredBufferShaderResourceViewInitializer Desc = Initializer.AsStructuredBufferSRV();
-				
-				FMetalStructuredBuffer* StructuredBuffer = ResourceCast(Desc.Buffer);
-
-				FMetalShaderResourceView* SRV = new FMetalShaderResourceView;
-				SRV->SourceVertexBuffer = nullptr;
-				SRV->SourceIndexBuffer = nullptr;
-				SRV->TextureView = nullptr;
-				SRV->SourceStructuredBuffer = StructuredBuffer;
-				
-				SRV->Offset = Desc.StartOffsetBytes;
-				SRV->Format = 0;
-				SRV->Stride = StructuredBuffer->GetStride();
-				
-				return SRV;
-			} // StructuredBufferSRV
-				
-			case FShaderResourceViewInitializer::EType::IndexBufferSRV:
-			{
-				FShaderResourceViewInitializer::FIndexBufferShaderResourceViewInitializer Desc = Initializer.AsIndexBufferSRV();
-				
-				FMetalIndexBuffer* IndexBuffer = ResourceCast(Desc.Buffer);
-				
-				FMetalShaderResourceView* SRV = new FMetalShaderResourceView;
-				SRV->SourceVertexBuffer = nullptr;
-				SRV->SourceIndexBuffer = IndexBuffer;
-				SRV->TextureView = nullptr;
-				
-				if(!IndexBuffer)
-				{
-					SRV->Format = PF_R16_UINT;
-					SRV->Stride = 0;
-					SRV->Offset = 0;
-				}
-				else
-				{
-					SRV->Format = (IndexBuffer->IndexType == mtlpp::IndexType::UInt16) ? PF_R16_UINT : PF_R32_UINT;
-					
-					const uint32 Stride = Desc.Buffer->GetStride();
-					check(Stride == ((SRV->Format == PF_R16_UINT) ? 2 : 4));
-					
-					SRV->Offset = Desc.StartOffsetBytes;
-					SRV->Stride = Stride;
-					
-					FMetalLinearTextureDescriptor LinearTextureDesc(Desc.StartOffsetBytes, Desc.NumElements, Stride);
-					SRV->InitLinearTextureDescriptor(LinearTextureDesc);
-					
-					IndexBuffer->CreateLinearTexture((EPixelFormat)SRV->Format, IndexBuffer, &LinearTextureDesc);
-				}
-					  
-				return SRV;
-			} // IndexBufferSRV
-					  
-			default:
-			{
-				checkNoEntry();
-				return nullptr;
-			}
-		}
+	@autoreleasepool {
+		return new FMetalShaderResourceView(Initializer);
 	}
 }
 
-void FMetalDynamicRHI::RHIUpdateShaderResourceView(FRHIShaderResourceView* SRVRHI, FRHIBuffer* VertexBufferRHI, uint32 Stride, uint8 Format)
+void FMetalDynamicRHI::RHIUpdateShaderResourceView(FRHIShaderResourceView* SRVRHI, FRHIBuffer* BufferRHI, uint32 Stride, uint8 Format)
 {
 	check(SRVRHI);
 	FMetalShaderResourceView* SRV = ResourceCast(SRVRHI);
-	if (!VertexBufferRHI)
-	{
-		SRV->SourceVertexBuffer = nullptr;
-		SRV->TextureView = nullptr;
-		SRV->SourceIndexBuffer = nullptr;
-		SRV->SourceStructuredBuffer = nullptr;
-		SRV->Offset = 0;
-		SRV->Format = Format;
-		SRV->Stride = Stride;
-	}
-	else if (SRV->SourceVertexBuffer != VertexBufferRHI)
-	{
-		FMetalVertexBuffer* VertexBuffer = ResourceCast(VertexBufferRHI);
-		SRV->SourceVertexBuffer = VertexBuffer;
-		SRV->TextureView = nullptr;
-		SRV->SourceIndexBuffer = nullptr;
-		SRV->SourceStructuredBuffer = nullptr;
-		SRV->Offset = 0;
-		SRV->Format = Format;
-		SRV->Stride = Stride;
-	}
+	check(!SRV->bTexture);
+
+	FMetalResourceMultiBuffer* OldBuffer = SRV->SourceBuffer;
+
+	SRV->SourceBuffer = ResourceCast(BufferRHI);
+	SRV->Stride = Stride;
+	SRV->Format = Format;
+
+	if (SRV->SourceBuffer)
+		SRV->SourceBuffer->AddRef();
+
+	if (OldBuffer)
+		OldBuffer->Release();
 }
 
-void FMetalDynamicRHI::RHIUpdateShaderResourceView(FRHIShaderResourceView* SRVRHI, FRHIBuffer* IndexBufferRHI)
+void FMetalDynamicRHI::RHIUpdateShaderResourceView(FRHIShaderResourceView* SRVRHI, FRHIBuffer* BufferRHI)
 {
 	check(SRVRHI);
 	FMetalShaderResourceView* SRV = ResourceCast(SRVRHI);
-	if (!IndexBufferRHI)
-	{
-		SRV->SourceVertexBuffer = nullptr;
-		SRV->TextureView = nullptr;
-		SRV->SourceIndexBuffer = nullptr;
-		SRV->SourceStructuredBuffer = nullptr;
-		SRV->Offset = 0;
-		SRV->Format = PF_R16_UINT;
-		SRV->Stride = 0;
-	}
-	else if (SRV->SourceIndexBuffer != IndexBufferRHI)
-	{
-		FMetalIndexBuffer* IndexBuffer = ResourceCast(IndexBufferRHI);
-		SRV->SourceVertexBuffer = nullptr;
-		SRV->TextureView = nullptr;
-		SRV->SourceIndexBuffer = IndexBuffer;
-		SRV->SourceStructuredBuffer = nullptr;
-		SRV->Offset = 0;
-		SRV->Format = (IndexBuffer->IndexType == mtlpp::IndexType::UInt16) ? PF_R16_UINT : PF_R32_UINT;
-		SRV->Stride = 0;
-	}
+	check(!SRV->bTexture);
+
+	FMetalResourceMultiBuffer* OldBuffer = SRV->SourceBuffer;
+
+	SRV->SourceBuffer = ResourceCast(BufferRHI);
+	SRV->Stride = 0;
+
+	SRV->Format = SRV->SourceBuffer && SRV->SourceBuffer->IndexType != mtlpp::IndexType::UInt16
+		? PF_R32_UINT
+		: PF_R16_UINT;
+
+	if (SRV->SourceBuffer)
+		SRV->SourceBuffer->AddRef();
+
+	if (OldBuffer)
+		OldBuffer->Release();
 }
 
 #if UE_METAL_RHI_SUPPORT_CLEAR_UAV_WITH_BLIT_ENCODER
@@ -447,24 +401,11 @@ void FMetalRHICommandContext::ClearUAVWithBlitEncoder(FRHIUnorderedAccessView* U
 {
 	SCOPED_AUTORELEASE_POOL;
 
-	FRHICommandListImmediate& RHICmdList = FRHICommandListExecutor::GetImmediateCommandList();
-	FMetalUnorderedAccessView* UnorderedAccessView = ResourceCast(UnorderedAccessViewRHI);
-	FMetalBuffer Buffer;
-	uint32 Size = 0;
+	FMetalResourceMultiBuffer* SourceBuffer = ResourceCast(UnorderedAccessViewRHI)->GetSourceBuffer();
+	FMetalBuffer Buffer = SourceBuffer->GetCurrentBuffer();
+	uint32 Size = SourceBuffer->GetSize();
 
-	switch (Type)
-	{
-		case EMetalRHIClearUAVType::VertexBuffer:
-			check(EnumHasAnyFlags(UnorderedAccessView->SourceView->SourceVertexBuffer->GetUsage(), BUF_ByteAddressBuffer));
-			Buffer = UnorderedAccessView->SourceView->SourceVertexBuffer->GetCurrentBuffer();
-			Size = UnorderedAccessView->SourceView->SourceVertexBuffer->GetSize();
-			break;
-
-		case EMetalRHIClearUAVType::StructuredBuffer:
-			Buffer = UnorderedAccessView->SourceView->SourceStructuredBuffer->GetCurrentBuffer();
-			Size = UnorderedAccessView->SourceView->SourceStructuredBuffer->GetSize();
-			break;
-	};
+	check(Type != EMetalRHIClearUAVType::VertexBuffer || EnumHasAnyFlags(SourceBuffer->GetUsage(), BUF_ByteAddressBuffer));
 
 	uint32 AlignedSize = Align(Size, BufferOffsetAlignment);
 	FMetalPooledBufferArgs Args(GetMetalDeviceContext().GetDevice(), AlignedSize, BUF_Dynamic, mtlpp::StorageMode::Shared);
@@ -482,8 +423,8 @@ void FMetalRHICommandContext::ClearUAVWithBlitEncoder(FRHIUnorderedAccessView* U
 void FMetalRHICommandContext::RHIClearUAVFloat(FRHIUnorderedAccessView* UnorderedAccessViewRHI, const FVector4f& Values)
 {
 #if UE_METAL_RHI_SUPPORT_CLEAR_UAV_WITH_BLIT_ENCODER
-	FMetalUnorderedAccessView* UnorderedAccessView = ResourceCast(UnorderedAccessViewRHI);
-	if (UnorderedAccessView->SourceView->SourceStructuredBuffer)
+	FMetalUnorderedAccessView* UAV = ResourceCast(UnorderedAccessViewRHI);
+	if (!UAV->bTexture && EnumHasAnyFlags(UAV->GetSourceBuffer()->GetUsage(), BUF_StructuredBuffer))
 	{
 		ClearUAVWithBlitEncoder(UnorderedAccessViewRHI, EMetalRHIClearUAVType::StructuredBuffer, *(uint32*)&Values.X);
 	}
@@ -498,8 +439,8 @@ void FMetalRHICommandContext::RHIClearUAVFloat(FRHIUnorderedAccessView* Unordere
 void FMetalRHICommandContext::RHIClearUAVUint(FRHIUnorderedAccessView* UnorderedAccessViewRHI, const FUintVector4& Values)
 {
 #if UE_METAL_RHI_SUPPORT_CLEAR_UAV_WITH_BLIT_ENCODER
-	FMetalUnorderedAccessView* UnorderedAccessView = ResourceCast(UnorderedAccessViewRHI);
-	if (UnorderedAccessView->SourceView->SourceStructuredBuffer)
+	FMetalUnorderedAccessView* UAV = ResourceCast(UnorderedAccessViewRHI);
+	if (!UAV->bTexture && EnumHasAnyFlags(UAV->GetSourceBuffer()->GetUsage(), BUF_StructuredBuffer))
 	{
 		ClearUAVWithBlitEncoder(UnorderedAccessViewRHI, EMetalRHIClearUAVType::StructuredBuffer, *(uint32*)&Values.X);
 	}
@@ -519,7 +460,7 @@ void FMetalRHICommandContext::ClearUAV(TRHICommandList_RecursiveHazardous<FMetal
 		// The Metal validation layer will complain about resources with a
 		// signed format bound against an unsigned data format type as the
 		// shader parameter.
-		switch (GPixelFormats[UnorderedAccessView->SourceView->Format].UnrealFormat)
+		switch (GPixelFormats[UnorderedAccessView->Format].UnrealFormat)
 		{
 			case PF_R32_SINT:
 			case PF_R16_SINT:
@@ -531,37 +472,24 @@ void FMetalRHICommandContext::ClearUAV(TRHICommandList_RecursiveHazardous<FMetal
 				break;
 		}
 
-		if (UnorderedAccessView->SourceView->SourceVertexBuffer)
+		if (UnorderedAccessView->bTexture)
 		{
-#if UE_METAL_RHI_SUPPORT_CLEAR_UAV_WITH_BLIT_ENCODER
-			if (EnumHasAnyFlags(UnorderedAccessView->SourceView->SourceVertexBuffer->GetUsage(), BUF_ByteAddressBuffer))
-			{
-				ClearUAVWithBlitEncoder(UnorderedAccessView, EMetalRHIClearUAVType::VertexBuffer, *(const uint32*)ClearValue);
-			}
-			else
-#endif // UE_METAL_RHI_SUPPORT_CLEAR_UAV_WITH_BLIT_ENCODER
-			{
-				uint32 NumElements = UnorderedAccessView->SourceView->SourceVertexBuffer->GetSize() / GPixelFormats[UnorderedAccessView->SourceView->Format].BlockBytes;
-				ClearUAVShader_T<EClearReplacementResourceType::Buffer, 4, false>(RHICmdList, UnorderedAccessView, NumElements, 1, 1, ClearValue, ValueType);
-			}
-		}
-		else if (UnorderedAccessView->SourceView->SourceTexture)
-		{
-			FIntVector SizeXYZ = UnorderedAccessView->SourceView->SourceTexture->GetSizeXYZ();
+			FMetalSurface* Texture = UnorderedAccessView->GetSourceTexture();
+			FIntVector SizeXYZ = Texture->GetSizeXYZ();
 
-			if (FRHITexture2D* Texture2D = UnorderedAccessView->SourceView->SourceTexture->GetTexture2D())
+			if (FRHITexture2D* Texture2D = Texture->GetTexture2D())
 			{
 				ClearUAVShader_T<EClearReplacementResourceType::Texture2D, 4, false>(RHICmdList, UnorderedAccessView, SizeXYZ.X, SizeXYZ.Y, SizeXYZ.Z, ClearValue, ValueType);
 			}
-			else if (FRHITexture2DArray* Texture2DArray = UnorderedAccessView->SourceView->SourceTexture->GetTexture2DArray())
+			else if (FRHITexture2DArray* Texture2DArray = Texture->GetTexture2DArray())
 			{
 				ClearUAVShader_T<EClearReplacementResourceType::Texture2DArray, 4, false>(RHICmdList, UnorderedAccessView, SizeXYZ.X, SizeXYZ.Y, SizeXYZ.Z, ClearValue, ValueType);
 			}
-			else if (FRHITexture3D* Texture3D = UnorderedAccessView->SourceView->SourceTexture->GetTexture3D())
+			else if (FRHITexture3D* Texture3D = Texture->GetTexture3D())
 			{
 				ClearUAVShader_T<EClearReplacementResourceType::Texture3D, 4, false>(RHICmdList, UnorderedAccessView, SizeXYZ.X, SizeXYZ.Y, SizeXYZ.Z, ClearValue, ValueType);
 			}
-			else if (FRHITextureCube* TextureCube = UnorderedAccessView->SourceView->SourceTexture->GetTextureCube())
+			else if (FRHITextureCube* TextureCube = Texture->GetTextureCube())
 			{
 				ClearUAVShader_T<EClearReplacementResourceType::Texture2DArray, 4, false>(RHICmdList, UnorderedAccessView, SizeXYZ.X, SizeXYZ.Y, SizeXYZ.Z, ClearValue, ValueType);
 			}
@@ -572,8 +500,19 @@ void FMetalRHICommandContext::ClearUAV(TRHICommandList_RecursiveHazardous<FMetal
 		}
 		else
 		{
-			// TODO: ensure(0);
-			UE_LOG(LogRHI, Warning, TEXT("Metal RHI ClearUAV does not yet support clearing of a UAV without a SourceView."));
+			FMetalResourceMultiBuffer* SourceBuffer = UnorderedAccessView->GetSourceBuffer();
+
+#if UE_METAL_RHI_SUPPORT_CLEAR_UAV_WITH_BLIT_ENCODER
+			if (EnumHasAnyFlags(SourceBuffer->GetUsage(), BUF_ByteAddressBuffer))
+			{
+				ClearUAVWithBlitEncoder(UnorderedAccessView, EMetalRHIClearUAVType::VertexBuffer, *(const uint32*)ClearValue);
+			}
+			else
+#endif // UE_METAL_RHI_SUPPORT_CLEAR_UAV_WITH_BLIT_ENCODER
+			{
+				uint32 NumElements = SourceBuffer->GetSize() / GPixelFormats[UnorderedAccessView->Format].BlockBytes;
+				ClearUAVShader_T<EClearReplacementResourceType::Buffer, 4, false>(RHICmdList, UnorderedAccessView, NumElements, 1, 1, ClearValue, ValueType);
+			}
 		}
 	} // @autoreleasepool
 }
diff --git a/Engine/Source/Runtime/Apple/MetalRHI/Private/MetalViewport.cpp b/Engine/Source/Runtime/Apple/MetalRHI/Private/MetalViewport.cpp
--- a/Engine/Source/Runtime/Apple/MetalRHI/Private/MetalViewport.cpp
+++ b/Engine/Source/Runtime/Apple/MetalRHI/Private/MetalViewport.cpp
@@ -238,29 +238,35 @@ void FMetalViewport::Resize(uint32 InSizeX, uint32 InSizeY, bool bInIsFullscreen
 
     {
         FScopeLock Lock(&Mutex);
-        FRHIResourceCreateInfo CreateInfo(TEXT("BackBuffer"));
-        FTexture2DRHIRef NewBackBuffer;
-        FTexture2DRHIRef DoubleBuffer;
-        if (GMetalSupportsIntermediateBackBuffer)
-        {
-            NewBackBuffer = (FMetalTexture2D*)(FRHITexture2D*)GDynamicRHI->RHICreateTexture2D(InSizeX, InSizeY, Format, 1, 1, TexCreate_RenderTargetable, ERHIAccess::Unknown, CreateInfo);
-            
-            if (GMetalSeparatePresentThread)
-            {
-                DoubleBuffer = GDynamicRHI->RHICreateTexture2D(InSizeX, InSizeY, Format, 1, 1, TexCreate_RenderTargetable, ERHIAccess::Unknown, CreateInfo);
-                ((FMetalTexture2D*)DoubleBuffer.GetReference())->Surface.Viewport = this;
-            }
-        }
-        else
+
+		TRefCountPtr<FMetalSurface> NewBackBuffer;
+		TRefCountPtr<FMetalSurface> DoubleBuffer;
+
+		ETextureCreateFlags Flags = GMetalSupportsIntermediateBackBuffer
+			? TexCreate_RenderTargetable
+			: TexCreate_RenderTargetable | TexCreate_Presentable;
+
+		FRHITextureCreateDesc CreateDesc = FRHITextureCreateDesc::Create2D(
+			  TEXT("BackBuffer")
+			, { (int32)InSizeX, (int32)InSizeY }
+			, Format
+			, FClearValueBinding::None
+			, Flags
+		);
+
+		NewBackBuffer = new FMetalSurface(CreateDesc);
+		NewBackBuffer->Viewport = this;
+
+        if (GMetalSupportsIntermediateBackBuffer && GMetalSeparatePresentThread)
         {
-            NewBackBuffer = (FMetalTexture2D*)(FRHITexture2D*)GDynamicRHI->RHICreateTexture2D(InSizeX, InSizeY, Format, 1, 1, TexCreate_RenderTargetable | TexCreate_Presentable, ERHIAccess::Unknown, CreateInfo);
+            DoubleBuffer = new FMetalSurface(CreateDesc);
+            DoubleBuffer->Viewport = this;
         }
-        ((FMetalTexture2D*)NewBackBuffer.GetReference())->Surface.Viewport = this;
-        
-        BackBuffer[Index] = (FMetalTexture2D*)NewBackBuffer.GetReference();
+
+        BackBuffer[Index] = NewBackBuffer;
         if (GMetalSeparatePresentThread)
         {
-            BackBuffer[EMetalViewportAccessRHI] = (FMetalTexture2D*)DoubleBuffer.GetReference();
+            BackBuffer[EMetalViewportAccessRHI] = DoubleBuffer;
         }
         else
         {
@@ -269,7 +275,7 @@ void FMetalViewport::Resize(uint32 InSizeX, uint32 InSizeY, bool bInIsFullscreen
     }
 }
 
-TRefCountPtr<FMetalTexture2D> FMetalViewport::GetBackBuffer(EMetalViewportAccessFlag Accessor) const
+TRefCountPtr<FMetalSurface> FMetalViewport::GetBackBuffer(EMetalViewportAccessFlag Accessor) const
 {
 	FScopeLock Lock(&Mutex);
 	
@@ -405,7 +411,7 @@ void FMetalViewport::ReleaseDrawable()
 
 		if (!GMetalSupportsIntermediateBackBuffer && IsValidRef(BackBuffer[GetViewportIndex(EMetalViewportAccessRHI)]))
 		{
-			BackBuffer[GetViewportIndex(EMetalViewportAccessRHI)]->Surface.Texture = nil;
+			BackBuffer[GetViewportIndex(EMetalViewportAccessRHI)]->Texture = nil;
 		}
 	}
 }
@@ -470,10 +476,10 @@ void FMetalViewport::Present(FMetalCommandQueue& CommandQueue, bool bLockToVsync
 						
 						if (GMetalSupportsIntermediateBackBuffer)
 						{
-							TRefCountPtr<FMetalTexture2D> Texture = LastCompleteFrame;
+							TRefCountPtr<FMetalSurface> Texture = LastCompleteFrame;
 							check(IsValidRef(Texture));
 							
-							FMetalTexture Src = Texture->Surface.Texture;
+							FMetalTexture Src = Texture->Texture;
 							FMetalTexture Dst = LocalDrawable.texture;
 							
 							NSUInteger Width = FMath::Min(Src.GetWidth(), Dst.GetWidth());
@@ -585,8 +591,8 @@ void FMetalViewport::Swap()
 		check(IsValidRef(BackBuffer[0]));
 		check(IsValidRef(BackBuffer[1]));
 		
-		TRefCountPtr<FMetalTexture2D> BB0 = BackBuffer[0];
-		TRefCountPtr<FMetalTexture2D> BB1 = BackBuffer[1];
+		TRefCountPtr<FMetalSurface> BB0 = BackBuffer[0];
+		TRefCountPtr<FMetalSurface> BB1 = BackBuffer[1];
 		
 		BackBuffer[0] = BB1;
 		BackBuffer[1] = BB0;
diff --git a/Engine/Source/Runtime/Apple/MetalRHI/Public/MetalResources.h b/Engine/Source/Runtime/Apple/MetalRHI/Public/MetalResources.h
--- a/Engine/Source/Runtime/Apple/MetalRHI/Public/MetalResources.h
+++ b/Engine/Source/Runtime/Apple/MetalRHI/Public/MetalResources.h
@@ -87,27 +87,25 @@ private:
 class FMetalTexture : public mtlpp::Texture
 {
 public:
-	FMetalTexture(ns::Ownership retain = ns::Ownership::Retain) : mtlpp::Texture(retain) { }
+	FMetalTexture(ns::Ownership retain = ns::Ownership::Retain)
+		: mtlpp::Texture(retain)
+	{}
+
 	FMetalTexture(ns::Protocol<id<MTLTexture>>::type handle, ns::Ownership retain = ns::Ownership::Retain)
-	: mtlpp::Texture(handle, nullptr, retain) {}
+		: mtlpp::Texture(handle, nullptr, retain)
+	{}
 	
 	FMetalTexture(mtlpp::Texture&& rhs)
-	: mtlpp::Texture((mtlpp::Texture&&)rhs)
-	{
-		
-	}
+		: mtlpp::Texture((mtlpp::Texture&&)rhs)
+	{}
 	
 	FMetalTexture(const FMetalTexture& rhs)
-	: mtlpp::Texture(rhs)
-	{
-		
-	}
+		: mtlpp::Texture(rhs)
+	{}
 	
 	FMetalTexture(FMetalTexture&& rhs)
-	: mtlpp::Texture((mtlpp::Texture&&)rhs)
-	{
-		
-	}
+		: mtlpp::Texture((mtlpp::Texture&&)rhs)
+	{}
 	
 	FMetalTexture& operator=(const FMetalTexture& rhs)
 	{
@@ -135,24 +133,43 @@ public:
 	}
 };
 
-/** Texture/RT wrapper. */
-class METALRHI_API FMetalSurface
+struct FMetalTextureDesc
+{
+	FMetalTextureDesc(FRHITextureDesc const& InDesc);
+	
+	mtlpp::TextureDescriptor Desc;
+	mtlpp::PixelFormat MTLFormat;
+	bool bMemoryless = false;
+	bool bIsRenderTarget = false;
+	uint8 FormatKey = 0;
+};
+
+struct FMetalTextureCreateDesc : public FRHITextureCreateDesc, public FMetalTextureDesc
+{
+	FMetalTextureCreateDesc(FRHITextureCreateDesc const& CreateDesc)
+		: FRHITextureCreateDesc(CreateDesc)
+		, FMetalTextureDesc(CreateDesc)
+	{
+		// @todo: texture type unification - Metal can override NumSamples based on command line options.
+		// We should instead require the renderer to do this.
+		NumSamples = Desc.GetSampleCount();
+	}
+};
+
+// Metal RHI texture resource
+class METALRHI_API FMetalSurface : public FRHITexture
 {
 public:
 
 	/** 
 	 * Constructor that will create Texture and Color/DepthBuffers as needed
 	 */
-	FMetalSurface(ERHIResourceType ResourceType, EPixelFormat Format, uint32 SizeX, uint32 SizeY, uint32 SizeZ, uint32 NumSamples, bool bArray, uint32 ArraySize, uint32 NumMips, ETextureCreateFlags Flags, FResourceBulkDataInterface* BulkData);
-
-	FMetalSurface(FMetalSurface& Source, NSRange MipRange);
-	
-	FMetalSurface(FMetalSurface& Source, NSRange MipRange, EPixelFormat Format, bool bSRGBForceDisable);
+	FMetalSurface(FMetalTextureCreateDesc const& CreateDesc);
 	
 	/**
 	 * Destructor
 	 */
-	~FMetalSurface();
+	virtual ~FMetalSurface();
 
 	/** Prepare for texture-view support - need only call this once on the source texture which is to be viewed. */
 	void PrepareTextureView();
@@ -207,13 +224,9 @@ public:
 	ns::AutoReleased<FMetalTexture> GetCurrentTexture();
 
 	FMetalTexture Reallocate(FMetalTexture Texture, mtlpp::TextureUsage UsageModifier);
-	void ReplaceTexture(FMetalContext& Context, FMetalTexture OldTexture, FMetalTexture NewTexture);
 	void MakeAliasable(void);
-	void MakeUnAliasable(void);
-	
-	ERHIResourceType Type;
-	EPixelFormat PixelFormat;
-	uint8 FormatKey;
+
+	uint8 const FormatKey;
 	//texture used for store actions and binding to shader params
 	FMetalTexture Texture;
 	//if surface is MSAA, texture used to bind for RT
@@ -226,8 +239,7 @@ public:
 	// iOS A9+ where depth resolve is available
 	// iOS < A9 where depth resolve is unavailable.
 	FMetalTexture MSAAResolveTexture;
-	uint32 SizeX, SizeY, SizeZ;
-	bool bIsCubemap;
+
 	int16 volatile Written;
 	int16 GPUReadback = 0;
 	enum EMetalGPUReadbackFlags : int16
@@ -238,8 +250,6 @@ public:
 		ReadbackFenceComplete 			= 1 << ReadbackFenceCompleteShift,
 		ReadbackRequestedAndComplete 	= ReadbackRequested | ReadbackFenceComplete
 	};
-	
-	ETextureCreateFlags Flags;
 
 	uint32 BufferLocks;
 
@@ -248,125 +258,23 @@ public:
 	
 	// For back-buffers, the owning viewport.
 	class FMetalViewport* Viewport;
-	
-	TSet<class FMetalShaderResourceView*> SRVs;
-
-private:
-	void Init(FMetalSurface& Source, NSRange MipRange);
-	
-	void Init(FMetalSurface& Source, NSRange MipRange, EPixelFormat Format, bool bSRGBForceDisable);
-	
-private:
-	// The movie playback IOSurface/CVTexture wrapper to avoid page-off
-	CFTypeRef ImageSurfaceRef;
-	
-	// Texture view surfaces don't own their resources, only reference
-	bool bTextureView;
-	
-	// Count of outstanding async. texture uploads
-	static volatile int64 ActiveUploads;
-};
 
-class FMetalTexture2D : public FRHITexture2D
-{
-public:
-	/** The surface info */
-	FMetalSurface Surface;
-
-	// Constructor, just calls base and Surface constructor
-	FMetalTexture2D(EPixelFormat Format, uint32 SizeX, uint32 SizeY, uint32 NumMips, uint32 NumSamples, ETextureCreateFlags Flags, FResourceBulkDataInterface* BulkData, const FClearValueBinding& InClearValue)
-		: FRHITexture2D(SizeX, SizeY, NumMips, NumSamples, Format, Flags, InClearValue)
-		, Surface(RRT_Texture2D, Format, SizeX, SizeY, 1, NumSamples, /*bArray=*/ false, 1, NumMips, Flags, BulkData)
-	{
-	}
-	
-	virtual ~FMetalTexture2D()
-	{
-	}
-	
 	virtual void* GetTextureBaseRHI() override final
 	{
-		return &Surface;
-	}
-	
-	virtual void* GetNativeResource() const override final
-	{
-		return Surface.Texture;
+		return this;
 	}
-};
-
-class FMetalTexture2DArray : public FRHITexture2DArray
-{
-public:
-	/** The surface info */
-	FMetalSurface Surface;
 
-	// Constructor, just calls base and Surface constructor
-	FMetalTexture2DArray(EPixelFormat Format, uint32 SizeX, uint32 SizeY, uint32 ArraySize, uint32 NumMips, ETextureCreateFlags Flags, FResourceBulkDataInterface* BulkData, const FClearValueBinding& InClearValue)
-		: FRHITexture2DArray(SizeX, SizeY, ArraySize, NumMips, 1, Format, Flags, InClearValue)
-		, Surface(RRT_Texture2DArray, Format, SizeX, SizeY, 1, /*NumSamples=*/1, /*bArray=*/ true, ArraySize, NumMips, Flags, BulkData)
-	{
-	}
-	
-	virtual ~FMetalTexture2DArray()
-	{
-	}
-	
-	virtual void* GetTextureBaseRHI() override final
-	{
-		return &Surface;
-	}
-};
-
-class FMetalTexture3D : public FRHITexture3D
-{
-public:
-	/** The surface info */
-	FMetalSurface Surface;
-
-	// Constructor, just calls base and Surface constructor
-	FMetalTexture3D(EPixelFormat Format, uint32 SizeX, uint32 SizeY, uint32 SizeZ, uint32 NumMips, ETextureCreateFlags Flags, FResourceBulkDataInterface* BulkData, const FClearValueBinding& InClearValue)
-		: FRHITexture3D(SizeX, SizeY, SizeZ, NumMips, Format, Flags, InClearValue)
-		, Surface(RRT_Texture3D, Format, SizeX, SizeY, SizeZ, /*NumSamples=*/1, /*bArray=*/ false, 1, NumMips, Flags, BulkData)
-	{
-	}
-	
-	virtual ~FMetalTexture3D()
-	{
-	}
-	
-	virtual void* GetTextureBaseRHI() override final
-	{
-		return &Surface;
-	}
-};
-
-class FMetalTextureCube : public FRHITextureCube
-{
-public:
-	/** The surface info */
-	FMetalSurface Surface;
-
-	// Constructor, just calls base and Surface constructor
-	FMetalTextureCube(EPixelFormat Format, uint32 Size, bool bArray, uint32 ArraySize, uint32 NumMips, ETextureCreateFlags Flags, FResourceBulkDataInterface* BulkData, const FClearValueBinding& InClearValue)
-		: FRHITextureCube(Size, NumMips, Format, Flags, InClearValue)
-		, Surface(RRT_TextureCube, Format, Size, Size, 6, /*NumSamples=*/1, bArray, ArraySize, NumMips, Flags, BulkData)
-	{
-	}
-	
-	virtual ~FMetalTextureCube()
+	virtual void* GetNativeResource() const override final
 	{
+		return Texture;
 	}
 	
-	virtual void* GetTextureBaseRHI() override final
-	{
-		return &Surface;
-	}
+private:
+	// The movie playback IOSurface/CVTexture wrapper to avoid page-off
+	CFTypeRef ImageSurfaceRef;
 
-	virtual void* GetNativeResource() const override final
-	{
-		return Surface.Texture;
-	}
+	// Count of outstanding async. texture uploads
+	static volatile int64 ActiveUploads;
 };
 
 @interface FMetalBufferData : FApplePlatformObject<NSObject>
@@ -390,34 +298,13 @@ ENUM_CLASS_FLAGS(EMetalBufferUsage);
 class FMetalLinearTextureDescriptor
 {
 public:
-	FMetalLinearTextureDescriptor()
-		: StartOffsetBytes(0)
-		, NumElements(UINT_MAX)
-		, BytesPerElement(0)
-	{
-		// void
-	}
+	FMetalLinearTextureDescriptor() = default;
 
 	FMetalLinearTextureDescriptor(uint32 InStartOffsetBytes, uint32 InNumElements, uint32 InBytesPerElement)
 		: StartOffsetBytes(InStartOffsetBytes)
-		, NumElements(InNumElements)
-		, BytesPerElement(InBytesPerElement)
-	{
-		// void
-	}
-	
-	FMetalLinearTextureDescriptor(const FMetalLinearTextureDescriptor& Other)
-		: StartOffsetBytes(Other.StartOffsetBytes)
-		, NumElements(Other.NumElements)
-		, BytesPerElement(Other.BytesPerElement)
-	{
-		// void
-	}
-	
-	~FMetalLinearTextureDescriptor()
-	{
-		// void
-	}
+		, NumElements     (InNumElements)
+		, BytesPerElement (InBytesPerElement)
+	{}
 
 	friend uint32 GetTypeHash(FMetalLinearTextureDescriptor const& Key)
 	{
@@ -434,9 +321,9 @@ public:
 		       && BytesPerElement  == Other.BytesPerElement;
 	}
 
-	uint32 StartOffsetBytes;
-	uint32 NumElements;
-	uint32 BytesPerElement;
+	uint32 StartOffsetBytes = 0;
+	uint32 NumElements      = UINT_MAX;
+	uint32 BytesPerElement  = 0;
 };
 
 class FMetalRHIBuffer
@@ -599,50 +486,118 @@ typedef FMetalResourceMultiBuffer FMetalIndexBuffer;
 typedef FMetalResourceMultiBuffer FMetalVertexBuffer;
 typedef FMetalResourceMultiBuffer FMetalStructuredBuffer;
 
-class FMetalShaderResourceView : public FRHIShaderResourceView
-{
+class FMetalResourceViewBase
+{
+protected:
+	// Constructor for buffers
+	FMetalResourceViewBase(
+		  FRHIBuffer* InBuffer
+		, uint32 InStartOffsetBytes
+		, uint32 InNumElements
+		, EPixelFormat InFormat
+	);
+
+	// Constructor for textures
+	FMetalResourceViewBase(
+		  FRHITexture* InTexture
+		, EPixelFormat InFormat
+		, uint8 InMipLevel
+		, uint8 InNumMipLevels
+		, ERHITextureSRVOverrideSRGBType InSRGBOverride
+		, uint32 InFirstArraySlice
+		, uint32 InNumArraySlices
+	);
+
 public:
+	~FMetalResourceViewBase();
 
-	// The vertex buffer this SRV comes from (can be null)
-	TRefCountPtr<FMetalVertexBuffer> SourceVertexBuffer;
-	
-	// The index buffer this SRV comes from (can be null)
-	TRefCountPtr<FMetalIndexBuffer> SourceIndexBuffer;
+	inline FMetalResourceMultiBuffer* GetSourceBuffer () const { check(!bTexture); return SourceBuffer; }
 
-	// The texture that this SRV come from
-	TRefCountPtr<FRHITexture> SourceTexture;
-	
-	// The source structured buffer (can be null)
-	TRefCountPtr<FMetalStructuredBuffer> SourceStructuredBuffer;
-	
-	FMetalSurface* TextureView;
-	uint32 Offset;
-	uint8 MipLevel          : 4;
-	uint8 bSRGBForceDisable : 1;
-	uint8 Reserved          : 3;
-	uint8 NumMips;
-	uint8 Format;
-	uint8 Stride;
+	inline FMetalSurface*             GetSourceTexture() const { check(bTexture); return SourceTexture; }
+	inline FMetalTexture const&       GetTextureView  () const { check(bTexture); return TextureView;   }
 
-	void InitLinearTextureDescriptor(const FMetalLinearTextureDescriptor& InLinearTextureDescriptor);
+private:
+	// Needed for RHIUpdateShaderResourceView
+	friend class FMetalDynamicRHI;
 
-	FMetalShaderResourceView();
-	~FMetalShaderResourceView();
-	
-	ns::AutoReleased<FMetalTexture> GetLinearTexture(bool const bUAV);
+	union
+	{
+		FMetalResourceMultiBuffer* SourceBuffer;
+		FMetalSurface* SourceTexture;
+	};
 
-private:
-	FMetalLinearTextureDescriptor* LinearTextureDesc;
-};
+	TUniquePtr<FMetalLinearTextureDescriptor> LinearTextureDesc = nullptr;
+	FMetalTexture TextureView = nullptr;
 
+public:
+	uint8 const bTexture : 1;
+	uint8       bSRGBForceDisable : 1;
+	uint8       MipLevel : 4;
+	uint8       Reserved : 2;
+	uint8       NumMips;
+	uint8       Format;
+	uint8       Stride;
+	uint32      Offset;
 
+	ns::AutoReleased<FMetalTexture> GetLinearTexture();
+};
 
-class FMetalUnorderedAccessView : public FRHIUnorderedAccessView
+class FMetalShaderResourceView final : public FRHIShaderResourceView, public FMetalResourceViewBase
 {
 public:
-	
-	// the potential resources to refer to with the UAV object
-	TRefCountPtr<FMetalShaderResourceView> SourceView;
+	explicit FMetalShaderResourceView(const FShaderResourceViewInitializer& Initializer)
+		: FMetalResourceViewBase(
+			  Initializer.AsBufferSRV().Buffer
+			, Initializer.AsBufferSRV().StartOffsetBytes
+			, Initializer.AsBufferSRV().NumElements
+			, Initializer.AsBufferSRV().Format
+		)
+	{}
+
+	explicit FMetalShaderResourceView(FRHITexture* Texture, const FRHITextureSRVCreateInfo& CreateInfo)
+		: FMetalResourceViewBase(
+			  Texture
+			, CreateInfo.Format
+			, CreateInfo.MipLevel
+			, CreateInfo.NumMipLevels
+			, CreateInfo.SRGBOverride
+			, CreateInfo.FirstArraySlice
+			, CreateInfo.NumArraySlices
+		)
+	{}
+
+	virtual ~FMetalShaderResourceView()
+	{}
+};
+
+class FMetalUnorderedAccessView final : public FRHIUnorderedAccessView, public FMetalResourceViewBase
+{
+public:
+	explicit FMetalUnorderedAccessView(FRHIBuffer* Buffer, EPixelFormat Format)
+		: FMetalResourceViewBase(Buffer, 0, UINT_MAX, Format)
+	{}
+
+	explicit FMetalUnorderedAccessView(FRHIBuffer* Buffer, bool bUseUAVCounter, bool bAppendBuffer)
+		: FMetalResourceViewBase(Buffer, 0, UINT_MAX, PF_Unknown)
+	{
+		checkf(!bUseUAVCounter, TEXT("UAV counters not implemented."));
+		checkf(!bAppendBuffer, TEXT("UAV append buffers not implemented."));
+	}
+
+	explicit FMetalUnorderedAccessView(FRHITexture* Texture, uint32 MipLevel, uint16 FirstArraySlice, uint16 NumArraySlices)
+		: FMetalResourceViewBase(
+			  Texture
+			, PF_Unknown
+			, MipLevel
+			, 1 // NumMipLevels
+			, ERHITextureSRVOverrideSRGBType::SRGBO_ForceDisable
+			, FirstArraySlice
+			, NumArraySlices
+		)
+	{}
+
+	virtual ~FMetalUnorderedAccessView()
+	{}
 };
 
 class FMetalGPUFence final : public FRHIGPUFence
@@ -714,26 +669,6 @@ struct TMetalResourceTraits<FRHIComputeShader>
 	typedef FMetalComputeShader TConcreteType;
 };
 template<>
-struct TMetalResourceTraits<FRHITexture3D>
-{
-	typedef FMetalTexture3D TConcreteType;
-};
-template<>
-struct TMetalResourceTraits<FRHITexture2D>
-{
-	typedef FMetalTexture2D TConcreteType;
-};
-template<>
-struct TMetalResourceTraits<FRHITexture2DArray>
-{
-	typedef FMetalTexture2DArray TConcreteType;
-};
-template<>
-struct TMetalResourceTraits<FRHITextureCube>
-{
-	typedef FMetalTextureCube TConcreteType;
-};
-template<>
 struct TMetalResourceTraits<FRHIRenderQuery>
 {
 	typedef FMetalRHIRenderQuery TConcreteType;
diff --git a/Engine/Source/Runtime/Apple/MetalRHI/Public/MetalViewport.h b/Engine/Source/Runtime/Apple/MetalRHI/Public/MetalViewport.h
--- a/Engine/Source/Runtime/Apple/MetalRHI/Public/MetalViewport.h
+++ b/Engine/Source/Runtime/Apple/MetalRHI/Public/MetalViewport.h
@@ -36,7 +36,7 @@ public:
 
 	void Resize(uint32 InSizeX, uint32 InSizeY, bool bInIsFullscreen,EPixelFormat Format);
 	
-	TRefCountPtr<FMetalTexture2D> GetBackBuffer(EMetalViewportAccessFlag Accessor) const;
+	TRefCountPtr<FMetalSurface> GetBackBuffer(EMetalViewportAccessFlag Accessor) const;
 	id<CAMetalDrawable> GetDrawable(EMetalViewportAccessFlag Accessor);
 	FMetalTexture GetDrawableTexture(EMetalViewportAccessFlag Accessor);
 	ns::AutoReleased<FMetalTexture> GetCurrentTexture(EMetalViewportAccessFlag Accessor);
@@ -65,7 +65,7 @@ private:
 
 private:
 	id<CAMetalDrawable> Drawable;
-	TRefCountPtr<FMetalTexture2D> BackBuffer[2];
+	TRefCountPtr<FMetalSurface> BackBuffer[2];
 	mutable FCriticalSection Mutex;
 	
 	ns::AutoReleased<FMetalTexture> DrawableTextures[2];
@@ -73,7 +73,7 @@ private:
 	uint32 DisplayID;
 	FMetalViewportPresentHandler Block;
 	volatile int32 FrameAvailable;
-	TRefCountPtr<FMetalTexture2D> LastCompleteFrame;
+	TRefCountPtr<FMetalSurface> LastCompleteFrame;
 	bool bIsFullScreen;
 
 #if PLATFORM_MAC
diff --git a/Engine/Source/Runtime/Core/Public/PixelFormat.h b/Engine/Source/Runtime/Core/Public/PixelFormat.h
--- a/Engine/Source/Runtime/Core/Public/PixelFormat.h
+++ b/Engine/Source/Runtime/Core/Public/PixelFormat.h
@@ -10,7 +10,7 @@
 #undef PF_MAX
 #endif
 
-enum EPixelFormat
+enum EPixelFormat : uint8
 {
 	PF_Unknown              =0,
 	PF_A32B32G32R32F        =1,
diff --git a/Engine/Source/Runtime/D3D12RHI/Private/D3D12Adapter.h b/Engine/Source/Runtime/D3D12RHI/Private/D3D12Adapter.h
--- a/Engine/Source/Runtime/D3D12RHI/Private/D3D12Adapter.h
+++ b/Engine/Source/Runtime/D3D12RHI/Private/D3D12Adapter.h
@@ -350,10 +350,14 @@ public:
 		D3D12_RESOURCE_STATES InCreateState,
 		bool bHasInitialData,
 		const FRHIGPUMask& InGPUMask,
-		ED3D12ResourceTransientMode TransientMode,
 		ID3D12ResourceAllocator* ResourceAllocator,
 		const TCHAR* InDebugName);
 
+	void CreateUAVAliasResource(
+		D3D12_CLEAR_VALUE* ClearValuePtr,
+		const TCHAR* DebugName,
+		FD3D12ResourceLocation& Location);
+
 	template <typename ObjectType, typename CreationCoreFunction>
 	inline ObjectType* CreateLinkedObject(FRHIGPUMask GPUMask, const CreationCoreFunction& pfnCreationCore)
 	{
@@ -438,7 +442,6 @@ protected:
 		uint32 Alignment,
 		FD3D12Buffer* Buffer,
 		FD3D12ResourceLocation& ResourceLocation,
-		ED3D12ResourceTransientMode TransientMode,
 		ID3D12ResourceAllocator* ResourceAllocator,
 		const TCHAR* InDebugName);
 
diff --git a/Engine/Source/Runtime/D3D12RHI/Private/D3D12Allocation.cpp b/Engine/Source/Runtime/D3D12RHI/Private/D3D12Allocation.cpp
--- a/Engine/Source/Runtime/D3D12RHI/Private/D3D12Allocation.cpp
+++ b/Engine/Source/Runtime/D3D12RHI/Private/D3D12Allocation.cpp
@@ -1671,7 +1671,7 @@ HRESULT FD3D12TextureAllocatorPool::AllocateTexture(
 	const TCHAR* Name)
 {
 	// The top mip level must be less than 64 KB to use 4 KB alignment
-	bool b4KAligment = TextureCanBe4KAligned(Desc, (EPixelFormat)UEFormat);
+	bool b4KAligment = FD3D12Texture::CanBe4KAligned(Desc, (EPixelFormat)UEFormat);
 	Desc.Alignment = b4KAligment ?	D3D12_SMALL_RESOURCE_PLACEMENT_ALIGNMENT : (Desc.SampleDesc.Count > 1 ? D3D12_DEFAULT_MSAA_RESOURCE_PLACEMENT_ALIGNMENT : D3D12_DEFAULT_RESOURCE_PLACEMENT_ALIGNMENT);
 	const D3D12_RESOURCE_ALLOCATION_INFO Info = GetParentDevice()->GetDevice()->GetResourceAllocationInfo(0, 1, &Desc);
 
@@ -1896,7 +1896,7 @@ HRESULT FD3D12TextureAllocatorPool::AllocateTexture(FD3D12ResourceDesc Desc, con
 		Desc.SampleDesc.Count == 1)// Multi-Sample texures have much larger alignment requirements (4MB vs 64KB)
 	{
 		// The top mip level must be less than 64k
-		if (TextureCanBe4KAligned(Desc, (EPixelFormat)UEFormat))
+		if (FD3D12Texture::CanBe4KAligned(Desc, (EPixelFormat)UEFormat))
 		{
 			Desc.Alignment = D3D12_SMALL_RESOURCE_PLACEMENT_ALIGNMENT; // request 4k alignment
 			return ReadOnlyTexturePool.AllocateTexture(Desc, ClearValue, TextureLocation, InitialState, Name);
@@ -1907,12 +1907,11 @@ HRESULT FD3D12TextureAllocatorPool::AllocateTexture(FD3D12ResourceDesc Desc, con
 	FD3D12Resource* Resource = nullptr;
 
 	const D3D12_HEAP_PROPERTIES HeapProps = CD3DX12_HEAP_PROPERTIES(D3D12_HEAP_TYPE_DEFAULT, GetGPUMask().GetNative(), GetVisibilityMask().GetNative());
+	const D3D12_RESOURCE_ALLOCATION_INFO Info = GetParentDevice()->GetDevice()->GetResourceAllocationInfo(0, 1, &Desc);
 
 	// UAV Aliasing needs a Heap to create the aliased resource in.
 	if (Desc.NeedsUAVAliasWorkarounds())
 	{
-		const D3D12_RESOURCE_ALLOCATION_INFO Info = GetParentDevice()->GetDevice()->GetResourceAllocationInfo(0, 1, &Desc);
-
 		D3D12_HEAP_DESC HeapDesc{};
 		HeapDesc.SizeInBytes = Info.SizeInBytes;
 		HeapDesc.Properties = HeapProps;
@@ -1933,8 +1932,7 @@ HRESULT FD3D12TextureAllocatorPool::AllocateTexture(FD3D12ResourceDesc Desc, con
 
 		if (SUCCEEDED(hr))
 		{
-			TextureLocation.SetType(FD3D12ResourceLocation::ResourceLocationType::eStandAlone);
-			TextureLocation.SetResource(Resource);
+			TextureLocation.AsStandAlone(Resource, Info.SizeInBytes);
 		}
 
 		return hr;
@@ -1944,8 +1942,7 @@ HRESULT FD3D12TextureAllocatorPool::AllocateTexture(FD3D12ResourceDesc Desc, con
 
 	if (SUCCEEDED(hr))
 	{
-		TextureLocation.SetType(FD3D12ResourceLocation::ResourceLocationType::eStandAlone);
-		TextureLocation.SetResource(Resource);
+		TextureLocation.AsStandAlone(Resource, Info.SizeInBytes);
 	}
 
 	return hr;
diff --git a/Engine/Source/Runtime/D3D12RHI/Private/D3D12Buffer.cpp b/Engine/Source/Runtime/D3D12RHI/Private/D3D12Buffer.cpp
--- a/Engine/Source/Runtime/D3D12RHI/Private/D3D12Buffer.cpp
+++ b/Engine/Source/Runtime/D3D12RHI/Private/D3D12Buffer.cpp
@@ -334,7 +334,6 @@ void FD3D12Adapter::AllocateBuffer(FD3D12Device* Device,
 	uint32 Alignment,
 	FD3D12Buffer* Buffer,
 	FD3D12ResourceLocation& ResourceLocation,
-	ED3D12ResourceTransientMode TransientMode,
 	ID3D12ResourceAllocator* ResourceAllocator,
 	const TCHAR* InDebugName)
 {
@@ -376,7 +375,6 @@ FD3D12Buffer* FD3D12Adapter::CreateRHIBuffer(
 	D3D12_RESOURCE_STATES InCreateState,
 	bool bHasInitialData,
 	const FRHIGPUMask& InGPUMask,
-	ED3D12ResourceTransientMode TransientMode,
 	ID3D12ResourceAllocator* ResourceAllocator,
 	const TCHAR* InDebugName)
 {
@@ -407,7 +405,7 @@ FD3D12Buffer* FD3D12Adapter::CreateRHIBuffer(
 
 			if ((Device->GetGPUIndex() == FirstGPUIndex) || EnumHasAnyFlags(InUsage, BUF_MultiGPUAllocate))
 			{
-				AllocateBuffer(Device, InDesc, Size, InUsage, InResourceStateMode, InCreateState, Alignment, NewBuffer, NewBuffer->ResourceLocation, TransientMode, ResourceAllocator, InDebugName);
+				AllocateBuffer(Device, InDesc, Size, InUsage, InResourceStateMode, InCreateState, Alignment, NewBuffer, NewBuffer->ResourceLocation, ResourceAllocator, InDebugName);
 				NewBuffer0 = NewBuffer;
 			}
 			else
@@ -434,7 +432,7 @@ FD3D12Buffer* FD3D12Adapter::CreateRHIBuffer(
 			}
 #endif // NAME_OBJECTS
 
-			AllocateBuffer(Device, InDesc, Size, InUsage, InResourceStateMode, InCreateState, Alignment, NewBuffer, NewBuffer->ResourceLocation, TransientMode, ResourceAllocator, InDebugName);
+			AllocateBuffer(Device, InDesc, Size, InUsage, InResourceStateMode, InCreateState, Alignment, NewBuffer, NewBuffer->ResourceLocation, ResourceAllocator, InDebugName);
 			
 			// Unlock immediately if there is no initial data
 			if (!bHasInitialData)
@@ -571,7 +569,7 @@ FBufferRHIRef FD3D12DynamicRHI::CreateBuffer(FRHICommandListImmediate* RHICmdLis
 	return CreateD3D12Buffer(RHICmdList, Size, Usage, Stride, InResourceState, CreateInfo);
 }
 
-FD3D12Buffer* FD3D12DynamicRHI::CreateD3D12Buffer(class FRHICommandListImmediate* RHICmdList, uint32 Size, EBufferUsageFlags Usage, uint32 Stride, ERHIAccess InResourceState, FRHIResourceCreateInfo& CreateInfo, ED3D12ResourceTransientMode TransientMode, ID3D12ResourceAllocator* ResourceAllocator)
+FD3D12Buffer* FD3D12DynamicRHI::CreateD3D12Buffer(class FRHICommandListImmediate* RHICmdList, uint32 Size, EBufferUsageFlags Usage, uint32 Stride, ERHIAccess InResourceState, FRHIResourceCreateInfo& CreateInfo, ID3D12ResourceAllocator* ResourceAllocator)
 {
 	D3D12_RESOURCE_DESC Desc;
 	uint32 Alignment;
@@ -597,7 +595,7 @@ FD3D12Buffer* FD3D12DynamicRHI::CreateD3D12Buffer(class FRHICommandListImmediate
 	bool bHasInitialData = CreateInfo.ResourceArray != nullptr;
 
 	FD3D12SyncPoint CopyQueueSyncPoint;
-	FD3D12Buffer* Buffer = GetAdapter().CreateRHIBuffer(Desc, Alignment, Stride, Size, Usage, StateMode, CreateState, bHasInitialData, CreateInfo.GPUMask, TransientMode, ResourceAllocator, CreateInfo.DebugName);
+	FD3D12Buffer* Buffer = GetAdapter().CreateRHIBuffer(Desc, Alignment, Stride, Size, Usage, StateMode, CreateState, bHasInitialData, CreateInfo.GPUMask, ResourceAllocator, CreateInfo.DebugName);
 	check(Buffer->ResourceLocation.IsValid());
 
 	// Copy the resource data if available 
@@ -609,10 +607,10 @@ FD3D12Buffer* FD3D12DynamicRHI::CreateD3D12Buffer(class FRHICommandListImmediate
 	return Buffer;
 }
 
-FRHIBuffer* FD3D12DynamicRHI::CreateBuffer(const FRHIBufferCreateInfo& CreateInfo, const TCHAR* DebugName, ERHIAccess InitialState, ED3D12ResourceTransientMode TransientMode, ID3D12ResourceAllocator* ResourceAllocator)
+FRHIBuffer* FD3D12DynamicRHI::CreateBuffer(const FRHIBufferCreateInfo& CreateInfo, const TCHAR* DebugName, ERHIAccess InitialState, ID3D12ResourceAllocator* ResourceAllocator)
 {
 	FRHIResourceCreateInfo ResourceCreateInfo(DebugName);
-	return CreateD3D12Buffer(nullptr, CreateInfo.Size, CreateInfo.Usage, CreateInfo.Stride, InitialState, ResourceCreateInfo, TransientMode, ResourceAllocator);
+	return CreateD3D12Buffer(nullptr, CreateInfo.Size, CreateInfo.Usage, CreateInfo.Stride, InitialState, ResourceCreateInfo, ResourceAllocator);
 }
 
 void* FD3D12DynamicRHI::LockBuffer(FRHICommandListImmediate* RHICmdList, FD3D12Buffer* Buffer, uint32 BufferSize, EBufferUsageFlags BufferUsage, uint32 Offset, uint32 Size, EResourceLockMode LockMode)
diff --git a/Engine/Source/Runtime/D3D12RHI/Private/D3D12CommandContext.cpp b/Engine/Source/Runtime/D3D12RHI/Private/D3D12CommandContext.cpp
--- a/Engine/Source/Runtime/D3D12RHI/Private/D3D12CommandContext.cpp
+++ b/Engine/Source/Runtime/D3D12RHI/Private/D3D12CommandContext.cpp
@@ -1124,8 +1124,8 @@ void FD3D12CommandContextRedirector::RHITransferResources(const TArrayView<const
 		{
 			check(Param.Buffer == nullptr);
 
-			SrcResource = FD3D12CommandContext::RetrieveTextureBase(Param.Texture, Param.SrcGPUIndex)->GetResource();
-			DestResource = FD3D12CommandContext::RetrieveTextureBase(Param.Texture, Param.DestGPUIndex)->GetResource();
+			SrcResource = FD3D12CommandContext::RetrieveTexture(Param.Texture, Param.SrcGPUIndex)->GetResource();
+			DestResource = FD3D12CommandContext::RetrieveTexture(Param.Texture, Param.DestGPUIndex)->GetResource();
 		}
 		else
 		{
@@ -1177,8 +1177,8 @@ void FD3D12CommandContextRedirector::RHITransferResources(const TArrayView<const
 
 		if (Param.Texture)
 		{
-			FD3D12TextureBase* SrcTexture = FD3D12CommandContext::RetrieveTextureBase(Param.Texture, Param.SrcGPUIndex);
-			FD3D12TextureBase* DestTexture = FD3D12CommandContext::RetrieveTextureBase(Param.Texture, Param.DestGPUIndex);
+			FD3D12Texture* SrcTexture = FD3D12CommandContext::RetrieveTexture(Param.Texture, Param.SrcGPUIndex);
+			FD3D12Texture* DestTexture = FD3D12CommandContext::RetrieveTexture(Param.Texture, Param.DestGPUIndex);
 
 			// If the texture size is zero (Max.Z == 0, set in the constructor), copy the whole resource
 			if (Param.Max.Z == 0)
diff --git a/Engine/Source/Runtime/D3D12RHI/Private/D3D12CommandContext.h b/Engine/Source/Runtime/D3D12RHI/Private/D3D12CommandContext.h
--- a/Engine/Source/Runtime/D3D12RHI/Private/D3D12CommandContext.h
+++ b/Engine/Source/Runtime/D3D12RHI/Private/D3D12CommandContext.h
@@ -148,7 +148,7 @@ public:
 	// Tracks the currently set state blocks.
 	FD3D12RenderTargetView* CurrentRenderTargets[D3D12_SIMULTANEOUS_RENDER_TARGET_COUNT];
 	FD3D12DepthStencilView* CurrentDepthStencilTarget;
-	FD3D12TextureBase* CurrentDepthTexture;
+	FD3D12Texture* CurrentDepthTexture;
 	uint32 NumSimultaneousRenderTargets;
 
 	/** Track the currently bound uniform buffers. */
@@ -247,8 +247,8 @@ public:
 	template<typename TPixelShader>
 	void ResolveTextureUsingShader(
 		FRHICommandList_RecursiveHazardous& RHICmdList,
-		FD3D12Texture2D* SourceTexture,
-		FD3D12Texture2D* DestTexture,
+		FD3D12Texture* SourceTexture,
+		FD3D12Texture* DestTexture,
 		FD3D12RenderTargetView* DestSurfaceRTV,
 		FD3D12DepthStencilView* DestSurfaceDSV,
 		const D3D12_RESOURCE_DESC& ResolveTargetDesc,
@@ -442,15 +442,15 @@ public:
 		return RetrieveObject<ObjectType, RHIType>(RHIObject, GetGPUIndex());
 	}
 
-	static inline FD3D12TextureBase* RetrieveTextureBase(FRHITexture* Texture, uint32 GPUIndex)
+	static inline FD3D12Texture* RetrieveTexture(FRHITexture* Texture, uint32 GPUIndex)
 	{
-		FD3D12TextureBase* RHITexture = GetD3D12TextureFromRHITexture(Texture);
+		FD3D12Texture* RHITexture = GetD3D12TextureFromRHITexture(Texture);
 		return RHITexture ? RHITexture->GetLinkedObject(GPUIndex) : nullptr;
 	}
 
-	FORCEINLINE_DEBUGGABLE FD3D12TextureBase* RetrieveTextureBase(FRHITexture* Texture)
+	FORCEINLINE_DEBUGGABLE FD3D12Texture* RetrieveTexture(FRHITexture* Texture)
 	{
-		return RetrieveTextureBase(Texture, GetGPUIndex());
+		return RetrieveTexture(Texture, GetGPUIndex());
 	}
 
 	uint32 GetGPUIndex() const { return GPUMask.ToIndex(); }
diff --git a/Engine/Source/Runtime/D3D12RHI/Private/D3D12Commands.cpp b/Engine/Source/Runtime/D3D12RHI/Private/D3D12Commands.cpp
--- a/Engine/Source/Runtime/D3D12RHI/Private/D3D12Commands.cpp
+++ b/Engine/Source/Runtime/D3D12RHI/Private/D3D12Commands.cpp
@@ -252,7 +252,7 @@ void ProcessResource(FD3D12CommandContext& Context, const FRHITransitionInfo& In
 	}
 	case FRHITransitionInfo::EType::Texture:
 	{
-		FD3D12TextureBase* Texture = Context.RetrieveTextureBase(Info.Texture);
+		FD3D12Texture* Texture = Context.RetrieveTexture(Info.Texture);
 		check(Texture);
 		Function(Info, Texture->GetResource());
 		break;
@@ -395,7 +395,7 @@ static void HandleTransientAliasing(FD3D12CommandContext& Context, const FD3D12T
 		}
 		case FRHITransientAliasingInfo::EType::Texture:
 		{
-			FD3D12TextureBase* Texture = Context.RetrieveTextureBase(Info.Texture);
+			FD3D12Texture* Texture = Context.RetrieveTexture(Info.Texture);
 			check(Texture);
 			BaseShaderResource = Texture;
 			break;
@@ -423,7 +423,7 @@ static void HandleTransientAliasing(FD3D12CommandContext& Context, const FD3D12T
 					switch (Overlap.Type)
 					{
 					case FRHITransientAliasingOverlap::EType::Texture:
-						ResourceBefore = Context.RetrieveTextureBase(Overlap.Texture)->GetResource();
+						ResourceBefore = Context.RetrieveTexture(Overlap.Texture)->GetResource();
 						break;
 					case FRHITransientAliasingOverlap::EType::Buffer:
 						ResourceBefore = Context.RetrieveObject<FD3D12Buffer>(Overlap.Buffer)->GetResource();
@@ -782,7 +782,7 @@ void FD3D12CommandContext::RHISetComputePipelineState(FRHIComputePipelineState*
 
 void FD3D12CommandContext::RHISetShaderTexture(FRHIGraphicsShader* ShaderRHI, uint32 TextureIndex, FRHITexture* NewTextureRHI)
 {
-	FD3D12TextureBase* const NewTexture = RetrieveTextureBase(NewTextureRHI);
+	FD3D12Texture* const NewTexture = RetrieveTexture(NewTextureRHI);
 	switch (ShaderRHI->GetFrequency())
 	{
 	case SF_Vertex:
@@ -830,7 +830,7 @@ void FD3D12CommandContext::RHISetShaderTexture(FRHIGraphicsShader* ShaderRHI, ui
 void FD3D12CommandContext::RHISetShaderTexture(FRHIComputeShader* ComputeShaderRHI, uint32 TextureIndex, FRHITexture* NewTextureRHI)
 {
 	//ValidateBoundShader(StateCache, ComputeShaderRHI);
-	FD3D12TextureBase* const NewTexture = RetrieveTextureBase(NewTextureRHI);
+	FD3D12Texture* const NewTexture = RetrieveTexture(NewTextureRHI);
 	StateCache.SetShaderResourceView<SF_Compute>(NewTexture ? NewTexture->GetShaderResourceView() : nullptr, TextureIndex);
 }
 
@@ -1228,7 +1228,7 @@ void FD3D12CommandContext::SetRenderTargets(
 	const FRHIDepthRenderTargetView* NewDepthStencilTargetRHI
 	)
 {
-	FD3D12TextureBase* NewDepthStencilTarget = NewDepthStencilTargetRHI ? RetrieveTextureBase(NewDepthStencilTargetRHI->Texture) : nullptr;
+	FD3D12Texture* NewDepthStencilTarget = NewDepthStencilTargetRHI ? RetrieveTexture(NewDepthStencilTargetRHI->Texture) : nullptr;
 
 	check(NewNumSimultaneousRenderTargets <= MaxSimultaneousRenderTargets);
 
@@ -1263,7 +1263,7 @@ void FD3D12CommandContext::SetRenderTargets(
 		{
 			int32 RTMipIndex = NewRenderTargetsRHI[RenderTargetIndex].MipIndex;
 			int32 RTSliceIndex = NewRenderTargetsRHI[RenderTargetIndex].ArraySliceIndex;
-			FD3D12TextureBase* NewRenderTarget = RetrieveTextureBase(NewRenderTargetsRHI[RenderTargetIndex].Texture);
+			FD3D12Texture* NewRenderTarget = RetrieveTexture(NewRenderTargetsRHI[RenderTargetIndex].Texture);
 			RenderTargetView = NewRenderTarget->GetRenderTargetView(RTMipIndex, RTSliceIndex);
 
 			ensureMsgf(RenderTargetView, TEXT("Texture being set as render target has no RTV"));
@@ -1429,7 +1429,7 @@ void FD3D12CommandContext::SetRenderTargetsAndClear(const FRHISetRenderTargetsIn
 			VRSCombiners[1] = ConvertShadingRateCombiner(RenderTargetsInfo.ShadingRateTextureCombiner); // Combiner 1 is used to mix rates from a texture and the previous combiner
 			if (RenderTargetsInfo.ShadingRateTexture != nullptr)
 			{
-				FD3D12Resource* Resource = RetrieveTextureBase(RenderTargetsInfo.ShadingRateTexture)->GetResource();
+				FD3D12Resource* Resource = RetrieveTexture(RenderTargetsInfo.ShadingRateTexture)->GetResource();
 				CommandListHandle.GraphicsCommandList5()->RSSetShadingRateImage(Resource->GetResource());
 			}
 			else
@@ -1601,11 +1601,11 @@ inline int32 SetShaderResourcesFromBuffer_Surface(FD3D12CommandContext& CmdConte
 			}
 			TextureRHI->SetLastRenderTime(CurrentTime);
 
-			FD3D12TextureBase* TextureD3D12 = CmdContext.RetrieveTextureBase(TextureRHI);
+			FD3D12Texture* TextureD3D12 = CmdContext.RetrieveTexture(TextureRHI);
 			FD3D12ShaderResourceView* D3D12Resource = TextureD3D12->GetShaderResourceView();
 			if (D3D12Resource == nullptr)
 			{
-				D3D12Resource = CmdContext.RetrieveTextureBase(GWhiteTexture->TextureRHI)->GetShaderResourceView();
+				D3D12Resource = CmdContext.RetrieveTexture(GWhiteTexture->TextureRHI)->GetShaderResourceView();
 			}
 
 #if ENABLE_RHI_VALIDATION
@@ -2536,13 +2536,13 @@ void FD3D12CommandContext::RHIBroadcastTemporalEffect(const FName& InEffectName,
 {
 #if WITH_MGPU
 	FMemMark Mark(FMemStack::Get());
-	TArray<FD3D12TextureBase*, TMemStackAllocator<>> Resources;
+	TArray<FD3D12Texture*, TMemStackAllocator<>> Resources;
 	Resources.Reserve(InTextures.Num());
 	for (FRHITexture* Texture : InTextures)
 	{
-		Resources.Emplace(RetrieveTextureBase(Texture));
+		Resources.Emplace(RetrieveTexture(Texture));
 	}
-	auto CopyFunction = [](FD3D12CommandListHandle& CommandList, FD3D12TextureBase* Dst, FD3D12TextureBase* Src)
+	auto CopyFunction = [](FD3D12CommandListHandle& CommandList, FD3D12Texture* Dst, FD3D12Texture* Src)
 	{
 		CommandList->CopyResource(Dst->GetResource()->GetResource(), Src->GetResource()->GetResource());
 	};
diff --git a/Engine/Source/Runtime/D3D12RHI/Private/D3D12RHI.cpp b/Engine/Source/Runtime/D3D12RHI/Private/D3D12RHI.cpp
--- a/Engine/Source/Runtime/D3D12RHI/Private/D3D12RHI.cpp
+++ b/Engine/Source/Runtime/D3D12RHI/Private/D3D12RHI.cpp
@@ -473,25 +473,25 @@ ID3D12Resource* FD3D12DynamicRHI::RHIGetResource(FRHITexture* InTexture) const
 
 uint32 FD3D12DynamicRHI::RHIGetResourceDeviceIndex(FRHITexture* InTexture) const
 {
-	FD3D12TextureBase* D3D12Texture = GetD3D12TextureFromRHITexture(InTexture);
+	FD3D12Texture* D3D12Texture = GetD3D12TextureFromRHITexture(InTexture);
 	return D3D12Texture->GetParentDevice()->GetGPUIndex();
 }
 
 int64 FD3D12DynamicRHI::RHIGetResourceMemorySize(FRHITexture* InTexture) const
 {
-	FD3D12TextureBase* D3D12Texture = GetD3D12TextureFromRHITexture(InTexture);
-	return D3D12Texture->GetMemorySize();
+	FD3D12Texture* D3D12Texture = GetD3D12TextureFromRHITexture(InTexture);
+	return D3D12Texture->ResourceLocation.GetSize();
 }
 
 bool FD3D12DynamicRHI::RHIIsResourcePlaced(FRHITexture* InTexture) const
 {
-	FD3D12TextureBase* D3D12Texture = GetD3D12TextureFromRHITexture(InTexture);
+	FD3D12Texture* D3D12Texture = GetD3D12TextureFromRHITexture(InTexture);
 	return D3D12Texture->GetResource()->IsPlacedResource();
 }
 
 D3D12_CPU_DESCRIPTOR_HANDLE FD3D12DynamicRHI::RHIGetRenderTargetView(FRHITexture* InTexture, int32 InMipIndex, int32 InArraySliceIndex) const
 {
-	FD3D12TextureBase* D3D12Texture = GetD3D12TextureFromRHITexture(InTexture);
+	FD3D12Texture* D3D12Texture = GetD3D12TextureFromRHITexture(InTexture);
 	FD3D12RenderTargetView* RTV = D3D12Texture->GetRenderTargetView(InMipIndex, InArraySliceIndex);
 	return RTV ? RTV->GetView() : D3D12_CPU_DESCRIPTOR_HANDLE{};
 }
@@ -517,7 +517,7 @@ void FD3D12DynamicRHI::RHIRegisterWork(uint32 InDeviceIndex, uint32 NumPrimitive
 
 void FD3D12DynamicRHI::RHIAddPendingBarrier(FRHITexture* InTexture, D3D12_RESOURCE_STATES InState, uint32 InSubResource)
 {
-	FD3D12TextureBase* D3D12Texture = GetD3D12TextureFromRHITexture(InTexture);
+	FD3D12Texture* D3D12Texture = GetD3D12TextureFromRHITexture(InTexture);
 	FD3D12CommandListHandle& CommandList = D3D12Texture->GetParentDevice()->GetDefaultCommandContext().CommandListHandle;
 
 	CommandList.AddPendingResourceBarrier(D3D12Texture->GetResource(), InState, InSubResource);
diff --git a/Engine/Source/Runtime/D3D12RHI/Private/D3D12RHIPrivate.h b/Engine/Source/Runtime/D3D12RHI/Private/D3D12RHIPrivate.h
--- a/Engine/Source/Runtime/D3D12RHI/Private/D3D12RHIPrivate.h
+++ b/Engine/Source/Runtime/D3D12RHI/Private/D3D12RHIPrivate.h
@@ -356,18 +356,13 @@ public:
 	virtual FShaderResourceViewRHIRef RHICreateShaderResourceView(const FShaderResourceViewInitializer& Initializer) final override;
 	virtual void RHIUpdateShaderResourceView(FRHIShaderResourceView* SRV, FRHIBuffer* Buffer, uint32 Stride, uint8 Format) final override;
 	virtual void RHIUpdateShaderResourceView(FRHIShaderResourceView* SRV, FRHIBuffer* Buffer) final override;
-	virtual uint64 RHICalcTexture2DPlatformSize(uint32 SizeX, uint32 SizeY, uint8 Format, uint32 NumMips, uint32 NumSamples, ETextureCreateFlags InFlags, const FRHIResourceCreateInfo& CreateInfo, uint32& OutAlign) override;
-	virtual uint64 RHICalcTexture2DArrayPlatformSize(uint32 SizeX, uint32 SizeY, uint32 ArraySize, uint8 Format, uint32 NumMips, uint32 NumSamples, ETextureCreateFlags Flags, const FRHIResourceCreateInfo& CreateInfo, uint32& OutAlign) final override;
-	virtual uint64 RHICalcTexture3DPlatformSize(uint32 SizeX, uint32 SizeY, uint32 SizeZ, uint8 Format, uint32 NumMips, ETextureCreateFlags InFlags, const FRHIResourceCreateInfo& CreateInfo, uint32& OutAlign) final override;
-	virtual uint64 RHICalcTextureCubePlatformSize(uint32 Size, uint8 Format, uint32 NumMips, ETextureCreateFlags InFlags, const FRHIResourceCreateInfo& CreateInfo, uint32& OutAlign) final override;
+	virtual FRHICalcTextureSizeResult RHICalcTexturePlatformSize(const FRHITextureDesc& Desc, uint32 FirstMipIndex) override;
 	virtual uint64 RHIGetMinimumAlignmentForBufferBackedSRV(EPixelFormat Format) final override;
 	virtual void RHIGetTextureMemoryStats(FTextureMemoryStats& OutStats) final override;
 	virtual bool RHIGetTextureMemoryVisualizeData(FColor* TextureData, int32 SizeX, int32 SizeY, int32 Pitch, int32 PixelSize) final override;
-	virtual FTexture2DRHIRef RHICreateTexture2D(uint32 SizeX, uint32 SizeY, uint8 Format, uint32 NumMips, uint32 NumSamples, ETextureCreateFlags Flags, ERHIAccess InResourceState, FRHIResourceCreateInfo& CreateInfo) override;
-	virtual FTexture2DRHIRef RHIAsyncCreateTexture2D(uint32 SizeX, uint32 SizeY, uint8 Format, uint32 NumMips, ETextureCreateFlags Flags, ERHIAccess InResourceState, void** InitialMipData, uint32 NumInitialMips) final override;
+	virtual FTextureRHIRef RHIAsyncCreateTexture2D(uint32 SizeX, uint32 SizeY, uint8 Format, uint32 NumMips, ETextureCreateFlags Flags, ERHIAccess InResourceState, void** InitialMipData, uint32 NumInitialMips) final override;
+	virtual FTextureRHIRef RHICreateTexture(const FRHITextureCreateDesc& CreateDesc) override;
 	virtual void RHICopySharedMips(FRHITexture2D* DestTexture2D, FRHITexture2D* SrcTexture2D) final override;
-	virtual FTexture2DArrayRHIRef RHICreateTexture2DArray(uint32 SizeX, uint32 SizeY, uint32 SizeZ, uint8 Format, uint32 NumMips, uint32 NumSamples, ETextureCreateFlags Flags, ERHIAccess InResourceState, FRHIResourceCreateInfo& CreateInfo) override;
-	virtual FTexture3DRHIRef RHICreateTexture3D(uint32 SizeX, uint32 SizeY, uint32 SizeZ, uint8 Format, uint32 NumMips, ETextureCreateFlags Flags, ERHIAccess InResourceState, FRHIResourceCreateInfo& CreateInfo) override;
 	virtual FShaderResourceViewRHIRef RHICreateShaderResourceView(FRHITexture* Texture, const FRHITextureSRVCreateInfo& CreateInfo) final override;
 	virtual uint32 RHIComputeMemorySize(FRHITexture* TextureRHI) final override;
 	virtual FTexture2DRHIRef RHIAsyncReallocateTexture2D(FRHITexture2D* Texture2D, int32 NewMipCount, int32 NewSizeX, int32 NewSizeY, FThreadSafeCounter* RequestStatus) override;
@@ -382,8 +377,6 @@ public:
 	virtual FUpdateTexture3DData BeginUpdateTexture3D_RenderThread(class FRHICommandListImmediate& RHICmdList, FRHITexture3D* Texture, uint32 MipIndex, const struct FUpdateTextureRegion3D& UpdateRegion) final override;
 	virtual void EndUpdateTexture3D_RenderThread(class FRHICommandListImmediate& RHICmdList, FUpdateTexture3DData& UpdateData) final override;
 	virtual void EndMultiUpdateTexture3D_RenderThread(class FRHICommandListImmediate& RHICmdList, TArray<FUpdateTexture3DData>& UpdateDataArray) final override;
-	virtual FTextureCubeRHIRef RHICreateTextureCube(uint32 Size, uint8 Format, uint32 NumMips, ETextureCreateFlags Flags, ERHIAccess InResourceState, FRHIResourceCreateInfo& CreateInfo) override;
-	virtual FTextureCubeRHIRef RHICreateTextureCubeArray(uint32 Size, uint32 ArraySize, uint8 Format, uint32 NumMips, ETextureCreateFlags Flags, ERHIAccess InResourceState, FRHIResourceCreateInfo& CreateInfo) override;
 	virtual void* RHILockTextureCubeFace(FRHITextureCube* Texture, uint32 FaceIndex, uint32 ArrayIndex, uint32 MipIndex, EResourceLockMode LockMode, uint32& DestStride, bool bLockWithinMiptail) final override;
 	virtual void RHIUnlockTextureCubeFace(FRHITextureCube* Texture, uint32 FaceIndex, uint32 ArrayIndex, uint32 MipIndex, bool bLockWithinMiptail) final override;
 	virtual void RHIBindDebugLabelName(FRHITexture* Texture, const TCHAR* Name) final override;
@@ -507,14 +500,13 @@ public:
 	{
 		return RHIFinalizeAsyncReallocateTexture2D(Texture2D, bBlockUntilCompleted);
 	}
-	virtual ETextureReallocationStatus CancelAsyncReallocateTexture2D_RenderThread(class FRHICommandListImmediate& RHICmdList, FRHITexture2D* Texture2D, bool bBlockUntilCompleted)\
+	virtual ETextureReallocationStatus CancelAsyncReallocateTexture2D_RenderThread(class FRHICommandListImmediate& RHICmdList, FRHITexture2D* Texture2D, bool bBlockUntilCompleted)
 	{
 		return RHICancelAsyncReallocateTexture2D(Texture2D, bBlockUntilCompleted);
 	}
 
-	virtual FTexture2DRHIRef RHICreateTexture2D_RenderThread(class FRHICommandListImmediate& RHICmdList, uint32 SizeX, uint32 SizeY, uint8 Format, uint32 NumMips, uint32 NumSamples, ETextureCreateFlags Flags, ERHIAccess InResourceState, FRHIResourceCreateInfo& CreateInfo) override;
-	virtual FTexture2DArrayRHIRef RHICreateTexture2DArray_RenderThread(class FRHICommandListImmediate& RHICmdList, uint32 SizeX, uint32 SizeY, uint32 SizeZ, uint8 Format, uint32 NumMips, uint32 NumSamples, ETextureCreateFlags Flags, ERHIAccess InResourceState, FRHIResourceCreateInfo& CreateInfo) override;
-	virtual FTexture3DRHIRef RHICreateTexture3D_RenderThread(class FRHICommandListImmediate& RHICmdList, uint32 SizeX, uint32 SizeY, uint32 SizeZ, uint8 Format, uint32 NumMips, ETextureCreateFlags Flags, ERHIAccess InResourceState, FRHIResourceCreateInfo& CreateInfo) override;
+	virtual FTextureRHIRef RHICreateTexture_RenderThread(class FRHICommandListImmediate& RHICmdList, const FRHITextureCreateDesc& CreateDesc) override;
+
 	virtual FUnorderedAccessViewRHIRef RHICreateUnorderedAccessView_RenderThread(class FRHICommandListImmediate& RHICmdList, FRHIBuffer* Buffer, bool bUseUAVCounter, bool bAppendBuffer) override final;
 	virtual FUnorderedAccessViewRHIRef RHICreateUnorderedAccessView_RenderThread(class FRHICommandListImmediate& RHICmdList, FRHITexture* Texture, uint32 MipLevel, uint16 FirstArraySlice, uint16 NumArraySlices) override final;
 	virtual FUnorderedAccessViewRHIRef RHICreateUnorderedAccessView_RenderThread(class FRHICommandListImmediate& RHICmdList, FRHITexture* Texture, uint32 MipLevel, uint8 Format, uint16 FirstArraySlice, uint16 NumArraySlices) override final;
@@ -528,8 +520,6 @@ public:
 	virtual FShaderResourceViewRHIRef CreateShaderResourceView_RenderThread(class FRHICommandListImmediate& RHICmdList, FRHIBuffer* Buffer) override final;
 	virtual FShaderResourceViewRHIRef CreateShaderResourceView_RenderThread(class FRHICommandListImmediate& RHICmdList, const FShaderResourceViewInitializer& Initializer) override final;
 	virtual FShaderResourceViewRHIRef RHICreateShaderResourceView_RenderThread(class FRHICommandListImmediate& RHICmdList, FRHIBuffer* Buffer) override final;
-	virtual FTextureCubeRHIRef RHICreateTextureCube_RenderThread(class FRHICommandListImmediate& RHICmdList, uint32 Size, uint8 Format, uint32 NumMips, ETextureCreateFlags Flags, ERHIAccess InResourceState, FRHIResourceCreateInfo& CreateInfo) override;
-	virtual FTextureCubeRHIRef RHICreateTextureCubeArray_RenderThread(class FRHICommandListImmediate& RHICmdList, uint32 Size, uint32 ArraySize, uint8 Format, uint32 NumMips, ETextureCreateFlags Flags, ERHIAccess InResourceState, FRHIResourceCreateInfo& CreateInfo) override;
 
 	virtual FRenderQueryRHIRef RHICreateRenderQuery_RenderThread(class FRHICommandListImmediate& RHICmdList, ERenderQueryType QueryType) override final
 	{
@@ -1138,8 +1128,8 @@ public:
 	virtual void* CreateVirtualTexture(ETextureCreateFlags InFlags, D3D12_RESOURCE_DESC& ResourceDesc, const struct FD3D12TextureLayout& TextureLayout, FD3D12Resource** ppResource, FPlatformMemory::FPlatformVirtualMemoryBlock& RawTextureBlock, D3D12_RESOURCE_STATES InitialUsage = D3D12_RESOURCE_STATE_PIXEL_SHADER_RESOURCE) = 0;
 	virtual void DestroyVirtualTexture(ETextureCreateFlags InFlags, void* RawTextureMemory, FPlatformMemory::FPlatformVirtualMemoryBlock& RawTextureBlock, uint64 CommittedTextureSize) = 0;
 #endif
-	virtual bool HandleSpecialLock(void*& MemoryOut, uint32 MipIndex, uint32 ArrayIndex, ETextureCreateFlags InFlags, EResourceLockMode LockMode, const FD3D12TextureLayout& TextureLayout, void* RawTextureMemory, uint32& DestStride) { return false; }
-	virtual bool HandleSpecialUnlock(FRHICommandListBase* RHICmdList, uint32 MipIndex, ETextureCreateFlags InFlags, const struct FD3D12TextureLayout& TextureLayout, void* RawTextureMemory) { return false; }
+	virtual bool HandleSpecialLock(void*& MemoryOut, uint32 MipIndex, uint32 ArrayIndex, FD3D12Texture* InTexture, EResourceLockMode LockMode, uint32& DestStride) { return false; }
+	virtual bool HandleSpecialUnlock(FRHICommandListBase* RHICmdList, uint32 MipIndex, FD3D12Texture* InTexture) { return false; }
 
 	FD3D12Adapter& GetAdapter(uint32_t Index = 0) { return *ChosenAdapters[Index]; }
 	const FD3D12Adapter& GetAdapter(uint32_t Index = 0) const { return *ChosenAdapters[Index]; }
@@ -1204,27 +1194,18 @@ protected:
 
 public:
 
-	template<typename BaseResourceType>
-	TD3D12Texture2D<BaseResourceType>* CreateD3D12Texture2D(class FRHICommandListImmediate* RHICmdList, uint32 SizeX, uint32 SizeY, uint32 SizeZ, bool bTextureArray, bool CubeTexture, EPixelFormat Format,
-		uint32 NumMips, uint32 NumSamples, ETextureCreateFlags InFlags, ERHIAccess InResourceState, FRHIResourceCreateInfo& CreateInfo, ED3D12ResourceTransientMode TransientMode = ED3D12ResourceTransientMode::NonTransient, ID3D12ResourceAllocator* ResourceAllocator = nullptr);
-
-	FD3D12Buffer* CreateD3D12Buffer(class FRHICommandListImmediate* RHICmdList, uint32 Size, EBufferUsageFlags Usage, uint32 Stride, ERHIAccess ResourceState, FRHIResourceCreateInfo& CreateInfo, ED3D12ResourceTransientMode TransientMode = ED3D12ResourceTransientMode::NonTransient, ID3D12ResourceAllocator* ResourceAllocator = nullptr);
+	virtual FD3D12ResourceDesc GetResourceDesc(const FRHITextureDesc& CreateInfo) const;
 
-	FD3D12Texture3D* CreateD3D12Texture3D(class FRHICommandListImmediate* RHICmdList, uint32 SizeX, uint32 SizeY, uint32 SizeZ, EPixelFormat Format, uint32 NumMips, ETextureCreateFlags InFlags, ERHIAccess InResourceState, FRHIResourceCreateInfo& CreateInfo, ED3D12ResourceTransientMode TransientMode = ED3D12ResourceTransientMode::NonTransient, ID3D12ResourceAllocator* ResourceAllocator = nullptr);
+	virtual FD3D12Texture* CreateD3D12Texture(const FRHITextureCreateDesc& CreateDesc, class FRHICommandListImmediate* RHICmdList, ID3D12ResourceAllocator* ResourceAllocator = nullptr);
+	FD3D12Buffer* CreateD3D12Buffer(class FRHICommandListImmediate* RHICmdList, uint32 Size, EBufferUsageFlags Usage, uint32 Stride, ERHIAccess ResourceState, FRHIResourceCreateInfo& CreateInfo, ID3D12ResourceAllocator* ResourceAllocator = nullptr);
+	virtual FD3D12Texture* CreateNewD3D12Texture(const FRHITextureCreateDesc& CreateDesc, class FD3D12Device* Device);
 
-
-	// Interface for FD3D12TransientResourceAllocator.
-	virtual D3D12_RESOURCE_DESC GetResourceDesc(const FRHITextureCreateInfo& CreateInfo) const;
-	virtual FRHITexture* CreateTexture(const FRHITextureCreateInfo& CreateInfo, const TCHAR* DebugName, ERHIAccess InitialState, ED3D12ResourceTransientMode AllocationMode, ID3D12ResourceAllocator* ResourceAllocator);
-	FRHIBuffer* CreateBuffer(const FRHIBufferCreateInfo& CreateInfo, const TCHAR* DebugName, ERHIAccess InitialState, ED3D12ResourceTransientMode TransientMode, ID3D12ResourceAllocator* ResourceAllocator);
+	FRHIBuffer* CreateBuffer(const FRHIBufferCreateInfo& CreateInfo, const TCHAR* DebugName, ERHIAccess InitialState, ID3D12ResourceAllocator* ResourceAllocator);
 
 protected:
 
-	template<typename BaseResourceType>
-	TD3D12Texture2D<BaseResourceType>* CreateTextureFromResource(bool bTextureArray, bool bCubeTexture, EPixelFormat Format, ETextureCreateFlags TexCreateFlags, const FClearValueBinding& ClearValueBinding, ID3D12Resource* Resource);
-
-	template<typename BaseResourceType>
-	TD3D12Texture2D<BaseResourceType>* CreateAliasedD3D12Texture2D(TD3D12Texture2D<BaseResourceType>* SourceTexture);
+	FD3D12Texture* CreateTextureFromResource(bool bTextureArray, bool bCubeTexture, EPixelFormat Format, ETextureCreateFlags TexCreateFlags, const FClearValueBinding& ClearValueBinding, ID3D12Resource* Resource);
+	FD3D12Texture* CreateAliasedD3D12Texture2D(FD3D12Texture* SourceTexture);
 
 	/**
 	 * Gets the best supported MSAA settings from the provided MSAA count to check against.
diff --git a/Engine/Source/Runtime/D3D12RHI/Private/D3D12RayTracing.cpp b/Engine/Source/Runtime/D3D12RHI/Private/D3D12RayTracing.cpp
--- a/Engine/Source/Runtime/D3D12RHI/Private/D3D12RayTracing.cpp
+++ b/Engine/Source/Runtime/D3D12RHI/Private/D3D12RayTracing.cpp
@@ -713,10 +713,9 @@ FD3D12RayTracingCompactionRequestHandler::FD3D12RayTracingCompactionRequestHandl
 
 	FRHIGPUMask GPUMask = FRHIGPUMask::FromIndex(GetParentDevice()->GetGPUIndex());
 	ID3D12ResourceAllocator* ResourceAllocator = nullptr;
-	ED3D12ResourceTransientMode TransientMode = ED3D12ResourceTransientMode::NonTransient;
 	bool bHasInitialData = false;
 	PostBuildInfoBuffer = GetParentDevice()->GetParentAdapter()->CreateRHIBuffer(PostBuildInfoBufferDesc, 8,
-		0, PostBuildInfoBufferDesc.Width, BUF_UnorderedAccess | BUF_SourceCopy, ED3D12ResourceStateMode::MultiState, D3D12_RESOURCE_STATE_UNORDERED_ACCESS, bHasInitialData, GPUMask, TransientMode, ResourceAllocator, TEXT("PostBuildInfoBuffer"));
+		0, PostBuildInfoBufferDesc.Width, BUF_UnorderedAccess | BUF_SourceCopy, ED3D12ResourceStateMode::MultiState, D3D12_RESOURCE_STATE_UNORDERED_ACCESS, bHasInitialData, GPUMask, ResourceAllocator, TEXT("PostBuildInfoBuffer"));
 	SetName(PostBuildInfoBuffer->GetResource(), TEXT("PostBuildInfoBuffer"));
 
 	PostBuildInfoStagingBuffer = RHICreateStagingBuffer();
@@ -1951,10 +1950,9 @@ public:
 		bool bHasInitialData = true;
 
 		ID3D12ResourceAllocator* ResourceAllocator = nullptr;
-		ED3D12ResourceTransientMode TransientMode = ED3D12ResourceTransientMode::NonTransient;
 		Buffer = Adapter->CreateRHIBuffer(
 			BufferDesc, BufferDesc.Alignment, 0, BufferDesc.Width, BUF_Static, ED3D12ResourceStateMode::MultiState,
-			D3D12_RESOURCE_STATE_COPY_DEST, bHasInitialData, GPUMask, TransientMode, ResourceAllocator, TEXT("Shader binding table"));
+			D3D12_RESOURCE_STATE_COPY_DEST, bHasInitialData, GPUMask, ResourceAllocator, TEXT("Shader binding table"));
 
 		// Use copy queue for uploading the data
 		FD3D12SyncPoint CopyQueueSyncPoint = Buffer->UploadResourceDataViaCopyQueue(&Data);
@@ -3277,7 +3275,6 @@ static TRefCountPtr<FD3D12Buffer> CreateRayTracingBuffer(FD3D12Adapter* Adapter,
 	TRefCountPtr<FD3D12Buffer> Result;
 
 	FString DebugNameString = DebugName.ToString();
-	ED3D12ResourceTransientMode TransientMode = ED3D12ResourceTransientMode::NonTransient;
 	ID3D12ResourceAllocator* ResourceAllocator = nullptr;
 	FRHIGPUMask GPUMask = FRHIGPUMask::FromIndex(GPUIndex);
 	bool bHasInitialData = false;
@@ -3289,7 +3286,7 @@ static TRefCountPtr<FD3D12Buffer> CreateRayTracingBuffer(FD3D12Adapter* Adapter,
 			BufferDesc, D3D12_RAYTRACING_ACCELERATION_STRUCTURE_BYTE_ALIGNMENT,
 			0, BufferDesc.Width, BUF_AccelerationStructure,
 			ED3D12ResourceStateMode::SingleState, D3D12_RESOURCE_STATE_RAYTRACING_ACCELERATION_STRUCTURE, bHasInitialData,
-			GPUMask, TransientMode, ResourceAllocator, *DebugNameString);
+			GPUMask, ResourceAllocator, *DebugNameString);
 	}
 	else if (Type == ERayTracingBufferType::Scratch)
 	{
@@ -3300,7 +3297,7 @@ static TRefCountPtr<FD3D12Buffer> CreateRayTracingBuffer(FD3D12Adapter* Adapter,
 			BufferDesc, D3D12_RAYTRACING_ACCELERATION_STRUCTURE_BYTE_ALIGNMENT,
 			0, BufferDesc.Width, BUF_UnorderedAccess,
 			ED3D12ResourceStateMode::Default, D3D12_RESOURCE_STATE_UNORDERED_ACCESS, bHasInitialData,
-			GPUMask, TransientMode, ResourceAllocator, *DebugNameString);
+			GPUMask, ResourceAllocator, *DebugNameString);
 
 		// Elevates the scratch buffer heap priority, which may help performance / stability in low memory conditions 
 		// (Acceleration structure already boosted from allocation side)
@@ -4677,7 +4674,7 @@ static bool SetRayTracingShaderResources(
 		FRHITexture* Resource = Textures[SRVIndex];
 		if (Resource)
 		{
-			FD3D12TextureBase* Texture = FD3D12CommandContext::RetrieveTextureBase(Resource, GPUIndex);
+			FD3D12Texture* Texture = FD3D12CommandContext::RetrieveTexture(Resource, GPUIndex);
 			FD3D12ShaderResourceView* SRV = Texture->GetShaderResourceView();
 			LocalSRVs[SRVIndex] = SRV->GetOfflineCpuHandle();
 			BoundSRVMask |= 1ull << SRVIndex;
@@ -4781,10 +4778,10 @@ static bool SetRayTracingShaderResources(
 					checkf(RHITexture != nullptr, TEXT("Missing required texture binding for slot %d in uniform buffer %d (UB layout name: '%s')"),
 						BindIndex, BufferIndex, *(Buffer->GetLayout().GetDebugName()));
 
-					FD3D12ShaderResourceView* SRV = FD3D12CommandContext::RetrieveTextureBase(RHITexture, GPUIndex)->GetShaderResourceView();
+					FD3D12ShaderResourceView* SRV = FD3D12CommandContext::RetrieveTexture(RHITexture, GPUIndex)->GetShaderResourceView();
 					if (!ensure(SRV))
 					{
-						SRV = FD3D12CommandContext::RetrieveTextureBase(GBlackTexture->TextureRHI, GPUIndex)->GetShaderResourceView();
+						SRV = FD3D12CommandContext::RetrieveTexture(GBlackTexture->TextureRHI, GPUIndex)->GetShaderResourceView();
 					}
 					check(SRV != nullptr);
 
