Fix iOS/tvOS game target build.

#rnx

diff --git a/Engine/Source/Runtime/Apple/AGXRHI/Private/AGXUAV.cpp b/Engine/Source/Runtime/Apple/AGXRHI/Private/AGXUAV.cpp
--- a/Engine/Source/Runtime/Apple/AGXRHI/Private/AGXUAV.cpp
+++ b/Engine/Source/Runtime/Apple/AGXRHI/Private/AGXUAV.cpp
@@ -99,7 +99,7 @@ FAGXResourceViewBase::FAGXResourceViewBase(
 		// TODO: Apple Silicon supports memoryless
 #if PLATFORM_IOS
 		// Memoryless targets can't have texture views (SRVs or UAVs)
-		if (Source.Texture.GetStorageMode() != mtlpp::StorageMode::Memoryless)
+		if (SourceTexture->Texture.GetStorageMode() != mtlpp::StorageMode::Memoryless)
 #endif
 		{
 			// Determine the appropriate metal format for the view.
@@ -116,9 +116,11 @@ FAGXResourceViewBase::FAGXResourceViewBase(
 						checkNoEntry();
 						break;
 
+#if PLATFORM_MAC
 					case mtlpp::PixelFormat::Depth24Unorm_Stencil8:
 						MetalFormat = mtlpp::PixelFormat::X24_Stencil8;
 						break;
+#endif
 
 					case mtlpp::PixelFormat::Depth32Float_Stencil8:
 						MetalFormat = mtlpp::PixelFormat::X32_Stencil8;
diff --git a/Engine/Source/Runtime/Apple/MetalRHI/Private/MetalUAV.cpp b/Engine/Source/Runtime/Apple/MetalRHI/Private/MetalUAV.cpp
--- a/Engine/Source/Runtime/Apple/MetalRHI/Private/MetalUAV.cpp
+++ b/Engine/Source/Runtime/Apple/MetalRHI/Private/MetalUAV.cpp
@@ -97,7 +97,7 @@ FMetalResourceViewBase::FMetalResourceViewBase(
 
 #if PLATFORM_IOS
 		// Memoryless targets can't have texture views (SRVs or UAVs)
-		if (Source.Texture.GetStorageMode() != mtlpp::StorageMode::Memoryless)
+		if (SourceTexture->Texture.GetStorageMode() != mtlpp::StorageMode::Memoryless)
 #endif
 		{
 			// Determine the appropriate metal format for the view.
@@ -111,7 +111,9 @@ FMetalResourceViewBase::FMetalResourceViewBase(
 				switch (SourceTexture->Texture.GetPixelFormat())
 				{
 				default: checkNoEntry(); break;
+#if PLATFORM_MAC
 				case mtlpp::PixelFormat::Depth24Unorm_Stencil8: MetalFormat = mtlpp::PixelFormat::X24_Stencil8; break;
+#endif
 				case mtlpp::PixelFormat::Depth32Float_Stencil8: MetalFormat = mtlpp::PixelFormat::X32_Stencil8; break;
 				}
 			}
