Disable Lumen on Mac with AMD GPUs

#rnx

diff --git a/Engine/Source/Runtime/Apple/MetalRHI/Private/MetalCommandQueue.cpp b/Engine/Source/Runtime/Apple/MetalRHI/Private/MetalCommandQueue.cpp
--- a/Engine/Source/Runtime/Apple/MetalRHI/Private/MetalCommandQueue.cpp
+++ b/Engine/Source/Runtime/Apple/MetalRHI/Private/MetalCommandQueue.cpp
@@ -146,6 +146,9 @@ FMetalCommandQueue::FMetalCommandQueue(mtlpp::Device InDevice, uint32 const MaxN
 #else // Assume that Mac & other platforms all support these from the start. They can diverge later.
 	const bool bIsNVIDIA = [Device.GetName().GetPtr() rangeOfString:@"Nvidia" options:NSCaseInsensitiveSearch].location != NSNotFound;
 	Features = EMetalFeaturesCountingQueries | EMetalFeaturesBaseVertexInstance | EMetalFeaturesIndirectBuffer | EMetalFeaturesLayeredRendering | EMetalFeaturesCubemapArrays;
+
+	FString DeviceName(Device.GetName());
+
 	if (!bIsNVIDIA)
 	{
 		Features |= EMetalFeaturesSetBufferOffset;
@@ -157,7 +160,6 @@ FMetalCommandQueue::FMetalCommandQueue(mtlpp::Device InDevice, uint32 const MaxN
 		// Assume that set*Bytes only works on macOS Sierra and above as no-one has tested it anywhere else.
 		Features |= EMetalFeaturesSetBytes;
 		
-		FString DeviceName(Device.GetName());
 		// On earlier OS versions Intel Broadwell couldn't suballocate properly
 		if (!(DeviceName.Contains(TEXT("Intel")) && (DeviceName.Contains(TEXT("5300")) || DeviceName.Contains(TEXT("6000")) || DeviceName.Contains(TEXT("6100")))) || FPlatformMisc::MacOSXVersionCompare(10,14,0) >= 0)
 		{
@@ -252,6 +254,16 @@ FMetalCommandQueue::FMetalCommandQueue(mtlpp::Device InDevice, uint32 const MaxN
 		Features |= EMetalFeaturesGPUTrace;
 	}
 
+	// For Release-5.0 we encountered a hard GPU hang on AMD GPUs.
+	if (DeviceName.Contains(TEXT("AMD")))
+	{
+		static const auto CVarLumenDiffuseIndirectAllow = IConsoleManager::Get().FindConsoleVariable(TEXT("r.Lumen.DiffuseIndirect.Allow"));
+		static const auto CVarLumenReflectionsAllow = IConsoleManager::Get().FindConsoleVariable(TEXT("r.Lumen.Reflections.Allow"));
+
+		CVarLumenDiffuseIndirectAllow->Set(0, ECVF_SetByCode);
+		CVarLumenReflectionsAllow->Set(0, ECVF_SetByCode);
+	}
+
 	PermittedOptions = 0;
 	PermittedOptions |= mtlpp::ResourceOptions::CpuCacheModeDefaultCache;
 	PermittedOptions |= mtlpp::ResourceOptions::CpuCacheModeWriteCombined;
