Compile fix for iOS devices

diff --git a/Engine/Source/Runtime/Apple/MetalRHI/Private/MetalCommandQueue.cpp b/Engine/Source/Runtime/Apple/MetalRHI/Private/MetalCommandQueue.cpp
--- a/Engine/Source/Runtime/Apple/MetalRHI/Private/MetalCommandQueue.cpp
+++ b/Engine/Source/Runtime/Apple/MetalRHI/Private/MetalCommandQueue.cpp
@@ -254,15 +254,22 @@ FMetalCommandQueue::FMetalCommandQueue(mtlpp::Device InDevice, uint32 const MaxN
 		Features |= EMetalFeaturesGPUTrace;
 	}
 
+#if PLATFORM_MAC
 	// For Release-5.0 we encountered a hard GPU hang on AMD GPUs.
 	if (DeviceName.Contains(TEXT("AMD")))
 	{
-		static const auto CVarLumenDiffuseIndirectAllow = IConsoleManager::Get().FindConsoleVariable(TEXT("r.Lumen.DiffuseIndirect.Allow"));
-		static const auto CVarLumenReflectionsAllow = IConsoleManager::Get().FindConsoleVariable(TEXT("r.Lumen.Reflections.Allow"));
-
-		CVarLumenDiffuseIndirectAllow->Set(0, ECVF_SetByCode);
-		CVarLumenReflectionsAllow->Set(0, ECVF_SetByCode);
+		const auto CVarLumenDiffuseIndirectAllow = IConsoleManager::Get().FindConsoleVariable(TEXT("r.Lumen.DiffuseIndirect.Allow"));
+		const auto CVarLumenReflectionsAllow = IConsoleManager::Get().FindConsoleVariable(TEXT("r.Lumen.Reflections.Allow"));
+		if (CVarLumenDiffuseIndirectAllow != nullptr)
+		{
+			CVarLumenDiffuseIndirectAllow->Set(0, ECVF_SetByCode);
+		}
+		if (CVarLumenReflectionsAllow != nullptr)
+		{
+			CVarLumenReflectionsAllow->Set(0, ECVF_SetByCode);
+		}
 	}
+#endif
 
 	PermittedOptions = 0;
 	PermittedOptions |= mtlpp::ResourceOptions::CpuCacheModeDefaultCache;
