Fixed metal build break.

diff --git a/Engine/Source/Runtime/Apple/MetalRHI/Public/MetalResources.h b/Engine/Source/Runtime/Apple/MetalRHI/Public/MetalResources.h
--- a/Engine/Source/Runtime/Apple/MetalRHI/Public/MetalResources.h
+++ b/Engine/Source/Runtime/Apple/MetalRHI/Public/MetalResources.h
@@ -589,7 +589,7 @@ public:
 	}
 
 	explicit FMetalUnorderedAccessView(FRHITexture* Texture, uint32 MipLevel, uint16 FirstArraySlice, uint16 NumArraySlices)
-		: FRHIUnorderedAccessView(Buffer)
+		: FRHIUnorderedAccessView(Texture)
 		, FMetalResourceViewBase(
 			  Texture
 			, PF_Unknown
