Remove (TexCreate_UAV | TexCreate_NoTiling) check from the non-linear tex create code path assert.   There are cases when we cannot create a buffer backed texture when those flags are set.  Allow fallback to standard texture create path as it is currently allowed in 5.0 release.

#rnx

diff --git a/Engine/Source/Runtime/Apple/AGXRHI/Private/AGXTexture.cpp b/Engine/Source/Runtime/Apple/AGXRHI/Private/AGXTexture.cpp
--- a/Engine/Source/Runtime/Apple/AGXRHI/Private/AGXTexture.cpp
+++ b/Engine/Source/Runtime/Apple/AGXRHI/Private/AGXTexture.cpp
@@ -635,7 +635,7 @@ FAGXSurface::FAGXSurface(FAGXTextureCreateDesc const& CreateDesc)
 		{
 			// If we are in here then either the texture description is not buffer compatable or these flags were not set
 			// assert that these flag combinations are not set as they require a buffer backed texture and the texture description is not compatible with that
-			checkf(!(EnumHasAllFlags(CreateDesc.Flags, TexCreate_UAV | TexCreate_NoTiling) || EnumHasAllFlags(CreateDesc.Flags, TexCreate_AtomicCompatible)), TEXT("Requested buffer backed texture that breaks Metal linear texture limitations: %s"), *FString([CreateDesc.Desc description]));
+			checkf(!EnumHasAllFlags(CreateDesc.Flags, TexCreate_AtomicCompatible), TEXT("Requested buffer backed texture that breaks Metal linear texture limitations: %s"), *FString([CreateDesc.Desc description]));
 			Texture = GetAGXDeviceContext().CreateTexture(this, CreateDesc.Desc);
 		}
 		
diff --git a/Engine/Source/Runtime/Apple/MetalRHI/Private/MetalTexture.cpp b/Engine/Source/Runtime/Apple/MetalRHI/Private/MetalTexture.cpp
--- a/Engine/Source/Runtime/Apple/MetalRHI/Private/MetalTexture.cpp
+++ b/Engine/Source/Runtime/Apple/MetalRHI/Private/MetalTexture.cpp
@@ -665,7 +665,7 @@ FMetalSurface::FMetalSurface(FMetalTextureCreateDesc const& CreateDesc)
 		{
 			// If we are in here then either the texture description is not buffer compatable or these flags were not set
 			// assert that these flag combinations are not set as they require a buffer backed texture and the texture description is not compatible with that
-			checkf(!(EnumHasAllFlags(CreateDesc.Flags, TexCreate_UAV | TexCreate_NoTiling) || EnumHasAllFlags(CreateDesc.Flags, TexCreate_AtomicCompatible)), TEXT("Requested buffer backed texture that breaks Metal linear texture limitations: %s"), *FString([CreateDesc.Desc description]));
+			checkf(!EnumHasAllFlags(CreateDesc.Flags, TexCreate_AtomicCompatible), TEXT("Requested buffer backed texture that breaks Metal linear texture limitations: %s"), *FString([CreateDesc.Desc description]));
 			Texture = GetMetalDeviceContext().CreateTexture(this, CreateDesc.Desc);
 		}
 		
