Fixup casing in copyright notice.

#rnx

diff --git a/Engine/Source/Runtime/Apple/MetalCpp/Private/MetalCpp.cpp b/Engine/Source/Runtime/Apple/MetalCpp/Private/MetalCpp.cpp
--- a/Engine/Source/Runtime/Apple/MetalCpp/Private/MetalCpp.cpp
+++ b/Engine/Source/Runtime/Apple/MetalCpp/Private/MetalCpp.cpp
@@ -1,4 +1,4 @@
-// Copyright Epic Games, Inc. All rights reserved.
+// Copyright Epic Games, Inc. All Rights Reserved.
 
 #include "MetalCpp.h"
 
diff --git a/Engine/Source/Runtime/Apple/MetalCpp/Public/MetalCpp.h b/Engine/Source/Runtime/Apple/MetalCpp/Public/MetalCpp.h
--- a/Engine/Source/Runtime/Apple/MetalCpp/Public/MetalCpp.h
+++ b/Engine/Source/Runtime/Apple/MetalCpp/Public/MetalCpp.h
@@ -1,4 +1,4 @@
-// Copyright Epic Games, Inc. All rights reserved.
+// Copyright Epic Games, Inc. All Rights Reserved.
 
 #pragma once
 
