AGXRHI: Scrub all usage of mtlpp::BlitCommandEncoder

#rnx

diff --git a/Engine/Source/Runtime/Apple/AGXRHI/Private/AGXCommandEncoder.cpp b/Engine/Source/Runtime/Apple/AGXRHI/Private/AGXCommandEncoder.cpp
--- a/Engine/Source/Runtime/Apple/AGXRHI/Private/AGXCommandEncoder.cpp
+++ b/Engine/Source/Runtime/Apple/AGXRHI/Private/AGXCommandEncoder.cpp
@@ -25,6 +25,7 @@ FAGXCommandEncoder::FAGXCommandEncoder(FAGXCommandList& CmdList, EAGXCommandEnco
 , bSupportsMetalFeaturesSetBytes(CmdList.GetCommandQueue().SupportsFeature(EAGXFeaturesSetBytes))
 , RingBuffer(EncoderRingBufferSize, BufferOffsetAlignment, FAGXCommandQueue::GetCompatibleResourceOptions((mtlpp::ResourceOptions)(mtlpp::ResourceOptions::HazardTrackingModeUntracked | BUFFER_RESOURCE_STORAGE_MANAGED)))
 , RenderPassDesc(nil)
+, BlitCommandEncoder(nil)
 #if ENABLE_METAL_GPUPROFILE
 , CommandBufferStats(nullptr)
 #endif
@@ -284,7 +285,7 @@ bool FAGXCommandEncoder::IsComputeCommandEncoderActive(void) const
 
 bool FAGXCommandEncoder::IsBlitCommandEncoderActive(void) const
 {
-	return BlitCommandEncoder.GetPtr() != nil;
+	return BlitCommandEncoder != nil;
 }
 
 bool FAGXCommandEncoder::IsImmediate(void) const
@@ -330,7 +331,7 @@ mtlpp::ComputeCommandEncoder& FAGXCommandEncoder::GetComputeCommandEncoder(void)
 	return ComputeCommandEncoder;
 }
 
-mtlpp::BlitCommandEncoder& FAGXCommandEncoder::GetBlitCommandEncoder(void)
+id<MTLBlitCommandEncoder> FAGXCommandEncoder::GetBlitCommandEncoder() const
 {
 	check(IsBlitCommandEncoderActive());
 	return BlitCommandEncoder;
@@ -443,8 +444,7 @@ void FAGXCommandEncoder::BeginBlitCommandEncoding(void)
 	check(CommandBuffer);
 	check(IsRenderCommandEncoderActive() == false && IsComputeCommandEncoderActive() == false && IsBlitCommandEncoderActive() == false);
 	
-	BlitCommandEncoder = MTLPP_VALIDATE(mtlpp::CommandBuffer, CommandBuffer, AGXSafeGetRuntimeDebuggingLevel() >= EAGXDebugLevelValidation, BlitCommandEncoder());
-	
+	BlitCommandEncoder = [[CommandBuffer.GetPtr() blitCommandEncoder] retain];	
 	EncoderNum++;
 	
 	NSString* Label = nil;
@@ -452,13 +452,13 @@ void FAGXCommandEncoder::BeginBlitCommandEncoding(void)
 	if(GetEmitDrawEvents())
 	{
 		Label = [NSString stringWithFormat:@"BlitEncoder: %@", [DebugGroups count] > 0 ? [DebugGroups lastObject] : (NSString*)CFSTR("InitialPass")];
-		BlitCommandEncoder.SetLabel(Label);
+		[BlitCommandEncoder setLabel:Label];
 		
 		if([DebugGroups count])
 		{
 			for (NSString* Group in DebugGroups)
 			{
-				BlitCommandEncoder.PushDebugGroup(Group);
+				[BlitCommandEncoder pushDebugGroup:Group];
 			}
 		}
 	}
@@ -579,7 +579,8 @@ void FAGXCommandEncoder::EndEncoding(void)
 		}
 		else if(IsBlitCommandEncoderActive())
 		{
-			BlitCommandEncoder.EndEncoding();
+			[BlitCommandEncoder endEncoding];
+			[BlitCommandEncoder release];
 			BlitCommandEncoder = nil;
 		}
 	}
@@ -640,7 +641,7 @@ void FAGXCommandEncoder::InsertDebugSignpost(ns::String const& String)
 		}
 		else if (BlitCommandEncoder)
 		{
-			BlitCommandEncoder.InsertDebugSignpost(String);
+			[BlitCommandEncoder insertDebugSignpost:String];
 		}
 	}
 }
@@ -665,7 +666,7 @@ void FAGXCommandEncoder::PushDebugGroup(ns::String const& String)
 		}
 		else if (BlitCommandEncoder)
 		{
-			BlitCommandEncoder.PushDebugGroup(String);
+			[BlitCommandEncoder pushDebugGroup:String];
 		}
 	}
 }
@@ -690,7 +691,7 @@ void FAGXCommandEncoder::PopDebugGroup(void)
 		}
 		else if (BlitCommandEncoder)
 		{
-			BlitCommandEncoder.PopDebugGroup();
+			[BlitCommandEncoder popDebugGroup];
 		}
 	}
 }
diff --git a/Engine/Source/Runtime/Apple/AGXRHI/Private/AGXCommandEncoder.h b/Engine/Source/Runtime/Apple/AGXRHI/Private/AGXCommandEncoder.h
--- a/Engine/Source/Runtime/Apple/AGXRHI/Private/AGXCommandEncoder.h
+++ b/Engine/Source/Runtime/Apple/AGXRHI/Private/AGXCommandEncoder.h
@@ -140,7 +140,7 @@ public:
 	mtlpp::ComputeCommandEncoder& GetComputeCommandEncoder(void);
 	
 	/** @returns The active blit command encoder or nil if there isn't one. */
-	mtlpp::BlitCommandEncoder& GetBlitCommandEncoder(void);
+	id<MTLBlitCommandEncoder> GetBlitCommandEncoder() const;
 	
 	/** @returns The number of encoded passes in the command buffer. */
 	uint32 NumEncodedPasses(void) const { return EncoderNum; }
@@ -460,7 +460,7 @@ public:
 	mtlpp::ParallelRenderCommandEncoder ParallelRenderCommandEncoder;
 	mtlpp::RenderCommandEncoder RenderCommandEncoder;
 	mtlpp::ComputeCommandEncoder ComputeCommandEncoder;
-	mtlpp::BlitCommandEncoder BlitCommandEncoder;
+	id<MTLBlitCommandEncoder> BlitCommandEncoder;
 	TArray<mtlpp::RenderCommandEncoder> ChildRenderCommandEncoders;
 	
 #if ENABLE_METAL_GPUPROFILE
diff --git a/Engine/Source/Runtime/Apple/AGXRHI/Private/AGXContext.cpp b/Engine/Source/Runtime/Apple/AGXRHI/Private/AGXContext.cpp
--- a/Engine/Source/Runtime/Apple/AGXRHI/Private/AGXContext.cpp
+++ b/Engine/Source/Runtime/Apple/AGXRHI/Private/AGXContext.cpp
@@ -1520,9 +1520,9 @@ void FAGXContext::SynchronizeTexture(FAGXTexture const& Texture, uint32 Slice, u
 	RenderPass.SynchronizeTexture(Texture, Slice, Level);
 }
 
-void FAGXContext::SynchroniseResource(mtlpp::Resource const& Resource)
+void FAGXContext::SynchronizeResource(mtlpp::Resource const& Resource)
 {
-	RenderPass.SynchroniseResource(Resource);
+	RenderPass.SynchronizeResource(Resource);
 }
 
 void FAGXContext::FillBuffer(FAGXBuffer const& Buffer, ns::Range Range, uint8 Value)
diff --git a/Engine/Source/Runtime/Apple/AGXRHI/Private/AGXContext.h b/Engine/Source/Runtime/Apple/AGXRHI/Private/AGXContext.h
--- a/Engine/Source/Runtime/Apple/AGXRHI/Private/AGXContext.h
+++ b/Engine/Source/Runtime/Apple/AGXRHI/Private/AGXContext.h
@@ -103,7 +103,7 @@ public:
 	
 	void SynchronizeTexture(FAGXTexture const& Texture, uint32 Slice, uint32 Level);
 	
-	void SynchroniseResource(mtlpp::Resource const& Resource);
+	void SynchronizeResource(mtlpp::Resource const& Resource);
 	
 	void FillBuffer(FAGXBuffer const& Buffer, ns::Range Range, uint8 Value);
 
diff --git a/Engine/Source/Runtime/Apple/AGXRHI/Private/AGXRenderPass.cpp b/Engine/Source/Runtime/Apple/AGXRHI/Private/AGXRenderPass.cpp
--- a/Engine/Source/Runtime/Apple/AGXRHI/Private/AGXRenderPass.cpp
+++ b/Engine/Source/Runtime/Apple/AGXRHI/Private/AGXRenderPass.cpp
@@ -574,65 +574,123 @@ void FAGXRenderPass::InsertTextureBarrier()
 void FAGXRenderPass::CopyFromTextureToBuffer(FAGXTexture const& Texture, uint32 sourceSlice, uint32 sourceLevel, mtlpp::Origin sourceOrigin, mtlpp::Size sourceSize, FAGXBuffer const& toBuffer, uint32 destinationOffset, uint32 destinationBytesPerRow, uint32 destinationBytesPerImage, mtlpp::BlitOption options)
 {
 	ConditionalSwitchToBlit();
-	mtlpp::BlitCommandEncoder& Encoder = CurrentEncoder.GetBlitCommandEncoder();
-	check(Encoder.GetPtr());
+	
+	id<MTLBlitCommandEncoder> Encoder = CurrentEncoder.GetBlitCommandEncoder();
+	check(Encoder);
 	
 	METAL_GPUPROFILE(FAGXProfiler::GetProfiler()->EncodeBlit(CurrentEncoder.GetCommandBufferStats(), __FUNCTION__));
 	{
-		MTLPP_VALIDATE(mtlpp::BlitCommandEncoder, Encoder, AGXSafeGetRuntimeDebuggingLevel() >= EAGXDebugLevelValidation, Copy(Texture, sourceSlice, sourceLevel, sourceOrigin, sourceSize, toBuffer, destinationOffset, destinationBytesPerRow, destinationBytesPerImage, options));
+		[Encoder			 copyFromTexture:Texture.GetPtr()
+								 sourceSlice:(NSUInteger)sourceSlice
+								 sourceLevel:(NSUInteger)sourceLevel
+								sourceOrigin:*reinterpret_cast<MTLOrigin*>(&sourceOrigin)
+								  sourceSize:*reinterpret_cast<MTLSize*>(&sourceSize)
+									toBuffer:toBuffer.GetPtr()
+						   destinationOffset:(NSUInteger)destinationOffset + toBuffer.GetOffset()
+					  destinationBytesPerRow:(NSUInteger)destinationBytesPerRow
+					destinationBytesPerImage:(NSUInteger)destinationBytesPerImage
+									 options:*reinterpret_cast<MTLBlitOption*>(&options)];
 	}
+	
 	ConditionalSubmit();
 }
 
 void FAGXRenderPass::CopyFromBufferToTexture(FAGXBuffer const& Buffer, uint32 sourceOffset, uint32 sourceBytesPerRow, uint32 sourceBytesPerImage, mtlpp::Size sourceSize, FAGXTexture const& toTexture, uint32 destinationSlice, uint32 destinationLevel, mtlpp::Origin destinationOrigin, mtlpp::BlitOption options)
 {
 	ConditionalSwitchToBlit();
-	mtlpp::BlitCommandEncoder& Encoder = CurrentEncoder.GetBlitCommandEncoder();
-	check(Encoder.GetPtr());
 	
+	id<MTLBlitCommandEncoder> Encoder = CurrentEncoder.GetBlitCommandEncoder();
+	check(Encoder);
 	
 	METAL_GPUPROFILE(FAGXProfiler::GetProfiler()->EncodeBlit(CurrentEncoder.GetCommandBufferStats(), __FUNCTION__));
+	
 	if (options == mtlpp::BlitOption::None)
 	{
-		MTLPP_VALIDATE(mtlpp::BlitCommandEncoder, Encoder, AGXSafeGetRuntimeDebuggingLevel() >= EAGXDebugLevelValidation, Copy(Buffer, sourceOffset, sourceBytesPerRow, sourceBytesPerImage, sourceSize, toTexture, destinationSlice, destinationLevel, destinationOrigin));
+		[Encoder		 copyFromBuffer:Buffer.GetPtr()
+						   sourceOffset:(NSUInteger)sourceOffset + Buffer.GetOffset()
+					  sourceBytesPerRow:(NSUInteger)sourceBytesPerRow
+					sourceBytesPerImage:(NSUInteger)sourceBytesPerImage
+							 sourceSize:*reinterpret_cast<MTLSize*>(&sourceSize)
+							  toTexture:toTexture.GetPtr()
+					   destinationSlice:(NSUInteger)destinationSlice
+					   destinationLevel:(NSUInteger)destinationLevel
+					  destinationOrigin:*reinterpret_cast<MTLOrigin*>(&destinationOrigin)];
 	}
 	else
 	{
-		MTLPP_VALIDATE(mtlpp::BlitCommandEncoder, Encoder, AGXSafeGetRuntimeDebuggingLevel() >= EAGXDebugLevelValidation, Copy(Buffer, sourceOffset, sourceBytesPerRow, sourceBytesPerImage, sourceSize, toTexture, destinationSlice, destinationLevel, destinationOrigin, options));
+		[Encoder		 copyFromBuffer:Buffer.GetPtr()
+						   sourceOffset:(NSUInteger)sourceOffset + Buffer.GetOffset()
+					  sourceBytesPerRow:(NSUInteger)sourceBytesPerRow
+					sourceBytesPerImage:(NSUInteger)sourceBytesPerImage
+							 sourceSize:*reinterpret_cast<MTLSize*>(&sourceSize)
+							  toTexture:toTexture.GetPtr()
+					   destinationSlice:(NSUInteger)destinationSlice
+					   destinationLevel:(NSUInteger)destinationLevel
+					  destinationOrigin:*reinterpret_cast<MTLOrigin*>(&destinationOrigin)
+								options:*reinterpret_cast<MTLBlitOption*>(&options)];
 	}
+	
 	ConditionalSubmit();
 }
 
 void FAGXRenderPass::CopyFromTextureToTexture(FAGXTexture const& Texture, uint32 sourceSlice, uint32 sourceLevel, mtlpp::Origin sourceOrigin, mtlpp::Size sourceSize, FAGXTexture const& toTexture, uint32 destinationSlice, uint32 destinationLevel, mtlpp::Origin destinationOrigin)
 {
 	ConditionalSwitchToBlit();
-	mtlpp::BlitCommandEncoder& Encoder = CurrentEncoder.GetBlitCommandEncoder();
-	check(Encoder.GetPtr());
+	
+	id<MTLBlitCommandEncoder> Encoder = CurrentEncoder.GetBlitCommandEncoder();
+	check(Encoder);
 	
 	METAL_GPUPROFILE(FAGXProfiler::GetProfiler()->EncodeBlit(CurrentEncoder.GetCommandBufferStats(), __FUNCTION__));
-	MTLPP_VALIDATE(mtlpp::BlitCommandEncoder, Encoder, AGXSafeGetRuntimeDebuggingLevel() >= EAGXDebugLevelValidation, Copy(Texture, sourceSlice, sourceLevel, sourceOrigin, sourceSize, toTexture, destinationSlice, destinationLevel, destinationOrigin));
+	
+	[Encoder	  copyFromTexture:Texture.GetPtr()
+					  sourceSlice:(NSUInteger)sourceSlice
+					  sourceLevel:(NSUInteger)sourceLevel
+					 sourceOrigin:*reinterpret_cast<MTLOrigin*>(&sourceOrigin)
+					   sourceSize:*reinterpret_cast<MTLSize*>(&sourceSize)
+						toTexture:toTexture.GetPtr()
+				 destinationSlice:(NSUInteger)destinationSlice
+				 destinationLevel:(NSUInteger)destinationLevel
+				destinationOrigin:*reinterpret_cast<MTLOrigin*>(&destinationOrigin)];
+	
 	ConditionalSubmit();
 }
 
 void FAGXRenderPass::CopyFromBufferToBuffer(FAGXBuffer const& SourceBuffer, NSUInteger SourceOffset, FAGXBuffer const& DestinationBuffer, NSUInteger DestinationOffset, NSUInteger Size)
 {
 	ConditionalSwitchToBlit();
-	mtlpp::BlitCommandEncoder& Encoder = CurrentEncoder.GetBlitCommandEncoder();
-	check(Encoder.GetPtr());
+	
+	id<MTLBlitCommandEncoder> Encoder = CurrentEncoder.GetBlitCommandEncoder();
+	check(Encoder);
 	
 	METAL_GPUPROFILE(FAGXProfiler::GetProfiler()->EncodeBlit(CurrentEncoder.GetCommandBufferStats(), __FUNCTION__));
-	MTLPP_VALIDATE(mtlpp::BlitCommandEncoder, Encoder, AGXSafeGetRuntimeDebuggingLevel() >= EAGXDebugLevelValidation, Copy(SourceBuffer, SourceOffset, DestinationBuffer, DestinationOffset, Size));
+
+	[Encoder	   copyFromBuffer:SourceBuffer.GetPtr()
+					 sourceOffset:SourceOffset + SourceBuffer.GetOffset()
+						 toBuffer:DestinationBuffer.GetPtr()
+				destinationOffset:DestinationOffset + DestinationBuffer.GetOffset()
+							 size:Size];
+	
 	ConditionalSubmit();
 }
 
 void FAGXRenderPass::PresentTexture(FAGXTexture const& Texture, uint32 sourceSlice, uint32 sourceLevel, mtlpp::Origin sourceOrigin, mtlpp::Size sourceSize, FAGXTexture const& toTexture, uint32 destinationSlice, uint32 destinationLevel, mtlpp::Origin destinationOrigin)
 {
 	ConditionalSwitchToBlit();
-	mtlpp::BlitCommandEncoder& Encoder = CurrentEncoder.GetBlitCommandEncoder();
-	check(Encoder.GetPtr());
+	
+	id<MTLBlitCommandEncoder> Encoder = CurrentEncoder.GetBlitCommandEncoder();
+	check(Encoder);
 	
 	METAL_GPUPROFILE(FAGXProfiler::GetProfiler()->EncodeBlit(CurrentEncoder.GetCommandBufferStats(), __FUNCTION__));
-	MTLPP_VALIDATE(mtlpp::BlitCommandEncoder, Encoder, AGXSafeGetRuntimeDebuggingLevel() >= EAGXDebugLevelValidation, Copy(Texture, sourceSlice, sourceLevel, sourceOrigin, sourceSize, toTexture, destinationSlice, destinationLevel, destinationOrigin));
+
+	[Encoder	  copyFromTexture:Texture.GetPtr()
+					  sourceSlice:(NSUInteger)sourceSlice
+					  sourceLevel:(NSUInteger)sourceLevel
+					 sourceOrigin:*reinterpret_cast<MTLOrigin*>(&sourceOrigin)
+					   sourceSize:*reinterpret_cast<MTLSize*>(&sourceSize)
+						toTexture:toTexture.GetPtr()
+				 destinationSlice:(NSUInteger)destinationSlice
+				 destinationLevel:(NSUInteger)destinationLevel
+				destinationOrigin:*reinterpret_cast<MTLOrigin*>(&destinationOrigin)];
 }
 
 void FAGXRenderPass::SynchronizeTexture(FAGXTexture const& Texture, uint32 Slice, uint32 Level)
@@ -640,25 +698,29 @@ void FAGXRenderPass::SynchronizeTexture(FAGXTexture const& Texture, uint32 Slice
 	check(Texture);
 #if PLATFORM_MAC
 	ConditionalSwitchToBlit();
-	mtlpp::BlitCommandEncoder& Encoder = CurrentEncoder.GetBlitCommandEncoder();
-	check(Encoder.GetPtr());
+	
+	id<MTLBlitCommandEncoder> Encoder = CurrentEncoder.GetBlitCommandEncoder();
+	check(Encoder);
 	
 	// METAL_GPUPROFILE(FAGXProfiler::GetProfiler()->EncodeBlit(CurrentEncoder.GetCommandBufferStats(), __FUNCTION__));
-	MTLPP_VALIDATE(mtlpp::BlitCommandEncoder, Encoder, AGXSafeGetRuntimeDebuggingLevel() >= EAGXDebugLevelValidation, Synchronize(Texture, Slice, Level));
+	[Encoder synchronizeTexture:Texture.GetPtr() slice:(NSUInteger)Slice level:(NSUInteger)Level];
+	
 	ConditionalSubmit();
 #endif
 }
 
-void FAGXRenderPass::SynchroniseResource(mtlpp::Resource const& Resource)
+void FAGXRenderPass::SynchronizeResource(mtlpp::Resource const& Resource)
 {
 	check(Resource);
 #if PLATFORM_MAC
 	ConditionalSwitchToBlit();
-	mtlpp::BlitCommandEncoder& Encoder = CurrentEncoder.GetBlitCommandEncoder();
-	check(Encoder.GetPtr());
+	
+	id<MTLBlitCommandEncoder> Encoder = CurrentEncoder.GetBlitCommandEncoder();
+	check(Encoder);
 	
 	// METAL_GPUPROFILE(FAGXProfiler::GetProfiler()->EncodeBlit(CurrentEncoder.GetCommandBufferStats(), __FUNCTION__));
-	MTLPP_VALIDATE(mtlpp::BlitCommandEncoder, Encoder, AGXSafeGetRuntimeDebuggingLevel() >= EAGXDebugLevelValidation, Synchronize(Resource));
+	[Encoder synchronizeResource:Resource.GetPtr()];
+
 	ConditionalSubmit();
 #endif
 }
@@ -667,7 +729,7 @@ void FAGXRenderPass::FillBuffer(FAGXBuffer const& Buffer, ns::Range Range, uint8
 {
 	check(Buffer);
 	
-	mtlpp::BlitCommandEncoder TargetEncoder;
+	id<MTLBlitCommandEncoder> TargetEncoder = nil;
 	bool bAsync = !CurrentEncoder.HasBufferBindingHistory(Buffer);
 	if(bAsync)
 	{
@@ -682,9 +744,9 @@ void FAGXRenderPass::FillBuffer(FAGXBuffer const& Buffer, ns::Range Range, uint8
 		METAL_GPUPROFILE(FAGXProfiler::GetProfiler()->EncodeBlit(CurrentEncoder.GetCommandBufferStats(), FString::Printf(TEXT("FillBuffer: %p %llu %llu"), Buffer.GetPtr(), Buffer.GetOffset() + Range.Location, Range.Length)));
 	}
 	
-	check(TargetEncoder.GetPtr());
+	check(TargetEncoder);
 	
-	MTLPP_VALIDATE(mtlpp::BlitCommandEncoder, TargetEncoder, AGXSafeGetRuntimeDebuggingLevel() >= EAGXDebugLevelValidation, Fill(Buffer, Range, Value));
+	[TargetEncoder fillBuffer:Buffer.GetPtr() range:NSMakeRange(Range.Location + Buffer.GetOffset(), Range.Length) value:Value];
 	
 	if (!bAsync)
 	{
@@ -694,7 +756,7 @@ void FAGXRenderPass::FillBuffer(FAGXBuffer const& Buffer, ns::Range Range, uint8
 
 bool FAGXRenderPass::AsyncCopyFromBufferToTexture(FAGXBuffer const& Buffer, uint32 sourceOffset, uint32 sourceBytesPerRow, uint32 sourceBytesPerImage, mtlpp::Size sourceSize, FAGXTexture const& toTexture, uint32 destinationSlice, uint32 destinationLevel, mtlpp::Origin destinationOrigin, mtlpp::BlitOption options)
 {
-	mtlpp::BlitCommandEncoder TargetEncoder;
+	id<MTLBlitCommandEncoder> TargetEncoder = nil;
 	bool bAsync = !CurrentEncoder.HasTextureBindingHistory(toTexture);
 	if(bAsync)
 	{
@@ -709,15 +771,32 @@ bool FAGXRenderPass::AsyncCopyFromBufferToTexture(FAGXBuffer const& Buffer, uint
 		METAL_GPUPROFILE(FAGXProfiler::GetProfiler()->EncodeBlit(CurrentEncoder.GetCommandBufferStats(), __FUNCTION__));
 	}
 	
-	check(TargetEncoder.GetPtr());
+	check(TargetEncoder);
 	
 	if (options == mtlpp::BlitOption::None)
 	{
-		MTLPP_VALIDATE(mtlpp::BlitCommandEncoder, TargetEncoder, AGXSafeGetRuntimeDebuggingLevel() >= EAGXDebugLevelValidation, Copy(Buffer, sourceOffset, sourceBytesPerRow, sourceBytesPerImage, sourceSize, toTexture, destinationSlice, destinationLevel, destinationOrigin));
+		[TargetEncoder		 copyFromBuffer:Buffer.GetPtr()
+							   sourceOffset:(NSUInteger)sourceOffset + Buffer.GetOffset()
+						  sourceBytesPerRow:(NSUInteger)sourceBytesPerRow
+						sourceBytesPerImage:(NSUInteger)sourceBytesPerImage
+								 sourceSize:*reinterpret_cast<MTLSize*>(&sourceSize)
+								  toTexture:toTexture.GetPtr()
+						   destinationSlice:(NSUInteger)destinationSlice
+						   destinationLevel:(NSUInteger)destinationLevel
+						  destinationOrigin:*reinterpret_cast<MTLOrigin*>(&destinationOrigin)];
 	}
 	else
 	{
-		MTLPP_VALIDATE(mtlpp::BlitCommandEncoder, TargetEncoder, AGXSafeGetRuntimeDebuggingLevel() >= EAGXDebugLevelValidation, Copy(Buffer, sourceOffset, sourceBytesPerRow, sourceBytesPerImage, sourceSize, toTexture, destinationSlice, destinationLevel, destinationOrigin, options));
+		[TargetEncoder		 copyFromBuffer:Buffer.GetPtr()
+							   sourceOffset:(NSUInteger)sourceOffset + Buffer.GetOffset()
+						  sourceBytesPerRow:(NSUInteger)sourceBytesPerRow
+						sourceBytesPerImage:(NSUInteger)sourceBytesPerImage
+								 sourceSize:*reinterpret_cast<MTLSize*>(&sourceSize)
+								  toTexture:toTexture.GetPtr()
+						   destinationSlice:(NSUInteger)destinationSlice
+						   destinationLevel:(NSUInteger)destinationLevel
+						  destinationOrigin:*reinterpret_cast<MTLOrigin*>(&destinationOrigin)
+									options:*reinterpret_cast<MTLBlitOption*>(&options)];
 	}
 	
 	return bAsync;
@@ -725,7 +804,7 @@ bool FAGXRenderPass::AsyncCopyFromBufferToTexture(FAGXBuffer const& Buffer, uint
 
 bool FAGXRenderPass::AsyncCopyFromTextureToTexture(FAGXTexture const& Texture, uint32 sourceSlice, uint32 sourceLevel, mtlpp::Origin sourceOrigin, mtlpp::Size sourceSize, FAGXTexture const& toTexture, uint32 destinationSlice, uint32 destinationLevel, mtlpp::Origin destinationOrigin)
 {
-	mtlpp::BlitCommandEncoder TargetEncoder;
+	id<MTLBlitCommandEncoder> TargetEncoder = nil;
 	bool bAsync = !CurrentEncoder.HasTextureBindingHistory(toTexture);
 	if(bAsync)
 	{
@@ -740,10 +819,18 @@ bool FAGXRenderPass::AsyncCopyFromTextureToTexture(FAGXTexture const& Texture, u
 		METAL_GPUPROFILE(FAGXProfiler::GetProfiler()->EncodeBlit(CurrentEncoder.GetCommandBufferStats(), __FUNCTION__));
 	}
 	
-	check(TargetEncoder.GetPtr());
-	
-	MTLPP_VALIDATE(mtlpp::BlitCommandEncoder, TargetEncoder, AGXSafeGetRuntimeDebuggingLevel() >= EAGXDebugLevelValidation, Copy(Texture, sourceSlice, sourceLevel, sourceOrigin, sourceSize, toTexture, destinationSlice, destinationLevel, destinationOrigin));
+	check(TargetEncoder);
 	
+	[TargetEncoder	  copyFromTexture:Texture.GetPtr()
+						  sourceSlice:(NSUInteger)sourceSlice
+						  sourceLevel:(NSUInteger)sourceLevel
+						 sourceOrigin:*reinterpret_cast<MTLOrigin*>(&sourceOrigin)
+						   sourceSize:*reinterpret_cast<MTLSize*>(&sourceSize)
+							toTexture:toTexture.GetPtr()
+					 destinationSlice:(NSUInteger)destinationSlice
+					 destinationLevel:(NSUInteger)destinationLevel
+					destinationOrigin:*reinterpret_cast<MTLOrigin*>(&destinationOrigin)];
+
 	return bAsync;
 }
 
@@ -754,7 +841,7 @@ bool FAGXRenderPass::CanAsyncCopyToBuffer(FAGXBuffer const& DestinationBuffer)
 
 void FAGXRenderPass::AsyncCopyFromBufferToBuffer(FAGXBuffer const& SourceBuffer, NSUInteger SourceOffset, FAGXBuffer const& DestinationBuffer, NSUInteger DestinationOffset, NSUInteger Size)
 {
-	mtlpp::BlitCommandEncoder TargetEncoder;
+	id<MTLBlitCommandEncoder> TargetEncoder = nil;
 	bool bAsync = !CurrentEncoder.HasBufferBindingHistory(DestinationBuffer);
 	if(bAsync)
 	{
@@ -769,9 +856,13 @@ void FAGXRenderPass::AsyncCopyFromBufferToBuffer(FAGXBuffer const& SourceBuffer,
 		METAL_GPUPROFILE(FAGXProfiler::GetProfiler()->EncodeBlit(CurrentEncoder.GetCommandBufferStats(), FString::Printf(TEXT("AsyncCopyFromBufferToBuffer: %p %llu %llu"), DestinationBuffer.GetPtr(), DestinationBuffer.GetOffset() + DestinationOffset, Size)));
 	}
 	
-	check(TargetEncoder.GetPtr());
+	check(TargetEncoder);
 	
-    MTLPP_VALIDATE(mtlpp::BlitCommandEncoder, TargetEncoder, AGXSafeGetRuntimeDebuggingLevel() >= EAGXDebugLevelValidation, Copy(SourceBuffer, SourceOffset, DestinationBuffer, DestinationOffset, Size));
+	[TargetEncoder	   copyFromBuffer:SourceBuffer.GetPtr()
+						 sourceOffset:SourceOffset + SourceBuffer.GetOffset()
+							 toBuffer:DestinationBuffer.GetPtr()
+					destinationOffset:DestinationOffset + DestinationBuffer.GetOffset()
+								 size:Size];
 }
 
 FAGXBuffer FAGXRenderPass::AllocateTemporyBufferForCopy(FAGXBuffer const& DestinationBuffer, NSUInteger Size, NSUInteger Align)
@@ -794,11 +885,13 @@ void FAGXRenderPass::AsyncGenerateMipmapsForTexture(FAGXTexture const& Texture)
 	// This must be a plain old error
 	check(!CurrentEncoder.HasTextureBindingHistory(Texture));
 	ConditionalSwitchToAsyncBlit();
-	mtlpp::BlitCommandEncoder Encoder = PrologueEncoder.GetBlitCommandEncoder();
-	check(Encoder.GetPtr());
+	
+	id<MTLBlitCommandEncoder> Encoder = PrologueEncoder.GetBlitCommandEncoder();
+	check(Encoder);
 	
 	METAL_GPUPROFILE(FAGXProfiler::GetProfiler()->EncodeBlit(CurrentEncoder.GetCommandBufferStats(), __FUNCTION__));
-	MTLPP_VALIDATE(mtlpp::BlitCommandEncoder, Encoder, AGXSafeGetRuntimeDebuggingLevel() >= EAGXDebugLevelValidation, GenerateMipmaps(Texture));
+
+	[Encoder generateMipmapsForTexture:Texture.GetPtr()];
 }
 
 void FAGXRenderPass::End(void)
@@ -863,17 +956,16 @@ void FAGXRenderPass::TransitionResources(mtlpp::Resource const& Resource)
 void FAGXRenderPass::InsertDebugEncoder()
 {
 	FAGXBuffer NewBuf = CurrentEncoder.GetRingBuffer().NewBuffer(BufferOffsetAlignment, BufferOffsetAlignment);
-	
 	check(NewBuf);
 	
-	mtlpp::BlitCommandEncoder TargetEncoder;
 	ConditionalSwitchToBlit();
-	TargetEncoder = CurrentEncoder.GetBlitCommandEncoder();
-	METAL_GPUPROFILE(FAGXProfiler::GetProfiler()->EncodeBlit(CurrentEncoder.GetCommandBufferStats(), __FUNCTION__));
 	
-	check(TargetEncoder.GetPtr());
+	id<MTLBlitCommandEncoder> TargetEncoder = CurrentEncoder.GetBlitCommandEncoder();
+	check(TargetEncoder);
+	
+	METAL_GPUPROFILE(FAGXProfiler::GetProfiler()->EncodeBlit(CurrentEncoder.GetCommandBufferStats(), __FUNCTION__));
 	
-	MTLPP_VALIDATE(mtlpp::BlitCommandEncoder, TargetEncoder, AGXSafeGetRuntimeDebuggingLevel() >= EAGXDebugLevelValidation, Fill(NewBuf, ns::Range(0, BufferOffsetAlignment), 0xff));
+	[TargetEncoder fillBuffer:NewBuf.GetPtr() range:NSMakeRange(0 + NewBuf.GetOffset(), BufferOffsetAlignment) value:0xFF];
 	
 	ConditionalSubmit();
 }
diff --git a/Engine/Source/Runtime/Apple/AGXRHI/Private/AGXRenderPass.h b/Engine/Source/Runtime/Apple/AGXRHI/Private/AGXRenderPass.h
--- a/Engine/Source/Runtime/Apple/AGXRHI/Private/AGXRenderPass.h
+++ b/Engine/Source/Runtime/Apple/AGXRHI/Private/AGXRenderPass.h
@@ -60,7 +60,7 @@ public:
     
     void SynchronizeTexture(FAGXTexture const& Texture, uint32 Slice, uint32 Level);
     
-	void SynchroniseResource(mtlpp::Resource const& Resource);
+	void SynchronizeResource(mtlpp::Resource const& Resource);
     
 	void FillBuffer(FAGXBuffer const& Buffer, ns::Range Range, uint8 Value);
 	
diff --git a/Engine/Source/Runtime/Apple/AGXRHI/Private/AGXStateCache.cpp b/Engine/Source/Runtime/Apple/AGXRHI/Private/AGXStateCache.cpp
--- a/Engine/Source/Runtime/Apple/AGXRHI/Private/AGXStateCache.cpp
+++ b/Engine/Source/Runtime/Apple/AGXRHI/Private/AGXStateCache.cpp
@@ -1859,10 +1859,10 @@ void FAGXStateCache::FlushVisibilityResults(FAGXCommandEncoder& CommandEncoder)
 		
         CommandEncoder.BeginBlitCommandEncoding();
 		
-		mtlpp::BlitCommandEncoder& Encoder = CommandEncoder.GetBlitCommandEncoder();
+		id<MTLBlitCommandEncoder> Encoder = CommandEncoder.GetBlitCommandEncoder();
 
 		// METAL_GPUPROFILE(FAGXProfiler::GetProfiler()->EncodeBlit(CommandEncoder.GetCommandBufferStats(), __FUNCTION__));
-		MTLPP_VALIDATE(mtlpp::BlitCommandEncoder, Encoder, AGXSafeGetRuntimeDebuggingLevel() >= EAGXDebugLevelValidation, Synchronize(VisibilityResults->Buffer));
+		[Encoder synchronizeResource:VisibilityResults->Buffer.GetPtr()];
 		
 		VisibilityWritten = 0;
 	}
diff --git a/Engine/Source/Runtime/Apple/AGXRHI/Private/AGXVertexBuffer.cpp b/Engine/Source/Runtime/Apple/AGXRHI/Private/AGXVertexBuffer.cpp
--- a/Engine/Source/Runtime/Apple/AGXRHI/Private/AGXVertexBuffer.cpp
+++ b/Engine/Source/Runtime/Apple/AGXRHI/Private/AGXVertexBuffer.cpp
@@ -601,7 +601,7 @@ void* FAGXRHIBuffer::Lock(bool bIsOnRHIThread, EResourceLockMode InLockMode, uin
 			SCOPE_CYCLE_COUNTER(STAT_AGXBufferPageOffTime);
 			
 			// Synchronise the buffer with the CPU
-			GetAGXDeviceContext().SynchroniseResource(GetCurrentBacking().Buffer);
+			GetAGXDeviceContext().SynchronizeResource(GetCurrentBacking().Buffer);
 			
 			//kick the current command buffer.
 			GetAGXDeviceContext().SubmitCommandBufferAndWait();
