Port changes to Metal RHI from CL 19811451 to AGXRHI

---

Fixes for Lumen Scene rendering on Metal.

- Fix Metal buffer pool policy to separate by storage.

#rnx

diff --git a/Engine/Source/Runtime/Apple/AGXRHI/Private/AGXBuffer.cpp b/Engine/Source/Runtime/Apple/AGXRHI/Private/AGXBuffer.cpp
--- a/Engine/Source/Runtime/Apple/AGXRHI/Private/AGXBuffer.cpp
+++ b/Engine/Source/Runtime/Apple/AGXRHI/Private/AGXBuffer.cpp
@@ -927,14 +927,13 @@ uint32 FAGXBufferPoolPolicyData::GetPoolBucketIndex(CreationArguments Args)
 	check( Size <= BucketSizes[Lower] );
 	check( (Lower == 0 ) || ( Size > BucketSizes[Lower-1] ) );
 	
-	return Lower;
+	return (int32)Args.Storage * NumPoolBucketSizes + Lower;
 }
 
 uint32 FAGXBufferPoolPolicyData::GetPoolBucketSize(uint32 Bucket)
 {
 	check(Bucket < NumPoolBuckets);
-	uint32 Index = Bucket;
-	checkf(Index < NumPoolBucketSizes, TEXT("%d %d"), Index, NumPoolBucketSizes);
+	uint32 Index = Bucket % NumPoolBucketSizes;
 	return BucketSizes[Index];
 }
 
diff --git a/Engine/Source/Runtime/Apple/AGXRHI/Private/AGXBuffer.h b/Engine/Source/Runtime/Apple/AGXRHI/Private/AGXBuffer.h
--- a/Engine/Source/Runtime/Apple/AGXRHI/Private/AGXBuffer.h
+++ b/Engine/Source/Runtime/Apple/AGXRHI/Private/AGXBuffer.h
@@ -209,9 +209,10 @@ public:
 	typedef FAGXPooledBufferArgs CreationArguments;
 	enum
 	{
+		NumResourceStorageModes = 4, /* Corresponds to MTLStorageMode types: managed, shared, private, memoryless */
 		NumSafeFrames = 1, /** Number of frames to leave buffers before reclaiming/reusing */
 		NumPoolBucketSizes = NumBucketSizes, /** Number of pool bucket sizes */
-		NumPoolBuckets = NumPoolBucketSizes, /** Number of pool bucket sizes - all entries must use consistent ResourceOptions */
+		NumPoolBuckets = NumPoolBucketSizes * NumResourceStorageModes, /** Number of pool bucket sizes - all entries must use consistent ResourceOptions, so the total number of pool buckets is the number of pool bucket sizes x the number of resource storage modes */
 		NumToDrainPerFrame = 65536, /** Max. number of resources to cull in a single frame */
 		CullAfterFramesNum = 30 /** Resources are culled if unused for more frames than this */
 	};
