AGXRHI: Scrub all usage of mtlpp::ComputeCommandEncoder

#rnx

diff --git a/Engine/Source/Runtime/Apple/AGXRHI/Private/AGXCommandEncoder.cpp b/Engine/Source/Runtime/Apple/AGXRHI/Private/AGXCommandEncoder.cpp
--- a/Engine/Source/Runtime/Apple/AGXRHI/Private/AGXCommandEncoder.cpp
+++ b/Engine/Source/Runtime/Apple/AGXRHI/Private/AGXCommandEncoder.cpp
@@ -25,6 +25,7 @@ FAGXCommandEncoder::FAGXCommandEncoder(FAGXCommandList& CmdList, EAGXCommandEnco
 , bSupportsMetalFeaturesSetBytes(CmdList.GetCommandQueue().SupportsFeature(EAGXFeaturesSetBytes))
 , RingBuffer(EncoderRingBufferSize, BufferOffsetAlignment, FAGXCommandQueue::GetCompatibleResourceOptions((mtlpp::ResourceOptions)(mtlpp::ResourceOptions::HazardTrackingModeUntracked | BUFFER_RESOURCE_STORAGE_MANAGED)))
 , RenderPassDesc(nil)
+, ComputeCommandEncoder(nil)
 , BlitCommandEncoder(nil)
 #if ENABLE_METAL_GPUPROFILE
 , CommandBufferStats(nullptr)
@@ -160,12 +161,12 @@ void FAGXCommandEncoder::ResetLive(void)
 	{
 		for (uint32 i = 0; i < ML_MaxBuffers; i++)
 		{
-			ComputeCommandEncoder.SetBuffer(nil, 0, i);
+			[ComputeCommandEncoder setBuffer:nil offset:0 atIndex:i];
 		}
 		
 		for (uint32 i = 0; i < ML_MaxTextures; i++)
 		{
-			ComputeCommandEncoder.SetTexture(nil, i);
+			[ComputeCommandEncoder setTexture:nil atIndex:i];
 		}
 	}
 }
@@ -280,7 +281,7 @@ bool FAGXCommandEncoder::IsRenderCommandEncoderActive(void) const
 
 bool FAGXCommandEncoder::IsComputeCommandEncoderActive(void) const
 {
-	return ComputeCommandEncoder.GetPtr() != nil;
+	return ComputeCommandEncoder != nil;
 }
 
 bool FAGXCommandEncoder::IsBlitCommandEncoderActive(void) const
@@ -325,7 +326,7 @@ mtlpp::RenderCommandEncoder& FAGXCommandEncoder::GetRenderCommandEncoder(void)
 	return RenderCommandEncoder;
 }
 
-mtlpp::ComputeCommandEncoder& FAGXCommandEncoder::GetComputeCommandEncoder(void)
+id<MTLComputeCommandEncoder> FAGXCommandEncoder::GetComputeCommandEncoder() const
 {
 	check(IsComputeCommandEncoderActive());
 	return ComputeCommandEncoder;
@@ -413,11 +414,11 @@ void FAGXCommandEncoder::BeginComputeCommandEncoding(mtlpp::DispatchType Dispatc
 	
 	if (DispatchType == mtlpp::DispatchType::Serial)
 	{
-		ComputeCommandEncoder = MTLPP_VALIDATE(mtlpp::CommandBuffer, CommandBuffer, AGXSafeGetRuntimeDebuggingLevel() >= EAGXDebugLevelValidation, ComputeCommandEncoder());
+		ComputeCommandEncoder = [[CommandBuffer.GetPtr() computeCommandEncoder] retain];
 	}
 	else
 	{
-		ComputeCommandEncoder = MTLPP_VALIDATE(mtlpp::CommandBuffer, CommandBuffer, AGXSafeGetRuntimeDebuggingLevel() >= EAGXDebugLevelValidation, ComputeCommandEncoder(DispatchType));
+		ComputeCommandEncoder = [[CommandBuffer.GetPtr() computeCommandEncoderWithDispatchType:*reinterpret_cast<MTLDispatchType*>(&DispatchType)] retain];
 	}
 
 	EncoderNum++;
@@ -427,13 +428,13 @@ void FAGXCommandEncoder::BeginComputeCommandEncoding(mtlpp::DispatchType Dispatc
 	if(GetEmitDrawEvents())
 	{
 		Label = [NSString stringWithFormat:@"ComputeEncoder: %@", [DebugGroups count] > 0 ? [DebugGroups lastObject] : (NSString*)CFSTR("InitialPass")];
-		ComputeCommandEncoder.SetLabel(Label);
+		[ComputeCommandEncoder setLabel:Label];
 		
 		if([DebugGroups count])
 		{
 			for (NSString* Group in DebugGroups)
 			{
-				ComputeCommandEncoder.PushDebugGroup(Group);
+				[ComputeCommandEncoder pushDebugGroup:Group];
 			}
 		}
 	}
@@ -574,7 +575,8 @@ void FAGXCommandEncoder::EndEncoding(void)
 		}
 		else if(IsComputeCommandEncoderActive())
 		{
-			ComputeCommandEncoder.EndEncoding();
+			[ComputeCommandEncoder endEncoding];
+			[ComputeCommandEncoder release];
 			ComputeCommandEncoder = nil;
 		}
 		else if(IsBlitCommandEncoderActive())
@@ -637,7 +639,7 @@ void FAGXCommandEncoder::InsertDebugSignpost(ns::String const& String)
 		}
 		else if (ComputeCommandEncoder)
 		{
-			ComputeCommandEncoder.InsertDebugSignpost(String);
+			[ComputeCommandEncoder insertDebugSignpost:String];
 		}
 		else if (BlitCommandEncoder)
 		{
@@ -662,7 +664,7 @@ void FAGXCommandEncoder::PushDebugGroup(ns::String const& String)
 		}
 		else if (ComputeCommandEncoder)
 		{
-			ComputeCommandEncoder.PushDebugGroup(String);
+			[ComputeCommandEncoder pushDebugGroup:String];
 		}
 		else if (BlitCommandEncoder)
 		{
@@ -687,7 +689,7 @@ void FAGXCommandEncoder::PopDebugGroup(void)
 		}
 		else if (ComputeCommandEncoder)
 		{
-			ComputeCommandEncoder.PopDebugGroup();
+			[ComputeCommandEncoder popDebugGroup];
 		}
 		else if (BlitCommandEncoder)
 		{
@@ -951,7 +953,7 @@ void FAGXCommandEncoder::SetShaderBytes(mtlpp::FunctionType const FunctionType,
 					break;
 				case mtlpp::FunctionType::Kernel:
 					check(ComputeCommandEncoder);
-					ComputeCommandEncoder.SetBytes(Bytes, Length, Index);
+					[ComputeCommandEncoder setBytes:static_cast<const void*>(Bytes) length:Length atIndex:Index];
 					break;
 				default:
 					check(false);
@@ -1004,7 +1006,7 @@ void FAGXCommandEncoder::SetShaderBufferOffset(mtlpp::FunctionType FunctionType,
 			break;
 		case mtlpp::FunctionType::Kernel:
 			check (ComputeCommandEncoder);
-			ComputeCommandEncoder.SetBufferOffset(Offset + ShaderBuffers[uint32(FunctionType)].Buffers[index].GetOffset(), index);
+			[ComputeCommandEncoder setBufferOffset:(Offset + ShaderBuffers[uint32(FunctionType)].Buffers[index].GetOffset()) atIndex:index];
 			break;
 		default:
 			check(false);
@@ -1029,7 +1031,7 @@ void FAGXCommandEncoder::SetShaderTexture(mtlpp::FunctionType FunctionType, FAGX
 		case mtlpp::FunctionType::Kernel:
 			check (ComputeCommandEncoder);
 			FenceResource(Texture);
-			ComputeCommandEncoder.SetTexture(Texture, index);
+			[ComputeCommandEncoder setTexture:Texture.GetPtr() atIndex:index];
 			break;
 		default:
 			check(false);
@@ -1068,7 +1070,7 @@ void FAGXCommandEncoder::SetShaderSamplerState(mtlpp::FunctionType FunctionType,
 			break;
 		case mtlpp::FunctionType::Kernel:
 			check (ComputeCommandEncoder);
-			ComputeCommandEncoder.SetSamplerState(Sampler, index);
+			[ComputeCommandEncoder setSamplerState:Sampler.GetPtr() atIndex:index];
 			break;
 		default:
 			check(false);
@@ -1094,7 +1096,7 @@ void FAGXCommandEncoder::SetComputePipelineState(FAGXShaderPipeline* State)
 {
 	check (ComputeCommandEncoder);
 	{
-		ComputeCommandEncoder.SetComputePipelineState(State->ComputePipelineState);
+		[ComputeCommandEncoder setComputePipelineState:State->ComputePipelineState.GetPtr()];
 	}
 }
 
@@ -1174,7 +1176,7 @@ void FAGXCommandEncoder::SetShaderBufferInternal(mtlpp::FunctionType Function, u
 				Binding.Bound |= (1 << Index);
 				check(ComputeCommandEncoder);
 				FenceResource(Buffer);
-				ComputeCommandEncoder.SetBuffer(Buffer, Offset, Index);
+				[ComputeCommandEncoder setBuffer:Buffer.GetPtr() offset:(Offset + Buffer.GetOffset()) atIndex:Index];
 				break;
 
 			default:
@@ -1212,7 +1214,7 @@ void FAGXCommandEncoder::SetShaderBufferInternal(mtlpp::FunctionType Function, u
 			case mtlpp::FunctionType::Kernel:
 				Binding.Bound |= (1 << Index);
 				check(ComputeCommandEncoder);
-				ComputeCommandEncoder.SetBytes(Bytes, Len, Index);
+				[ComputeCommandEncoder setBytes:static_cast<const void*>(Bytes) length:(NSUInteger)Len atIndex:Index];
 				break;
 
 			default:
diff --git a/Engine/Source/Runtime/Apple/AGXRHI/Private/AGXCommandEncoder.h b/Engine/Source/Runtime/Apple/AGXRHI/Private/AGXCommandEncoder.h
--- a/Engine/Source/Runtime/Apple/AGXRHI/Private/AGXCommandEncoder.h
+++ b/Engine/Source/Runtime/Apple/AGXRHI/Private/AGXCommandEncoder.h
@@ -137,7 +137,7 @@ public:
 	mtlpp::RenderCommandEncoder& GetRenderCommandEncoder(void);
 	
 	/** @returns The active compute command encoder or nil if there isn't one. */
-	mtlpp::ComputeCommandEncoder& GetComputeCommandEncoder(void);
+	id<MTLComputeCommandEncoder> GetComputeCommandEncoder() const;
 	
 	/** @returns The active blit command encoder or nil if there isn't one. */
 	id<MTLBlitCommandEncoder> GetBlitCommandEncoder() const;
@@ -459,7 +459,7 @@ public:
 	mtlpp::CommandBuffer CommandBuffer;
 	mtlpp::ParallelRenderCommandEncoder ParallelRenderCommandEncoder;
 	mtlpp::RenderCommandEncoder RenderCommandEncoder;
-	mtlpp::ComputeCommandEncoder ComputeCommandEncoder;
+	id<MTLComputeCommandEncoder> ComputeCommandEncoder;
 	id<MTLBlitCommandEncoder> BlitCommandEncoder;
 	TArray<mtlpp::RenderCommandEncoder> ChildRenderCommandEncoders;
 	
diff --git a/Engine/Source/Runtime/Apple/AGXRHI/Private/AGXRenderPass.cpp b/Engine/Source/Runtime/Apple/AGXRHI/Private/AGXRenderPass.cpp
--- a/Engine/Source/Runtime/Apple/AGXRHI/Private/AGXRenderPass.cpp
+++ b/Engine/Source/Runtime/Apple/AGXRHI/Private/AGXRenderPass.cpp
@@ -474,10 +474,10 @@ void FAGXRenderPass::Dispatch(uint32 ThreadGroupCountX, uint32 ThreadGroupCountY
 		
 		METAL_GPUPROFILE(FAGXProfiler::GetProfiler()->EncodeDispatch(PrologueEncoder.GetCommandBufferStats(), __FUNCTION__));
 		
-		mtlpp::Size ThreadgroupCounts = mtlpp::Size(ComputeShader->NumThreadsX, ComputeShader->NumThreadsY, ComputeShader->NumThreadsZ);
+		MTLSize ThreadgroupCounts = MTLSizeMake(ComputeShader->NumThreadsX, ComputeShader->NumThreadsY, ComputeShader->NumThreadsZ);
 		check(ComputeShader->NumThreadsX > 0 && ComputeShader->NumThreadsY > 0 && ComputeShader->NumThreadsZ > 0);
-		mtlpp::Size Threadgroups = mtlpp::Size(ThreadGroupCountX, ThreadGroupCountY, ThreadGroupCountZ);
-		PrologueEncoder.GetComputeCommandEncoder().DispatchThreadgroups(Threadgroups, ThreadgroupCounts);
+		MTLSize Threadgroups = MTLSizeMake(ThreadGroupCountX, ThreadGroupCountY, ThreadGroupCountZ);
+		[PrologueEncoder.GetComputeCommandEncoder() dispatchThreadgroups:Threadgroups threadsPerThreadgroup:ThreadgroupCounts];
 		
 		ConditionalSubmit();
 	}
@@ -494,10 +494,10 @@ void FAGXRenderPass::Dispatch(uint32 ThreadGroupCountX, uint32 ThreadGroupCountY
 		
 		METAL_GPUPROFILE(FAGXProfiler::GetProfiler()->EncodeDispatch(CurrentEncoder.GetCommandBufferStats(), __FUNCTION__));
 		
-		mtlpp::Size ThreadgroupCounts = mtlpp::Size(ComputeShader->NumThreadsX, ComputeShader->NumThreadsY, ComputeShader->NumThreadsZ);
+		MTLSize ThreadgroupCounts = MTLSizeMake(ComputeShader->NumThreadsX, ComputeShader->NumThreadsY, ComputeShader->NumThreadsZ);
 		check(ComputeShader->NumThreadsX > 0 && ComputeShader->NumThreadsY > 0 && ComputeShader->NumThreadsZ > 0);
-		mtlpp::Size Threadgroups = mtlpp::Size(ThreadGroupCountX, ThreadGroupCountY, ThreadGroupCountZ);
-		CurrentEncoder.GetComputeCommandEncoder().DispatchThreadgroups(Threadgroups, ThreadgroupCounts);
+		MTLSize Threadgroups = MTLSizeMake(ThreadGroupCountX, ThreadGroupCountY, ThreadGroupCountZ);
+		[CurrentEncoder.GetComputeCommandEncoder() dispatchThreadgroups:Threadgroups threadsPerThreadgroup:ThreadgroupCounts];
 		
 		ConditionalSubmit();
 	}
@@ -520,10 +520,12 @@ void FAGXRenderPass::DispatchIndirect(FAGXVertexBuffer* ArgumentBuffer, uint32 A
 		check(ComputeShader);
 		
 		METAL_GPUPROFILE(FAGXProfiler::GetProfiler()->EncodeDispatch(PrologueEncoder.GetCommandBufferStats(), __FUNCTION__));
-		mtlpp::Size ThreadgroupCounts = mtlpp::Size(ComputeShader->NumThreadsX, ComputeShader->NumThreadsY, ComputeShader->NumThreadsZ);
+		MTLSize ThreadgroupCounts = MTLSizeMake(ComputeShader->NumThreadsX, ComputeShader->NumThreadsY, ComputeShader->NumThreadsZ);
 		check(ComputeShader->NumThreadsX > 0 && ComputeShader->NumThreadsY > 0 && ComputeShader->NumThreadsZ > 0);
 		
-		PrologueEncoder.GetComputeCommandEncoder().DispatchThreadgroupsWithIndirectBuffer(ArgumentBuffer->GetCurrentBuffer(), ArgumentOffset, ThreadgroupCounts);
+		[PrologueEncoder.GetComputeCommandEncoder() dispatchThreadgroupsWithIndirectBuffer:ArgumentBuffer->GetCurrentBuffer().GetPtr()
+																	  indirectBufferOffset:(ArgumentOffset + ArgumentBuffer->GetCurrentBuffer().GetOffset())
+																	 threadsPerThreadgroup:ThreadgroupCounts];
 		
 		ConditionalSubmit();
 	}
@@ -539,10 +541,12 @@ void FAGXRenderPass::DispatchIndirect(FAGXVertexBuffer* ArgumentBuffer, uint32 A
 		check(ComputeShader);
 		
 		METAL_GPUPROFILE(FAGXProfiler::GetProfiler()->EncodeDispatch(CurrentEncoder.GetCommandBufferStats(), __FUNCTION__));
-		mtlpp::Size ThreadgroupCounts = mtlpp::Size(ComputeShader->NumThreadsX, ComputeShader->NumThreadsY, ComputeShader->NumThreadsZ);
+		MTLSize ThreadgroupCounts = MTLSizeMake(ComputeShader->NumThreadsX, ComputeShader->NumThreadsY, ComputeShader->NumThreadsZ);
 		check(ComputeShader->NumThreadsX > 0 && ComputeShader->NumThreadsY > 0 && ComputeShader->NumThreadsZ > 0);
 
-		CurrentEncoder.GetComputeCommandEncoder().DispatchThreadgroupsWithIndirectBuffer(ArgumentBuffer->GetCurrentBuffer(), ArgumentOffset, ThreadgroupCounts);
+		[CurrentEncoder.GetComputeCommandEncoder() dispatchThreadgroupsWithIndirectBuffer:ArgumentBuffer->GetCurrentBuffer()
+																	 indirectBufferOffset:(ArgumentOffset + ArgumentBuffer->GetCurrentBuffer().GetOffset())
+																	threadsPerThreadgroup:ThreadgroupCounts];
 
 		ConditionalSubmit();
 	}
@@ -1339,12 +1343,3 @@ uint32 FAGXRenderPass::GetCommandBufferIndex(void) const
 		return GetAGXDeviceContext().GetCurrentRenderPass().GetCommandBufferIndex();
 	}
 }
-
-#if PLATFORM_MAC
-@protocol IMTLRenderCommandEncoder
-- (void)memoryBarrierWithResources:(const id<MTLResource>[])resources count:(NSUInteger)count afterStages:(MTLRenderStages)after beforeStages:(MTLRenderStages)before;
-@end
-#endif
-@protocol IMTLComputeCommandEncoder
-- (void)memoryBarrierWithResources:(const id<MTLResource>[])resources count:(NSUInteger)count;
-@end
