AGXRHI: Scrub all usage of mtlpp::RenderCommandEncoder

#rnx

diff --git a/Engine/Source/Runtime/Apple/AGXRHI/Private/AGXCommandEncoder.cpp b/Engine/Source/Runtime/Apple/AGXRHI/Private/AGXCommandEncoder.cpp
--- a/Engine/Source/Runtime/Apple/AGXRHI/Private/AGXCommandEncoder.cpp
+++ b/Engine/Source/Runtime/Apple/AGXRHI/Private/AGXCommandEncoder.cpp
@@ -25,6 +25,7 @@ FAGXCommandEncoder::FAGXCommandEncoder(FAGXCommandList& CmdList, EAGXCommandEnco
 , bSupportsMetalFeaturesSetBytes(CmdList.GetCommandQueue().SupportsFeature(EAGXFeaturesSetBytes))
 , RingBuffer(EncoderRingBufferSize, BufferOffsetAlignment, FAGXCommandQueue::GetCompatibleResourceOptions((mtlpp::ResourceOptions)(mtlpp::ResourceOptions::HazardTrackingModeUntracked | BUFFER_RESOURCE_STORAGE_MANAGED)))
 , RenderPassDesc(nil)
+, RenderCommandEncoder(nil)
 , ComputeCommandEncoder(nil)
 , BlitCommandEncoder(nil)
 #if ENABLE_METAL_GPUPROFILE
@@ -147,14 +148,14 @@ void FAGXCommandEncoder::ResetLive(void)
 	{
 		for (uint32 i = 0; i < ML_MaxBuffers; i++)
 		{
-			RenderCommandEncoder.SetVertexBuffer(nil, 0, i);
-			RenderCommandEncoder.SetFragmentBuffer(nil, 0, i);
+			[RenderCommandEncoder setVertexBuffer:nil offset:0 atIndex:i];
+			[RenderCommandEncoder setFragmentBuffer:nil offset:0 atIndex:i];
 		}
 		
 		for (uint32 i = 0; i < ML_MaxTextures; i++)
 		{
-			RenderCommandEncoder.SetVertexTexture(nil, i);
-			RenderCommandEncoder.SetFragmentTexture(nil, i);
+			[RenderCommandEncoder setVertexTexture:nil atIndex:i];
+			[RenderCommandEncoder setFragmentTexture:nil atIndex:i];
 		}
 	}
 	else if (IsComputeCommandEncoderActive())
@@ -276,7 +277,7 @@ bool FAGXCommandEncoder::IsParallelRenderCommandEncoderActive(void) const
 	
 bool FAGXCommandEncoder::IsRenderCommandEncoderActive(void) const
 {
-	return RenderCommandEncoder.GetPtr() != nil || ParallelRenderCommandEncoder.GetPtr() != nil;
+	return RenderCommandEncoder != nil || ParallelRenderCommandEncoder.GetPtr() != nil;
 }
 
 bool FAGXCommandEncoder::IsComputeCommandEncoderActive(void) const
@@ -320,7 +321,7 @@ mtlpp::RenderCommandEncoder& FAGXCommandEncoder::GetChildRenderCommandEncoder(ui
 	return ChildRenderCommandEncoders[Index];
 }
 
-mtlpp::RenderCommandEncoder& FAGXCommandEncoder::GetRenderCommandEncoder(void)
+id<MTLRenderCommandEncoder> FAGXCommandEncoder::GetRenderCommandEncoder() const
 {
 	check(IsRenderCommandEncoderActive() && RenderCommandEncoder);
 	return RenderCommandEncoder;
@@ -382,12 +383,12 @@ void FAGXCommandEncoder::BeginRenderCommandEncoding(void)
 	
 	if (!CommandList.IsParallel() || Type == EAGXCommandEncoderPrologue)
 	{
-		RenderCommandEncoder = MTLPP_VALIDATE(mtlpp::CommandBuffer, CommandBuffer, AGXSafeGetRuntimeDebuggingLevel() >= EAGXDebugLevelValidation, RenderCommandEncoder(RenderPassDesc));
+		RenderCommandEncoder = [[CommandBuffer.GetPtr() renderCommandEncoderWithDescriptor:RenderPassDesc.GetPtr()] retain];
 		EncoderNum++;	
 	}
 	else
 	{
-		RenderCommandEncoder = GetAGXDeviceContext().GetParallelRenderCommandEncoder(CommandList.GetParallelIndex(), ParallelRenderCommandEncoder, CommandBuffer);
+		RenderCommandEncoder = [GetAGXDeviceContext().GetParallelRenderCommandEncoder(CommandList.GetParallelIndex(), ParallelRenderCommandEncoder, CommandBuffer).GetPtr() retain];
 	}
 	
 	NSString* Label = nil;
@@ -395,13 +396,13 @@ void FAGXCommandEncoder::BeginRenderCommandEncoding(void)
 	if(GetEmitDrawEvents())
 	{
 		Label = [NSString stringWithFormat:@"RenderEncoder: %@", [DebugGroups count] > 0 ? [DebugGroups lastObject] : (NSString*)CFSTR("InitialPass")];
-		RenderCommandEncoder.SetLabel(Label);
+		[RenderCommandEncoder setLabel:Label];
 		
 		if([DebugGroups count])
 		{
 			for (NSString* Group in DebugGroups)
 			{
-				RenderCommandEncoder.PushDebugGroup(Group);
+				[RenderCommandEncoder pushDebugGroup:Group];
 			}
 		}
 	}
@@ -418,7 +419,7 @@ void FAGXCommandEncoder::BeginComputeCommandEncoding(mtlpp::DispatchType Dispatc
 	}
 	else
 	{
-		ComputeCommandEncoder = [[CommandBuffer.GetPtr() computeCommandEncoderWithDispatchType:*reinterpret_cast<MTLDispatchType*>(&DispatchType)] retain];
+		ComputeCommandEncoder = [[CommandBuffer.GetPtr() computeCommandEncoderWithDispatchType:(MTLDispatchType)DispatchType] retain];
 	}
 
 	EncoderNum++;
@@ -484,24 +485,25 @@ void FAGXCommandEncoder::EndEncoding(void)
 						{
 							mtlpp::StoreAction Action = ColorStoreActions[i];
 							check(Action != mtlpp::StoreAction::Unknown);
-							RenderCommandEncoder.SetColorStoreAction((mtlpp::StoreAction)Action, i);
+							[RenderCommandEncoder setColorStoreAction:(MTLStoreAction)Action atIndex:i];
 						}
 					}
 					if (RenderPassDesc.GetDepthAttachment().GetTexture() && RenderPassDesc.GetDepthAttachment().GetStoreAction() == mtlpp::StoreAction::Unknown)
 					{
 						mtlpp::StoreAction Action = DepthStoreAction;
 						check(Action != mtlpp::StoreAction::Unknown);
-						RenderCommandEncoder.SetDepthStoreAction((mtlpp::StoreAction)Action);
+						[RenderCommandEncoder setDepthStoreAction:(MTLStoreAction)Action];
 					}
 					if (RenderPassDesc.GetStencilAttachment().GetTexture() && RenderPassDesc.GetStencilAttachment().GetStoreAction() == mtlpp::StoreAction::Unknown)
 					{
 						mtlpp::StoreAction Action = StencilStoreAction;
 						check(Action != mtlpp::StoreAction::Unknown);
-						RenderCommandEncoder.SetStencilStoreAction((mtlpp::StoreAction)Action);
+						[RenderCommandEncoder setStencilStoreAction:(MTLStoreAction)Action];
 					}
 				}
 				
-				RenderCommandEncoder.EndEncoding();
+				[RenderCommandEncoder endEncoding];
+				[RenderCommandEncoder release];
 				RenderCommandEncoder = nil;
 			}
 			
@@ -630,7 +632,7 @@ void FAGXCommandEncoder::InsertDebugSignpost(ns::String const& String)
 	{
 		if (RenderCommandEncoder)
 		{
-			RenderCommandEncoder.InsertDebugSignpost(String);
+			[RenderCommandEncoder insertDebugSignpost:String.GetPtr()];
 		}
 		else if (ParallelRenderCommandEncoder && !IsParallel())
 		{
@@ -655,7 +657,7 @@ void FAGXCommandEncoder::PushDebugGroup(ns::String const& String)
 		[DebugGroups addObject:String];
 		if (RenderCommandEncoder)
 		{
-			RenderCommandEncoder.PushDebugGroup(String);
+			[RenderCommandEncoder pushDebugGroup:String.GetPtr()];
 		}
 		else if (ParallelRenderCommandEncoder && !IsParallel())
 		{
@@ -680,7 +682,7 @@ void FAGXCommandEncoder::PopDebugGroup(void)
 		[DebugGroups removeLastObject];
 		if (RenderCommandEncoder)
 		{
-			RenderCommandEncoder.PopDebugGroup();
+			[RenderCommandEncoder popDebugGroup];
 		}
 		else if (ParallelRenderCommandEncoder && !IsParallel())
 		{
@@ -758,7 +760,7 @@ void FAGXCommandEncoder::SetRenderPipelineState(FAGXShaderPipeline* PipelineStat
 {
 	check (RenderCommandEncoder);
 	{
-		RenderCommandEncoder.SetRenderPipelineState(PipelineState->RenderPipelineState);
+		[RenderCommandEncoder setRenderPipelineState:PipelineState->RenderPipelineState.GetPtr()];
 	}
 }
 
@@ -768,13 +770,13 @@ void FAGXCommandEncoder::SetViewport(mtlpp::Viewport const Viewport[], uint32 Nu
 	check(NumActive >= 1 && NumActive < ML_MaxViewports);
 	if (NumActive == 1)
 	{
-		RenderCommandEncoder.SetViewport(Viewport[0]);
+		[RenderCommandEncoder setViewport:*reinterpret_cast<const MTLViewport*>(&Viewport[0])];
 	}
 #if PLATFORM_MAC
 	else
 	{
 		check(FAGXCommandQueue::SupportsFeature(EAGXFeaturesMultipleViewports));
-		RenderCommandEncoder.SetViewports(Viewport, NumActive);
+		[RenderCommandEncoder setViewports:reinterpret_cast<const MTLViewport*>(Viewport) count:(NSUInteger)NumActive];
 	}
 #endif
 }
@@ -783,7 +785,7 @@ void FAGXCommandEncoder::SetFrontFacingWinding(mtlpp::Winding const InFrontFacin
 {
     check (RenderCommandEncoder);
 	{
-		RenderCommandEncoder.SetFrontFacingWinding(InFrontFacingWinding);
+		[RenderCommandEncoder setFrontFacingWinding:(MTLWinding)InFrontFacingWinding];
 	}
 }
 
@@ -791,7 +793,7 @@ void FAGXCommandEncoder::SetCullMode(mtlpp::CullMode const InCullMode)
 {
     check (RenderCommandEncoder);
 	{
-		RenderCommandEncoder.SetCullMode(InCullMode);
+		[RenderCommandEncoder setCullMode:(MTLCullMode)InCullMode];
 	}
 }
 
@@ -799,7 +801,7 @@ void FAGXCommandEncoder::SetDepthBias(float const InDepthBias, float const InSlo
 {
     check (RenderCommandEncoder);
 	{
-		RenderCommandEncoder.SetDepthBias(InDepthBias, InSlopeScale, InClamp);
+		[RenderCommandEncoder setDepthBias:InDepthBias slopeScale:InSlopeScale clamp:InClamp];
 	}
 }
 
@@ -809,13 +811,13 @@ void FAGXCommandEncoder::SetScissorRect(mtlpp::ScissorRect const Rect[], uint32
 	check(NumActive >= 1 && NumActive < ML_MaxViewports);
 	if (NumActive == 1)
 	{
-		RenderCommandEncoder.SetScissorRect(Rect[0]);
+		[RenderCommandEncoder setScissorRect:*reinterpret_cast<const MTLScissorRect*>(&Rect[0])];
 	}
 #if PLATFORM_MAC
 	else
 	{
 		check(FAGXCommandQueue::SupportsFeature(EAGXFeaturesMultipleViewports));
-		RenderCommandEncoder.SetScissorRects(Rect, NumActive);
+		[RenderCommandEncoder setScissorRects:reinterpret_cast<const MTLScissorRect*>(Rect) count:(NSUInteger)NumActive];
 	}
 #endif
 }
@@ -824,7 +826,7 @@ void FAGXCommandEncoder::SetTriangleFillMode(mtlpp::TriangleFillMode const InFil
 {
     check(RenderCommandEncoder);
 	{
-		RenderCommandEncoder.SetTriangleFillMode(InFillMode);
+		[RenderCommandEncoder setTriangleFillMode:(MTLTriangleFillMode)InFillMode];
 	}
 }
 
@@ -832,7 +834,7 @@ void FAGXCommandEncoder::SetDepthClipMode(mtlpp::DepthClipMode const InDepthClip
 {
 	check(RenderCommandEncoder);
 	{
-		RenderCommandEncoder.SetDepthClipMode(InDepthClipMode);
+		[RenderCommandEncoder setDepthClipMode:(MTLDepthClipMode)InDepthClipMode];
 	}
 }
 
@@ -840,7 +842,7 @@ void FAGXCommandEncoder::SetBlendColor(float const Red, float const Green, float
 {
 	check(RenderCommandEncoder);
 	{
-		RenderCommandEncoder.SetBlendColor(Red, Green, Blue, Alpha);
+		[RenderCommandEncoder setBlendColorRed:Red green:Green blue:Blue alpha:Alpha];
 	}
 }
 
@@ -848,7 +850,7 @@ void FAGXCommandEncoder::SetDepthStencilState(mtlpp::DepthStencilState const& In
 {
     check (RenderCommandEncoder);
 	{
-		RenderCommandEncoder.SetDepthStencilState(InDepthStencilState);
+		[RenderCommandEncoder setDepthStencilState:InDepthStencilState.GetPtr()];
 	}
 }
 
@@ -856,7 +858,7 @@ void FAGXCommandEncoder::SetStencilReferenceValue(uint32 const ReferenceValue)
 {
     check (RenderCommandEncoder);
 	{
-		RenderCommandEncoder.SetStencilReferenceValue(ReferenceValue);
+		[RenderCommandEncoder setStencilReferenceValue:ReferenceValue];
 	}
 }
 
@@ -865,7 +867,7 @@ void FAGXCommandEncoder::SetVisibilityResultMode(mtlpp::VisibilityResultMode con
     check (RenderCommandEncoder);
 	{
 		check(Mode == mtlpp::VisibilityResultMode::Disabled || RenderPassDesc.GetVisibilityResultBuffer());
-		RenderCommandEncoder.SetVisibilityResultMode(Mode, Offset);
+		[RenderCommandEncoder setVisibilityResultMode:(MTLVisibilityResultMode)Mode offset:Offset];
 	}
 }
 	
@@ -945,11 +947,11 @@ void FAGXCommandEncoder::SetShaderBytes(mtlpp::FunctionType const FunctionType,
 			{
 				case mtlpp::FunctionType::Vertex:
 					check(RenderCommandEncoder);
-					RenderCommandEncoder.SetVertexData(Bytes, Length, Index);
+					[RenderCommandEncoder setVertexBytes:static_cast<const void*>(Bytes) length:Length atIndex:Index];
 					break;
 				case mtlpp::FunctionType::Fragment:
 					check(RenderCommandEncoder);
-					RenderCommandEncoder.SetFragmentData(Bytes, Length, Index);
+					[RenderCommandEncoder setFragmentBytes:static_cast<const void*>(Bytes) length:Length atIndex:Index];
 					break;
 				case mtlpp::FunctionType::Kernel:
 					check(ComputeCommandEncoder);
@@ -998,11 +1000,11 @@ void FAGXCommandEncoder::SetShaderBufferOffset(mtlpp::FunctionType FunctionType,
 	{
 		case mtlpp::FunctionType::Vertex:
 			check (RenderCommandEncoder);
-			RenderCommandEncoder.SetVertexBufferOffset(Offset + ShaderBuffers[uint32(FunctionType)].Buffers[index].GetOffset(), index);
+			[RenderCommandEncoder setVertexBufferOffset:(Offset + ShaderBuffers[uint32(FunctionType)].Buffers[index].GetOffset()) atIndex:index];
 			break;
 		case mtlpp::FunctionType::Fragment:
 			check(RenderCommandEncoder);
-			RenderCommandEncoder.SetFragmentBufferOffset(Offset + ShaderBuffers[uint32(FunctionType)].Buffers[index].GetOffset(), index);
+			[RenderCommandEncoder setFragmentBufferOffset:(Offset + ShaderBuffers[uint32(FunctionType)].Buffers[index].GetOffset()) atIndex:index];
 			break;
 		case mtlpp::FunctionType::Kernel:
 			check (ComputeCommandEncoder);
@@ -1022,11 +1024,11 @@ void FAGXCommandEncoder::SetShaderTexture(mtlpp::FunctionType FunctionType, FAGX
 		case mtlpp::FunctionType::Vertex:
 			check (RenderCommandEncoder);
 			FenceResource(Texture);
-			RenderCommandEncoder.SetVertexTexture(Texture, index);
+			[RenderCommandEncoder setVertexTexture:Texture.GetPtr() atIndex:index];
 			break;
 		case mtlpp::FunctionType::Fragment:
 			check(RenderCommandEncoder);
-			RenderCommandEncoder.SetFragmentTexture(Texture, index);
+			[RenderCommandEncoder setFragmentTexture:Texture.GetPtr() atIndex:index];
 			break;
 		case mtlpp::FunctionType::Kernel:
 			check (ComputeCommandEncoder);
@@ -1062,11 +1064,11 @@ void FAGXCommandEncoder::SetShaderSamplerState(mtlpp::FunctionType FunctionType,
 	{
 		case mtlpp::FunctionType::Vertex:
        		check (RenderCommandEncoder);
-			RenderCommandEncoder.SetVertexSamplerState(Sampler, index);
+			[RenderCommandEncoder setVertexSamplerState:Sampler.GetPtr() atIndex:index];
 			break;
 		case mtlpp::FunctionType::Fragment:
 			check (RenderCommandEncoder);
-			RenderCommandEncoder.SetFragmentSamplerState(Sampler, index);
+			[RenderCommandEncoder setFragmentSamplerState:Sampler.GetPtr() atIndex:index];
 			break;
 		case mtlpp::FunctionType::Kernel:
 			check (ComputeCommandEncoder);
@@ -1163,13 +1165,13 @@ void FAGXCommandEncoder::SetShaderBufferInternal(mtlpp::FunctionType Function, u
 				Binding.Bound |= (1 << Index);
 				check(RenderCommandEncoder);
 				FenceResource(Buffer);
-				RenderCommandEncoder.SetVertexBuffer(Buffer, Offset, Index);
+				[RenderCommandEncoder setVertexBuffer:Buffer.GetPtr() offset:(Offset + Buffer.GetOffset()) atIndex:Index];
 				break;
 
 			case mtlpp::FunctionType::Fragment:
 				Binding.Bound |= (1 << Index);
 				check(RenderCommandEncoder);
-				RenderCommandEncoder.SetFragmentBuffer(Buffer, Offset, Index);
+				[RenderCommandEncoder setFragmentBuffer:Buffer.GetPtr() offset:(Offset + Buffer.GetOffset()) atIndex:Index];
 				break;
 
 			case mtlpp::FunctionType::Kernel:
@@ -1202,13 +1204,13 @@ void FAGXCommandEncoder::SetShaderBufferInternal(mtlpp::FunctionType Function, u
 			case mtlpp::FunctionType::Vertex:
 				Binding.Bound |= (1 << Index);
 				check(RenderCommandEncoder);
-				RenderCommandEncoder.SetVertexData(Bytes, Len, Index);
+				[RenderCommandEncoder setVertexBytes:static_cast<const void*>(Bytes) length:(NSUInteger)Len atIndex:Index];
 				break;
 
 			case mtlpp::FunctionType::Fragment:
 				Binding.Bound |= (1 << Index);
 				check(RenderCommandEncoder);
-				RenderCommandEncoder.SetFragmentData(Bytes, Len, Index);
+				[RenderCommandEncoder setFragmentBytes:static_cast<const void*>(Bytes) length:(NSUInteger)Len atIndex:Index];
 				break;
 
 			case mtlpp::FunctionType::Kernel:
diff --git a/Engine/Source/Runtime/Apple/AGXRHI/Private/AGXCommandEncoder.h b/Engine/Source/Runtime/Apple/AGXRHI/Private/AGXCommandEncoder.h
--- a/Engine/Source/Runtime/Apple/AGXRHI/Private/AGXCommandEncoder.h
+++ b/Engine/Source/Runtime/Apple/AGXRHI/Private/AGXCommandEncoder.h
@@ -134,7 +134,7 @@ public:
 	mtlpp::RenderCommandEncoder& GetChildRenderCommandEncoder(uint32 Index);
 
 	/** @returns The active render command encoder or nil if there isn't one. */
-	mtlpp::RenderCommandEncoder& GetRenderCommandEncoder(void);
+	id<MTLRenderCommandEncoder> GetRenderCommandEncoder() const;
 	
 	/** @returns The active compute command encoder or nil if there isn't one. */
 	id<MTLComputeCommandEncoder> GetComputeCommandEncoder() const;
@@ -458,7 +458,7 @@ public:
 	
 	mtlpp::CommandBuffer CommandBuffer;
 	mtlpp::ParallelRenderCommandEncoder ParallelRenderCommandEncoder;
-	mtlpp::RenderCommandEncoder RenderCommandEncoder;
+	id<MTLRenderCommandEncoder> RenderCommandEncoder;
 	id<MTLComputeCommandEncoder> ComputeCommandEncoder;
 	id<MTLBlitCommandEncoder> BlitCommandEncoder;
 	TArray<mtlpp::RenderCommandEncoder> ChildRenderCommandEncoders;
diff --git a/Engine/Source/Runtime/Apple/AGXRHI/Private/AGXCommands.cpp b/Engine/Source/Runtime/Apple/AGXRHI/Private/AGXCommands.cpp
--- a/Engine/Source/Runtime/Apple/AGXRHI/Private/AGXCommands.cpp
+++ b/Engine/Source/Runtime/Apple/AGXRHI/Private/AGXCommands.cpp
@@ -40,14 +40,14 @@ public:
 };
 static TGlobalResource<FVector4VertexDeclaration> FVector4VertexDeclaration;
 
-mtlpp::PrimitiveType AGXTranslatePrimitiveType(uint32 PrimitiveType)
+MTLPrimitiveType AGXTranslatePrimitiveType(uint32 PrimitiveType)
 {
 	switch (PrimitiveType)
 	{
-		case PT_TriangleList:	return mtlpp::PrimitiveType::Triangle;
-		case PT_TriangleStrip:	return mtlpp::PrimitiveType::TriangleStrip;
-		case PT_LineList:		return mtlpp::PrimitiveType::Line;
-		case PT_PointList:		return mtlpp::PrimitiveType::Point;
+		case PT_TriangleList:	return MTLPrimitiveTypeTriangle;
+		case PT_TriangleStrip:	return MTLPrimitiveTypeTriangleStrip;
+		case PT_LineList:		return MTLPrimitiveTypeLine;
+		case PT_PointList:		return MTLPrimitiveTypePoint;
 		// Metal doesn't actually actually draw in control-point patch-lists because of the way the compute shader stage works - it can handle any arbitrary patch size and will output triangles.
 		case PT_1_ControlPointPatchList:
 		case PT_2_ControlPointPatchList:
@@ -82,11 +82,11 @@ mtlpp::PrimitiveType AGXTranslatePrimitiveType(uint32 PrimitiveType)
 		case PT_31_ControlPointPatchList:
 		case PT_32_ControlPointPatchList:
 		{
-			return mtlpp::PrimitiveType::Triangle;
+			return MTLPrimitiveTypeTriangle;
 		}
 		default:
 			METAL_FATAL_ERROR(TEXT("Unsupported primitive type %d"), (int32)PrimitiveType);
-			return mtlpp::PrimitiveType::Triangle;
+			return MTLPrimitiveTypeTriangle;
 	}
 }
 
diff --git a/Engine/Source/Runtime/Apple/AGXRHI/Private/AGXRHIPrivate.h b/Engine/Source/Runtime/Apple/AGXRHI/Private/AGXRHIPrivate.h
--- a/Engine/Source/Runtime/Apple/AGXRHI/Private/AGXRHIPrivate.h
+++ b/Engine/Source/Runtime/Apple/AGXRHI/Private/AGXRHIPrivate.h
@@ -250,7 +250,7 @@ FORCEINLINE mtlpp::LoadAction GetMetalRTLoadAction(ERenderTargetLoadAction LoadA
 	}
 }
 
-mtlpp::PrimitiveType AGXTranslatePrimitiveType(uint32 PrimitiveType);
+MTLPrimitiveType AGXTranslatePrimitiveType(uint32 PrimitiveType);
 
 #if PLATFORM_MAC
 mtlpp::PrimitiveTopologyClass AGXTranslatePrimitiveTopology(uint32 PrimitiveType);
diff --git a/Engine/Source/Runtime/Apple/AGXRHI/Private/AGXRenderPass.cpp b/Engine/Source/Runtime/Apple/AGXRHI/Private/AGXRenderPass.cpp
--- a/Engine/Source/Runtime/Apple/AGXRHI/Private/AGXRenderPass.cpp
+++ b/Engine/Source/Runtime/Apple/AGXRHI/Private/AGXRenderPass.cpp
@@ -302,7 +302,11 @@ void FAGXRenderPass::DrawPrimitive(uint32 PrimitiveType, uint32 BaseVertexIndex,
 	uint32 NumVertices = GetVertexCountForPrimitiveCount(NumPrimitives, PrimitiveType);
 	
 	METAL_GPUPROFILE(FAGXProfiler::GetProfiler()->EncodeDraw(CurrentEncoder.GetCommandBufferStats(), __FUNCTION__, NumPrimitives, NumVertices, NumInstances));
-	CurrentEncoder.GetRenderCommandEncoder().Draw(AGXTranslatePrimitiveType(PrimitiveType), BaseVertexIndex, NumVertices, NumInstances);
+
+	[CurrentEncoder.GetRenderCommandEncoder() drawPrimitives:AGXTranslatePrimitiveType(PrimitiveType)
+												 vertexStart:(NSUInteger)BaseVertexIndex
+												 vertexCount:(NSUInteger)NumVertices
+											   instanceCount:(NSUInteger)NumInstances];
 
 	ConditionalSubmit();	
 }
@@ -321,7 +325,10 @@ void FAGXRenderPass::DrawPrimitiveIndirect(uint32 PrimitiveType, FAGXVertexBuffe
 		PrepareToRender(PrimitiveType);
 		
 		METAL_GPUPROFILE(FAGXProfiler::GetProfiler()->EncodeDraw(CurrentEncoder.GetCommandBufferStats(), __FUNCTION__, 1, 1, 1));
-		CurrentEncoder.GetRenderCommandEncoder().Draw(AGXTranslatePrimitiveType(PrimitiveType), TheBackingBuffer, ArgumentOffset);
+
+		[CurrentEncoder.GetRenderCommandEncoder() drawPrimitives:AGXTranslatePrimitiveType(PrimitiveType)
+												  indirectBuffer:TheBackingBuffer.GetPtr()
+											indirectBufferOffset:((NSUInteger)ArgumentOffset + TheBackingBuffer.GetOffset())];
 
 		ConditionalSubmit();
 	}
@@ -390,13 +397,26 @@ void FAGXRenderPass::DrawIndexedPrimitive(FAGXBuffer const& IndexBuffer, uint32
 	uint32 NumIndices = GetVertexCountForPrimitiveCount(NumPrimitives, PrimitiveType);
 	
 	METAL_GPUPROFILE(FAGXProfiler::GetProfiler()->EncodeDraw(CurrentEncoder.GetCommandBufferStats(), __FUNCTION__, NumPrimitives, NumVertices, NumInstances));
+
 	if (GRHISupportsBaseVertexIndex && GRHISupportsFirstInstance)
 	{
-		CurrentEncoder.GetRenderCommandEncoder().DrawIndexed(AGXTranslatePrimitiveType(PrimitiveType), NumIndices, ((IndexStride == 2) ? mtlpp::IndexType::UInt16 : mtlpp::IndexType::UInt32), IndexBuffer, StartIndex * IndexStride, NumInstances, BaseVertexIndex, FirstInstance);
+		[CurrentEncoder.GetRenderCommandEncoder() drawIndexedPrimitives:AGXTranslatePrimitiveType(PrimitiveType)
+															 indexCount:(NSUInteger)NumIndices
+															  indexType:((IndexStride == 2) ? MTLIndexTypeUInt16 : MTLIndexTypeUInt32)
+															indexBuffer:IndexBuffer.GetPtr()
+													  indexBufferOffset:((NSUInteger)(StartIndex * IndexStride) + IndexBuffer.GetOffset())
+														  instanceCount:(NSUInteger)NumInstances
+															 baseVertex:(NSUInteger)BaseVertexIndex
+														   baseInstance:(NSUInteger)FirstInstance];
 	}
 	else
 	{
-		CurrentEncoder.GetRenderCommandEncoder().DrawIndexed(AGXTranslatePrimitiveType(PrimitiveType), NumIndices, ((IndexStride == 2) ? mtlpp::IndexType::UInt16 : mtlpp::IndexType::UInt32), IndexBuffer, StartIndex * IndexStride, NumInstances);
+		[CurrentEncoder.GetRenderCommandEncoder() drawIndexedPrimitives:AGXTranslatePrimitiveType(PrimitiveType)
+															 indexCount:(NSUInteger)NumIndices
+															  indexType:((IndexStride == 2) ? MTLIndexTypeUInt16 : MTLIndexTypeUInt32)
+															indexBuffer:IndexBuffer.GetPtr()
+													  indexBufferOffset:((NSUInteger)(StartIndex * IndexStride) + IndexBuffer.GetOffset())
+														  instanceCount:(NSUInteger)NumInstances];
 	}
 	
 	ConditionalSubmit();
@@ -422,7 +442,13 @@ void FAGXRenderPass::DrawIndexedIndirect(FAGXIndexBuffer* IndexBuffer, uint32 Pr
 		PrepareToRender(PrimitiveType);
 		
 		METAL_GPUPROFILE(FAGXProfiler::GetProfiler()->EncodeDraw(CurrentEncoder.GetCommandBufferStats(), __FUNCTION__, 1, 1, 1));
-		CurrentEncoder.GetRenderCommandEncoder().DrawIndexed(AGXTranslatePrimitiveType(PrimitiveType), (mtlpp::IndexType)IndexBuffer->IndexType, TheBackingIndexBuffer, 0, TheBackingBuffer, (DrawArgumentsIndex * 5 * sizeof(uint32)));
+
+		[CurrentEncoder.GetRenderCommandEncoder()	drawIndexedPrimitives:AGXTranslatePrimitiveType(PrimitiveType)
+																indexType:(MTLIndexType)IndexBuffer->IndexType
+															  indexBuffer:TheBackingIndexBuffer.GetPtr()
+														indexBufferOffset:(0 + TheBackingIndexBuffer.GetOffset())
+														   indirectBuffer:TheBackingBuffer.GetPtr()
+													 indirectBufferOffset:((NSUInteger)(DrawArgumentsIndex * 5 * sizeof(uint32)) + TheBackingBuffer.GetOffset())];
 
 		ConditionalSubmit();
 	}
@@ -449,7 +475,13 @@ void FAGXRenderPass::DrawIndexedPrimitiveIndirect(uint32 PrimitiveType,FAGXIndex
 		PrepareToRender(PrimitiveType);
 		
 		METAL_GPUPROFILE(FAGXProfiler::GetProfiler()->EncodeDraw(CurrentEncoder.GetCommandBufferStats(), __FUNCTION__, 1, 1, 1));
-		CurrentEncoder.GetRenderCommandEncoder().DrawIndexed(AGXTranslatePrimitiveType(PrimitiveType), (mtlpp::IndexType)IndexBuffer->IndexType, TheBackingIndexBuffer, 0, TheBackingBuffer, ArgumentOffset);
+
+		[CurrentEncoder.GetRenderCommandEncoder()	drawIndexedPrimitives:AGXTranslatePrimitiveType(PrimitiveType)
+																indexType:(MTLIndexType)IndexBuffer->IndexType
+															  indexBuffer:TheBackingIndexBuffer.GetPtr()
+														indexBufferOffset:(0 + TheBackingIndexBuffer.GetOffset())
+														   indirectBuffer:TheBackingBuffer.GetPtr()
+													 indirectBufferOffset:((NSUInteger)ArgumentOffset + TheBackingBuffer.GetOffset())];
 
 		ConditionalSubmit();
 	}
@@ -569,8 +601,9 @@ void FAGXRenderPass::InsertTextureBarrier()
 #if PLATFORM_MAC
 	check(CurrentEncoder.IsRenderCommandEncoderActive());
 	
-	id <MTLRenderCommandEncoder> RenderEncoder = CurrentEncoder.GetRenderCommandEncoder().GetPtr();
+	id<MTLRenderCommandEncoder> RenderEncoder = CurrentEncoder.GetRenderCommandEncoder();
 	check(RenderEncoder);
+	
 	[RenderEncoder memoryBarrierWithScope:MTLBarrierScopeRenderTargets afterStages:MTLRenderStageFragment beforeStages:MTLRenderStageVertex];
 #endif
 }
