AGXRHI: Scrub all usage of mtlpp::ParallelRenderCommandEncoder

#rnx

diff --git a/Engine/Source/Runtime/Apple/AGXRHI/Private/AGXCommandEncoder.cpp b/Engine/Source/Runtime/Apple/AGXRHI/Private/AGXCommandEncoder.cpp
--- a/Engine/Source/Runtime/Apple/AGXRHI/Private/AGXCommandEncoder.cpp
+++ b/Engine/Source/Runtime/Apple/AGXRHI/Private/AGXCommandEncoder.cpp
@@ -25,6 +25,7 @@ FAGXCommandEncoder::FAGXCommandEncoder(FAGXCommandList& CmdList, EAGXCommandEnco
 , bSupportsMetalFeaturesSetBytes(CmdList.GetCommandQueue().SupportsFeature(EAGXFeaturesSetBytes))
 , RingBuffer(EncoderRingBufferSize, BufferOffsetAlignment, FAGXCommandQueue::GetCompatibleResourceOptions((mtlpp::ResourceOptions)(mtlpp::ResourceOptions::HazardTrackingModeUntracked | BUFFER_RESOURCE_STORAGE_MANAGED)))
 , RenderPassDesc(nil)
+, ParallelRenderCommandEncoder(nil)
 , RenderCommandEncoder(nil)
 , ComputeCommandEncoder(nil)
 , BlitCommandEncoder(nil)
@@ -270,24 +271,24 @@ void FAGXCommandEncoder::CommitCommandBuffer(uint32 const Flags)
 
 #pragma mark - Public Command Encoder Accessors -
 	
-bool FAGXCommandEncoder::IsParallelRenderCommandEncoderActive(void) const
+bool FAGXCommandEncoder::IsParallelRenderCommandEncoderActive() const
 {
-	return ParallelRenderCommandEncoder.GetPtr() != nil;
+	return (ParallelRenderCommandEncoder != nil);
 }
 	
-bool FAGXCommandEncoder::IsRenderCommandEncoderActive(void) const
+bool FAGXCommandEncoder::IsRenderCommandEncoderActive() const
 {
-	return RenderCommandEncoder != nil || ParallelRenderCommandEncoder.GetPtr() != nil;
+	return (RenderCommandEncoder != nil) || (ParallelRenderCommandEncoder != nil);
 }
 
-bool FAGXCommandEncoder::IsComputeCommandEncoderActive(void) const
+bool FAGXCommandEncoder::IsComputeCommandEncoderActive() const
 {
-	return ComputeCommandEncoder != nil;
+	return (ComputeCommandEncoder != nil);
 }
 
-bool FAGXCommandEncoder::IsBlitCommandEncoderActive(void) const
+bool FAGXCommandEncoder::IsBlitCommandEncoderActive() const
 {
-	return BlitCommandEncoder != nil;
+	return (BlitCommandEncoder != nil);
 }
 
 bool FAGXCommandEncoder::IsImmediate(void) const
@@ -310,12 +311,12 @@ mtlpp::RenderPassDescriptor const& FAGXCommandEncoder::GetRenderPassDescriptor(v
 	return RenderPassDesc;
 }
 
-mtlpp::ParallelRenderCommandEncoder& FAGXCommandEncoder::GetParallelRenderCommandEncoder(void)
+id<MTLParallelRenderCommandEncoder> FAGXCommandEncoder::GetParallelRenderCommandEncoder() const
 {
 	return ParallelRenderCommandEncoder;
 }
 
-mtlpp::RenderCommandEncoder& FAGXCommandEncoder::GetChildRenderCommandEncoder(uint32 Index)
+id<MTLRenderCommandEncoder> FAGXCommandEncoder::GetChildRenderCommandEncoder(uint32 Index) const
 {
 	check(IsParallelRenderCommandEncoderActive() && Index < ChildRenderCommandEncoders.Num());
 	return ChildRenderCommandEncoders[Index];
@@ -348,7 +349,7 @@ void FAGXCommandEncoder::BeginParallelRenderCommandEncoding(uint32 NumChildren)
 	check(CommandBuffer);
 	check(IsRenderCommandEncoderActive() == false && IsComputeCommandEncoderActive() == false && IsBlitCommandEncoderActive() == false);
 	
-	ParallelRenderCommandEncoder = MTLPP_VALIDATE(mtlpp::CommandBuffer, CommandBuffer, AGXSafeGetRuntimeDebuggingLevel() >= EAGXDebugLevelValidation, ParallelRenderCommandEncoder(RenderPassDesc));
+	ParallelRenderCommandEncoder = [[CommandBuffer.GetPtr() parallelRenderCommandEncoderWithDescriptor:RenderPassDesc.GetPtr()] retain];
 	
 	EncoderNum++;
 	
@@ -357,20 +358,20 @@ void FAGXCommandEncoder::BeginParallelRenderCommandEncoding(uint32 NumChildren)
 	if(GetEmitDrawEvents())
 	{
 		Label = [NSString stringWithFormat:@"ParallelRenderCommandEncoder: %@", [DebugGroups count] > 0 ? [DebugGroups lastObject] : (NSString*)CFSTR("InitialPass")];
-		ParallelRenderCommandEncoder.SetLabel(Label);
+		[ParallelRenderCommandEncoder setLabel:Label];
 		
 		if([DebugGroups count])
 		{
 			for (NSString* Group in DebugGroups)
 			{
-				ParallelRenderCommandEncoder.PushDebugGroup(Group);
+				[ParallelRenderCommandEncoder pushDebugGroup:Group];
 			}
 		}
 	}
 	
 	for (uint32 i = 0; i < NumChildren; i++)
 	{
-		mtlpp::RenderCommandEncoder CommandEncoder = MTLPP_VALIDATE(mtlpp::ParallelRenderCommandEncoder, ParallelRenderCommandEncoder, AGXSafeGetRuntimeDebuggingLevel() >= EAGXDebugLevelValidation, GetRenderCommandEncoder());
+		id<MTLRenderCommandEncoder> CommandEncoder = [[ParallelRenderCommandEncoder renderCommandEncoder] retain];
 		ChildRenderCommandEncoders.Add(CommandEncoder);
 	}
 }
@@ -388,7 +389,7 @@ void FAGXCommandEncoder::BeginRenderCommandEncoding(void)
 	}
 	else
 	{
-		RenderCommandEncoder = [GetAGXDeviceContext().GetParallelRenderCommandEncoder(CommandList.GetParallelIndex(), ParallelRenderCommandEncoder, CommandBuffer).GetPtr() retain];
+		RenderCommandEncoder = GetAGXDeviceContext().GetParallelRenderCommandEncoder(CommandList.GetParallelIndex(), &ParallelRenderCommandEncoder, CommandBuffer);
 	}
 	
 	NSString* Label = nil;
@@ -446,7 +447,7 @@ void FAGXCommandEncoder::BeginBlitCommandEncoding(void)
 	check(CommandBuffer);
 	check(IsRenderCommandEncoderActive() == false && IsComputeCommandEncoderActive() == false && IsBlitCommandEncoderActive() == false);
 	
-	BlitCommandEncoder = [[CommandBuffer.GetPtr() blitCommandEncoder] retain];	
+	BlitCommandEncoder = [[CommandBuffer.GetPtr() blitCommandEncoder] retain];
 	EncoderNum++;
 	
 	NSString* Label = nil;
@@ -474,7 +475,7 @@ void FAGXCommandEncoder::EndEncoding(void)
 		{
 			if (RenderCommandEncoder)
 			{
-				if (ParallelRenderCommandEncoder.GetPtr() == nil)
+				if (ParallelRenderCommandEncoder == nil)
 				{
 					check(RenderPassDesc);
 					
@@ -536,7 +537,8 @@ void FAGXCommandEncoder::EndEncoding(void)
 				EncoderNum = 0;
 
 				CommandBuffer = nil;
-				
+
+				[ParallelRenderCommandEncoder release];
 				ParallelRenderCommandEncoder = nil;
 			}
 
@@ -552,24 +554,25 @@ void FAGXCommandEncoder::EndEncoding(void)
 						{
 							mtlpp::StoreAction Action = ColorStoreActions[i];
 							check(Action != mtlpp::StoreAction::Unknown);
-							ParallelRenderCommandEncoder.SetColorStoreAction((mtlpp::StoreAction)Action, i);
+							[ParallelRenderCommandEncoder setColorStoreAction:(MTLStoreAction)Action atIndex:i];
 						}
 					}
 					if (RenderPassDesc.GetDepthAttachment().GetTexture() && RenderPassDesc.GetDepthAttachment().GetStoreAction() == mtlpp::StoreAction::Unknown)
 					{
 						mtlpp::StoreAction Action = DepthStoreAction;
 						check(Action != mtlpp::StoreAction::Unknown);
-						ParallelRenderCommandEncoder.SetDepthStoreAction((mtlpp::StoreAction)Action);
+						[ParallelRenderCommandEncoder setDepthStoreAction:(MTLStoreAction)Action];
 					}
 					if (RenderPassDesc.GetStencilAttachment().GetTexture() && RenderPassDesc.GetStencilAttachment().GetStoreAction() == mtlpp::StoreAction::Unknown)
 					{
 						mtlpp::StoreAction Action = StencilStoreAction;
 						check(Action != mtlpp::StoreAction::Unknown);
-						ParallelRenderCommandEncoder.SetStencilStoreAction((mtlpp::StoreAction)Action);
+						[ParallelRenderCommandEncoder setStencilStoreAction:(MTLStoreAction)Action];
 					}
 				}
 
-				ParallelRenderCommandEncoder.EndEncoding();
+				[ParallelRenderCommandEncoder endEncoding];
+				[ParallelRenderCommandEncoder release];
 				ParallelRenderCommandEncoder = nil;
 
 				ChildRenderCommandEncoders.Empty();
@@ -636,16 +639,15 @@ void FAGXCommandEncoder::InsertDebugSignpost(ns::String const& String)
 		}
 		else if (ParallelRenderCommandEncoder && !IsParallel())
 		{
-			ParallelRenderCommandEncoder.InsertDebugSignpost(String);
-			METAL_DEBUG_LAYER(EAGXDebugLevelFastValidation, ParallelRenderCommandEncoder.InsertDebugSignpost(String));
+			[ParallelRenderCommandEncoder insertDebugSignpost:String.GetPtr()];
 		}
 		else if (ComputeCommandEncoder)
 		{
-			[ComputeCommandEncoder insertDebugSignpost:String];
+			[ComputeCommandEncoder insertDebugSignpost:String.GetPtr()];
 		}
 		else if (BlitCommandEncoder)
 		{
-			[BlitCommandEncoder insertDebugSignpost:String];
+			[BlitCommandEncoder insertDebugSignpost:String.GetPtr()];
 		}
 	}
 }
@@ -654,23 +656,22 @@ void FAGXCommandEncoder::PushDebugGroup(ns::String const& String)
 {
 	if (String)
 	{
-		[DebugGroups addObject:String];
+		[DebugGroups addObject:String.GetPtr()];
 		if (RenderCommandEncoder)
 		{
 			[RenderCommandEncoder pushDebugGroup:String.GetPtr()];
 		}
 		else if (ParallelRenderCommandEncoder && !IsParallel())
 		{
-			ParallelRenderCommandEncoder.PushDebugGroup(String);
-			METAL_DEBUG_LAYER(EAGXDebugLevelFastValidation, ParallelRenderCommandEncoder.PushDebugGroup(String));
+			[ParallelRenderCommandEncoder pushDebugGroup:String.GetPtr()];
 		}
 		else if (ComputeCommandEncoder)
 		{
-			[ComputeCommandEncoder pushDebugGroup:String];
+			[ComputeCommandEncoder pushDebugGroup:String.GetPtr()];
 		}
 		else if (BlitCommandEncoder)
 		{
-			[BlitCommandEncoder pushDebugGroup:String];
+			[BlitCommandEncoder pushDebugGroup:String.GetPtr()];
 		}
 	}
 }
@@ -686,8 +687,7 @@ void FAGXCommandEncoder::PopDebugGroup(void)
 		}
 		else if (ParallelRenderCommandEncoder && !IsParallel())
 		{
-			ParallelRenderCommandEncoder.PopDebugGroup();
-			METAL_DEBUG_LAYER(EAGXDebugLevelFastValidation, ParallelRenderCommandEncoder.PopDebugGroup());
+			[ParallelRenderCommandEncoder popDebugGroup];
 		}
 		else if (ComputeCommandEncoder)
 		{
diff --git a/Engine/Source/Runtime/Apple/AGXRHI/Private/AGXCommandEncoder.h b/Engine/Source/Runtime/Apple/AGXRHI/Private/AGXCommandEncoder.h
--- a/Engine/Source/Runtime/Apple/AGXRHI/Private/AGXCommandEncoder.h
+++ b/Engine/Source/Runtime/Apple/AGXRHI/Private/AGXCommandEncoder.h
@@ -98,16 +98,16 @@ public:
 #pragma mark - Public Command Encoder Accessors -
 	
 	/** @returns True if and only if there is an active parallel render command encoder, otherwise false. */
-	bool IsParallelRenderCommandEncoderActive(void) const;
+	bool IsParallelRenderCommandEncoderActive() const;
 	
 	/** @returns True if and only if there is an active render command encoder, otherwise false. */
-	bool IsRenderCommandEncoderActive(void) const;
+	bool IsRenderCommandEncoderActive() const;
 	
 	/** @returns True if and only if there is an active compute command encoder, otherwise false. */
-	bool IsComputeCommandEncoderActive(void) const;
+	bool IsComputeCommandEncoderActive() const;
 	
 	/** @returns True if and only if there is an active blit command encoder, otherwise false. */
-	bool IsBlitCommandEncoderActive(void) const;
+	bool IsBlitCommandEncoderActive() const;
 	
 	/**
 	 * True iff the command-encoder submits immediately to the command-queue, false if it performs any buffering.
@@ -128,10 +128,10 @@ public:
 	mtlpp::RenderPassDescriptor const& GetRenderPassDescriptor(void) const;
 	
 	/** @returns The active render command encoder or nil if there isn't one. */
-	mtlpp::ParallelRenderCommandEncoder& GetParallelRenderCommandEncoder(void);
+	id<MTLParallelRenderCommandEncoder> GetParallelRenderCommandEncoder() const;
 	
 	/** @returns The child render command encoder of the current parallel render encoder for Index. */
-	mtlpp::RenderCommandEncoder& GetChildRenderCommandEncoder(uint32 Index);
+	id<MTLRenderCommandEncoder> GetChildRenderCommandEncoder(uint32 Index) const;
 
 	/** @returns The active render command encoder or nil if there isn't one. */
 	id<MTLRenderCommandEncoder> GetRenderCommandEncoder() const;
@@ -457,11 +457,11 @@ public:
 	mtlpp::RenderPassDescriptor RenderPassDesc;
 	
 	mtlpp::CommandBuffer CommandBuffer;
-	mtlpp::ParallelRenderCommandEncoder ParallelRenderCommandEncoder;
+	id<MTLParallelRenderCommandEncoder> ParallelRenderCommandEncoder;
 	id<MTLRenderCommandEncoder> RenderCommandEncoder;
 	id<MTLComputeCommandEncoder> ComputeCommandEncoder;
 	id<MTLBlitCommandEncoder> BlitCommandEncoder;
-	TArray<mtlpp::RenderCommandEncoder> ChildRenderCommandEncoders;
+	TArray< id<MTLRenderCommandEncoder> > ChildRenderCommandEncoders;
 	
 #if ENABLE_METAL_GPUPROFILE
 	FAGXCommandBufferStats* CommandBufferStats;
diff --git a/Engine/Source/Runtime/Apple/AGXRHI/Private/AGXContext.cpp b/Engine/Source/Runtime/Apple/AGXRHI/Private/AGXContext.cpp
--- a/Engine/Source/Runtime/Apple/AGXRHI/Private/AGXContext.cpp
+++ b/Engine/Source/Runtime/Apple/AGXRHI/Private/AGXContext.cpp
@@ -1602,7 +1602,7 @@ void FAGXDeviceContext::SetParallelRenderPassDescriptor(FRHIRenderPassInfo const
 	}
 }
 
-mtlpp::RenderCommandEncoder FAGXDeviceContext::GetParallelRenderCommandEncoder(uint32 Index, mtlpp::ParallelRenderCommandEncoder& ParallelEncoder, mtlpp::CommandBuffer& CommandBuffer)
+id<MTLRenderCommandEncoder> FAGXDeviceContext::GetParallelRenderCommandEncoder(uint32 Index, id<MTLParallelRenderCommandEncoder>* ParallelEncoder, mtlpp::CommandBuffer& CommandBuffer)
 {
 	FScopeLock Lock(&FreeListMutex);
 	
diff --git a/Engine/Source/Runtime/Apple/AGXRHI/Private/AGXContext.h b/Engine/Source/Runtime/Apple/AGXRHI/Private/AGXContext.h
--- a/Engine/Source/Runtime/Apple/AGXRHI/Private/AGXContext.h
+++ b/Engine/Source/Runtime/Apple/AGXRHI/Private/AGXContext.h
@@ -213,7 +213,7 @@ public:
 
 	void BeginParallelRenderCommandEncoding(uint32 Num);
 	void SetParallelRenderPassDescriptor(FRHIRenderPassInfo const& TargetInfo);
-	mtlpp::RenderCommandEncoder GetParallelRenderCommandEncoder(uint32 Index, mtlpp::ParallelRenderCommandEncoder& ParallelEncoder, mtlpp::CommandBuffer& CommandBuffer);
+	id<MTLRenderCommandEncoder> GetParallelRenderCommandEncoder(uint32 Index, id<MTLParallelRenderCommandEncoder>* ParallelEncoder, mtlpp::CommandBuffer& CommandBuffer);
 	void EndParallelRenderCommandEncoding(void);
 	
 	/** Get the index of the bound Metal device in the global list of rendering devices. */
diff --git a/Engine/Source/Runtime/Apple/AGXRHI/Private/AGXRenderPass.cpp b/Engine/Source/Runtime/Apple/AGXRHI/Private/AGXRenderPass.cpp
--- a/Engine/Source/Runtime/Apple/AGXRHI/Private/AGXRenderPass.cpp
+++ b/Engine/Source/Runtime/Apple/AGXRHI/Private/AGXRenderPass.cpp
@@ -1053,10 +1053,10 @@ bool FAGXRenderPass::IsWithinParallelPass(void)
 	return bWithinRenderPass && CurrentEncoder.IsParallelRenderCommandEncoderActive();
 }
 
-mtlpp::RenderCommandEncoder FAGXRenderPass::GetParallelRenderCommandEncoder(uint32 Index, mtlpp::ParallelRenderCommandEncoder& ParallelEncoder)
+id<MTLRenderCommandEncoder> FAGXRenderPass::GetParallelRenderCommandEncoder(uint32 Index, id<MTLParallelRenderCommandEncoder>* ParallelEncoder)
 {
 	check(IsWithinParallelPass());
-	ParallelEncoder = CurrentEncoder.GetParallelRenderCommandEncoder();
+	*ParallelEncoder = CurrentEncoder.GetParallelRenderCommandEncoder();
 	return CurrentEncoder.GetChildRenderCommandEncoder(Index);
 }
 
diff --git a/Engine/Source/Runtime/Apple/AGXRHI/Private/AGXRenderPass.h b/Engine/Source/Runtime/Apple/AGXRHI/Private/AGXRenderPass.h
--- a/Engine/Source/Runtime/Apple/AGXRHI/Private/AGXRenderPass.h
+++ b/Engine/Source/Runtime/Apple/AGXRHI/Private/AGXRenderPass.h
@@ -146,7 +146,7 @@ public:
 	 * Get a child render command-encoder and the parent parallel command-encoder when within a parallel pass.
 	 * @returns A valid render command encoder or nil.
 	 */
-    mtlpp::RenderCommandEncoder GetParallelRenderCommandEncoder(uint32 Index, mtlpp::ParallelRenderCommandEncoder& ParallelEncoder);
+    id<MTLRenderCommandEncoder> GetParallelRenderCommandEncoder(uint32 Index, id<MTLParallelRenderCommandEncoder>* ParallelEncoder);
 	
 	void InsertTextureBarrier();
 	
