AGXRHI: Scrub all usage of mtlpp::RenderPassDescriptor

- Replace all uses of mtlpp::RenderPassDescriptor with direct usage of Metal
- Replace all uses of mtlpp::StoreAction with MTLStoreAction
- Replace all uses of mtlpp::LoadAction with MTLLoadAction

#rnx

diff --git a/Engine/Source/Runtime/Apple/AGXRHI/Private/AGXCommandEncoder.cpp b/Engine/Source/Runtime/Apple/AGXRHI/Private/AGXCommandEncoder.cpp
--- a/Engine/Source/Runtime/Apple/AGXRHI/Private/AGXCommandEncoder.cpp
+++ b/Engine/Source/Runtime/Apple/AGXRHI/Private/AGXCommandEncoder.cpp
@@ -51,10 +51,10 @@ FAGXCommandEncoder::FAGXCommandEncoder(FAGXCommandList& CmdList, EAGXCommandEnco
 	
 	for (uint32 i = 0; i < MaxSimultaneousRenderTargets; i++)
 	{
-		ColorStoreActions[i] = mtlpp::StoreAction::Unknown;
+		ColorStoreActions[i] = MTLStoreActionUnknown;
 	}
-	DepthStoreAction = mtlpp::StoreAction::Unknown;
-	StencilStoreAction = mtlpp::StoreAction::Unknown;
+	DepthStoreAction = MTLStoreActionUnknown;
+	StencilStoreAction = MTLStoreActionUnknown;
 }
 
 FAGXCommandEncoder::~FAGXCommandEncoder(void)
@@ -109,10 +109,10 @@ void FAGXCommandEncoder::Reset(void)
 	{
 		for (uint32 i = 0; i < MaxSimultaneousRenderTargets; i++)
 		{
-			ColorStoreActions[i] = mtlpp::StoreAction::Unknown;
+			ColorStoreActions[i] = MTLStoreActionUnknown;
 		}
-		DepthStoreAction = mtlpp::StoreAction::Unknown;
-		StencilStoreAction = mtlpp::StoreAction::Unknown;
+		DepthStoreAction = MTLStoreActionUnknown;
+		StencilStoreAction = MTLStoreActionUnknown;
 	}
 	
 	for (uint32 Frequency = 0; Frequency < uint32(mtlpp::FunctionType::Kernel)+1; Frequency++)
@@ -301,12 +301,12 @@ bool FAGXCommandEncoder::IsParallel(void) const
 	return CommandList.IsParallel() && (Type == EAGXCommandEncoderCurrent);
 }
 
-bool FAGXCommandEncoder::IsRenderPassDescriptorValid(void) const
+bool FAGXCommandEncoder::IsRenderPassDescriptorValid() const
 {
 	return (RenderPassDesc != nil);
 }
 
-mtlpp::RenderPassDescriptor const& FAGXCommandEncoder::GetRenderPassDescriptor(void) const
+MTLRenderPassDescriptor const* FAGXCommandEncoder::GetRenderPassDescriptor() const
 {
 	return RenderPassDesc;
 }
@@ -349,7 +349,7 @@ void FAGXCommandEncoder::BeginParallelRenderCommandEncoding(uint32 NumChildren)
 	check(CommandBuffer);
 	check(IsRenderCommandEncoderActive() == false && IsComputeCommandEncoderActive() == false && IsBlitCommandEncoderActive() == false);
 	
-	ParallelRenderCommandEncoder = [[CommandBuffer.GetPtr() parallelRenderCommandEncoderWithDescriptor:RenderPassDesc.GetPtr()] retain];
+	ParallelRenderCommandEncoder = [[CommandBuffer.GetPtr() parallelRenderCommandEncoderWithDescriptor:RenderPassDesc] retain];
 	
 	EncoderNum++;
 	
@@ -384,7 +384,7 @@ void FAGXCommandEncoder::BeginRenderCommandEncoding(void)
 	
 	if (!CommandList.IsParallel() || Type == EAGXCommandEncoderPrologue)
 	{
-		RenderCommandEncoder = [[CommandBuffer.GetPtr() renderCommandEncoderWithDescriptor:RenderPassDesc.GetPtr()] retain];
+		RenderCommandEncoder = [[CommandBuffer.GetPtr() renderCommandEncoderWithDescriptor:RenderPassDesc] retain];
 		EncoderNum++;	
 	}
 	else
@@ -479,27 +479,32 @@ void FAGXCommandEncoder::EndEncoding(void)
 				{
 					check(RenderPassDesc);
 					
-					ns::Array<mtlpp::RenderPassColorAttachmentDescriptor> ColorAttachments = RenderPassDesc.GetColorAttachments();
+					MTLRenderPassColorAttachmentDescriptorArray* ColorAttachments = [RenderPassDesc colorAttachments];
 					for (uint32 i = 0; i < MaxSimultaneousRenderTargets; i++)
 					{
-						if (ColorAttachments[i].GetTexture() && ColorAttachments[i].GetStoreAction() == mtlpp::StoreAction::Unknown)
+						MTLRenderPassColorAttachmentDescriptor* ColorAttachment = [ColorAttachments objectAtIndexedSubscript:i];
+						if ([ColorAttachment texture] && ([ColorAttachment storeAction] == MTLStoreActionUnknown))
 						{
-							mtlpp::StoreAction Action = ColorStoreActions[i];
-							check(Action != mtlpp::StoreAction::Unknown);
-							[RenderCommandEncoder setColorStoreAction:(MTLStoreAction)Action atIndex:i];
+							MTLStoreAction Action = ColorStoreActions[i];
+							check(Action != MTLStoreActionUnknown);
+							[RenderCommandEncoder setColorStoreAction:Action atIndex:i];
 						}
 					}
-					if (RenderPassDesc.GetDepthAttachment().GetTexture() && RenderPassDesc.GetDepthAttachment().GetStoreAction() == mtlpp::StoreAction::Unknown)
+
+					MTLRenderPassDepthAttachmentDescriptor* DepthAttachment = [RenderPassDesc depthAttachment];
+					if ([DepthAttachment texture] && ([DepthAttachment storeAction] == MTLStoreActionUnknown))
 					{
-						mtlpp::StoreAction Action = DepthStoreAction;
-						check(Action != mtlpp::StoreAction::Unknown);
-						[RenderCommandEncoder setDepthStoreAction:(MTLStoreAction)Action];
+						MTLStoreAction Action = DepthStoreAction;
+						check(Action != MTLStoreActionUnknown);
+						[RenderCommandEncoder setDepthStoreAction:Action];
 					}
-					if (RenderPassDesc.GetStencilAttachment().GetTexture() && RenderPassDesc.GetStencilAttachment().GetStoreAction() == mtlpp::StoreAction::Unknown)
+
+					MTLRenderPassStencilAttachmentDescriptor* StencilAttachment = [RenderPassDesc stencilAttachment];
+					if ([StencilAttachment texture] && ([StencilAttachment storeAction] == MTLStoreActionUnknown))
 					{
-						mtlpp::StoreAction Action = StencilStoreAction;
-						check(Action != mtlpp::StoreAction::Unknown);
-						[RenderCommandEncoder setStencilStoreAction:(MTLStoreAction)Action];
+						MTLStoreAction Action = StencilStoreAction;
+						check(Action != MTLStoreActionUnknown);
+						[RenderCommandEncoder setStencilStoreAction:Action];
 					}
 				}
 				
@@ -547,27 +552,32 @@ void FAGXCommandEncoder::EndEncoding(void)
 				{
 					check(RenderPassDesc);
 					
-					ns::Array<mtlpp::RenderPassColorAttachmentDescriptor> ColorAttachments = RenderPassDesc.GetColorAttachments();
+					MTLRenderPassColorAttachmentDescriptorArray* ColorAttachments = [RenderPassDesc colorAttachments];
 					for (uint32 i = 0; i < MaxSimultaneousRenderTargets; i++)
 					{
-						if (ColorAttachments[i].GetTexture() && ColorAttachments[i].GetStoreAction() == mtlpp::StoreAction::Unknown)
+						MTLRenderPassColorAttachmentDescriptor* ColorAttachment = [ColorAttachments objectAtIndexedSubscript:i];
+						if ([ColorAttachment texture] && ([ColorAttachment storeAction] == MTLStoreActionUnknown))
 						{
-							mtlpp::StoreAction Action = ColorStoreActions[i];
-							check(Action != mtlpp::StoreAction::Unknown);
-							[ParallelRenderCommandEncoder setColorStoreAction:(MTLStoreAction)Action atIndex:i];
+							MTLStoreAction Action = ColorStoreActions[i];
+							check(Action != MTLStoreActionUnknown);
+							[ParallelRenderCommandEncoder setColorStoreAction:Action atIndex:i];
 						}
 					}
-					if (RenderPassDesc.GetDepthAttachment().GetTexture() && RenderPassDesc.GetDepthAttachment().GetStoreAction() == mtlpp::StoreAction::Unknown)
+
+					MTLRenderPassDepthAttachmentDescriptor* DepthAttachment = [RenderPassDesc depthAttachment];
+					if ([DepthAttachment texture] && ([DepthAttachment storeAction] == MTLStoreActionUnknown))
 					{
-						mtlpp::StoreAction Action = DepthStoreAction;
-						check(Action != mtlpp::StoreAction::Unknown);
-						[ParallelRenderCommandEncoder setDepthStoreAction:(MTLStoreAction)Action];
+						MTLStoreAction Action = DepthStoreAction;
+						check(Action != MTLStoreActionUnknown);
+						[ParallelRenderCommandEncoder setDepthStoreAction:Action];
 					}
-					if (RenderPassDesc.GetStencilAttachment().GetTexture() && RenderPassDesc.GetStencilAttachment().GetStoreAction() == mtlpp::StoreAction::Unknown)
+
+					MTLRenderPassStencilAttachmentDescriptor* StencilAttachment = [RenderPassDesc stencilAttachment];
+					if ([StencilAttachment texture] && ([StencilAttachment storeAction] == MTLStoreActionUnknown))
 					{
-						mtlpp::StoreAction Action = StencilStoreAction;
-						check(Action != mtlpp::StoreAction::Unknown);
-						[ParallelRenderCommandEncoder setStencilStoreAction:(MTLStoreAction)Action];
+						MTLStoreAction Action = StencilStoreAction;
+						check(Action != MTLStoreActionUnknown);
+						[ParallelRenderCommandEncoder setStencilStoreAction:Action];
 					}
 				}
 
@@ -709,27 +719,27 @@ FAGXCommandBufferStats* FAGXCommandEncoder::GetCommandBufferStats(void)
 
 #pragma mark - Public Render State Mutators -
 
-void FAGXCommandEncoder::SetRenderPassDescriptor(mtlpp::RenderPassDescriptor RenderPass)
+void FAGXCommandEncoder::SetRenderPassDescriptor(MTLRenderPassDescriptor* InRenderPassDesc)
 {
 	check(IsRenderCommandEncoderActive() == false && IsComputeCommandEncoderActive() == false && IsBlitCommandEncoderActive() == false);
-	check(RenderPass);
+	check(InRenderPassDesc);
 	
-	if(RenderPass.GetPtr() != RenderPassDesc.GetPtr())
+	if (InRenderPassDesc != RenderPassDesc)
 	{
 		AGXSafeReleaseMetalRenderPassDescriptor(RenderPassDesc);
-		RenderPassDesc = RenderPass;
+		RenderPassDesc = InRenderPassDesc;
 		{
 			for (uint32 i = 0; i < MaxSimultaneousRenderTargets; i++)
 			{
-				ColorStoreActions[i] = mtlpp::StoreAction::Unknown;
+				ColorStoreActions[i] = MTLStoreActionUnknown;
 			}
-			DepthStoreAction = mtlpp::StoreAction::Unknown;
-			StencilStoreAction = mtlpp::StoreAction::Unknown;
+			DepthStoreAction = MTLStoreActionUnknown;
+			StencilStoreAction = MTLStoreActionUnknown;
 		}
 	}
 	check(RenderPassDesc);
 	
-	for (uint32 Frequency = 0; Frequency < uint32(mtlpp::FunctionType::Kernel)+1; Frequency++)
+	for (uint32 Frequency = 0; Frequency < uint32(MTLFunctionTypeKernel)+1; Frequency++)
 	{
 		for (uint32 i = 0; i < ML_MaxBuffers; i++)
 		{
@@ -743,7 +753,7 @@ void FAGXCommandEncoder::SetRenderPassDescriptor(mtlpp::RenderPassDescriptor Ren
 	}
 }
 
-void FAGXCommandEncoder::SetRenderPassStoreActions(mtlpp::StoreAction const* const ColorStore, mtlpp::StoreAction const DepthStore, mtlpp::StoreAction const StencilStore)
+void FAGXCommandEncoder::SetRenderPassStoreActions(MTLStoreAction const* ColorStore, MTLStoreAction DepthStore, MTLStoreAction StencilStore)
 {
 	check(RenderPassDesc);
 	{
@@ -866,7 +876,7 @@ void FAGXCommandEncoder::SetVisibilityResultMode(mtlpp::VisibilityResultMode con
 {
     check (RenderCommandEncoder);
 	{
-		check(Mode == mtlpp::VisibilityResultMode::Disabled || RenderPassDesc.GetVisibilityResultBuffer());
+		check(Mode == mtlpp::VisibilityResultMode::Disabled || [RenderPassDesc visibilityResultBuffer]);
 		[RenderCommandEncoder setVisibilityResultMode:(MTLVisibilityResultMode)Mode offset:Offset];
 	}
 }
diff --git a/Engine/Source/Runtime/Apple/AGXRHI/Private/AGXCommandEncoder.h b/Engine/Source/Runtime/Apple/AGXRHI/Private/AGXCommandEncoder.h
--- a/Engine/Source/Runtime/Apple/AGXRHI/Private/AGXCommandEncoder.h
+++ b/Engine/Source/Runtime/Apple/AGXRHI/Private/AGXCommandEncoder.h
@@ -12,6 +12,48 @@ class FAGXCommandQueue;
 class FAGXGraphicsPipelineState;
 struct FAGXCommandBufferFence;
 
+#if 0
+template<typename T>
+class TObjCWrapper : public FThreadSafeRefCountedObject
+{
+	using Pointer = T*;
+
+public:
+	UE_NONCOPYABLE(TObjCWrapper)
+
+	TObjCWrapper(Pointer InObject = nil)
+		: Impl([InObject retain])
+	{
+	}
+
+	~TObjCWrapper()
+	{
+		if (Impl)
+		{
+			[Impl release];
+			Impl = nil;
+		}
+	}
+
+	bool IsValid() const
+	{
+		return (Impl != nil);
+	}
+
+	Pointer AllocInit()
+	{
+		check(!IsValid());
+		Impl = [[T alloc] init];
+		return Impl;
+	}
+
+private:
+	Pointer Impl;
+};
+
+typedef TRefCountPtr<TObjCWrapper<MTLRenderPassDescriptor> > FMTLRenderPassDescriptor;
+#endif // 0
+
 /**
  * Enumeration for submission hints to avoid unclear bool values.
  */
@@ -122,10 +164,10 @@ public:
 	bool IsParallel(void) const;
 
 	/** @returns True if and only if there is valid render pass descriptor set on the encoder, otherwise false. */
-	bool IsRenderPassDescriptorValid(void) const;
+	bool IsRenderPassDescriptorValid() const;
 	
 	/** @returns The current render pass descriptor. */
-	mtlpp::RenderPassDescriptor const& GetRenderPassDescriptor(void) const;
+	MTLRenderPassDescriptor const* GetRenderPassDescriptor() const;
 	
 	/** @returns The active render command encoder or nil if there isn't one. */
 	id<MTLParallelRenderCommandEncoder> GetParallelRenderCommandEncoder() const;
@@ -148,13 +190,13 @@ public:
 #pragma mark - Public Command Encoder Mutators -
 
 	/**
- 	 * Begins encoding rendering commands into the current command buffer. No other encoder may be active & the mtlpp::RenderPassDescriptor must previously have been set.
+ 	 * Begins encoding rendering commands into the current command buffer. No other encoder may be active & the MTLRenderPassDescriptor must previously have been set.
 	 * @param NumChildren The number of child render-encoders to create. 
 	 */
 	void BeginParallelRenderCommandEncoding(uint32 NumChildren);
 
 	/**
- 	 * Begins encoding rendering commands into the current command buffer. No other encoder may be active & the mtlpp::RenderPassDescriptor must previously have been set.
+ 	 * Begins encoding rendering commands into the current command buffer. No other encoder may be active & the MTLRenderPassDescriptor must previously have been set.
 	 */
 	void BeginRenderCommandEncoding(void);
 	
@@ -201,15 +243,15 @@ public:
 	 * Set the render pass descriptor - no encoder may be active when this function is called.
 	 * @param RenderPass The render pass descriptor to set. May be nil.
 	 */
-	void SetRenderPassDescriptor(mtlpp::RenderPassDescriptor RenderPass);
+	void SetRenderPassDescriptor(MTLRenderPassDescriptor* InRenderPassDesc);
 	
 	/**
 	 * Set the render pass store actions, call after SetRenderPassDescriptor but before EndEncoding.
 	 * @param ColorStore The store actions for color targets.
-	 * @param DepthStore The store actions for the depth buffer - use mtlpp::StoreAction::Unknown if no depth-buffer bound.
-	 * @param StencilStore The store actions for the stencil buffer - use mtlpp::StoreAction::Unknown if no stencil-buffer bound.
+	 * @param DepthStore The store actions for the depth buffer - use MTLStoreActionUnknown if no depth-buffer bound.
+	 * @param StencilStore The store actions for the stencil buffer - use MTLStoreActionUnknown if no stencil-buffer bound.
 	 */
-	void SetRenderPassStoreActions(mtlpp::StoreAction const* const ColorStore, mtlpp::StoreAction const DepthStore, mtlpp::StoreAction const StencilStore);
+	void SetRenderPassStoreActions(MTLStoreAction const* ColorStore, MTLStoreAction DepthStore, MTLStoreAction StencilStore);
 	
 	/*
 	 * Sets the current render pipeline state object.
@@ -448,13 +490,13 @@ public:
     
 	FAGXBufferBindings ShaderBuffers[int(mtlpp::FunctionType::Kernel)+1];
 	
-	mtlpp::StoreAction ColorStoreActions[MaxSimultaneousRenderTargets];
-	mtlpp::StoreAction DepthStoreAction;
-	mtlpp::StoreAction StencilStoreAction;
+	MTLStoreAction ColorStoreActions[MaxSimultaneousRenderTargets];
+	MTLStoreAction DepthStoreAction;
+	MTLStoreAction StencilStoreAction;
 	
 	FAGXSubBufferRing RingBuffer;
 	
-	mtlpp::RenderPassDescriptor RenderPassDesc;
+	MTLRenderPassDescriptor* RenderPassDesc;
 	
 	mtlpp::CommandBuffer CommandBuffer;
 	id<MTLParallelRenderCommandEncoder> ParallelRenderCommandEncoder;
diff --git a/Engine/Source/Runtime/Apple/AGXRHI/Private/AGXRHIPrivate.h b/Engine/Source/Runtime/Apple/AGXRHI/Private/AGXRHIPrivate.h
--- a/Engine/Source/Runtime/Apple/AGXRHI/Private/AGXRHIPrivate.h
+++ b/Engine/Source/Runtime/Apple/AGXRHI/Private/AGXRHIPrivate.h
@@ -203,7 +203,7 @@ void AGXSafeReleaseMetalTexture(FAGXTexture& Object);
 void AGXSafeReleaseMetalBuffer(FAGXBuffer& Buffer);
 
 // Safely release a render pass descriptor so that it may be reused.
-void AGXSafeReleaseMetalRenderPassDescriptor(mtlpp::RenderPassDescriptor& Desc);
+void AGXSafeReleaseMetalRenderPassDescriptor(MTLRenderPassDescriptor* Desc);
 
 // Access the underlying surface object from any kind of texture
 FAGXSurface* AGXGetMetalSurfaceFromRHITexture(FRHITexture* Texture);
@@ -239,14 +239,14 @@ FORCEINLINE int32 GetMetalCubeFace(ECubeFace Face)
 	}
 }
 
-FORCEINLINE mtlpp::LoadAction GetMetalRTLoadAction(ERenderTargetLoadAction LoadAction)
+FORCEINLINE MTLLoadAction GetMetalRTLoadAction(ERenderTargetLoadAction LoadAction)
 {
 	switch(LoadAction)
 	{
-		case ERenderTargetLoadAction::ENoAction: return mtlpp::LoadAction::DontCare;
-		case ERenderTargetLoadAction::ELoad: return mtlpp::LoadAction::Load;
-		case ERenderTargetLoadAction::EClear: return mtlpp::LoadAction::Clear;
-		default: return mtlpp::LoadAction::DontCare;
+		case ERenderTargetLoadAction::ENoAction: return MTLLoadActionDontCare;
+		case ERenderTargetLoadAction::ELoad: return MTLLoadActionLoad;
+		case ERenderTargetLoadAction::EClear: return MTLLoadActionClear;
+		default: return MTLLoadActionDontCare;
 	}
 }
 
diff --git a/Engine/Source/Runtime/Apple/AGXRHI/Private/AGXRenderPass.cpp b/Engine/Source/Runtime/Apple/AGXRHI/Private/AGXRenderPass.cpp
--- a/Engine/Source/Runtime/Apple/AGXRHI/Private/AGXRenderPass.cpp
+++ b/Engine/Source/Runtime/Apple/AGXRHI/Private/AGXRenderPass.cpp
@@ -124,11 +124,11 @@ void FAGXRenderPass::Submit(EAGXSubmitFlags Flags)
 	}
 }
 
-void FAGXRenderPass::BeginParallelRenderPass(mtlpp::RenderPassDescriptor RenderPass, uint32 NumParallelContextsInPass)
+void FAGXRenderPass::BeginParallelRenderPass(MTLRenderPassDescriptor* InRenderPassDesc, uint32 NumParallelContextsInPass)
 {
 	check(!bWithinRenderPass);
 	check(!RenderPassDesc);
-	check(RenderPass);
+	check(InRenderPassDesc);
 	check(CurrentEncoder.GetCommandBuffer());
 
 	if (!CurrentEncoder.GetParallelRenderCommandEncoder())
@@ -143,21 +143,21 @@ void FAGXRenderPass::BeginParallelRenderPass(mtlpp::RenderPassDescriptor RenderP
 			CurrentEncoder.EndEncoding();
 		}
 
-		CurrentEncoder.SetRenderPassDescriptor(RenderPass);
+		CurrentEncoder.SetRenderPassDescriptor(InRenderPassDesc);
 
 		CurrentEncoder.BeginParallelRenderCommandEncoding(NumParallelContextsInPass);
 
-		RenderPassDesc = RenderPass;
+		RenderPassDesc = InRenderPassDesc;
 
 		bWithinRenderPass = true;
 	}
 }
 
-void FAGXRenderPass::BeginRenderPass(mtlpp::RenderPassDescriptor RenderPass)
+void FAGXRenderPass::BeginRenderPass(MTLRenderPassDescriptor* InRenderPassDesc)
 {
 	check(!bWithinRenderPass);
 	check(!RenderPassDesc);
-	check(RenderPass);
+	check(InRenderPassDesc);
 	check(!CurrentEncoder.IsRenderCommandEncoderActive());
 	if (!CmdList.IsParallel() && !CmdList.IsImmediate() && !CurrentEncoder.GetCommandBuffer())
 	{
@@ -178,7 +178,7 @@ void FAGXRenderPass::BeginRenderPass(mtlpp::RenderPassDescriptor RenderPass)
 	State.SetStateDirty();
 	State.SetRenderTargetsActive(true);
 	
-	RenderPassDesc = RenderPass;
+	RenderPassDesc = InRenderPassDesc;
 	
 	CurrentEncoder.SetRenderPassDescriptor(RenderPassDesc);
 
@@ -194,79 +194,32 @@ void FAGXRenderPass::BeginRenderPass(mtlpp::RenderPassDescriptor RenderPass)
 	check(!PrologueEncoder.IsBlitCommandEncoderActive() && !PrologueEncoder.IsComputeCommandEncoderActive());
 }
 
-void FAGXRenderPass::RestartRenderPass(mtlpp::RenderPassDescriptor RenderPass)
+void FAGXRenderPass::RestartRenderPass(MTLRenderPassDescriptor* InRenderPassDesc)
 {
 	check(bWithinRenderPass);
 	check(RenderPassDesc);
 	check(CmdList.IsParallel() || CurrentEncoder.GetCommandBuffer());
 	
-	mtlpp::RenderPassDescriptor StartDesc;
-	if (RenderPass != nil)
+	MTLRenderPassDescriptor* StartDesc = nil;
+	if (InRenderPassDesc != nil)
 	{
 		// Just restart with the render pass we were given - the caller should have ensured that this is restartable
 		check(State.CanRestartRenderPass());
-		StartDesc = RenderPass;
+		StartDesc = InRenderPassDesc;
 	}
-	else if (State.PrepareToRestart(CurrentEncoder.IsRenderPassDescriptorValid() && (State.GetRenderPassDescriptor().GetPtr() == CurrentEncoder.GetRenderPassDescriptor().GetPtr())))
+	else if (State.PrepareToRestart(CurrentEncoder.IsRenderPassDescriptorValid() && (State.GetRenderPassDescriptor() == CurrentEncoder.GetRenderPassDescriptor())))
 	{
 		// Restart with the render pass we have in the state cache - the state cache says its safe
 		StartDesc = State.GetRenderPassDescriptor();
 	}
 	else
 	{
-		
-		METAL_FATAL_ERROR(TEXT("Failed to restart render pass with descriptor: %s"), *FString([RenderPassDesc.GetPtr() description]));
+		METAL_FATAL_ERROR(TEXT("Failed to restart render pass with descriptor: %s"), *FString([RenderPassDesc description]));
 	}
 	check(StartDesc);
 	
 	RenderPassDesc = StartDesc;
 	
-#if METAL_DEBUG_OPTIONS
-	if ((GetAGXDeviceContext().GetCommandQueue().GetRuntimeDebuggingLevel() >= EAGXDebugLevelValidation))
-	{
-		bool bAllLoadActionsOK = true;
-		ns::Array<mtlpp::RenderPassColorAttachmentDescriptor> Attachments = RenderPassDesc.GetColorAttachments();
-		for(uint i = 0; i < 8; i++)
-		{
-			mtlpp::RenderPassColorAttachmentDescriptor Desc = Attachments[i];
-			if(Desc && Desc.GetTexture())
-			{
-				bAllLoadActionsOK &= (Desc.GetLoadAction() != mtlpp::LoadAction::Clear);
-			}
-		}
-		if(RenderPassDesc.GetDepthAttachment() && RenderPassDesc.GetDepthAttachment().GetTexture())
-		{
-			bAllLoadActionsOK &= (RenderPassDesc.GetDepthAttachment().GetLoadAction() != mtlpp::LoadAction::Clear);
-		}
-		if(RenderPassDesc.GetStencilAttachment() && RenderPassDesc.GetStencilAttachment().GetTexture())
-		{
-			bAllLoadActionsOK &= (RenderPassDesc.GetStencilAttachment().GetLoadAction() != mtlpp::LoadAction::Clear);
-		}
-		
-		if (!bAllLoadActionsOK)
-		{
-			UE_LOG(LogAGX, Warning, TEXT("Tried to restart render encoding with a clear operation - this would erroneously re-clear any existing draw calls: %s"), *FString([RenderPassDesc.GetPtr() description]));
-			
-			for(uint i = 0; i< 8; i++)
-			{
-				mtlpp::RenderPassColorAttachmentDescriptor Desc = Attachments[i];
-				if(Desc && Desc.GetTexture())
-				{
-					Desc.SetLoadAction(mtlpp::LoadAction::Load);
-				}
-			}
-			if(RenderPassDesc.GetDepthAttachment() && RenderPassDesc.GetDepthAttachment().GetTexture())
-			{
-				RenderPassDesc.GetDepthAttachment().SetLoadAction(mtlpp::LoadAction::Load);
-			}
-			if(RenderPassDesc.GetStencilAttachment() && RenderPassDesc.GetStencilAttachment().GetTexture())
-			{
-				RenderPassDesc.GetStencilAttachment().SetLoadAction(mtlpp::LoadAction::Load);
-			}
-		}
-	}
-#endif
-	
 	// EndEncoding should provide the encoder fence...
 	if (CurrentEncoder.IsBlitCommandEncoderActive() || CurrentEncoder.IsComputeCommandEncoderActive() || CurrentEncoder.IsRenderCommandEncoderActive())
 	{
diff --git a/Engine/Source/Runtime/Apple/AGXRHI/Private/AGXRenderPass.h b/Engine/Source/Runtime/Apple/AGXRHI/Private/AGXRenderPass.h
--- a/Engine/Source/Runtime/Apple/AGXRHI/Private/AGXRenderPass.h
+++ b/Engine/Source/Runtime/Apple/AGXRHI/Private/AGXRenderPass.h
@@ -25,11 +25,11 @@ public:
 	
     void Begin(bool const bParallelBegin = false);
 	
-    void BeginParallelRenderPass(mtlpp::RenderPassDescriptor RenderPass, uint32 NumParallelContextsInPass);
+    void BeginParallelRenderPass(MTLRenderPassDescriptor* InRenderPassDesc, uint32 NumParallelContextsInPass);
 
-    void BeginRenderPass(mtlpp::RenderPassDescriptor RenderPass);
+    void BeginRenderPass(MTLRenderPassDescriptor* InRenderPassDesc);
 
-    void RestartRenderPass(mtlpp::RenderPassDescriptor RenderPass);
+    void RestartRenderPass(MTLRenderPassDescriptor* InRenderPassDesc);
     
     void DrawPrimitive(uint32 PrimitiveType, uint32 BaseVertexIndex, uint32 NumPrimitives, uint32 NumInstances);
     
@@ -187,7 +187,7 @@ private:
 	// Disjoint ranges *are* permitted!
 	TMap<id<MTLBuffer>, TArray<NSRange>> OutstandingBufferUploads;
 
-    mtlpp::RenderPassDescriptor RenderPassDesc;
+    MTLRenderPassDescriptor* RenderPassDesc;
     
 	mtlpp::DispatchType ComputeDispatchType;
     uint32 NumOutstandingOps;
diff --git a/Engine/Source/Runtime/Apple/AGXRHI/Private/AGXStateCache.cpp b/Engine/Source/Runtime/Apple/AGXRHI/Private/AGXStateCache.cpp
--- a/Engine/Source/Runtime/Apple/AGXRHI/Private/AGXStateCache.cpp
+++ b/Engine/Source/Runtime/Apple/AGXRHI/Private/AGXStateCache.cpp
@@ -47,13 +47,13 @@ static mtlpp::DepthClipMode TranslateDepthClipMode(ERasterizerDepthClipMode Dept
 	}
 }
 
-FORCEINLINE mtlpp::StoreAction GetMetalRTStoreAction(ERenderTargetStoreAction StoreAction)
+FORCEINLINE MTLStoreAction GetMetalRTStoreAction(ERenderTargetStoreAction StoreAction)
 {
 	switch(StoreAction)
 	{
-		case ERenderTargetStoreAction::ENoAction: return mtlpp::StoreAction::DontCare;
-		case ERenderTargetStoreAction::EStore: return mtlpp::StoreAction::Store;
-		//default store action in the desktop renderers needs to be mtlpp::StoreAction::StoreAndMultisampleResolve.  Trying to express the renderer by the requested maxrhishaderplatform
+		case ERenderTargetStoreAction::ENoAction: return MTLStoreActionDontCare;
+		case ERenderTargetStoreAction::EStore: return MTLStoreActionStore;
+		//default store action in the desktop renderers needs to be MTLStoreActionStoreAndMultisampleResolve.  Trying to express the renderer by the requested maxrhishaderplatform
         //because we may render to the same MSAA target twice in two separate passes.  BasePass, then some stuff, then translucency for example and we need to not lose the prior MSAA contents to do this properly.
 		case ERenderTargetStoreAction::EMultisampleResolve:
 		{
@@ -61,33 +61,33 @@ FORCEINLINE mtlpp::StoreAction GetMetalRTStoreAction(ERenderTargetStoreAction St
 			static bool bSupportsMSAAStoreResolve = FAGXCommandQueue::SupportsFeature(EAGXFeaturesMSAAStoreAndResolve) && (GMaxRHIFeatureLevel >= ERHIFeatureLevel::SM5);
             if (bNoMSAA)
             {
-                return mtlpp::StoreAction::Store;
+                return MTLStoreActionStore;
             }
 			else if (bSupportsMSAAStoreResolve)
 			{
-				return mtlpp::StoreAction::StoreAndMultisampleResolve;
+				return MTLStoreActionStoreAndMultisampleResolve;
 			}
 			else
 			{
-				return mtlpp::StoreAction::MultisampleResolve;
+				return MTLStoreActionMultisampleResolve;
 			}
 		}
-		default: return mtlpp::StoreAction::DontCare;
+		default: return MTLStoreActionDontCare;
 	}
 }
 
-FORCEINLINE mtlpp::StoreAction GetConditionalMetalRTStoreAction(bool bMSAATarget)
+FORCEINLINE MTLStoreAction GetConditionalMetalRTStoreAction(bool bMSAATarget)
 {
 	if (bMSAATarget)
 	{
 		//this func should only be getting called when an encoder had to abnormally break.  In this case we 'must' do StoreAndResolve because the encoder will be restarted later
 		//with the original MSAA rendertarget and the original data must still be there to continue the render properly.
 		check(FAGXCommandQueue::SupportsFeature(EAGXFeaturesMSAAStoreAndResolve));
-		return mtlpp::StoreAction::StoreAndMultisampleResolve;
+		return MTLStoreActionStoreAndMultisampleResolve;
 	}
 	else
 	{
-		return mtlpp::StoreAction::Store;
+		return MTLStoreActionStore;
 	}	
 }
 
@@ -101,50 +101,51 @@ public:
 	
 	~FAGXRenderPassDescriptorPool()
 	{
-		
+		MTLRenderPassDescriptor* Desc = nil;
+		while (nil != (Desc = Cache.Pop()))
+		{
+			[Desc release];
+		}
 	}
 	
-	mtlpp::RenderPassDescriptor CreateDescriptor()
+	MTLRenderPassDescriptor* CreateDescriptor()
 	{
 		MTLRenderPassDescriptor* Desc = Cache.Pop();
 		if (!Desc)
 		{
-			Desc = [MTLRenderPassDescriptor new];
+			Desc = [[MTLRenderPassDescriptor renderPassDescriptor] retain];
 		}
-		return mtlpp::RenderPassDescriptor(Desc);
+		return Desc;
 	}
 	
-	void ReleaseDescriptor(mtlpp::RenderPassDescriptor& Desc)
+	void ReleaseDescriptor(MTLRenderPassDescriptor* Desc)
 	{
-		mtlpp::Texture EmptyTex;
-		
-		ns::Array<mtlpp::RenderPassColorAttachmentDescriptor> Attachements = Desc.GetColorAttachments();
+		MTLRenderPassColorAttachmentDescriptorArray* Attachments = [Desc colorAttachments];
 		for (uint32 i = 0; i < MaxSimultaneousRenderTargets; i++)
 		{
-			mtlpp::RenderPassColorAttachmentDescriptor Color = Attachements[i];
-			Color.SetTexture(EmptyTex);
-			Color.SetResolveTexture(EmptyTex);
-			Color.SetStoreAction(mtlpp::StoreAction::Store);
+			MTLRenderPassColorAttachmentDescriptor* Color = [Attachments objectAtIndexedSubscript:i];
+			[Color setTexture:nil];
+			[Color setResolveTexture:nil];
+			[Color setStoreAction:MTLStoreActionStore];
 		}
 		
-		mtlpp::RenderPassDepthAttachmentDescriptor Depth = Desc.GetDepthAttachment();
-		Depth.SetTexture(EmptyTex);
-		Depth.SetResolveTexture(EmptyTex);
-		Depth.SetStoreAction(mtlpp::StoreAction::Store);
-
-		mtlpp::RenderPassStencilAttachmentDescriptor Stencil = Desc.GetStencilAttachment();
-		Stencil.SetTexture(EmptyTex);
-		Stencil.SetResolveTexture(EmptyTex);
-		Stencil.SetStoreAction(mtlpp::StoreAction::Store);
-
-		mtlpp::Buffer Empty;
-		Desc.SetVisibilityResultBuffer(Empty);
+		MTLRenderPassDepthAttachmentDescriptor* Depth = [Desc depthAttachment];
+		[Depth setTexture:nil];
+		[Depth setResolveTexture:nil];
+		[Depth setStoreAction:MTLStoreActionStore];
+
+		MTLRenderPassStencilAttachmentDescriptor* Stencil = [Desc stencilAttachment];
+		[Stencil setTexture:nil];
+		[Stencil setResolveTexture:nil];
+		[Stencil setStoreAction:MTLStoreActionStore];
+
+		[Desc setVisibilityResultBuffer:nil];
 		
 #if PLATFORM_MAC
-		Desc.SetRenderTargetArrayLength(1);
+		[Desc setRenderTargetArrayLength:1];
 #endif
 		
-		Cache.Push(Desc.GetPtr());
+		Cache.Push(Desc);
 	}
 	
 	static FAGXRenderPassDescriptorPool& Get()
@@ -157,17 +158,17 @@ private:
 	TLockFreePointerListLIFO<MTLRenderPassDescriptor> Cache;
 };
 
-void AGXSafeReleaseMetalRenderPassDescriptor(mtlpp::RenderPassDescriptor& Desc)
+void AGXSafeReleaseMetalRenderPassDescriptor(MTLRenderPassDescriptor* Desc)
 {
-	if (Desc.GetPtr())
+	if (Desc)
 	{
 		FAGXRenderPassDescriptorPool::Get().ReleaseDescriptor(Desc);
 	}
 }
 
 FAGXStateCache::FAGXStateCache(bool const bInImmediate)
-: DepthStore(mtlpp::StoreAction::Unknown)
-, StencilStore(mtlpp::StoreAction::Unknown)
+: DepthStore(MTLStoreActionUnknown)
+, StencilStore(MTLStoreActionUnknown)
 , VisibilityResults(nullptr)
 , VisibilityMode(mtlpp::VisibilityResultMode::Disabled)
 , VisibilityOffset(0)
@@ -197,7 +198,7 @@ FAGXStateCache::FAGXStateCache(bool const bInImmediate)
 	
 	for (uint32 i = 0; i < MaxSimultaneousRenderTargets; i++)
 	{
-		ColorStore[i] = mtlpp::StoreAction::Unknown;
+		ColorStore[i] = MTLStoreActionUnknown;
 	}
 	
 	FMemory::Memzero(RenderPassInfo);
@@ -313,10 +314,10 @@ void FAGXStateCache::Reset(void)
 	
 	for (uint32 i = 0; i < MaxSimultaneousRenderTargets; i++)
 	{
-		ColorStore[i] = mtlpp::StoreAction::Unknown;
+		ColorStore[i] = MTLStoreActionUnknown;
 	}
-	DepthStore = mtlpp::StoreAction::Unknown;
-	StencilStore = mtlpp::StoreAction::Unknown;
+	DepthStore = MTLStoreActionUnknown;
+	StencilStore = MTLStoreActionUnknown;
 	
 	BlendFactor = FLinearColor::Transparent;
 	FrameBufferSize = CGSizeMake(0.0, 0.0);
@@ -431,26 +432,26 @@ bool FAGXStateCache::SetRenderPassInfo(FRHIRenderPassInfo const& InRenderTargets
 		bool bNeedsClear = false;
 		
 		//Create local store action states if we support deferred store
-		mtlpp::StoreAction NewColorStore[MaxSimultaneousRenderTargets];
+		MTLStoreAction NewColorStore[MaxSimultaneousRenderTargets];
 		for (uint32 i = 0; i < MaxSimultaneousRenderTargets; ++i)
 		{
-			NewColorStore[i] = mtlpp::StoreAction::Unknown;
+			NewColorStore[i] = MTLStoreActionUnknown;
 		}
 		
-		mtlpp::StoreAction NewDepthStore = mtlpp::StoreAction::Unknown;
-		mtlpp::StoreAction NewStencilStore = mtlpp::StoreAction::Unknown;
+		MTLStoreAction NewDepthStore = MTLStoreActionUnknown;
+		MTLStoreAction NewStencilStore = MTLStoreActionUnknown;
 		
 		// back this up for next frame
 		RenderPassInfo = InRenderTargets;
 		
 		// at this point, we need to fully set up an encoder/command buffer, so make a new one (autoreleased)
-		mtlpp::RenderPassDescriptor RenderPass = FAGXRenderPassDescriptorPool::Get().CreateDescriptor();
+		MTLRenderPassDescriptor* RenderPass = FAGXRenderPassDescriptorPool::Get().CreateDescriptor();
 	
 		// if we need to do queries, write to the supplied query buffer
 		if (IsFeatureLevelSupported(GMaxRHIShaderPlatform, ERHIFeatureLevel::ES3_1))
 		{
 			VisibilityResults = QueryBuffer;
-			RenderPass.SetVisibilityResultBuffer(QueryBuffer ? QueryBuffer->Buffer : nil);
+			[RenderPass setVisibilityResultBuffer:(QueryBuffer ? QueryBuffer->Buffer.GetPtr() : nil)];
 		}
 		else
 		{
@@ -482,7 +483,7 @@ bool FAGXStateCache::SetRenderPassInfo(FRHIRenderPassInfo const& InRenderTargets
 		
 		bCanRestartRenderPass = true;
 		
-		ns::Array<mtlpp::RenderPassColorAttachmentDescriptor> Attachements = RenderPass.GetColorAttachments();
+		MTLRenderPassColorAttachmentDescriptorArray* Attachements = [RenderPass colorAttachments];
 		
 		uint32 NumColorRenderTargets = RenderPassInfo.GetNumColorRenderTargets();
 		
@@ -564,7 +565,7 @@ bool FAGXStateCache::SetRenderPassInfo(FRHIRenderPassInfo const& InRenderTargets
 						break;
 				}
 	
-				mtlpp::RenderPassColorAttachmentDescriptor ColorAttachment = Attachements[RenderTargetIndex];
+				MTLRenderPassColorAttachmentDescriptor* ColorAttachment = [Attachements objectAtIndexedSubscript:RenderTargetIndex];
 	
 				ERenderTargetStoreAction HighLevelStoreAction = GetStoreAction(RenderTargetView.Action);
 				ERenderTargetLoadAction HighLevelLoadAction = GetLoadAction(RenderTargetView.Action);
@@ -590,10 +591,10 @@ bool FAGXStateCache::SetRenderPassInfo(FRHIRenderPassInfo const& InRenderTargets
 					}
 #endif
 					// set up an MSAA attachment
-					ColorAttachment.SetTexture(Surface.MSAATexture);
+					[ColorAttachment setTexture:Surface.MSAATexture.GetPtr()];
 					NewColorStore[RenderTargetIndex] = GetMetalRTStoreAction(ERenderTargetStoreAction::EMultisampleResolve);
-					ColorAttachment.SetStoreAction(!bMemoryless && GRHIDeviceId > 2 ? mtlpp::StoreAction::Unknown : NewColorStore[RenderTargetIndex]);
-					ColorAttachment.SetResolveTexture(Surface.MSAAResolveTexture ? Surface.MSAAResolveTexture : Surface.Texture);
+					[ColorAttachment setStoreAction:(!bMemoryless && GRHIDeviceId > 2 ? MTLStoreActionUnknown : NewColorStore[RenderTargetIndex])];
+					[ColorAttachment setResolveTexture:(Surface.MSAAResolveTexture ? Surface.MSAAResolveTexture.GetPtr() : Surface.Texture.GetPtr())];
 					SampleCount = Surface.MSAATexture.GetSampleCount();
 				}
 				else
@@ -607,37 +608,37 @@ bool FAGXStateCache::SetRenderPassInfo(FRHIRenderPassInfo const& InRenderTargets
 					}
 #endif
 					// set up non-MSAA attachment
-					ColorAttachment.SetTexture(Surface.Texture);
+					[ColorAttachment setTexture:Surface.Texture.GetPtr()];
 					NewColorStore[RenderTargetIndex] = GetMetalRTStoreAction(HighLevelStoreAction);
-					ColorAttachment.SetStoreAction(!bMemoryless ? mtlpp::StoreAction::Unknown : NewColorStore[RenderTargetIndex]);
-                    SampleCount = 1;
+					[ColorAttachment setStoreAction:(!bMemoryless ? MTLStoreActionUnknown : NewColorStore[RenderTargetIndex])];
+					SampleCount = 1;
 				}
 				
-				ColorAttachment.SetLevel(RenderTargetView.MipIndex);
+				[ColorAttachment setLevel:RenderTargetView.MipIndex];
 				if(Surface.GetDesc().IsTexture3D())
 				{
-					ColorAttachment.SetSlice(0);
-					ColorAttachment.SetDepthPlane(ArraySliceIndex);
+					[ColorAttachment setSlice:0];
+					[ColorAttachment setDepthPlane:ArraySliceIndex];
 				}
 				else
 				{
-					ColorAttachment.SetSlice(ArraySliceIndex);
+					[ColorAttachment setSlice:ArraySliceIndex];
 				}
 				
-				ColorAttachment.SetLoadAction((Surface.Written || !bImmediate || bRestart) ? GetMetalRTLoadAction(HighLevelLoadAction) : mtlpp::LoadAction::Clear);
+				[ColorAttachment setLoadAction:((Surface.Written || !bImmediate || bRestart) ? GetMetalRTLoadAction(HighLevelLoadAction) : MTLLoadActionClear)];
 				FPlatformAtomics::InterlockedExchange(&Surface.Written, 1);
 				
-				bNeedsClear |= (ColorAttachment.GetLoadAction() == mtlpp::LoadAction::Clear);
+				bNeedsClear |= ([ColorAttachment loadAction] == MTLLoadActionClear);
 				
 				const FClearValueBinding& ClearValue = RenderPassInfo.ColorRenderTargets[RenderTargetIndex].RenderTarget->GetClearBinding();
 				if (ClearValue.ColorBinding == EClearBinding::EColorBound)
 				{
 					const FLinearColor& ClearColor = ClearValue.GetClearColor();
-					ColorAttachment.SetClearColor(mtlpp::ClearColor(ClearColor.R, ClearColor.G, ClearColor.B, ClearColor.A));
+					[ColorAttachment setClearColor:MTLClearColorMake(ClearColor.R, ClearColor.G, ClearColor.B, ClearColor.A)];
 				}
 
 				bCanRestartRenderPass &= 	!bMemoryless &&
-											ColorAttachment.GetLoadAction() == mtlpp::LoadAction::Load &&
+											[ColorAttachment loadAction] == MTLLoadActionLoad &&
 											HighLevelStoreAction != ERenderTargetStoreAction::ENoAction;
 	
 				bHasValidRenderTarget = true;
@@ -663,7 +664,7 @@ bool FAGXStateCache::SetRenderPassInfo(FRHIRenderPassInfo const& InRenderTargets
 			{
 				METAL_FATAL_ASSERT(ArrayTargets == BoundTargets, TEXT("All color render targets must be layered when performing multi-layered rendering under Metal (%d != %d)."), ArrayTargets, BoundTargets);
 					RenderTargetArraySize = ArrayRenderLayers;
-					RenderPass.SetRenderTargetArrayLength(ArrayRenderLayers);
+					[RenderPass setRenderTargetArrayLength:ArrayRenderLayers];
 			}
 #endif
 		}
@@ -694,7 +695,7 @@ bool FAGXStateCache::SetRenderPassInfo(FRHIRenderPassInfo const& InRenderTargets
 				METAL_FATAL_ASSERT(GetAGXDeviceContext().SupportsFeature(EAGXFeaturesLayeredRendering), TEXT("Layered rendering is unsupported on this device (%d)."), ArrayRenderLayers);
 #if PLATFORM_MAC
 					RenderTargetArraySize = ArrayRenderLayers;
-					RenderPass.SetRenderTargetArrayLength(ArrayRenderLayers);
+					[RenderPass setRenderTargetArrayLength:ArrayRenderLayers];
 #endif
 			}
 			
@@ -790,11 +791,11 @@ bool FAGXStateCache::SetRenderPassInfo(FRHIRenderPassInfo const& InRenderTargets
 				DepthClearValue = 1.0f;
 			}
 
-           bool const bCombinedDepthStencilUsingStencil = (DepthTexture && (mtlpp::PixelFormat)DepthTexture.GetPixelFormat() != mtlpp::PixelFormat::Depth32Float && RenderPassInfo.DepthStencilRenderTarget.ExclusiveDepthStencil.IsUsingStencil());
+			bool const bCombinedDepthStencilUsingStencil = (DepthTexture && (mtlpp::PixelFormat)DepthTexture.GetPixelFormat() != mtlpp::PixelFormat::Depth32Float && RenderPassInfo.DepthStencilRenderTarget.ExclusiveDepthStencil.IsUsingStencil());
 			bool const bUsingDepth = (RenderPassInfo.DepthStencilRenderTarget.ExclusiveDepthStencil.IsUsingDepth() || (bCombinedDepthStencilUsingStencil));
 			if (DepthTexture && bUsingDepth)
 			{
-				mtlpp::RenderPassDepthAttachmentDescriptor DepthAttachment;
+				MTLRenderPassDepthAttachmentDescriptor* DepthAttachment = [RenderPass depthAttachment];
 				
 				DepthFormatKey = Surface.FormatKey;
 				
@@ -803,10 +804,10 @@ bool FAGXStateCache::SetRenderPassInfo(FRHIRenderPassInfo const& InRenderTargets
 				ERenderTargetStoreAction DepthStoreAction = GetStoreAction(DepthActions);
 
 				// set up the depth attachment
-				DepthAttachment.SetTexture(DepthTexture);
-				DepthAttachment.SetLoadAction(GetMetalRTLoadAction(DepthLoadAction));
+				[DepthAttachment setTexture:DepthTexture.GetPtr()];
+				[DepthAttachment setLoadAction:GetMetalRTLoadAction(DepthLoadAction)];
 				
-				bNeedsClear |= (DepthAttachment.GetLoadAction() == mtlpp::LoadAction::Clear);
+				bNeedsClear |= ([DepthAttachment loadAction] == MTLLoadActionClear);
 				
 				ERenderTargetStoreAction HighLevelStoreAction = (Surface.MSAATexture && !bDepthStencilSampleCountMismatchFixup) ? ERenderTargetStoreAction::EMultisampleResolve : DepthStoreAction;
 				if (bUsingDepth && (HighLevelStoreAction == ERenderTargetStoreAction::ENoAction || bDepthStencilSampleCountMismatchFixup))
@@ -827,7 +828,7 @@ bool FAGXStateCache::SetRenderPassInfo(FRHIRenderPassInfo const& InRenderTargets
 				bDepthTextureMemoryless = DepthTexture.GetStorageMode() == mtlpp::StorageMode::Memoryless;
 				if (bDepthTextureMemoryless)
 				{
-					DepthAttachment.SetLoadAction(mtlpp::LoadAction::Clear);
+					[DepthAttachment setLoadAction:MTLLoadActionClear];
 					
 					if (bSupportsMSAADepthResolve && Surface.MSAATexture && DepthStoreAction == ERenderTargetStoreAction::EMultisampleResolve)
 					{
@@ -845,22 +846,22 @@ bool FAGXStateCache::SetRenderPassInfo(FRHIRenderPassInfo const& InRenderTargets
 #endif
                 //needed to quiet the metal validation that runs when you end renderpass. (it requires some kind of 'resolve' for an msaa target)
 				//But with deferredstore we don't set the real one until submit time.
-				NewDepthStore = !Surface.MSAATexture || bSupportsMSAADepthResolve ? GetMetalRTStoreAction(HighLevelStoreAction) : mtlpp::StoreAction::DontCare;
-				DepthAttachment.SetStoreAction(!bDepthTextureMemoryless && Surface.MSAATexture && GRHIDeviceId > 2 ? mtlpp::StoreAction::Unknown : NewDepthStore);
-				DepthAttachment.SetClearDepth(DepthClearValue);
+				NewDepthStore = !Surface.MSAATexture || bSupportsMSAADepthResolve ? GetMetalRTStoreAction(HighLevelStoreAction) : MTLStoreActionDontCare;
+				[DepthAttachment setStoreAction:(!bDepthTextureMemoryless && Surface.MSAATexture && GRHIDeviceId > 2 ? MTLStoreActionUnknown : NewDepthStore)];
+				[DepthAttachment setClearDepth:DepthClearValue];
 				check(SampleCount > 0);
 
-				if (Surface.MSAATexture && bSupportsMSAADepthResolve && DepthAttachment.GetStoreAction() != mtlpp::StoreAction::DontCare)
+				if (Surface.MSAATexture && bSupportsMSAADepthResolve && [DepthAttachment storeAction] != MTLStoreActionDontCare)
 				{
                     if (!bDepthStencilSampleCountMismatchFixup)
                     {
-                        DepthAttachment.SetResolveTexture(Surface.MSAAResolveTexture ? Surface.MSAAResolveTexture : Surface.Texture);
+                        [DepthAttachment setResolveTexture:(Surface.MSAAResolveTexture ? Surface.MSAAResolveTexture.GetPtr() : Surface.Texture.GetPtr())];
                     }
 #if PLATFORM_MAC
 					//would like to assert and do manual custom resolve, but that is causing some kind of weird corruption.
 					//checkf(false, TEXT("Depth resolves need to do 'max' for correctness.  MacOS does not expose this yet unless the spec changed."));
 #else
-					DepthAttachment.SetDepthResolveFilter(mtlpp::MultisampleDepthResolveFilter::Max);
+					[DepthAttachment setDepthResolveFilter:MTLMultisampleDepthResolveFilterMax];
 #endif
 				}
 				
@@ -871,11 +872,11 @@ bool FAGXStateCache::SetRenderPassInfo(FRHIRenderPassInfo const& InRenderTargets
 				bCanRestartRenderPass &=	(DepthSampleCount <= 1 || bDepthMSAARestart) &&
 											(
 												(RenderPassInfo.DepthStencilRenderTarget.DepthStencilTarget == FallbackDepthStencilSurface) ||
-												((DepthAttachment.GetLoadAction() == mtlpp::LoadAction::Load) && (bDepthMSAARestart || !RenderPassInfo.DepthStencilRenderTarget.ExclusiveDepthStencil.IsDepthWrite() || DepthStoreAction == ERenderTargetStoreAction::EStore))
+												(([DepthAttachment loadAction] == MTLLoadActionLoad) && (bDepthMSAARestart || !RenderPassInfo.DepthStencilRenderTarget.ExclusiveDepthStencil.IsDepthWrite() || DepthStoreAction == ERenderTargetStoreAction::EStore))
 											);
 				
 				// and assign it
-				RenderPass.SetDepthAttachment(DepthAttachment);
+				[RenderPass setDepthAttachment:DepthAttachment];
 			}
 	
             //if we're dealing with a samplecount mismatch we just bail on stencil entirely as stencil
@@ -885,7 +886,7 @@ bool FAGXStateCache::SetRenderPassInfo(FRHIRenderPassInfo const& InRenderTargets
 			bool const bUsingStencil = RenderPassInfo.DepthStencilRenderTarget.ExclusiveDepthStencil.IsUsingStencil() || (bCombinedDepthStencilUsingDepth);
 			if (StencilTexture && bUsingStencil)
 			{
-				mtlpp::RenderPassStencilAttachmentDescriptor StencilAttachment;
+				MTLRenderPassStencilAttachmentDescriptor* StencilAttachment = [RenderPass stencilAttachment];
 				
 				StencilFormatKey = Surface.FormatKey;
 				
@@ -894,10 +895,10 @@ bool FAGXStateCache::SetRenderPassInfo(FRHIRenderPassInfo const& InRenderTargets
 				ERenderTargetStoreAction StencilStoreAction = GetStoreAction(StencilActions);
 	
 				// set up the stencil attachment
-				StencilAttachment.SetTexture(StencilTexture);
-				StencilAttachment.SetLoadAction(GetMetalRTLoadAction(StencilLoadAction));
+				[StencilAttachment setTexture:StencilTexture.GetPtr()];
+				[StencilAttachment setLoadAction:GetMetalRTLoadAction(StencilLoadAction)];
 				
-				bNeedsClear |= (StencilAttachment.GetLoadAction() == mtlpp::LoadAction::Clear);
+				bNeedsClear |= ([StencilAttachment loadAction] == MTLLoadActionClear);
 				
 				ERenderTargetStoreAction HighLevelStoreAction = StencilStoreAction;
 				if (bUsingStencil && (HighLevelStoreAction == ERenderTargetStoreAction::ENoAction || bDepthStencilSampleCountMismatchFixup))
@@ -911,7 +912,7 @@ bool FAGXStateCache::SetRenderPassInfo(FRHIRenderPassInfo const& InRenderTargets
 				{
 					bStencilMemoryless = true;
 					HighLevelStoreAction = ERenderTargetStoreAction::ENoAction;
-					StencilAttachment.SetLoadAction(mtlpp::LoadAction::Clear);
+					[StencilAttachment setLoadAction:MTLLoadActionClear];
 				}
 				else
 				{
@@ -921,13 +922,13 @@ bool FAGXStateCache::SetRenderPassInfo(FRHIRenderPassInfo const& InRenderTargets
 				
 				// For the case where Depth+Stencil is MSAA we can't Resolve depth and Store stencil - we can only Resolve + DontCare or StoreResolve + Store (on newer H/W and iOS).
 				// We only allow use of StoreResolve in the Desktop renderers as the mobile renderer does not and should not assume hardware support for it.
-				NewStencilStore = (StencilTexture.GetSampleCount() == 1  || GetMetalRTStoreAction(ERenderTargetStoreAction::EMultisampleResolve) == mtlpp::StoreAction::StoreAndMultisampleResolve) ? GetMetalRTStoreAction(HighLevelStoreAction) : mtlpp::StoreAction::DontCare;
-				StencilAttachment.SetStoreAction(!bStencilMemoryless && StencilTexture.GetSampleCount() > 1 && GRHIDeviceId > 2 ? mtlpp::StoreAction::Unknown : NewStencilStore);
-				StencilAttachment.SetClearStencil(StencilClearValue);
+				NewStencilStore = (StencilTexture.GetSampleCount() == 1  || GetMetalRTStoreAction(ERenderTargetStoreAction::EMultisampleResolve) == MTLStoreActionStoreAndMultisampleResolve) ? GetMetalRTStoreAction(HighLevelStoreAction) : MTLStoreActionDontCare;
+				[StencilAttachment setStoreAction:(!bStencilMemoryless && StencilTexture.GetSampleCount() > 1 && GRHIDeviceId > 2 ? MTLStoreActionUnknown : NewStencilStore)];
+				[StencilAttachment setClearStencil:StencilClearValue];
 
 				if (SampleCount == 0)
 				{
-					SampleCount = StencilAttachment.GetTexture().GetSampleCount();
+					SampleCount = [[StencilAttachment texture] sampleCount];
 				}
 				
 				bHasValidRenderTarget = true;
@@ -938,11 +939,11 @@ bool FAGXStateCache::SetRenderPassInfo(FRHIRenderPassInfo const& InRenderTargets
 				bCanRestartRenderPass &= 	(bStencilMSAARestart || SampleCount <= 1) &&
 											(
 												(RenderPassInfo.DepthStencilRenderTarget.DepthStencilTarget == FallbackDepthStencilSurface) ||
-												((StencilAttachment.GetLoadAction() == mtlpp::LoadAction::Load) && (bStencilMSAARestart || !RenderPassInfo.DepthStencilRenderTarget.ExclusiveDepthStencil.IsStencilWrite() || (StencilStoreAction == ERenderTargetStoreAction::EStore)))
+												(([StencilAttachment loadAction] == MTLLoadActionLoad) && (bStencilMSAARestart || !RenderPassInfo.DepthStencilRenderTarget.ExclusiveDepthStencil.IsStencilWrite() || (StencilStoreAction == ERenderTargetStoreAction::EStore)))
 											);
 				
 				// and assign it
-				RenderPass.SetStencilAttachment(StencilAttachment);
+				[RenderPass setStencilAttachment:StencilAttachment];
 			}
 		}
 		
@@ -1269,11 +1270,11 @@ bool FAGXStateCache::NeedsToSetRenderTarget(const FRHIRenderPassInfo& InRenderPa
 		if (InRenderPassInfo.DepthStencilRenderTarget.DepthStencilTarget && (GetStoreAction(GetDepthActions(InRenderPassInfo.DepthStencilRenderTarget.Action)) > GetStoreAction(GetDepthActions(RenderPassInfo.DepthStencilRenderTarget.Action)) || GetStoreAction(GetStencilActions(InRenderPassInfo.DepthStencilRenderTarget.Action)) > GetStoreAction(GetStencilActions(RenderPassInfo.DepthStencilRenderTarget.Action))))
 		{
 			// Don't break the encoder if we can just change the store actions.
-			mtlpp::StoreAction NewDepthStore = DepthStore;
-			mtlpp::StoreAction NewStencilStore = StencilStore;
+			MTLStoreAction NewDepthStore = DepthStore;
+			MTLStoreAction NewStencilStore = StencilStore;
 			if (GetStoreAction(GetDepthActions(InRenderPassInfo.DepthStencilRenderTarget.Action)) > GetStoreAction(GetDepthActions(RenderPassInfo.DepthStencilRenderTarget.Action)))
 			{
-				if (RenderPassDesc.GetDepthAttachment().GetTexture())
+				if ([[RenderPassDesc depthAttachment] texture])
 				{
 					FAGXSurface& Surface = *AGXGetMetalSurfaceFromRHITexture(RenderPassInfo.DepthStencilRenderTarget.DepthStencilTarget);
 					
@@ -1300,11 +1301,11 @@ bool FAGXStateCache::NeedsToSetRenderTarget(const FRHIRenderPassInfo& InRenderPa
 			
 			if (GetStoreAction(GetStencilActions(InRenderPassInfo.DepthStencilRenderTarget.Action)) > GetStoreAction(GetStencilActions(RenderPassInfo.DepthStencilRenderTarget.Action)))
 			{
-				if (RenderPassDesc.GetStencilAttachment().GetTexture())
+				if ([[RenderPassDesc stencilAttachment] texture])
 				{
 					NewStencilStore = GetMetalRTStoreAction(GetStoreAction(GetStencilActions(RenderPassInfo.DepthStencilRenderTarget.Action)));
 #if PLATFORM_IOS
-					if (RenderPassDesc.GetStencilAttachment().GetTexture().GetStorageMode() == mtlpp::StorageMode::Memoryless)
+					if ([[[RenderPassDesc stencilAttachment] texture] storageMode] == MTLStorageModeMemoryless)
 					{
 						NewStencilStore = GetMetalRTStoreAction(ERenderTargetStoreAction::ENoAction);
 					}
@@ -1828,20 +1829,21 @@ void FAGXStateCache::SetRenderStoreActions(FAGXCommandEncoder& CommandEncoder, b
 	{
 		if (bConditionalSwitch)
 		{
-			ns::Array<mtlpp::RenderPassColorAttachmentDescriptor> ColorAttachments = RenderPassDesc.GetColorAttachments();
+			MTLRenderPassColorAttachmentDescriptorArray* ColorAttachments = [RenderPassDesc colorAttachments];
 			for (int32 RenderTargetIndex = 0; RenderTargetIndex < RenderPassInfo.GetNumColorRenderTargets(); RenderTargetIndex++)
 			{
 				FRHIRenderPassInfo::FColorEntry& RenderTargetView = RenderPassInfo.ColorRenderTargets[RenderTargetIndex];
 				if(RenderTargetView.RenderTarget != nil)
 				{
-					const bool bMultiSampled = (ColorAttachments[RenderTargetIndex].GetTexture().GetSampleCount() > 1);
+					const bool bMultiSampled = ([[[ColorAttachments objectAtIndexedSubscript:RenderTargetIndex] texture] sampleCount] > 1);
 					ColorStore[RenderTargetIndex] = GetConditionalMetalRTStoreAction(bMultiSampled);
 				}
 			}
 			
 			if (RenderPassInfo.DepthStencilRenderTarget.DepthStencilTarget)
 			{
-				const bool bMultiSampled = RenderPassDesc.GetDepthAttachment().GetTexture() && (RenderPassDesc.GetDepthAttachment().GetTexture().GetSampleCount() > 1);
+				id<MTLTexture> DepthAttachmentTexture = [[RenderPassDesc depthAttachment] texture];
+				const bool bMultiSampled = DepthAttachmentTexture && ([DepthAttachmentTexture sampleCount] > 1);
 				DepthStore = GetConditionalMetalRTStoreAction(bMultiSampled);
 				StencilStore = GetConditionalMetalRTStoreAction(false);
 			}
@@ -1912,8 +1914,8 @@ void FAGXStateCache::SetRenderState(FAGXCommandEncoder& CommandEncoder, FAGXComm
             
             if (DepthStencilState && RenderPassDesc && AGXSafeGetRuntimeDebuggingLevel() >= EAGXDebugLevelFastValidation)
             {
-                METAL_FATAL_ASSERT(DepthStencilState->bIsDepthWriteEnabled == false || (RenderPassDesc.GetDepthAttachment() && RenderPassDesc.GetDepthAttachment().GetTexture()) , TEXT("Attempting to set a depth-stencil state that writes depth but no depth texture is configured!\nState: %s\nRender Pass: %s"), *FString([DepthStencilState->State.GetPtr() description]), *FString([RenderPassDesc.GetPtr() description]));
-                METAL_FATAL_ASSERT(DepthStencilState->bIsStencilWriteEnabled == false || (RenderPassDesc.GetStencilAttachment() && RenderPassDesc.GetStencilAttachment().GetTexture()), TEXT("Attempting to set a depth-stencil state that writes stencil but no stencil texture is configured!\nState: %s\nRender Pass: %s"), *FString([DepthStencilState->State.GetPtr() description]), *FString([RenderPassDesc.GetPtr() description]));
+                METAL_FATAL_ASSERT(DepthStencilState->bIsDepthWriteEnabled == false || [[RenderPassDesc depthAttachment] texture] , TEXT("Attempting to set a depth-stencil state that writes depth but no depth texture is configured!\nState: %s\nRender Pass: %s"), *FString([DepthStencilState->State.GetPtr() description]), *FString([RenderPassDesc description]));
+                METAL_FATAL_ASSERT(DepthStencilState->bIsStencilWriteEnabled == false || [[RenderPassDesc stencilAttachment] texture], TEXT("Attempting to set a depth-stencil state that writes stencil but no stencil texture is configured!\nState: %s\nRender Pass: %s"), *FString([DepthStencilState->State.GetPtr() description]), *FString([RenderPassDesc description]));
             }
             
 			CommandEncoder.SetDepthStencilState(DepthStencilState ? DepthStencilState->State : nil);
@@ -2250,12 +2252,12 @@ void FAGXStateCache::DiscardRenderTargets(bool Depth, bool Stencil, uint32 Color
 	{
 		switch (DepthStore)
 		{
-			case mtlpp::StoreAction::Unknown:
-			case mtlpp::StoreAction::Store:
-				DepthStore = mtlpp::StoreAction::DontCare;
+			case MTLStoreActionUnknown:
+			case MTLStoreActionStore:
+				DepthStore = MTLStoreActionDontCare;
 				break;
-			case mtlpp::StoreAction::StoreAndMultisampleResolve:
-				DepthStore = mtlpp::StoreAction::MultisampleResolve;
+			case MTLStoreActionStoreAndMultisampleResolve:
+				DepthStore = MTLStoreActionMultisampleResolve;
 				break;
 			default:
 				break;
@@ -2264,7 +2266,7 @@ void FAGXStateCache::DiscardRenderTargets(bool Depth, bool Stencil, uint32 Color
 
 	if (Stencil)
 	{
-		StencilStore = mtlpp::StoreAction::DontCare;
+		StencilStore = MTLStoreActionDontCare;
 	}
 
 	for (uint32 Index = 0; Index < MaxSimultaneousRenderTargets; ++Index)
@@ -2273,12 +2275,12 @@ void FAGXStateCache::DiscardRenderTargets(bool Depth, bool Stencil, uint32 Color
 		{
 			switch (ColorStore[Index])
 			{
-				case mtlpp::StoreAction::Unknown:
-				case mtlpp::StoreAction::Store:
-					ColorStore[Index] = mtlpp::StoreAction::DontCare;
+				case MTLStoreActionUnknown:
+				case MTLStoreActionStore:
+					ColorStore[Index] = MTLStoreActionDontCare;
 					break;
-				case mtlpp::StoreAction::StoreAndMultisampleResolve:
-					ColorStore[Index] = mtlpp::StoreAction::MultisampleResolve;
+				case MTLStoreActionStoreAndMultisampleResolve:
+					ColorStore[Index] = MTLStoreActionMultisampleResolve;
 					break;
 				default:
 					break;
diff --git a/Engine/Source/Runtime/Apple/AGXRHI/Private/AGXStateCache.h b/Engine/Source/Runtime/Apple/AGXRHI/Private/AGXStateCache.h
--- a/Engine/Source/Runtime/Apple/AGXRHI/Private/AGXStateCache.h
+++ b/Engine/Source/Runtime/Apple/AGXRHI/Private/AGXStateCache.h
@@ -139,7 +139,7 @@ public:
 	bool NeedsToSetRenderTarget(const FRHIRenderPassInfo& RenderPassInfo);
 	bool HasValidDepthStencilSurface() const { return IsValidRef(DepthStencilSurface); }
     bool CanRestartRenderPass() const { return bCanRestartRenderPass; }
-	mtlpp::RenderPassDescriptor GetRenderPassDescriptor(void) const { return RenderPassDesc; }
+	MTLRenderPassDescriptor* GetRenderPassDescriptor(void) const { return RenderPassDesc; }
 	uint32 GetSampleCount(void) const { return SampleCount; }
     bool IsLinearBuffer(EAGXShaderStages ShaderStage, uint32 BindIndex);
 	FAGXShaderPipeline* GetPipelineState() const;
@@ -254,9 +254,9 @@ private:
 	FAGXTextureBindings ShaderTextures[EAGXShaderStages::Num];
 	FAGXSamplerBindings ShaderSamplers[EAGXShaderStages::Num];
 	
-	mtlpp::StoreAction ColorStore[MaxSimultaneousRenderTargets];
-	mtlpp::StoreAction DepthStore;
-	mtlpp::StoreAction StencilStore;
+	MTLStoreAction ColorStore[MaxSimultaneousRenderTargets];
+	MTLStoreAction DepthStore;
+	MTLStoreAction StencilStore;
 
 	FAGXQueryBuffer* VisibilityResults;
 	mtlpp::VisibilityResultMode VisibilityMode;
@@ -287,7 +287,7 @@ private:
 	FTextureRHIRef DepthStencilResolve;
 	/** A fallback depth-stencil surface for draw calls that write to depth without a depth-stencil surface bound. */
 	FTexture2DRHIRef FallbackDepthStencilSurface;
-	mtlpp::RenderPassDescriptor RenderPassDesc;
+	MTLRenderPassDescriptor* RenderPassDesc;
 	uint32 RasterBits;
     uint8 PipelineBits;
 	bool bIsRenderTargetActive;
