AGXRHI: Scrub use of mtlpp::DispatchType

#rnx

diff --git a/Engine/Source/Runtime/Apple/AGXRHI/Private/AGXCommandEncoder.cpp b/Engine/Source/Runtime/Apple/AGXRHI/Private/AGXCommandEncoder.cpp
--- a/Engine/Source/Runtime/Apple/AGXRHI/Private/AGXCommandEncoder.cpp
+++ b/Engine/Source/Runtime/Apple/AGXRHI/Private/AGXCommandEncoder.cpp
@@ -409,18 +409,18 @@ void FAGXCommandEncoder::BeginRenderCommandEncoding(void)
 	}
 }
 
-void FAGXCommandEncoder::BeginComputeCommandEncoding(mtlpp::DispatchType DispatchType)
+void FAGXCommandEncoder::BeginComputeCommandEncoding(MTLDispatchType DispatchType)
 {
 	check(CommandBuffer);
 	check(IsRenderCommandEncoderActive() == false && IsComputeCommandEncoderActive() == false && IsBlitCommandEncoderActive() == false);
 	
-	if (DispatchType == mtlpp::DispatchType::Serial)
+	if (DispatchType == MTLDispatchTypeSerial)
 	{
 		ComputeCommandEncoder = [[CommandBuffer.GetPtr() computeCommandEncoder] retain];
 	}
 	else
 	{
-		ComputeCommandEncoder = [[CommandBuffer.GetPtr() computeCommandEncoderWithDispatchType:(MTLDispatchType)DispatchType] retain];
+		ComputeCommandEncoder = [[CommandBuffer.GetPtr() computeCommandEncoderWithDispatchType:DispatchType] retain];
 	}
 
 	EncoderNum++;
diff --git a/Engine/Source/Runtime/Apple/AGXRHI/Private/AGXCommandEncoder.h b/Engine/Source/Runtime/Apple/AGXRHI/Private/AGXCommandEncoder.h
--- a/Engine/Source/Runtime/Apple/AGXRHI/Private/AGXCommandEncoder.h
+++ b/Engine/Source/Runtime/Apple/AGXRHI/Private/AGXCommandEncoder.h
@@ -201,7 +201,7 @@ public:
 	void BeginRenderCommandEncoding(void);
 	
 	/** Begins encoding compute commands into the current command buffer. No other encoder may be active. */
-	void BeginComputeCommandEncoding(mtlpp::DispatchType Type = mtlpp::DispatchType::Serial);
+	void BeginComputeCommandEncoding(MTLDispatchType Type = MTLDispatchTypeSerial);
 	
 	/** Begins encoding blit commands into the current command buffer. No other encoder may be active. */
 	void BeginBlitCommandEncoding(void);
diff --git a/Engine/Source/Runtime/Apple/AGXRHI/Private/AGXRenderPass.cpp b/Engine/Source/Runtime/Apple/AGXRHI/Private/AGXRenderPass.cpp
--- a/Engine/Source/Runtime/Apple/AGXRHI/Private/AGXRenderPass.cpp
+++ b/Engine/Source/Runtime/Apple/AGXRHI/Private/AGXRenderPass.cpp
@@ -42,7 +42,7 @@ FAGXRenderPass::FAGXRenderPass(FAGXCommandList& InCmdList, FAGXStateCache& Cache
 , CurrentEncoder(InCmdList, EAGXCommandEncoderCurrent)
 , PrologueEncoder(InCmdList, EAGXCommandEncoderPrologue)
 , RenderPassDesc(nil)
-, ComputeDispatchType(mtlpp::DispatchType::Serial)
+, ComputeDispatchType(MTLDispatchTypeSerial)
 , NumOutstandingOps(0)
 , bWithinRenderPass(false)
 {
@@ -54,7 +54,7 @@ FAGXRenderPass::~FAGXRenderPass(void)
 	check(!PrologueEncoder.GetCommandBuffer());
 }
 
-void FAGXRenderPass::SetDispatchType(mtlpp::DispatchType Type)
+void FAGXRenderPass::SetDispatchType(MTLDispatchType Type)
 {
 	ComputeDispatchType = Type;
 }
diff --git a/Engine/Source/Runtime/Apple/AGXRHI/Private/AGXRenderPass.h b/Engine/Source/Runtime/Apple/AGXRHI/Private/AGXRenderPass.h
--- a/Engine/Source/Runtime/Apple/AGXRHI/Private/AGXRenderPass.h
+++ b/Engine/Source/Runtime/Apple/AGXRHI/Private/AGXRenderPass.h
@@ -21,7 +21,7 @@ public:
 	~FAGXRenderPass(void);
 	
 #pragma mark -
-	void SetDispatchType(mtlpp::DispatchType Type);
+	void SetDispatchType(MTLDispatchType Type);
 	
     void Begin(bool const bParallelBegin = false);
 	
@@ -189,7 +189,7 @@ private:
 
     MTLRenderPassDescriptor* RenderPassDesc;
     
-	mtlpp::DispatchType ComputeDispatchType;
+	MTLDispatchType ComputeDispatchType;
     uint32 NumOutstandingOps;
     bool bWithinRenderPass;
 };
