AGXRHI: Scrub mtlpp::Viewport and mtlpp::ScissorRect

#rnx

diff --git a/Engine/Source/Runtime/Apple/AGXRHI/Private/AGXCommandEncoder.cpp b/Engine/Source/Runtime/Apple/AGXRHI/Private/AGXCommandEncoder.cpp
--- a/Engine/Source/Runtime/Apple/AGXRHI/Private/AGXCommandEncoder.cpp
+++ b/Engine/Source/Runtime/Apple/AGXRHI/Private/AGXCommandEncoder.cpp
@@ -774,36 +774,36 @@ void FAGXCommandEncoder::SetRenderPipelineState(FAGXShaderPipeline* PipelineStat
 	}
 }
 
-void FAGXCommandEncoder::SetViewport(mtlpp::Viewport const Viewport[], uint32 NumActive)
+void FAGXCommandEncoder::SetViewport(MTLViewport const Viewport[], uint32 NumActive)
 {
 	check(RenderCommandEncoder);
 	check(NumActive >= 1 && NumActive < ML_MaxViewports);
 	if (NumActive == 1)
 	{
-		[RenderCommandEncoder setViewport:*reinterpret_cast<const MTLViewport*>(&Viewport[0])];
+		[RenderCommandEncoder setViewport:Viewport[0]];
 	}
 #if PLATFORM_MAC
 	else
 	{
 		check(FAGXCommandQueue::SupportsFeature(EAGXFeaturesMultipleViewports));
-		[RenderCommandEncoder setViewports:reinterpret_cast<const MTLViewport*>(Viewport) count:(NSUInteger)NumActive];
+		[RenderCommandEncoder setViewports:Viewport count:NumActive];
 	}
 #endif
 }
 
-void FAGXCommandEncoder::SetFrontFacingWinding(mtlpp::Winding const InFrontFacingWinding)
+void FAGXCommandEncoder::SetFrontFacingWinding(MTLWinding InFrontFacingWinding)
 {
     check (RenderCommandEncoder);
 	{
-		[RenderCommandEncoder setFrontFacingWinding:(MTLWinding)InFrontFacingWinding];
+		[RenderCommandEncoder setFrontFacingWinding:InFrontFacingWinding];
 	}
 }
 
-void FAGXCommandEncoder::SetCullMode(mtlpp::CullMode const InCullMode)
+void FAGXCommandEncoder::SetCullMode(MTLCullMode InCullMode)
 {
     check (RenderCommandEncoder);
 	{
-		[RenderCommandEncoder setCullMode:(MTLCullMode)InCullMode];
+		[RenderCommandEncoder setCullMode:InCullMode];
 	}
 }
 
@@ -815,19 +815,19 @@ void FAGXCommandEncoder::SetDepthBias(float const InDepthBias, float const InSlo
 	}
 }
 
-void FAGXCommandEncoder::SetScissorRect(mtlpp::ScissorRect const Rect[], uint32 NumActive)
+void FAGXCommandEncoder::SetScissorRect(MTLScissorRect const Rect[], uint32 NumActive)
 {
     check(RenderCommandEncoder);
 	check(NumActive >= 1 && NumActive < ML_MaxViewports);
 	if (NumActive == 1)
 	{
-		[RenderCommandEncoder setScissorRect:*reinterpret_cast<const MTLScissorRect*>(&Rect[0])];
+		[RenderCommandEncoder setScissorRect:Rect[0]];
 	}
 #if PLATFORM_MAC
 	else
 	{
 		check(FAGXCommandQueue::SupportsFeature(EAGXFeaturesMultipleViewports));
-		[RenderCommandEncoder setScissorRects:reinterpret_cast<const MTLScissorRect*>(Rect) count:(NSUInteger)NumActive];
+		[RenderCommandEncoder setScissorRects:Rect count:NumActive];
 	}
 #endif
 }
diff --git a/Engine/Source/Runtime/Apple/AGXRHI/Private/AGXCommandEncoder.h b/Engine/Source/Runtime/Apple/AGXRHI/Private/AGXCommandEncoder.h
--- a/Engine/Source/Runtime/Apple/AGXRHI/Private/AGXCommandEncoder.h
+++ b/Engine/Source/Runtime/Apple/AGXRHI/Private/AGXCommandEncoder.h
@@ -264,19 +264,19 @@ public:
 	 * @param Viewport The array of viewport dimensions to use.
 	 * @param NumActive The number of active viewport dimensions to use.
 	 */
-	void SetViewport(mtlpp::Viewport const Viewport[], uint32 NumActive);
+	void SetViewport(MTLViewport const Viewport[], uint32 NumActive);
 	
 	/*
 	 * The winding order of front-facing primitives.
 	 * @param FrontFacingWinding The front face winding.
 	 */
-	void SetFrontFacingWinding(mtlpp::Winding const FrontFacingWinding);
+	void SetFrontFacingWinding(MTLWinding const FrontFacingWinding);
 	
 	/*
 	 * Controls if primitives are culled when front facing, back facing, or not culled at all.
 	 * @param CullMode The cull mode.
 	 */
-	void SetCullMode(mtlpp::CullMode const CullMode);
+	void SetCullMode(MTLCullMode const CullMode);
 	
 	/*
 	 * Depth Bias.
@@ -291,7 +291,7 @@ public:
 	 * @param Rect The array of scissor rect dimensions.
 	 * @param NumActive The number of active scissor rect dimensions.
 	 */
-	void SetScissorRect(mtlpp::ScissorRect const Rect[], uint32 NumActive);
+	void SetScissorRect(MTLScissorRect const Rect[], uint32 NumActive);
 	
 	/*
 	 * Set how to rasterize triangle and triangle strip primitives.
diff --git a/Engine/Source/Runtime/Apple/AGXRHI/Private/AGXCommands.cpp b/Engine/Source/Runtime/Apple/AGXRHI/Private/AGXCommands.cpp
--- a/Engine/Source/Runtime/Apple/AGXRHI/Private/AGXCommands.cpp
+++ b/Engine/Source/Runtime/Apple/AGXRHI/Private/AGXCommands.cpp
@@ -173,7 +173,7 @@ void FAGXRHICommandContext::RHIDispatchIndirectComputeShader(FRHIBuffer* Argumen
 void FAGXRHICommandContext::RHISetViewport(float MinX, float MinY,float MinZ, float MaxX, float MaxY,float MaxZ)
 {
 	@autoreleasepool {
-	mtlpp::Viewport Viewport;
+	MTLViewport Viewport;
 	Viewport.originX = MinX;
 	Viewport.originY = MinY;
 	Viewport.width = MaxX - MinX;
@@ -190,7 +190,7 @@ void FAGXRHICommandContext::RHISetStereoViewport(float LeftMinX, float RightMinX
 	if (FAGXCommandQueue::SupportsFeature(EAGXFeaturesMultipleViewports))
 	{
 		@autoreleasepool {
-		mtlpp::Viewport Viewport[2];
+		MTLViewport Viewport[2];
 		
 		Viewport[0].originX = LeftMinX;
 		Viewport[0].originY = LeftMinY;
@@ -223,7 +223,7 @@ void FAGXRHICommandContext::RHISetMultipleViewports(uint32 Count, const FViewpor
 void FAGXRHICommandContext::RHISetScissorRect(bool bEnable,uint32 MinX,uint32 MinY,uint32 MaxX,uint32 MaxY)
 {
 	@autoreleasepool {
-	mtlpp::ScissorRect Scissor;
+	MTLScissorRect Scissor;
 	Scissor.x = MinX;
 	Scissor.y = MinY;
 	Scissor.width = MaxX - MinX;
@@ -232,7 +232,7 @@ void FAGXRHICommandContext::RHISetScissorRect(bool bEnable,uint32 MinX,uint32 Mi
 	// metal doesn't support 0 sized scissor rect
 	if (bEnable == false || Scissor.width == 0 || Scissor.height == 0)
 	{
-		mtlpp::Viewport const& Viewport = Context->GetCurrentState().GetViewport(0);
+		MTLViewport const& Viewport = Context->GetCurrentState().GetViewport(0);
 		CGSize FBSize = Context->GetCurrentState().GetFrameBufferSize();
 		
 		Scissor.x = Viewport.originX;
diff --git a/Engine/Source/Runtime/Apple/AGXRHI/Private/AGXContext.cpp b/Engine/Source/Runtime/Apple/AGXRHI/Private/AGXContext.cpp
--- a/Engine/Source/Runtime/Apple/AGXRHI/Private/AGXContext.cpp
+++ b/Engine/Source/Runtime/Apple/AGXRHI/Private/AGXContext.cpp
@@ -1273,7 +1273,12 @@ bool FAGXContext::PrepareToDraw(uint32 PrimitiveType)
 		
 		if (bBindDepthStencilForUAVRaster)
 		{
-			mtlpp::ScissorRect Rect(0, 0, (NSUInteger)FBSize.width, (NSUInteger)FBSize.height);
+			MTLScissorRect Rect;
+			Rect.x = 0;
+			Rect.y = 0;
+			Rect.width = (NSUInteger)FBSize.width;
+			Rect.height = (NSUInteger)FBSize.height;
+
 			StateCache.SetScissorRect(false, Rect);
 		}
 		
diff --git a/Engine/Source/Runtime/Apple/AGXRHI/Private/AGXStateCache.cpp b/Engine/Source/Runtime/Apple/AGXRHI/Private/AGXStateCache.cpp
--- a/Engine/Source/Runtime/Apple/AGXRHI/Private/AGXStateCache.cpp
+++ b/Engine/Source/Runtime/Apple/AGXRHI/Private/AGXStateCache.cpp
@@ -27,13 +27,13 @@ static mtlpp::TriangleFillMode TranslateFillMode(ERasterizerFillMode FillMode)
 	};
 }
 
-static mtlpp::CullMode TranslateCullMode(ERasterizerCullMode CullMode)
+static MTLCullMode TranslateCullMode(ERasterizerCullMode CullMode)
 {
 	switch (CullMode)
 	{
-		case CM_CCW:	return mtlpp::CullMode::Front;
-		case CM_CW:		return mtlpp::CullMode::Back;
-		default:		return mtlpp::CullMode::None;
+		case CM_CCW:	return MTLCullModeFront;
+		case CM_CW:		return MTLCullModeBack;
+		default:		return MTLCullModeNone;
 	}
 }
 
@@ -328,12 +328,12 @@ void FAGXStateCache::Reset(void)
     PipelineBits = EAGXPipelineFlagMask;
 }
 
-static bool MTLScissorRectEqual(mtlpp::ScissorRect const& Left, mtlpp::ScissorRect const& Right)
+static bool MTLScissorRectEqual(MTLScissorRect const& Left, MTLScissorRect const& Right)
 {
 	return Left.x == Right.x && Left.y == Right.y && Left.width == Right.width && Left.height == Right.height;
 }
 
-void FAGXStateCache::SetScissorRect(bool const bEnable, mtlpp::ScissorRect const& Rect)
+void FAGXStateCache::SetScissorRect(bool bEnable, MTLScissorRect const& Rect)
 {
 	if (bScissorRectEnabled != bEnable || !MTLScissorRectEqual(Scissor[0], Rect))
 	{
@@ -997,7 +997,7 @@ void FAGXStateCache::SetRenderTargetsActive(bool const bActive)
 	bIsRenderTargetActive = bActive;
 }
 
-static bool MTLViewportEqual(mtlpp::Viewport const& Left, mtlpp::Viewport const& Right)
+static bool MTLViewportEqual(MTLViewport const& Left, MTLViewport const& Right)
 {
 	return FMath::IsNearlyEqual(Left.originX, Right.originX) &&
 			FMath::IsNearlyEqual(Left.originY, Right.originY) &&
@@ -1007,7 +1007,7 @@ static bool MTLViewportEqual(mtlpp::Viewport const& Left, mtlpp::Viewport const&
 			FMath::IsNearlyEqual(Left.zfar, Right.zfar);
 }
 
-void FAGXStateCache::SetViewport(const mtlpp::Viewport& InViewport)
+void FAGXStateCache::SetViewport(MTLViewport const& InViewport)
 {
 	if (!MTLViewportEqual(Viewport[0], InViewport))
 	{
@@ -1020,7 +1020,7 @@ void FAGXStateCache::SetViewport(const mtlpp::Viewport& InViewport)
 	
 	if (!bScissorRectEnabled)
 	{
-		mtlpp::ScissorRect Rect;
+		MTLScissorRect Rect;
 		Rect.x = InViewport.originX;
 		Rect.y = InViewport.originY;
 		Rect.width = InViewport.width;
@@ -1029,7 +1029,7 @@ void FAGXStateCache::SetViewport(const mtlpp::Viewport& InViewport)
 	}
 }
 
-void FAGXStateCache::SetViewport(uint32 Index, const mtlpp::Viewport& InViewport)
+void FAGXStateCache::SetViewport(uint32 Index, MTLViewport const& InViewport)
 {
 	check(Index < ML_MaxViewports);
 	
@@ -1046,7 +1046,7 @@ void FAGXStateCache::SetViewport(uint32 Index, const mtlpp::Viewport& InViewport
 	// This always sets the scissor rect because the RHI doesn't bother to expose proper scissor states for multiple viewports.
 	// This will have to change if we want to guarantee correctness in the mid to long term.
 	{
-		mtlpp::ScissorRect Rect;
+		MTLScissorRect Rect;
 		Rect.x = InViewport.originX;
 		Rect.y = InViewport.originY;
 		Rect.width = InViewport.width;
@@ -1055,7 +1055,7 @@ void FAGXStateCache::SetViewport(uint32 Index, const mtlpp::Viewport& InViewport
 	}
 }
 
-void FAGXStateCache::SetScissorRect(uint32 Index, bool const bEnable, mtlpp::ScissorRect const& Rect)
+void FAGXStateCache::SetScissorRect(uint32 Index, bool bEnable, MTLScissorRect const& Rect)
 {
 	check(Index < ML_MaxViewports);
 	if (!MTLScissorRectEqual(Scissor[Index], Rect))
@@ -1068,7 +1068,7 @@ void FAGXStateCache::SetScissorRect(uint32 Index, bool const bEnable, mtlpp::Sci
 	ActiveScissors = Index + 1;
 }
 
-void FAGXStateCache::SetViewports(const mtlpp::Viewport InViewport[], uint32 Count)
+void FAGXStateCache::SetViewports(MTLViewport const InViewport[], uint32 Count)
 {
 	check(Count >= 1 && Count < ML_MaxViewports);
 	
@@ -1883,7 +1883,7 @@ void FAGXStateCache::SetRenderState(FAGXCommandEncoder& CommandEncoder, FAGXComm
 		}
 		if (RasterBits & EAGXRenderFlagFrontFacingWinding)
 		{
-			CommandEncoder.SetFrontFacingWinding(mtlpp::Winding::CounterClockwise);
+			CommandEncoder.SetFrontFacingWinding(MTLWindingCounterClockwise);
 		}
 		if (RasterBits & EAGXRenderFlagCullMode)
 		{
diff --git a/Engine/Source/Runtime/Apple/AGXRHI/Private/AGXStateCache.h b/Engine/Source/Runtime/Apple/AGXRHI/Private/AGXStateCache.h
--- a/Engine/Source/Runtime/Apple/AGXRHI/Private/AGXStateCache.h
+++ b/Engine/Source/Runtime/Apple/AGXRHI/Private/AGXStateCache.h
@@ -45,15 +45,15 @@ public:
 	/** Reset cached state for reuse */
 	void Reset(void);
 
-	void SetScissorRect(bool const bEnable, mtlpp::ScissorRect const& Rect);
+	void SetScissorRect(bool bEnable, MTLScissorRect const& Rect);
 	void SetBlendFactor(FLinearColor const& InBlendFactor);
 	void SetStencilRef(uint32 const InStencilRef);
 	void SetComputeShader(FAGXComputeShader* InComputeShader);
 	bool SetRenderPassInfo(FRHIRenderPassInfo const& InRenderTargets, FAGXQueryBuffer* QueryBuffer, bool const bRestart);
 	void InvalidateRenderTargets(void);
 	void SetRenderTargetsActive(bool const bActive);
-	void SetViewport(const mtlpp::Viewport& InViewport);
-	void SetViewports(const mtlpp::Viewport InViewport[], uint32 Count);
+	void SetViewport(MTLViewport const& InViewport);
+	void SetViewports(MTLViewport const InViewport[], uint32 Count);
 	void SetVertexStream(uint32 const Index, FAGXBuffer* Buffer, FAGXBufferData* Bytes, uint32 const Offset, uint32 const Length);
 	void SetGraphicsPipelineState(FAGXGraphicsPipelineState* State);
 	void BindUniformBuffer(EAGXShaderStages const Freq, uint32 const BufferIndex, FRHIUniformBuffer* BufferRHI);
@@ -129,7 +129,7 @@ public:
 	int32 GetNumRenderTargets() { return bHasValidColorTarget ? RenderPassInfo.GetNumColorRenderTargets() : -1; }
 	bool GetHasValidRenderTarget() const { return bHasValidRenderTarget; }
 	bool GetHasValidColorTarget() const { return bHasValidColorTarget; }
-	const mtlpp::Viewport& GetViewport(uint32 const Index) const { check(Index < ML_MaxViewports); return Viewport[Index]; }
+	MTLViewport const& GetViewport(uint32 const Index) const { check(Index < ML_MaxViewports); return Viewport[Index]; }
 	uint32 GetVertexBufferSize(uint32 const Index);
 	uint32 GetRenderTargetArraySize() const { return RenderTargetArraySize; }
 	const FRHIUniformBuffer** GetBoundUniformBuffers(EAGXShaderStages const Freq) { return (const FRHIUniformBuffer**)&BoundUniformBuffers[Freq][0]; }
@@ -175,8 +175,8 @@ private:
 	template <class ShaderType>
 	void SetResourcesFromTables(ShaderType Shader, uint32 ShaderStage);
 	
-	void SetViewport(uint32 Index, const mtlpp::Viewport& InViewport);
-	void SetScissorRect(uint32 Index, bool const bEnable, mtlpp::ScissorRect const& Rect);
+	void SetViewport(uint32 Index, MTLViewport const& InViewport);
+	void SetScissorRect(uint32 Index, bool bEnable, MTLScissorRect const& Rect);
 
 private:
 
@@ -274,8 +274,8 @@ private:
 	
 	uint32 RenderTargetArraySize;
 
-	mtlpp::Viewport Viewport[ML_MaxViewports];
-	mtlpp::ScissorRect Scissor[ML_MaxViewports];
+	MTLViewport Viewport[ML_MaxViewports];
+	MTLScissorRect Scissor[ML_MaxViewports];
 	
 	uint32 ActiveViewports;
 	uint32 ActiveScissors;
