AGXRHI: Scrub use of mtlpp::DepthStencilState, mtlpp::SamplerState and mtlpp::BlendState

- Replace all uses mtlpp::DepthStencilState, mtlpp::SamplerState and mtlpp::BlendState with direct usage of Metal
- Update misc. depth/stencil, sampler, blend factor and blend operation types to be Metal vs. mtlpp

#rnx

diff --git a/Engine/Source/Runtime/Apple/AGXRHI/Private/AGXCommandEncoder.cpp b/Engine/Source/Runtime/Apple/AGXRHI/Private/AGXCommandEncoder.cpp
--- a/Engine/Source/Runtime/Apple/AGXRHI/Private/AGXCommandEncoder.cpp
+++ b/Engine/Source/Runtime/Apple/AGXRHI/Private/AGXCommandEncoder.cpp
@@ -856,11 +856,11 @@ void FAGXCommandEncoder::SetBlendColor(float const Red, float const Green, float
 	}
 }
 
-void FAGXCommandEncoder::SetDepthStencilState(mtlpp::DepthStencilState const& InDepthStencilState)
+void FAGXCommandEncoder::SetDepthStencilState(id<MTLDepthStencilState> InDepthStencilState)
 {
     check (RenderCommandEncoder);
 	{
-		[RenderCommandEncoder setDepthStencilState:InDepthStencilState.GetPtr()];
+		[RenderCommandEncoder setDepthStencilState:InDepthStencilState];
 	}
 }
 
@@ -1067,22 +1067,22 @@ void FAGXCommandEncoder::SetShaderTexture(mtlpp::FunctionType FunctionType, FAGX
 	}
 }
 
-void FAGXCommandEncoder::SetShaderSamplerState(mtlpp::FunctionType FunctionType, mtlpp::SamplerState const& Sampler, NSUInteger index)
+void FAGXCommandEncoder::SetShaderSamplerState(mtlpp::FunctionType FunctionType, id<MTLSamplerState> Sampler, NSUInteger index)
 {
 	check(index < ML_MaxSamplers);
 	switch (FunctionType)
 	{
 		case mtlpp::FunctionType::Vertex:
        		check (RenderCommandEncoder);
-			[RenderCommandEncoder setVertexSamplerState:Sampler.GetPtr() atIndex:index];
+			[RenderCommandEncoder setVertexSamplerState:Sampler atIndex:index];
 			break;
 		case mtlpp::FunctionType::Fragment:
 			check (RenderCommandEncoder);
-			[RenderCommandEncoder setFragmentSamplerState:Sampler.GetPtr() atIndex:index];
+			[RenderCommandEncoder setFragmentSamplerState:Sampler atIndex:index];
 			break;
 		case mtlpp::FunctionType::Kernel:
 			check (ComputeCommandEncoder);
-			[ComputeCommandEncoder setSamplerState:Sampler.GetPtr() atIndex:index];
+			[ComputeCommandEncoder setSamplerState:Sampler atIndex:index];
 			break;
 		default:
 			check(false);
diff --git a/Engine/Source/Runtime/Apple/AGXRHI/Private/AGXCommandEncoder.h b/Engine/Source/Runtime/Apple/AGXRHI/Private/AGXCommandEncoder.h
--- a/Engine/Source/Runtime/Apple/AGXRHI/Private/AGXCommandEncoder.h
+++ b/Engine/Source/Runtime/Apple/AGXRHI/Private/AGXCommandEncoder.h
@@ -270,13 +270,13 @@ public:
 	 * The winding order of front-facing primitives.
 	 * @param FrontFacingWinding The front face winding.
 	 */
-	void SetFrontFacingWinding(MTLWinding const FrontFacingWinding);
+	void SetFrontFacingWinding(MTLWinding FrontFacingWinding);
 	
 	/*
 	 * Controls if primitives are culled when front facing, back facing, or not culled at all.
 	 * @param CullMode The cull mode.
 	 */
-	void SetCullMode(MTLCullMode const CullMode);
+	void SetCullMode(MTLCullMode CullMode);
 	
 	/*
 	 * Depth Bias.
@@ -318,7 +318,7 @@ public:
 	 * Set the DepthStencil state object.
 	 * @param DepthStencilState The depth-stencil state, must not be nil.
 	 */
-	void SetDepthStencilState(mtlpp::DepthStencilState const& DepthStencilState);
+	void SetDepthStencilState(id<MTLDepthStencilState> DepthStencilState);
 	
 	/*
 	 * Set the stencil reference value for both the back and front stencil buffers.
@@ -381,7 +381,7 @@ public:
 	 * @param Sampler The sampler state to bind or nil to clear.
 	 * @param Index The index to modify.
 	 */
-	void SetShaderSamplerState(mtlpp::FunctionType const FunctionType, mtlpp::SamplerState const& Sampler, NSUInteger const Index);
+	void SetShaderSamplerState(mtlpp::FunctionType FunctionType, id<MTLSamplerState> Sampler, NSUInteger Index);
 	
 	/*
 	 * Set the shader side-table data for FunctionType at Index.
diff --git a/Engine/Source/Runtime/Apple/AGXRHI/Private/AGXState.cpp b/Engine/Source/Runtime/Apple/AGXRHI/Private/AGXState.cpp
--- a/Engine/Source/Runtime/Apple/AGXRHI/Private/AGXState.cpp
+++ b/Engine/Source/Runtime/Apple/AGXRHI/Private/AGXState.cpp
@@ -17,13 +17,13 @@ static uint32 GetMetalMaxAnisotropy(ESamplerFilter Filter, uint32 MaxAniso)
 	}
 }
 
-static mtlpp::SamplerMinMagFilter TranslateZFilterMode(ESamplerFilter Filter)
+static MTLSamplerMinMagFilter TranslateZFilterMode(ESamplerFilter Filter)
 {	switch (Filter)
 	{
-		case SF_Point:				return mtlpp::SamplerMinMagFilter::Nearest;
-		case SF_AnisotropicPoint:	return mtlpp::SamplerMinMagFilter::Nearest;
-		case SF_AnisotropicLinear:	return mtlpp::SamplerMinMagFilter::Linear;
-		default:					return mtlpp::SamplerMinMagFilter::Linear;
+		case SF_Point:				return MTLSamplerMinMagFilterNearest;
+		case SF_AnisotropicPoint:	return MTLSamplerMinMagFilterNearest;
+		case SF_AnisotropicLinear:	return MTLSamplerMinMagFilterLinear;
+		default:					return MTLSamplerMinMagFilterLinear;
 	}
 }
 
@@ -78,48 +78,48 @@ static MTLStencilOperation TranslateStencilOp(EStencilOp StencilOp)
 	};
 }
 
-static mtlpp::BlendOperation TranslateBlendOp(EBlendOperation BlendOp)
+static MTLBlendOperation TranslateBlendOp(EBlendOperation BlendOp)
 {
 	switch(BlendOp)
 	{
-		case BO_Subtract:	return mtlpp::BlendOperation::Subtract;
-		case BO_Min:		return mtlpp::BlendOperation::Min;
-		case BO_Max:		return mtlpp::BlendOperation::Max;
-		default:			return mtlpp::BlendOperation::Add;
+		case BO_Subtract:	return MTLBlendOperationSubtract;
+		case BO_Min:		return MTLBlendOperationMin;
+		case BO_Max:		return MTLBlendOperationMax;
+		default:			return MTLBlendOperationAdd;
 	};
 }
 
 
-static mtlpp::BlendFactor TranslateBlendFactor(EBlendFactor BlendFactor)
+static MTLBlendFactor TranslateBlendFactor(EBlendFactor BlendFactor)
 {
-	switch(BlendFactor)
+	switch (BlendFactor)
 	{
-		case BF_One:					return mtlpp::BlendFactor::One;
-		case BF_SourceColor:			return mtlpp::BlendFactor::SourceColor;
-		case BF_InverseSourceColor:		return mtlpp::BlendFactor::OneMinusSourceColor;
-		case BF_SourceAlpha:			return mtlpp::BlendFactor::SourceAlpha;
-		case BF_InverseSourceAlpha:		return mtlpp::BlendFactor::OneMinusSourceAlpha;
-		case BF_DestAlpha:				return mtlpp::BlendFactor::DestinationAlpha;
-		case BF_InverseDestAlpha:		return mtlpp::BlendFactor::OneMinusDestinationAlpha;
-		case BF_DestColor:				return mtlpp::BlendFactor::DestinationColor;
-		case BF_InverseDestColor:		return mtlpp::BlendFactor::OneMinusDestinationColor;
-		case BF_Source1Color:			return mtlpp::BlendFactor::Source1Color;
-		case BF_InverseSource1Color:	return mtlpp::BlendFactor::OneMinusSource1Color;
-		case BF_Source1Alpha:			return mtlpp::BlendFactor::Source1Alpha;
-		case BF_InverseSource1Alpha:	return mtlpp::BlendFactor::OneMinusSource1Alpha;
-		default:						return mtlpp::BlendFactor::Zero;
+		case BF_One:					return MTLBlendFactorOne;
+		case BF_SourceColor:			return MTLBlendFactorSourceColor;
+		case BF_InverseSourceColor:		return MTLBlendFactorOneMinusSourceColor;
+		case BF_SourceAlpha:			return MTLBlendFactorSourceAlpha;
+		case BF_InverseSourceAlpha:		return MTLBlendFactorOneMinusSourceAlpha;
+		case BF_DestAlpha:				return MTLBlendFactorDestinationAlpha;
+		case BF_InverseDestAlpha:		return MTLBlendFactorOneMinusDestinationAlpha;
+		case BF_DestColor:				return MTLBlendFactorDestinationColor;
+		case BF_InverseDestColor:		return MTLBlendFactorOneMinusDestinationColor;
+		case BF_Source1Color:			return MTLBlendFactorSource1Color;
+		case BF_InverseSource1Color:	return MTLBlendFactorOneMinusSource1Color;
+		case BF_Source1Alpha:			return MTLBlendFactorSource1Alpha;
+		case BF_InverseSource1Alpha:	return MTLBlendFactorOneMinusSource1Alpha;
+		default:						return MTLBlendFactorZero;
 	};
 }
 
-static mtlpp::ColorWriteMask TranslateWriteMask(EColorWriteMask WriteMask)
+static MTLColorWriteMask TranslateWriteMask(EColorWriteMask WriteMask)
 {
 	uint32 Result = 0;
-	Result |= (WriteMask & CW_RED) ? (mtlpp::ColorWriteMask::Red) : 0;
-	Result |= (WriteMask & CW_GREEN) ? (mtlpp::ColorWriteMask::Green) : 0;
-	Result |= (WriteMask & CW_BLUE) ? (mtlpp::ColorWriteMask::Blue) : 0;
-	Result |= (WriteMask & CW_ALPHA) ? (mtlpp::ColorWriteMask::Alpha) : 0;
+	Result |= (WriteMask & CW_RED)   ? (uint32)(MTLColorWriteMaskRed)   : 0;
+	Result |= (WriteMask & CW_GREEN) ? (uint32)(MTLColorWriteMaskGreen) : 0;
+	Result |= (WriteMask & CW_BLUE)  ? (uint32)(MTLColorWriteMaskBlue)  : 0;
+	Result |= (WriteMask & CW_ALPHA) ? (uint32)(MTLColorWriteMaskAlpha) : 0;
 	
-	return (mtlpp::ColorWriteMask)Result;
+	return (MTLColorWriteMask)Result;
 }
 
 static EBlendOperation TranslateBlendOp(MTLBlendOperation BlendOp)
@@ -166,67 +166,80 @@ static EColorWriteMask TranslateWriteMask(MTLColorWriteMask WriteMask)
 	return (EColorWriteMask)Result;
 }
 
-template <typename InitializerType, typename StateType>
+struct FObjCWrapperRetained
+{
+	FObjCWrapperRetained(id InObject = nil)
+		: Object([InObject retain])
+	{
+	}
+
+	~FObjCWrapperRetained()
+	{
+		[Object release];
+	}
+
+	id Object;
+};
+
+template <typename InitializerType>
 class FAGXStateObjectCache
 {
 public:
 	FAGXStateObjectCache()
 	{
-		
 	}
 	
 	~FAGXStateObjectCache()
 	{
-		
 	}
 	
-	StateType Find(InitializerType Init)
+	id Find(InitializerType Init)
 	{
-		if(IsRunningRHIInSeparateThread())
+		if (IsRunningRHIInSeparateThread())
 		{
 			Mutex.ReadLock();
 		}
 		
-		StateType* State = Cache.Find(Init);
-		
-		if(IsRunningRHIInSeparateThread())
+		TUniquePtr<FObjCWrapperRetained>* State = Cache.Find(Init);
+
+		if (IsRunningRHIInSeparateThread())
 		{
 			Mutex.ReadUnlock();
 		}
 		
-		return State ? *State : StateType(nullptr);
+		return State ? State->Get()->Object : nil;
 	}
 	
-	void Add(InitializerType Init, StateType const& State)
+	void Add(InitializerType Init, id State)
 	{
-		if(IsRunningRHIInSeparateThread())
+		if (IsRunningRHIInSeparateThread())
 		{
 			Mutex.WriteLock();
 		}
 		
-		Cache.Add(Init, State);
+		Cache.Add(Init, MakeUnique<FObjCWrapperRetained>(State));
 		
-		if(IsRunningRHIInSeparateThread())
+		if (IsRunningRHIInSeparateThread())
 		{
 			Mutex.WriteUnlock();
 		}
 	}
 	
 private:
-	TMap<InitializerType, StateType> Cache;
+	TMap<InitializerType, TUniquePtr<FObjCWrapperRetained> > Cache;
 	FRWLock Mutex;
 };
 
-static FAGXStateObjectCache<FSamplerStateInitializerRHI, FAGXSampler> Samplers;
+static FAGXStateObjectCache<FSamplerStateInitializerRHI> Samplers;
 
-static FAGXSampler FindOrCreateSamplerState(const FSamplerStateInitializerRHI& Initializer)
+static id<MTLSamplerState> FindOrCreateSamplerState(const FSamplerStateInitializerRHI& Initializer)
 {
-	FAGXSampler State = Samplers.Find(Initializer);
-	if (!State.GetPtr())
+	id<MTLSamplerState> State = (id<MTLSamplerState>)Samplers.Find(Initializer);
+	if (!State)
 	{
 		MTLSamplerDescriptor* Desc = [[MTLSamplerDescriptor alloc] init];
 		
-		switch(Initializer.Filter)
+		switch (Initializer.Filter)
 		{
 			case SF_AnisotropicLinear:
 			case SF_AnisotropicPoint:
@@ -266,16 +279,18 @@ static FAGXSampler FindOrCreateSamplerState(const FSamplerStateInitializerRHI& I
 #if PLATFORM_MAC
 		[Desc setBorderColor:Initializer.BorderColor == 0 ? MTLSamplerBorderColorTransparentBlack : MTLSamplerBorderColorOpaqueWhite];
 #endif
-		State = FAGXSampler([GMtlDevice newSamplerStateWithDescriptor:Desc], ns::Ownership::Assign);
+		State = [GMtlDevice newSamplerStateWithDescriptor:Desc];
 		
 		[Desc release];
 		
 		Samplers.Add(Initializer, State);
+
+		[State release];
 	}
 	return State;
 }
 
-FAGXSamplerState::FAGXSamplerState(const FSamplerStateInitializerRHI& Initializer)
+FAGXSamplerState::FAGXSamplerState(FSamplerStateInitializerRHI const& Initializer)
 {
 	State = FindOrCreateSamplerState(Initializer);
 #if !PLATFORM_MAC
@@ -292,14 +307,13 @@ FAGXSamplerState::~FAGXSamplerState()
 {
 }
 
-FAGXRasterizerState::FAGXRasterizerState(const FRasterizerStateInitializerRHI& Initializer)
+FAGXRasterizerState::FAGXRasterizerState(FRasterizerStateInitializerRHI const& Initializer)
 {
 	State = Initializer;
 }
 
 FAGXRasterizerState::~FAGXRasterizerState()
 {
-	
 }
 
 bool FAGXRasterizerState::GetInitializer(FRasterizerStateInitializerRHI& Init)
@@ -308,14 +322,14 @@ bool FAGXRasterizerState::GetInitializer(FRasterizerStateInitializerRHI& Init)
 	return true;
 }
 
-static FAGXStateObjectCache<FDepthStencilStateInitializerRHI, mtlpp::DepthStencilState> DepthStencilStates;
+static FAGXStateObjectCache<FDepthStencilStateInitializerRHI> DepthStencilStates;
 
-FAGXDepthStencilState::FAGXDepthStencilState(const FDepthStencilStateInitializerRHI& InInitializer)
+FAGXDepthStencilState::FAGXDepthStencilState(FDepthStencilStateInitializerRHI const& InInitializer)
 {
 	Initializer = InInitializer;
 
-	State = DepthStencilStates.Find(Initializer);
-	if (!State.GetPtr())
+	State = (id<MTLDepthStencilState>)DepthStencilStates.Find(Initializer);
+	if (!State)
 	{
 		MTLDepthStencilDescriptor* Desc = [[MTLDepthStencilDescriptor alloc] init];
 		
@@ -364,11 +378,13 @@ FAGXDepthStencilState::FAGXDepthStencilState(const FDepthStencilStateInitializer
 		}
 		
 		// bake out the descriptor
-		State = mtlpp::DepthStencilState([GMtlDevice newDepthStencilStateWithDescriptor:Desc], nullptr, ns::Ownership::Assign);
+		State = [GMtlDevice newDepthStencilStateWithDescriptor:Desc];
 		
 		[Desc release];
 		
 		DepthStencilStates.Add(Initializer, State);
+
+		[State release];
 	}
 	
 	// cache some pipeline state info
@@ -389,55 +405,56 @@ bool FAGXDepthStencilState::GetInitializer(FDepthStencilStateInitializerRHI& Ini
 
 
 // statics
-static FAGXStateObjectCache<FBlendStateInitializerRHI::FRenderTarget, mtlpp::RenderPipelineColorAttachmentDescriptor> BlendStates;
+static FAGXStateObjectCache<FBlendStateInitializerRHI::FRenderTarget> BlendStates;
 TMap<uint32, uint8> FAGXBlendState::BlendSettingsToUniqueKeyMap;
 uint8 FAGXBlendState::NextKey = 0;
 FCriticalSection FAGXBlendState::Mutex;
 
-FAGXBlendState::FAGXBlendState(const FBlendStateInitializerRHI& Initializer)
+FAGXBlendState::FAGXBlendState(FBlendStateInitializerRHI const& Initializer)
 {
 	bUseIndependentRenderTargetBlendStates = Initializer.bUseIndependentRenderTargetBlendStates;
 	bUseAlphaToCoverage = Initializer.bUseAlphaToCoverage;
-	for(uint32 RenderTargetIndex = 0;RenderTargetIndex < MaxSimultaneousRenderTargets; ++RenderTargetIndex)
+	for (uint32 RenderTargetIndex = 0; RenderTargetIndex < MaxSimultaneousRenderTargets; ++RenderTargetIndex)
 	{
 		// which initializer to use
-		const FBlendStateInitializerRHI::FRenderTarget& Init =
+		FBlendStateInitializerRHI::FRenderTarget const& Init =
 			Initializer.bUseIndependentRenderTargetBlendStates
 				? Initializer.RenderTargets[RenderTargetIndex]
 				: Initializer.RenderTargets[0];
 
 		// make a new blend state
-		mtlpp::RenderPipelineColorAttachmentDescriptor BlendState = BlendStates.Find(Init);
+		MTLRenderPipelineColorAttachmentDescriptor* BlendState = (MTLRenderPipelineColorAttachmentDescriptor*)BlendStates.Find(Init);
 
-		if (!BlendState.GetPtr())
+		if (!BlendState)
 		{
-			BlendState = mtlpp::RenderPipelineColorAttachmentDescriptor();
-			
+			BlendState = [[MTLRenderPipelineColorAttachmentDescriptor alloc] init];
+
 			// set values
-			BlendState.SetBlendingEnabled(
+			[BlendState setBlendingEnabled:
 				Init.ColorBlendOp != BO_Add || Init.ColorDestBlend != BF_Zero || Init.ColorSrcBlend != BF_One ||
-				Init.AlphaBlendOp != BO_Add || Init.AlphaDestBlend != BF_Zero || Init.AlphaSrcBlend != BF_One);
-			BlendState.SetSourceRgbBlendFactor(TranslateBlendFactor(Init.ColorSrcBlend));
-			BlendState.SetDestinationRgbBlendFactor(TranslateBlendFactor(Init.ColorDestBlend));
-			BlendState.SetRgbBlendOperation(TranslateBlendOp(Init.ColorBlendOp));
-			BlendState.SetSourceAlphaBlendFactor(TranslateBlendFactor(Init.AlphaSrcBlend));
-			BlendState.SetDestinationAlphaBlendFactor(TranslateBlendFactor(Init.AlphaDestBlend));
-			BlendState.SetAlphaBlendOperation(TranslateBlendOp(Init.AlphaBlendOp));
-			BlendState.SetWriteMask(TranslateWriteMask(Init.ColorWriteMask));
+				Init.AlphaBlendOp != BO_Add || Init.AlphaDestBlend != BF_Zero || Init.AlphaSrcBlend != BF_One];
+			[BlendState setSourceRGBBlendFactor:TranslateBlendFactor(Init.ColorSrcBlend)];
+			[BlendState setDestinationRGBBlendFactor:TranslateBlendFactor(Init.ColorDestBlend)];
+			[BlendState setRgbBlendOperation:TranslateBlendOp(Init.ColorBlendOp)];
+			[BlendState setSourceAlphaBlendFactor:TranslateBlendFactor(Init.AlphaSrcBlend)];
+			[BlendState setDestinationAlphaBlendFactor:TranslateBlendFactor(Init.AlphaDestBlend)];
+			[BlendState setAlphaBlendOperation:TranslateBlendOp(Init.AlphaBlendOp)];
+			[BlendState setWriteMask:TranslateWriteMask(Init.ColorWriteMask)];
 			
 			BlendStates.Add(Init, BlendState);
+
+			[BlendState release];
 		}
-		
-		RenderTargetStates[RenderTargetIndex].BlendState = BlendState;
+
+		RenderTargetStates[RenderTargetIndex].BlendState = mtlpp::RenderPipelineColorAttachmentDescriptor(BlendState);
 
 		// get the unique key
 		uint32 BlendBitMask =
-			((uint32)BlendState.GetSourceRgbBlendFactor() << 0) | ((uint32)BlendState.GetDestinationRgbBlendFactor() << 4) | ((uint32)BlendState.GetRgbBlendOperation() << 8) |
-			((uint32)BlendState.GetSourceAlphaBlendFactor() << 11) | ((uint32)BlendState.GetDestinationAlphaBlendFactor() << 15) | ((uint32)BlendState.GetAlphaBlendOperation() << 19) |
-			((uint32)BlendState.GetWriteMask() << 22);
-		
+			((uint32)[BlendState sourceRGBBlendFactor] << 0) | ((uint32)[BlendState destinationRGBBlendFactor] << 4) | ((uint32)[BlendState rgbBlendOperation] << 8) |
+			((uint32)[BlendState sourceAlphaBlendFactor] << 11) | ((uint32)[BlendState destinationAlphaBlendFactor] << 15) | ((uint32)[BlendState alphaBlendOperation] << 19) |
+			((uint32)[BlendState writeMask] << 22);
 		
-		if(IsRunningRHIInSeparateThread())
+		if (IsRunningRHIInSeparateThread())
 		{
 			Mutex.Lock();
 		}
@@ -451,7 +468,7 @@ FAGXBlendState::FAGXBlendState(const FBlendStateInitializerRHI& Initializer)
 		}
 		// set the key
 		RenderTargetStates[RenderTargetIndex].BlendStateKey = *Key;
-		if(IsRunningRHIInSeparateThread())
+		if (IsRunningRHIInSeparateThread())
 		{
 			Mutex.Unlock();
 		}
diff --git a/Engine/Source/Runtime/Apple/AGXRHI/Private/AGXStateCache.cpp b/Engine/Source/Runtime/Apple/AGXRHI/Private/AGXStateCache.cpp
--- a/Engine/Source/Runtime/Apple/AGXRHI/Private/AGXStateCache.cpp
+++ b/Engine/Source/Runtime/Apple/AGXRHI/Private/AGXStateCache.cpp
@@ -1395,12 +1395,12 @@ void FAGXStateCache::SetShaderTexture(EAGXShaderStages const Frequency, FAGXText
 	}
 }
 
-void FAGXStateCache::SetShaderSamplerState(EAGXShaderStages const Frequency, FAGXSamplerState* const Sampler, NSUInteger const Index)
+void FAGXStateCache::SetShaderSamplerState(EAGXShaderStages Frequency, FAGXSamplerState* const Sampler, NSUInteger Index)
 {
 	check(Frequency < EAGXShaderStages::Num);
 	check(Index < ML_MaxSamplers);
 	
-	if (ShaderSamplers[Frequency].Samplers[Index].GetPtr() != (Sampler ? Sampler->State.GetPtr() : nil))
+	if (ShaderSamplers[Frequency].Samplers[Index] != (Sampler ? Sampler->State : nil))
 	{
 		if (Sampler)
 		{
@@ -1914,8 +1914,8 @@ void FAGXStateCache::SetRenderState(FAGXCommandEncoder& CommandEncoder, FAGXComm
             
             if (DepthStencilState && RenderPassDesc && AGXSafeGetRuntimeDebuggingLevel() >= EAGXDebugLevelFastValidation)
             {
-                METAL_FATAL_ASSERT(DepthStencilState->bIsDepthWriteEnabled == false || [[RenderPassDesc depthAttachment] texture] , TEXT("Attempting to set a depth-stencil state that writes depth but no depth texture is configured!\nState: %s\nRender Pass: %s"), *FString([DepthStencilState->State.GetPtr() description]), *FString([RenderPassDesc description]));
-                METAL_FATAL_ASSERT(DepthStencilState->bIsStencilWriteEnabled == false || [[RenderPassDesc stencilAttachment] texture], TEXT("Attempting to set a depth-stencil state that writes stencil but no stencil texture is configured!\nState: %s\nRender Pass: %s"), *FString([DepthStencilState->State.GetPtr() description]), *FString([RenderPassDesc description]));
+                METAL_FATAL_ASSERT(DepthStencilState->bIsDepthWriteEnabled == false || [[RenderPassDesc depthAttachment] texture] , TEXT("Attempting to set a depth-stencil state that writes depth but no depth texture is configured!\nState: %s\nRender Pass: %s"), *FString([DepthStencilState->State description]), *FString([RenderPassDesc description]));
+                METAL_FATAL_ASSERT(DepthStencilState->bIsStencilWriteEnabled == false || [[RenderPassDesc stencilAttachment] texture], TEXT("Attempting to set a depth-stencil state that writes stencil but no stencil texture is configured!\nState: %s\nRender Pass: %s"), *FString([DepthStencilState->State description]), *FString([RenderPassDesc description]));
             }
             
 			CommandEncoder.SetDepthStencilState(DepthStencilState ? DepthStencilState->State : nil);
diff --git a/Engine/Source/Runtime/Apple/AGXRHI/Private/AGXStateCache.h b/Engine/Source/Runtime/Apple/AGXRHI/Private/AGXStateCache.h
--- a/Engine/Source/Runtime/Apple/AGXRHI/Private/AGXStateCache.h
+++ b/Engine/Source/Runtime/Apple/AGXRHI/Private/AGXStateCache.h
@@ -95,7 +95,7 @@ public:
 	 * @param Sampler The sampler state to bind or nil to clear.
 	 * @param Index The index to modify.
 	 */
-	void SetShaderSamplerState(EAGXShaderStages const Frequency, FAGXSamplerState* const Sampler, NSUInteger const Index);
+	void SetShaderSamplerState(EAGXShaderStages Frequency, FAGXSamplerState* const Sampler, NSUInteger Index);
 
 	void SetShaderResourceView(FAGXContext* Context, EAGXShaderStages ShaderStage, uint32 BindIndex, FAGXShaderResourceView* RESTRICT SRV);
 	
@@ -228,9 +228,19 @@ private:
 	/** A structure of arrays for the current sampler binding settings. */
 	struct FAGXSamplerBindings
 	{
-		FAGXSamplerBindings() : Bound(0) {}
+		FAGXSamplerBindings()
+		{
+			for (uint32 SamplerIndex = 0; SamplerIndex < ML_MaxSamplers; ++SamplerIndex)
+			{
+				Samplers[SamplerIndex] = nil;
+			}
+
+			Bound = 0;
+		}
+
 		/** The bound sampler states or nil. */
-		ns::AutoReleased<FAGXSampler> Samplers[ML_MaxSamplers];
+		id<MTLSamplerState> Samplers[ML_MaxSamplers];
+
 		/** A bitmask for which samplers were bound by the application where a bit value of 1 is bound and 0 is unbound. */
 		uint16 Bound;
 	};
diff --git a/Engine/Source/Runtime/Apple/AGXRHI/Public/AGXState.h b/Engine/Source/Runtime/Apple/AGXRHI/Public/AGXState.h
--- a/Engine/Source/Runtime/Apple/AGXRHI/Public/AGXState.h
+++ b/Engine/Source/Runtime/Apple/AGXRHI/Public/AGXState.h
@@ -12,68 +12,6 @@ THIRD_PARTY_INCLUDES_START
 #include "mtlpp.hpp"
 THIRD_PARTY_INCLUDES_END
 
-class FAGXSampler : public mtlpp::SamplerState
-{
-public:
-	FAGXSampler(ns::Ownership retain = ns::Ownership::Retain) : mtlpp::SamplerState(nullptr, nullptr, retain) { }
-	FAGXSampler(ns::Protocol<id<MTLSamplerState>>::type handle, ns::Ownership retain = ns::Ownership::Retain)
-	: mtlpp::SamplerState(handle, nullptr, retain) {}
-	
-	FAGXSampler(mtlpp::SamplerState&& rhs)
-	: mtlpp::SamplerState((mtlpp::SamplerState&&)rhs)
-	{
-		
-	}
-	
-	FAGXSampler(const FAGXSampler& rhs)
-	: mtlpp::SamplerState(rhs)
-	{
-		
-	}
-	
-	FAGXSampler(const SamplerState& rhs)
-	: mtlpp::SamplerState(rhs)
-	{
-		
-	}
-	
-	FAGXSampler(FAGXSampler&& rhs)
-	: mtlpp::SamplerState((mtlpp::SamplerState&&)rhs)
-	{
-		
-	}
-	
-	FAGXSampler& operator=(const FAGXSampler& rhs)
-	{
-		if (this != &rhs)
-		{
-			mtlpp::SamplerState::operator=(rhs);
-		}
-		return *this;
-	}
-	
-	FAGXSampler& operator=(FAGXSampler&& rhs)
-	{
-		mtlpp::SamplerState::operator=((mtlpp::SamplerState&&)rhs);
-		return *this;
-	}
-	
-	inline bool operator==(FAGXSampler const& rhs) const
-	{
-		return GetPtr() == rhs.GetPtr();
-	}
-	
-	inline bool operator!=(FAGXSampler const& rhs) const
-	{
-		return GetPtr() != rhs.GetPtr();
-	}
-	
-	friend uint32 GetTypeHash(FAGXSampler const& Hash)
-	{
-		return GetTypeHash(Hash.GetPtr());
-	}
-};
-
 class FAGXSamplerState : public FRHISamplerState
 {
 public:
@@ -84,9 +22,9 @@ public:
 	FAGXSamplerState(const FSamplerStateInitializerRHI& Initializer);
 	~FAGXSamplerState();
 
-	FAGXSampler State;
+	id<MTLSamplerState> State;
 #if !PLATFORM_MAC
-	FAGXSampler NoAnisoState;
+	id<MTLSamplerState> NoAnisoState;
 #endif
 };
 
@@ -97,7 +35,7 @@ public:
 	/**
 	 * Constructor/destructor
 	 */
-	FAGXRasterizerState(const FRasterizerStateInitializerRHI& Initializer);
+	FAGXRasterizerState(FRasterizerStateInitializerRHI const& Initializer);
 	~FAGXRasterizerState();
 	
 	virtual bool GetInitializer(FRasterizerStateInitializerRHI& Init) override final;
@@ -112,13 +50,13 @@ public:
 	/**
 	 * Constructor/destructor
 	 */
-	FAGXDepthStencilState(const FDepthStencilStateInitializerRHI& Initializer);
+	FAGXDepthStencilState(FDepthStencilStateInitializerRHI const& Initializer);
 	~FAGXDepthStencilState();
 	
 	virtual bool GetInitializer(FDepthStencilStateInitializerRHI& Init) override final;
 	
 	FDepthStencilStateInitializerRHI Initializer;
-	mtlpp::DepthStencilState State;
+	id<MTLDepthStencilState> State;
 	bool bIsDepthWriteEnabled;
 	bool bIsStencilWriteEnabled;
 };
@@ -130,7 +68,7 @@ public:
 	/**
 	 * Constructor/destructor
 	 */
-	FAGXBlendState(const FBlendStateInitializerRHI& Initializer);
+	FAGXBlendState(FBlendStateInitializerRHI const& Initializer);
 	~FAGXBlendState();
 	
 	virtual bool GetInitializer(FBlendStateInitializerRHI& Init) override final;
