AGXRHI: Scrub mtlpp render/compute pipeline state and descriptor types

#rnx

diff --git a/Engine/Source/Runtime/Apple/AGXRHI/Private/AGXCommandEncoder.cpp b/Engine/Source/Runtime/Apple/AGXRHI/Private/AGXCommandEncoder.cpp
--- a/Engine/Source/Runtime/Apple/AGXRHI/Private/AGXCommandEncoder.cpp
+++ b/Engine/Source/Runtime/Apple/AGXRHI/Private/AGXCommandEncoder.cpp
@@ -768,10 +768,9 @@ void FAGXCommandEncoder::SetRenderPassStoreActions(MTLStoreAction const* ColorSt
 
 void FAGXCommandEncoder::SetRenderPipelineState(FAGXShaderPipeline* PipelineState)
 {
-	check (RenderCommandEncoder);
-	{
-		[RenderCommandEncoder setRenderPipelineState:PipelineState->RenderPipelineState.GetPtr()];
-	}
+	check(RenderCommandEncoder);
+
+	[RenderCommandEncoder setRenderPipelineState:PipelineState->RenderPipelineState];
 }
 
 void FAGXCommandEncoder::SetViewport(MTLViewport const Viewport[], uint32 NumActive)
@@ -1108,7 +1107,7 @@ void FAGXCommandEncoder::SetComputePipelineState(FAGXShaderPipeline* State)
 {
 	check (ComputeCommandEncoder);
 	{
-		[ComputeCommandEncoder setComputePipelineState:State->ComputePipelineState.GetPtr()];
+		[ComputeCommandEncoder setComputePipelineState:State->ComputePipelineState];
 	}
 }
 
diff --git a/Engine/Source/Runtime/Apple/AGXRHI/Private/AGXPipeline.cpp b/Engine/Source/Runtime/Apple/AGXRHI/Private/AGXPipeline.cpp
--- a/Engine/Source/Runtime/Apple/AGXRHI/Private/AGXPipeline.cpp
+++ b/Engine/Source/Runtime/Apple/AGXRHI/Private/AGXPipeline.cpp
@@ -322,22 +322,13 @@ APPLE_PLATFORM_OBJECT_ALLOC_OVERRIDES(FAGXShaderPipeline)
 	{
 		FAGXShaderPipelineCache::Get().RemoveRenderPipeline(self);
 	}
-	[super dealloc];
-}
 
-- (instancetype)init
-{
-	id Self = [super init];
-	if (Self)
-	{
-		RenderPipelineReflection = mtlpp::RenderPipelineReflection(nil);
-		ComputePipelineReflection = mtlpp::ComputePipelineReflection(nil);
-#if METAL_DEBUG_OPTIONS
-		RenderDesc = mtlpp::RenderPipelineDescriptor(nil);
-		ComputeDesc = mtlpp::ComputePipelineDescriptor(nil);
-#endif
-	}
-	return Self;
+	[RenderPipelineState release];
+	[ComputePipelineState release];
+	[RenderPipelineReflection release];
+	[ComputePipelineReflection release];
+
+	[super dealloc];
 }
 
 - (void)initResourceMask
@@ -349,7 +340,8 @@ APPLE_PLATFORM_OBJECT_ALLOC_OVERRIDES(FAGXShaderPipeline)
 		
 		if (AGXSafeGetRuntimeDebuggingLevel() < EAGXDebugLevelValidation)
 		{
-			RenderPipelineReflection = mtlpp::RenderPipelineReflection(nil);
+			[RenderPipelineReflection release];
+			RenderPipelineReflection = nil;
 		}
 	}
 	if (ComputePipelineReflection)
@@ -358,7 +350,8 @@ APPLE_PLATFORM_OBJECT_ALLOC_OVERRIDES(FAGXShaderPipeline)
 		
 		if (AGXSafeGetRuntimeDebuggingLevel() < EAGXDebugLevelValidation)
 		{
-			ComputePipelineReflection = mtlpp::ComputePipelineReflection(nil);
+			[ComputePipelineReflection release];
+			ComputePipelineReflection = nil;
 		}
 	}
 }
@@ -369,26 +362,23 @@ APPLE_PLATFORM_OBJECT_ALLOC_OVERRIDES(FAGXShaderPipeline)
 	{
 		case SF_Vertex:
 		{
-			MTLRenderPipelineReflection* Reflection = RenderPipelineReflection;
-			check(Reflection);
+			check(RenderPipelineReflection);
 			
-			Arguments = Reflection.vertexArguments;
+			Arguments = [RenderPipelineReflection vertexArguments];
 			break;
 		}
 		case SF_Pixel:
 		{
-			MTLRenderPipelineReflection* Reflection = RenderPipelineReflection;
-			check(Reflection);
+			check(RenderPipelineReflection);
 			
-			Arguments = Reflection.fragmentArguments;
+			Arguments = [RenderPipelineReflection fragmentArguments];
 			break;
 		}
 		case SF_Compute:
 		{
-			MTLComputePipelineReflection* Reflection = ComputePipelineReflection;
-			check(Reflection);
+			check(ComputePipelineReflection);
 			
-			Arguments = Reflection.arguments;
+			Arguments = [ComputePipelineReflection arguments];
 			break;
 		}
 		default:
@@ -480,12 +470,7 @@ static MTLVertexDescriptor* GetMaskedVertexDescriptor(MTLVertexDescriptor* Input
 	return InputDesc;
 }
 
-static bool ConfigureRenderPipelineDescriptor(mtlpp::RenderPipelineDescriptor& RenderPipelineDesc,
-#if PLATFORM_MAC
-											  mtlpp::RenderPipelineDescriptor& DebugPipelineDesc,
-#elif !PLATFORM_TVOS
-											  mtlpp::TileRenderPipelineDescriptor& DebugPipelineDesc,
-#endif
+static bool ConfigureRenderPipelineDescriptor(MTLRenderPipelineDescriptor* RenderPipelineDesc,
 											  FAGXGraphicsPipelineKey const& Key,
 											  const FGraphicsPipelineStateInitializer& Init)
 {
@@ -506,10 +491,7 @@ static bool ConfigureRenderPipelineDescriptor(mtlpp::RenderPipelineDescriptor& R
 	
 	FAGXBlendState* BlendState = (FAGXBlendState*)Init.BlendState;
 	
-	ns::Array<mtlpp::RenderPipelineColorAttachmentDescriptor> ColorAttachments = RenderPipelineDesc.GetColorAttachments();
-#if !PLATFORM_TVOS
-	auto DebugColorAttachements = DebugPipelineDesc.GetColorAttachments();
-#endif
+	MTLRenderPipelineColorAttachmentDescriptorArray* ColorAttachments = [RenderPipelineDesc colorAttachments];
 	
 	uint32 TargetWidth = 0;
 	for (uint32 ActiveTargetIndex = 0; ActiveTargetIndex < NumActiveTargets; ActiveTargetIndex++)
@@ -521,59 +503,39 @@ static bool ConfigureRenderPipelineDescriptor(mtlpp::RenderPipelineDescriptor& R
 		
 		TargetWidth += GPixelFormats[TargetFormat].BlockBytes;
 		
-		mtlpp::PixelFormat MetalFormat = (mtlpp::PixelFormat)GPixelFormats[TargetFormat].PlatformFormat;
+		MTLPixelFormat MetalFormat = (MTLPixelFormat)GPixelFormats[TargetFormat].PlatformFormat;
 		ETextureCreateFlags Flags = (ETextureCreateFlags)Init.RenderTargetFlags[ActiveTargetIndex];
 		if (EnumHasAnyFlags(Flags, TexCreate_SRGB))
 		{
 #if PLATFORM_MAC // Expand as R8_sRGB is iOS only.
-			if (MetalFormat == mtlpp::PixelFormat::R8Unorm)
+			if (MetalFormat == MTLPixelFormatR8Unorm)
 			{
-				MetalFormat = mtlpp::PixelFormat::RGBA8Unorm;
+				MetalFormat = MTLPixelFormatRGBA8Unorm;
 			}
 #endif
-			MetalFormat = AGXToSRGBFormat(MetalFormat);
+			MetalFormat = (MTLPixelFormat)AGXToSRGBFormat((mtlpp::PixelFormat)MetalFormat);
 		}
 		
-		mtlpp::RenderPipelineColorAttachmentDescriptor Attachment = ColorAttachments[ActiveTargetIndex];
-		Attachment.SetPixelFormat(MetalFormat);
+		MTLRenderPipelineColorAttachmentDescriptor* Attachment = [ColorAttachments objectAtIndexedSubscript:ActiveTargetIndex];
+		[Attachment setPixelFormat:MetalFormat];
 		
-#if !PLATFORM_TVOS
-		auto DebugAttachment = DebugColorAttachements[ActiveTargetIndex];
-		DebugAttachment.SetPixelFormat(MetalFormat);
-#endif
-		
-		mtlpp::RenderPipelineColorAttachmentDescriptor Blend = BlendState->RenderTargetStates[ActiveTargetIndex].BlendState;
-		if(TargetFormat != PF_Unknown)
+		MTLRenderPipelineColorAttachmentDescriptor* Blend = BlendState->RenderTargetStates[ActiveTargetIndex].BlendState;
+		if (TargetFormat != PF_Unknown)
 		{
 			// assign each property manually, would be nice if this was faster
-			Attachment.SetBlendingEnabled(Blend.IsBlendingEnabled());
-			Attachment.SetSourceRgbBlendFactor(Blend.GetSourceRgbBlendFactor());
-			Attachment.SetDestinationRgbBlendFactor(Blend.GetDestinationRgbBlendFactor());
-			Attachment.SetRgbBlendOperation(Blend.GetRgbBlendOperation());
-			Attachment.SetSourceAlphaBlendFactor(Blend.GetSourceAlphaBlendFactor());
-			Attachment.SetDestinationAlphaBlendFactor(Blend.GetDestinationAlphaBlendFactor());
-			Attachment.SetAlphaBlendOperation(Blend.GetAlphaBlendOperation());
-			Attachment.SetWriteMask(Blend.GetWriteMask());
-			
-#if PLATFORM_MAC
-			DebugAttachment.SetBlendingEnabled(Blend.IsBlendingEnabled());
-			DebugAttachment.SetSourceRgbBlendFactor(Blend.GetSourceRgbBlendFactor());
-			DebugAttachment.SetDestinationRgbBlendFactor(Blend.GetDestinationRgbBlendFactor());
-			DebugAttachment.SetRgbBlendOperation(Blend.GetRgbBlendOperation());
-			DebugAttachment.SetSourceAlphaBlendFactor(Blend.GetSourceAlphaBlendFactor());
-			DebugAttachment.SetDestinationAlphaBlendFactor(Blend.GetDestinationAlphaBlendFactor());
-			DebugAttachment.SetAlphaBlendOperation(Blend.GetAlphaBlendOperation());
-			DebugAttachment.SetWriteMask(Blend.GetWriteMask());
-#endif
+			[Attachment setBlendingEnabled:[Blend isBlendingEnabled]];
+			[Attachment setSourceRGBBlendFactor:[Blend sourceRGBBlendFactor]];
+			[Attachment setDestinationRGBBlendFactor:[Blend destinationRGBBlendFactor]];
+			[Attachment setRgbBlendOperation:[Blend rgbBlendOperation]];
+			[Attachment setSourceAlphaBlendFactor:[Blend sourceAlphaBlendFactor]];
+			[Attachment setDestinationAlphaBlendFactor:[Blend destinationAlphaBlendFactor]];
+			[Attachment setAlphaBlendOperation:[Blend alphaBlendOperation]];
+			[Attachment setWriteMask:[Blend writeMask]];
 		}
 		else
 		{
-			Attachment.SetBlendingEnabled(NO);
-			Attachment.SetWriteMask(mtlpp::ColorWriteMask::None);
-#if PLATFORM_MAC
-			DebugAttachment.SetBlendingEnabled(NO);
-			DebugAttachment.SetWriteMask(mtlpp::ColorWriteMask::None);
-#endif
+			[Attachment setBlendingEnabled:NO];
+			[Attachment setWriteMask:MTLColorWriteMaskNone];
 		}
 	}
 	
@@ -583,45 +545,32 @@ static bool ConfigureRenderPipelineDescriptor(mtlpp::RenderPipelineDescriptor& R
 		return false;
 	}
 	
-	switch(Init.DepthStencilTargetFormat)
+	switch (Init.DepthStencilTargetFormat)
 	{
 		case PF_DepthStencil:
 		{
-			mtlpp::PixelFormat MetalFormat = (mtlpp::PixelFormat)GPixelFormats[PF_DepthStencil].PlatformFormat;
-			if(MetalFormat == mtlpp::PixelFormat::Depth32Float)
+			MTLPixelFormat MetalFormat = (MTLPixelFormat)GPixelFormats[PF_DepthStencil].PlatformFormat;
+			if(MetalFormat == MTLPixelFormatDepth32Float)
 			{
 				if (Init.DepthTargetLoadAction != ERenderTargetLoadAction::ENoAction || Init.DepthTargetStoreAction != ERenderTargetStoreAction::ENoAction)
 				{
-					RenderPipelineDesc.SetDepthAttachmentPixelFormat(MetalFormat);
-#if PLATFORM_MAC
-					DebugPipelineDesc.SetDepthAttachmentPixelFormat(MetalFormat);
-#endif
+					[RenderPipelineDesc setDepthAttachmentPixelFormat:MetalFormat];
 				}
 				if (Init.StencilTargetLoadAction != ERenderTargetLoadAction::ENoAction || Init.StencilTargetStoreAction != ERenderTargetStoreAction::ENoAction)
 				{
-					RenderPipelineDesc.SetStencilAttachmentPixelFormat(mtlpp::PixelFormat::Stencil8);
-#if PLATFORM_MAC
-					DebugPipelineDesc.SetStencilAttachmentPixelFormat(mtlpp::PixelFormat::Stencil8);
-#endif
+					[RenderPipelineDesc setStencilAttachmentPixelFormat:MTLPixelFormatStencil8];
 				}
 			}
 			else
 			{
-				RenderPipelineDesc.SetDepthAttachmentPixelFormat(MetalFormat);
-				RenderPipelineDesc.SetStencilAttachmentPixelFormat(MetalFormat);
-#if PLATFORM_MAC
-				DebugPipelineDesc.SetDepthAttachmentPixelFormat(MetalFormat);
-				DebugPipelineDesc.SetStencilAttachmentPixelFormat(MetalFormat);
-#endif
+				[RenderPipelineDesc setDepthAttachmentPixelFormat:MetalFormat];
+				[RenderPipelineDesc setStencilAttachmentPixelFormat:MetalFormat];
 			}
 			break;
 		}
 		case PF_ShadowDepth:
 		{
-			RenderPipelineDesc.SetDepthAttachmentPixelFormat((mtlpp::PixelFormat)GPixelFormats[PF_ShadowDepth].PlatformFormat);
-#if PLATFORM_MAC
-			DebugPipelineDesc.SetDepthAttachmentPixelFormat((mtlpp::PixelFormat)GPixelFormats[PF_ShadowDepth].PlatformFormat);
-#endif
+			[RenderPipelineDesc setDepthAttachmentPixelFormat:(MTLPixelFormat)GPixelFormats[PF_ShadowDepth].PlatformFormat];
 			break;
 		}
 		default:
@@ -635,74 +584,62 @@ static bool ConfigureRenderPipelineDescriptor(mtlpp::RenderPipelineDescriptor& R
 	check(Init.BoundShaderState.GetGeometryShader() == nullptr);
 #endif
 	
-	if( RenderPipelineDesc.GetDepthAttachmentPixelFormat() == mtlpp::PixelFormat::Invalid &&
+	if ([RenderPipelineDesc depthAttachmentPixelFormat] == MTLPixelFormatInvalid &&
 		PixelShader && ( PixelShader->Bindings.InOutMask.IsFieldEnabled(CrossCompiler::FShaderBindingInOutMask::DepthStencilMaskIndex) || ( NumActiveTargets == 0 && PixelShader->Bindings.NumUAVs > 0) ) )
 	{
-		RenderPipelineDesc.SetDepthAttachmentPixelFormat((mtlpp::PixelFormat)GPixelFormats[PF_DepthStencil].PlatformFormat);
-		RenderPipelineDesc.SetStencilAttachmentPixelFormat((mtlpp::PixelFormat)GPixelFormats[PF_DepthStencil].PlatformFormat);
-		
-#if PLATFORM_MAC
-		DebugPipelineDesc.SetDepthAttachmentPixelFormat((mtlpp::PixelFormat)GPixelFormats[PF_DepthStencil].PlatformFormat);
-		DebugPipelineDesc.SetStencilAttachmentPixelFormat((mtlpp::PixelFormat)GPixelFormats[PF_DepthStencil].PlatformFormat);
-#endif
+		[RenderPipelineDesc setDepthAttachmentPixelFormat:(MTLPixelFormat)GPixelFormats[PF_DepthStencil].PlatformFormat];
+		[RenderPipelineDesc setStencilAttachmentPixelFormat:(MTLPixelFormat)GPixelFormats[PF_DepthStencil].PlatformFormat];
 	}
 	
 	static bool bNoMSAA = FParse::Param(FCommandLine::Get(), TEXT("nomsaa"));
 	uint16 NumSamples = !bNoMSAA ? FMath::Max(Init.NumSamples, (uint16)1u) : (uint16)1u;
-	RenderPipelineDesc.SetSampleCount(NumSamples);
-	RenderPipelineDesc.SetAlphaToCoverageEnabled(NumSamples > 1 && BlendState->bUseAlphaToCoverage);
+	[RenderPipelineDesc setSampleCount:NumSamples];
+	[RenderPipelineDesc setAlphaToCoverageEnabled:(NumSamples > 1 && BlendState->bUseAlphaToCoverage)];
 #if PLATFORM_MAC
-	RenderPipelineDesc.SetInputPrimitiveTopology(AGXTranslatePrimitiveTopology(Init.PrimitiveType));
-	DebugPipelineDesc.SetSampleCount(!bNoMSAA ? FMath::Max(Init.NumSamples, (uint16)1u) : (uint16)1u);
-	DebugPipelineDesc.SetInputPrimitiveTopology(mtlpp::PrimitiveTopologyClass::Point);
+	[RenderPipelineDesc setInputPrimitiveTopology:(MTLPrimitiveTopologyClass)AGXTranslatePrimitiveTopology(Init.PrimitiveType)];
 #endif
-	
 	if (FAGXCommandQueue::SupportsFeature(EAGXFeaturesPipelineBufferMutability))
 	{
 		FAGXVertexShader* VertexShader = (FAGXVertexShader*)Init.BoundShaderState.VertexShaderRHI;
 		
-		ns::AutoReleased<ns::Array<mtlpp::PipelineBufferDescriptor>> VertexPipelineBuffers = RenderPipelineDesc.GetVertexBuffers();
+		MTLPipelineBufferDescriptorArray* VertexPipelineBuffers = [RenderPipelineDesc vertexBuffers];
 		FMetalShaderBindings& VertexBindings = VertexShader->Bindings;
 		int8 VertexSideTable = VertexShader->SideTableBinding;
 		{
 			uint32 ImmutableBuffers = VertexBindings.ConstantBuffers;
-			while(ImmutableBuffers)
+			while (ImmutableBuffers)
 			{
 				uint32 Index = __builtin_ctz(ImmutableBuffers);
 				ImmutableBuffers &= ~(1 << Index);
 				
 				if (Index < ML_MaxBuffers)
 				{
-					ns::AutoReleased<mtlpp::PipelineBufferDescriptor> PipelineBuffer = VertexPipelineBuffers[Index];
-					PipelineBuffer.SetMutability(mtlpp::Mutability::Immutable);
+					[[VertexPipelineBuffers objectAtIndexedSubscript:Index] setMutability:MTLMutabilityImmutable];
 				}
 			}
 			if (VertexSideTable > 0)
 			{
-				ns::AutoReleased<mtlpp::PipelineBufferDescriptor> PipelineBuffer = VertexPipelineBuffers[VertexSideTable];
-				PipelineBuffer.SetMutability(mtlpp::Mutability::Immutable);
+				[[VertexPipelineBuffers objectAtIndexedSubscript:VertexSideTable] setMutability:MTLMutabilityImmutable];
 			}
 		}
 		
 		if (PixelShader)
 		{
-			ns::AutoReleased<ns::Array<mtlpp::PipelineBufferDescriptor>> FragmentPipelineBuffers = RenderPipelineDesc.GetFragmentBuffers();
+			MTLPipelineBufferDescriptorArray* FragmentPipelineBuffers = [RenderPipelineDesc fragmentBuffers];
 			uint32 ImmutableBuffers = PixelShader->Bindings.ConstantBuffers;
-			while(ImmutableBuffers)
+			while (ImmutableBuffers)
 			{
 				uint32 Index = __builtin_ctz(ImmutableBuffers);
 				ImmutableBuffers &= ~(1 << Index);
 				
 				if (Index < ML_MaxBuffers)
 				{
-					ns::AutoReleased<mtlpp::PipelineBufferDescriptor> PipelineBuffer = FragmentPipelineBuffers[Index];
-					PipelineBuffer.SetMutability(mtlpp::Mutability::Immutable);
+					[[FragmentPipelineBuffers objectAtIndexedSubscript:Index] setMutability:MTLMutabilityImmutable];
 				}
 			}
 			if (PixelShader->SideTableBinding > 0)
 			{
-				ns::AutoReleased<mtlpp::PipelineBufferDescriptor> PipelineBuffer = FragmentPipelineBuffers[PixelShader->SideTableBinding];
-				PipelineBuffer.SetMutability(mtlpp::Mutability::Immutable);
+				[[FragmentPipelineBuffers objectAtIndexedSubscript:PixelShader->SideTableBinding] setMutability:MTLMutabilityImmutable];
 			}
 		}
 	}
@@ -795,22 +732,22 @@ static void InitializeAGXPipelineCache()
 	}
 }
 
-static void RelatePipelineStateToCache(const mtlpp::RenderPipelineDescriptor & PipelineDesc, NSUInteger * outPipelineOpts)
+static void RelatePipelineStateToCache(MTLRenderPipelineDescriptor* PipelineDesc)
 {
 	if (GPSOBinaryArchive && (GPSOCacheMode != CacheMode::Ignore))
 	{
 		if (GPSOCacheMode == CacheMode::Recreate || GPSOCacheMode == CacheMode::Append)
 		{
-			__autoreleasing NSError * err = nil;
-			bool bAddedBinaryPSO = [GPSOBinaryArchive addRenderPipelineFunctionsWithDescriptor : PipelineDesc.GetPtr() error : &err];
-			if (err)
+			NSError* Error = nil;
+			bool bAddedBinaryPSO = [GPSOBinaryArchive addRenderPipelineFunctionsWithDescriptor:PipelineDesc error:&Error];
+			if (Error != nil)
 			{
-				UE_LOG(LogAGX, Warning, TEXT("AGX Pipeline Cache: Error adding Pipeline Functions to Binary Archive: %s"), *FString(err.localizedDescription.UTF8String));
+				UE_LOG(LogAGX, Warning, TEXT("AGX Pipeline Cache: Error adding Pipeline Functions to Binary Archive: %s"), *FString([Error localizedDescription]));
 			}
-            else if(bAddedBinaryPSO)
+            else if (bAddedBinaryPSO)
             {
                 GPSOHarvestCount++;
-                if(GAGXBinaryCacheDebugOutput && GPSOHarvestCount % 100)
+                if (GAGXBinaryCacheDebugOutput && GPSOHarvestCount % 100)
                 {
                     UE_LOG(LogAGX, Log, TEXT("AGX Pipeline Cache: Harvested PSO count: %d"), GPSOHarvestCount);
                 }
@@ -819,22 +756,22 @@ static void RelatePipelineStateToCache(const mtlpp::RenderPipelineDescriptor & P
 	}
 }
 
-static void RelatePipelineStateToCache(const mtlpp::ComputePipelineDescriptor & PipelineDesc, NSUInteger * outPipelineOpts)
+static void RelatePipelineStateToCache(MTLComputePipelineDescriptor* PipelineDesc)
 {
 	if (GPSOBinaryArchive && (GPSOCacheMode != CacheMode::Ignore))
 	{
 		if (GPSOCacheMode == CacheMode::Recreate || GPSOCacheMode == CacheMode::Append)
 		{
-			__autoreleasing NSError* err = nil;
-            bool bAddedBinaryPSO = [GPSOBinaryArchive addComputePipelineFunctionsWithDescriptor : PipelineDesc.GetPtr() error : &err] ;
-			if (err)
+			NSError* Error = nil;
+            bool bAddedBinaryPSO = [GPSOBinaryArchive addComputePipelineFunctionsWithDescriptor:PipelineDesc error:&Error] ;
+			if (Error != nil)
 			{
-				UE_LOG(LogAGX, Warning, TEXT("AGX Pipeline Cache: Error adding Pipeline Functions to Binary Archive: %s"), *FString(err.localizedDescription.UTF8String));
+				UE_LOG(LogAGX, Warning, TEXT("AGX Pipeline Cache: Error adding Pipeline Functions to Binary Archive: %s"), *FString([Error localizedDescription]));
 			}
-            else if(bAddedBinaryPSO)
+            else if (bAddedBinaryPSO)
             {
                 GPSOHarvestCount++;
-                if(GAGXBinaryCacheDebugOutput && GPSOHarvestCount % 100)
+                if (GAGXBinaryCacheDebugOutput && GPSOHarvestCount % 100)
                 {
                     UE_LOG(LogAGX, Log, TEXT("AGX Pipeline Cache: Harvested PSO count: %d"), GPSOHarvestCount);
                 }
@@ -878,13 +815,13 @@ static FAGXShaderPipeline* CreateMTLRenderPipeline(bool const bSync, FAGXGraphic
     FAGXVertexShader* VertexShader = (FAGXVertexShader*)Init.BoundShaderState.VertexShaderRHI;
     FAGXPixelShader* PixelShader = (FAGXPixelShader*)Init.BoundShaderState.PixelShaderRHI;
     
-    mtlpp::Function vertexFunction = VertexShader->GetFunction();
-    mtlpp::Function fragmentFunction = PixelShader ? PixelShader->GetFunction() : nil;
+ 	id<MTLFunction> vertexFunction = VertexShader->GetFunction().GetPtr();
+    id<MTLFunction> fragmentFunction = PixelShader ? PixelShader->GetFunction().GetPtr() : nil;
 
     FAGXShaderPipeline* Pipeline = nil;
     if (vertexFunction && ((PixelShader != nullptr) == (fragmentFunction != nil)))
     {
-		ns::Error Error;
+		NSError* Error = nil;
 
 		uint32 const NumActiveTargets = Init.ComputeNumValidRenderTargets();
         check(NumActiveTargets <= MaxSimultaneousRenderTargets);
@@ -892,66 +829,54 @@ static FAGXShaderPipeline* CreateMTLRenderPipeline(bool const bSync, FAGXGraphic
 		Pipeline = [FAGXShaderPipeline new];
 		METAL_DEBUG_OPTION(FMemory::Memzero(Pipeline->ResourceMask, sizeof(Pipeline->ResourceMask)));
 
-		mtlpp::RenderPipelineDescriptor RenderPipelineDesc;
-        mtlpp::ComputePipelineDescriptor ComputePipelineDesc(nil);
-#if PLATFORM_MAC
-        mtlpp::RenderPipelineDescriptor DebugPipelineDesc;
-#elif !PLATFORM_TVOS
-        mtlpp::TileRenderPipelineDescriptor DebugPipelineDesc;
-#endif
+		MTLRenderPipelineDescriptor* RenderPipelineDesc = [[MTLRenderPipelineDescriptor alloc] init];
 		
-#if PLATFORM_TVOS
 		if (!ConfigureRenderPipelineDescriptor(RenderPipelineDesc, Key, Init))
-#else
-		if (!ConfigureRenderPipelineDescriptor(RenderPipelineDesc, DebugPipelineDesc, Key, Init))
-#endif
 		{
 			return nil;
 		}
         
         FAGXVertexDeclaration* VertexDecl = (FAGXVertexDeclaration*)Init.BoundShaderState.VertexDeclarationRHI;
 		
-		RenderPipelineDesc.SetVertexDescriptor(GetMaskedVertexDescriptor(VertexDecl->Layout.VertexDesc, VertexShader->Bindings.InOutMask));
-		RenderPipelineDesc.SetVertexFunction(vertexFunction);
-		RenderPipelineDesc.SetFragmentFunction(fragmentFunction);
+		[RenderPipelineDesc setVertexDescriptor:GetMaskedVertexDescriptor(VertexDecl->Layout.VertexDesc, VertexShader->Bindings.InOutMask)];
+		[RenderPipelineDesc setVertexFunction:vertexFunction];
+		[RenderPipelineDesc setFragmentFunction:fragmentFunction];
 #if ENABLE_METAL_GPUPROFILE
-		ns::String VertexName = vertexFunction.GetName();
-		ns::String FragmentName = fragmentFunction ? fragmentFunction.GetName() : @"";
-		RenderPipelineDesc.SetLabel([NSString stringWithFormat:@"%@+%@", VertexName.GetPtr(), FragmentName.GetPtr()]);
+		NSString* VertexName = [vertexFunction name];
+		NSString* FragmentName = [fragmentFunction name];
+		[RenderPipelineDesc setLabel:[NSString stringWithFormat:@"%@+%@", VertexName, FragmentName ? FragmentName : @""]];
 #endif
 
-		NSUInteger RenderOption = mtlpp::PipelineOption::NoPipelineOption;
-		mtlpp::AutoReleasedRenderPipelineReflection* Reflection = nullptr;
-		mtlpp::AutoReleasedRenderPipelineReflection OutReflection;
-		Reflection = &OutReflection;
+		MTLPipelineOption RenderPipelineOptions = MTLPipelineOptionNone;
+		MTLAutoreleasedRenderPipelineReflection* Reflection = nil;
 		if (GetAGXDeviceContext().GetCommandQueue().GetRuntimeDebuggingLevel() >= EAGXDebugLevelFastValidation)
 		{
-			RenderOption = mtlpp::PipelineOption::ArgumentInfo | mtlpp::PipelineOption::BufferTypeInfo;
+			RenderPipelineOptions = MTLPipelineOptionArgumentInfo | MTLPipelineOptionBufferTypeInfo;
 		}
 
 		{
-			ns::AutoReleasedError RenderError;
-			METAL_GPUPROFILE(FAGXScopedCPUStats CPUStat(FString::Printf(TEXT("NewRenderPipeline: %s"), TEXT("")/**FString([RenderPipelineDesc.GetPtr() description])*/)));
-			RelatePipelineStateToCache(RenderPipelineDesc, &RenderOption);
-			Pipeline->RenderPipelineState = GMtlppDevice.NewRenderPipelineState(RenderPipelineDesc, (mtlpp::PipelineOption)RenderOption, Reflection, &RenderError);
-			if (Reflection)
+			METAL_GPUPROFILE(FAGXScopedCPUStats CPUStat(FString::Printf(TEXT("NewRenderPipeline: %s"), TEXT("")/**FString([RenderPipelineDesc description])*/)));
+			RelatePipelineStateToCache(RenderPipelineDesc);
+			Pipeline->RenderPipelineState = [GMtlDevice newRenderPipelineStateWithDescriptor:RenderPipelineDesc
+																					 options:RenderPipelineOptions
+																				  reflection:Reflection
+																				  	   error:&Error];
+			if (Reflection != nil)
 			{
-				Pipeline->RenderPipelineReflection = *Reflection;
-#if METAL_DEBUG_OPTIONS
-				Pipeline->RenderDesc = RenderPipelineDesc;
-#endif
+				Pipeline->RenderPipelineReflection = [*Reflection retain];
 			}
-			Error = RenderError;
 		}
 		
-		UE_CLOG((Pipeline->RenderPipelineState == nil), LogAGX, Error, TEXT("Failed to generate a pipeline state object: %s"), *FString(Error.GetPtr().description));
+		UE_CLOG((Pipeline->RenderPipelineState == nil), LogAGX, Error, TEXT("Failed to generate a pipeline state object: %s"), *FString([Error description]));
 		UE_CLOG((Pipeline->RenderPipelineState == nil), LogAGX, Error, TEXT("Vertex shader: %s"), *FString(VertexShader->GetSourceCode()));
 		UE_CLOG((Pipeline->RenderPipelineState == nil), LogAGX, Error, TEXT("Pixel shader: %s"), PixelShader ? *FString(PixelShader->GetSourceCode()) : TEXT("NULL"));
-		UE_CLOG((Pipeline->RenderPipelineState == nil), LogAGX, Error, TEXT("Descriptor: %s"), *FString(RenderPipelineDesc.GetPtr().description));
-		UE_CLOG((Pipeline->RenderPipelineState == nil), LogAGX, Error, TEXT("Failed to generate a render pipeline state object:\n\n %s\n\n"), *FString(Error.GetLocalizedDescription()));
+		UE_CLOG((Pipeline->RenderPipelineState == nil), LogAGX, Error, TEXT("Descriptor: %s"), *FString([RenderPipelineDesc description]));
+		UE_CLOG((Pipeline->RenderPipelineState == nil), LogAGX, Error, TEXT("Failed to generate a render pipeline state object:\n\n %s\n\n"), *FString([Error localizedDescription]));
 		
+		[RenderPipelineDesc release];
+
 		// We need to pass a failure up the chain, so we'll clean up here.
-		if(Pipeline->RenderPipelineState == nil)
+		if (Pipeline->RenderPipelineState == nil)
 		{
 			[Pipeline release];
 			return nil;
diff --git a/Engine/Source/Runtime/Apple/AGXRHI/Private/AGXPipeline.h b/Engine/Source/Runtime/Apple/AGXRHI/Private/AGXPipeline.h
--- a/Engine/Source/Runtime/Apple/AGXRHI/Private/AGXPipeline.h
+++ b/Engine/Source/Runtime/Apple/AGXRHI/Private/AGXPipeline.h
@@ -65,22 +65,19 @@ private:
 @interface FAGXShaderPipeline : FApplePlatformObject
 {
 @public
-	mtlpp::RenderPipelineState RenderPipelineState;
-	mtlpp::ComputePipelineState ComputePipelineState;
+	id<MTLRenderPipelineState> RenderPipelineState;
+	id<MTLComputePipelineState> ComputePipelineState;
 	TArray<uint32> BufferDataSizes[SF_NumStandardFrequencies];
 	TMap<uint8, uint8> TextureTypes[SF_NumStandardFrequencies];
 	FAGXDebugShaderResourceMask ResourceMask[SF_NumStandardFrequencies];
-	mtlpp::RenderPipelineReflection RenderPipelineReflection;
-	mtlpp::ComputePipelineReflection ComputePipelineReflection;
+	MTLRenderPipelineReflection* RenderPipelineReflection;
+	MTLComputePipelineReflection* ComputePipelineReflection;
 #if METAL_DEBUG_OPTIONS
-	ns::String VertexSource;
-	ns::String FragmentSource;
-	ns::String ComputeSource;
-	mtlpp::RenderPipelineDescriptor RenderDesc;
-	mtlpp::ComputePipelineDescriptor ComputeDesc;
+	NSString* VertexSource;
+	NSString* FragmentSource;
+	NSString* ComputeSource;
 #endif
 }
-- (instancetype)init;
 - (void)initResourceMask;
 - (void)initResourceMask:(EShaderFrequency)Frequency;
 @end
diff --git a/Engine/Source/Runtime/Apple/AGXRHI/Private/Shaders/Types/AGXComputeShader.cpp b/Engine/Source/Runtime/Apple/AGXRHI/Private/Shaders/Types/AGXComputeShader.cpp
--- a/Engine/Source/Runtime/Apple/AGXRHI/Private/Shaders/Types/AGXComputeShader.cpp
+++ b/Engine/Source/Runtime/Apple/AGXRHI/Private/Shaders/Types/AGXComputeShader.cpp
@@ -39,62 +39,60 @@ FAGXShaderPipeline* FAGXComputeShader::GetPipeline()
 {
 	if (!Pipeline)
 	{
-		mtlpp::Function Func = GetCompiledFunction();
-		check(Func);
+		id<MTLFunction> Func = GetCompiledFunction().GetPtr();
+		check(Func != nil);
 
-		ns::Error Error;
-		mtlpp::ComputePipelineDescriptor Descriptor;
-		Descriptor.SetLabel(Func.GetName());
-		Descriptor.SetComputeFunction(Func);
+		NSError* Error = nil;
+		MTLComputePipelineDescriptor* Descriptor = [[MTLComputePipelineDescriptor alloc] init];
+		[Descriptor setLabel:[Func name]];
+		[Descriptor setComputeFunction:Func];
 		if (FAGXCommandQueue::SupportsFeature(EAGXFeaturesTextureBuffers))
 		{
-			Descriptor.SetMaxTotalThreadsPerThreadgroup(NumThreadsX*NumThreadsY*NumThreadsZ);
+			[Descriptor setMaxTotalThreadsPerThreadgroup:(NumThreadsX * NumThreadsY * NumThreadsZ)];
 		}
 
 		if (FAGXCommandQueue::SupportsFeature(EAGXFeaturesPipelineBufferMutability))
 		{
-			ns::AutoReleased<ns::Array<mtlpp::PipelineBufferDescriptor>> PipelineBuffers = Descriptor.GetBuffers();
+			MTLPipelineBufferDescriptorArray* PipelineBuffers = [Descriptor buffers];
 
 			uint32 ImmutableBuffers = Bindings.ConstantBuffers;
-			while(ImmutableBuffers)
+			while (ImmutableBuffers)
 			{
 				uint32 Index = __builtin_ctz(ImmutableBuffers);
 				ImmutableBuffers &= ~(1 << Index);
 
 				if (Index < ML_MaxBuffers)
 				{
-					ns::AutoReleased<mtlpp::PipelineBufferDescriptor> PipelineBuffer = PipelineBuffers[Index];
-					PipelineBuffer.SetMutability(mtlpp::Mutability::Immutable);
+					[[PipelineBuffers objectAtIndexedSubscript:Index] setMutability:MTLMutabilityImmutable];
 				}
 			}
 			if (SideTableBinding > 0)
 			{
-				ns::AutoReleased<mtlpp::PipelineBufferDescriptor> PipelineBuffer = PipelineBuffers[SideTableBinding];
-				PipelineBuffer.SetMutability(mtlpp::Mutability::Immutable);
+				[[PipelineBuffers objectAtIndexedSubscript:SideTableBinding] setMutability:MTLMutabilityImmutable];
 			}
 		}
 
-		mtlpp::ComputePipelineState Kernel;
-		mtlpp::ComputePipelineReflection Reflection;
+		id<MTLComputePipelineState> Kernel = nil;
+#if METAL_DEBUG_OPTIONS
+		MTLAutoreleasedComputePipelineReflection* Reflection = nil;
+#endif
 
 		METAL_GPUPROFILE(FAGXScopedCPUStats CPUStat(FString::Printf(TEXT("NewComputePipeline: %d_%d"), SourceLen, SourceCRC)));
 #if METAL_DEBUG_OPTIONS
 		if (GetAGXDeviceContext().GetCommandQueue().GetRuntimeDebuggingLevel() >= EAGXDebugLevelFastValidation)
 		{
-			ns::AutoReleasedError ComputeError;
-			mtlpp::AutoReleasedComputePipelineReflection ComputeReflection;
-
-			NSUInteger ComputeOption = mtlpp::PipelineOption::ArgumentInfo | mtlpp::PipelineOption::BufferTypeInfo;
-			Kernel = GMtlppDevice.NewComputePipelineState(Descriptor, mtlpp::PipelineOption(ComputeOption), &ComputeReflection, &ComputeError);
-			Error = ComputeError;
-			Reflection = ComputeReflection;
+			Kernel = [GMtlDevice newComputePipelineStateWithDescriptor:Descriptor
+															  options:(MTLPipelineOptionArgumentInfo | MTLPipelineOptionBufferTypeInfo)
+														   reflection:Reflection
+														   		error:&Error];
 		}
 		else
 #endif // METAL_DEBUG_OPTIONS
 		{
-			ns::AutoReleasedError ComputeError;
-			Kernel = GMtlppDevice.NewComputePipelineState(Descriptor, mtlpp::PipelineOption(0), nil, &ComputeError);
-			Error = ComputeError;
+			Kernel = [GMtlDevice newComputePipelineStateWithDescriptor:Descriptor
+															  options:MTLPipelineOptionNone
+														   reflection:nil
+														   		error:&Error];
 		}
 
 		if (Kernel == nil)
@@ -106,14 +104,12 @@ FAGXShaderPipeline* FAGXComputeShader::GetPipeline()
 		Pipeline = [FAGXShaderPipeline new];
 		Pipeline->ComputePipelineState = Kernel;
 #if METAL_DEBUG_OPTIONS
-		Pipeline->ComputePipelineReflection = Reflection;
+		Pipeline->ComputePipelineReflection = [*Reflection retain];
 		Pipeline->ComputeSource = GetSourceCode();
-		if (Reflection)
-		{
-			Pipeline->ComputeDesc = Descriptor;
-		}
 #endif // METAL_DEBUG_OPTIONS
 		METAL_DEBUG_OPTION(FMemory::Memzero(Pipeline->ResourceMask, sizeof(Pipeline->ResourceMask)));
+
+		[Descriptor release];
 	}
 	check(Pipeline);
 
