AGXRHI: Use MTLPrimitiveTopologyClass type

#rnx

diff --git a/Engine/Source/Runtime/Apple/AGXRHI/Private/AGXContext.cpp b/Engine/Source/Runtime/Apple/AGXRHI/Private/AGXContext.cpp
--- a/Engine/Source/Runtime/Apple/AGXRHI/Private/AGXContext.cpp
+++ b/Engine/Source/Runtime/Apple/AGXRHI/Private/AGXContext.cpp
@@ -295,58 +295,6 @@ static id<MTLDevice> GetMTLDevice(uint32& DeviceIndex)
 	return SelectedDevice;
 #endif // PLATFORM_MAC_ARM64
 }
-
-mtlpp::PrimitiveTopologyClass AGXTranslatePrimitiveTopology(uint32 PrimitiveType)
-{
-	switch (PrimitiveType)
-	{
-		case PT_TriangleList:
-		case PT_TriangleStrip:
-			return mtlpp::PrimitiveTopologyClass::Triangle;
-		case PT_LineList:
-			return mtlpp::PrimitiveTopologyClass::Line;
-		case PT_PointList:
-			return mtlpp::PrimitiveTopologyClass::Point;
-		case PT_1_ControlPointPatchList:
-		case PT_2_ControlPointPatchList:
-		case PT_3_ControlPointPatchList:
-		case PT_4_ControlPointPatchList:
-		case PT_5_ControlPointPatchList:
-		case PT_6_ControlPointPatchList:
-		case PT_7_ControlPointPatchList:
-		case PT_8_ControlPointPatchList:
-		case PT_9_ControlPointPatchList:
-		case PT_10_ControlPointPatchList:
-		case PT_11_ControlPointPatchList:
-		case PT_12_ControlPointPatchList:
-		case PT_13_ControlPointPatchList:
-		case PT_14_ControlPointPatchList:
-		case PT_15_ControlPointPatchList:
-		case PT_16_ControlPointPatchList:
-		case PT_17_ControlPointPatchList:
-		case PT_18_ControlPointPatchList:
-		case PT_19_ControlPointPatchList:
-		case PT_20_ControlPointPatchList:
-		case PT_21_ControlPointPatchList:
-		case PT_22_ControlPointPatchList:
-		case PT_23_ControlPointPatchList:
-		case PT_24_ControlPointPatchList:
-		case PT_25_ControlPointPatchList:
-		case PT_26_ControlPointPatchList:
-		case PT_27_ControlPointPatchList:
-		case PT_28_ControlPointPatchList:
-		case PT_29_ControlPointPatchList:
-		case PT_30_ControlPointPatchList:
-		case PT_31_ControlPointPatchList:
-		case PT_32_ControlPointPatchList:
-		{
-			return mtlpp::PrimitiveTopologyClass::Triangle;
-		}
-		default:
-			UE_LOG(LogAGX, Fatal, TEXT("Unsupported primitive topology %d"), (int32)PrimitiveType);
-			return mtlpp::PrimitiveTopologyClass::Triangle;
-	}
-}
 #endif
 
 
diff --git a/Engine/Source/Runtime/Apple/AGXRHI/Private/AGXPipeline.cpp b/Engine/Source/Runtime/Apple/AGXRHI/Private/AGXPipeline.cpp
--- a/Engine/Source/Runtime/Apple/AGXRHI/Private/AGXPipeline.cpp
+++ b/Engine/Source/Runtime/Apple/AGXRHI/Private/AGXPipeline.cpp
@@ -41,6 +41,60 @@ static uint32 BlendBitOffsets[] = { Offset_BlendState0, Offset_BlendState1, Offs
 static uint32 RTBitOffsets[] = { Offset_RenderTargetFormat0, Offset_RenderTargetFormat1, Offset_RenderTargetFormat2, Offset_RenderTargetFormat3, Offset_RenderTargetFormat4, Offset_RenderTargetFormat5, Offset_RenderTargetFormat6, Offset_RenderTargetFormat7 };
 static_assert(Offset_RasterEnd < 64 && Offset_End < 128, "Offset_RasterEnd must be < 64 && Offset_End < 128");
 
+#if PLATFORM_MAC
+static MTLPrimitiveTopologyClass AGXTranslatePrimitiveTopology(uint32 PrimitiveType)
+{
+	switch (PrimitiveType)
+	{
+		case PT_TriangleList:
+		case PT_TriangleStrip:
+			return MTLPrimitiveTopologyClassTriangle;
+		case PT_LineList:
+			return MTLPrimitiveTopologyClassLine;
+		case PT_PointList:
+			return MTLPrimitiveTopologyClassPoint;
+		case PT_1_ControlPointPatchList:
+		case PT_2_ControlPointPatchList:
+		case PT_3_ControlPointPatchList:
+		case PT_4_ControlPointPatchList:
+		case PT_5_ControlPointPatchList:
+		case PT_6_ControlPointPatchList:
+		case PT_7_ControlPointPatchList:
+		case PT_8_ControlPointPatchList:
+		case PT_9_ControlPointPatchList:
+		case PT_10_ControlPointPatchList:
+		case PT_11_ControlPointPatchList:
+		case PT_12_ControlPointPatchList:
+		case PT_13_ControlPointPatchList:
+		case PT_14_ControlPointPatchList:
+		case PT_15_ControlPointPatchList:
+		case PT_16_ControlPointPatchList:
+		case PT_17_ControlPointPatchList:
+		case PT_18_ControlPointPatchList:
+		case PT_19_ControlPointPatchList:
+		case PT_20_ControlPointPatchList:
+		case PT_21_ControlPointPatchList:
+		case PT_22_ControlPointPatchList:
+		case PT_23_ControlPointPatchList:
+		case PT_24_ControlPointPatchList:
+		case PT_25_ControlPointPatchList:
+		case PT_26_ControlPointPatchList:
+		case PT_27_ControlPointPatchList:
+		case PT_28_ControlPointPatchList:
+		case PT_29_ControlPointPatchList:
+		case PT_30_ControlPointPatchList:
+		case PT_31_ControlPointPatchList:
+		case PT_32_ControlPointPatchList:
+		{
+			return MTLPrimitiveTopologyClassTriangle;
+		}
+		default:
+			UE_LOG(LogAGX, Fatal, TEXT("Unsupported primitive topology %d"), (int32)PrimitiveType);
+			return MTLPrimitiveTopologyClassTriangle;
+	}
+}
+#endif // PLATFORM_MAC
+
 struct FAGXGraphicsPipelineKey
 {
 	FAGXRenderPipelineHash RenderPipelineHash;
@@ -596,7 +650,7 @@ static bool ConfigureRenderPipelineDescriptor(MTLRenderPipelineDescriptor* Rende
 	[RenderPipelineDesc setSampleCount:NumSamples];
 	[RenderPipelineDesc setAlphaToCoverageEnabled:(NumSamples > 1 && BlendState->bUseAlphaToCoverage)];
 #if PLATFORM_MAC
-	[RenderPipelineDesc setInputPrimitiveTopology:(MTLPrimitiveTopologyClass)AGXTranslatePrimitiveTopology(Init.PrimitiveType)];
+	[RenderPipelineDesc setInputPrimitiveTopology:AGXTranslatePrimitiveTopology(Init.PrimitiveType)];
 #endif
 	if (FAGXCommandQueue::SupportsFeature(EAGXFeaturesPipelineBufferMutability))
 	{
diff --git a/Engine/Source/Runtime/Apple/AGXRHI/Private/AGXRHIPrivate.h b/Engine/Source/Runtime/Apple/AGXRHI/Private/AGXRHIPrivate.h
--- a/Engine/Source/Runtime/Apple/AGXRHI/Private/AGXRHIPrivate.h
+++ b/Engine/Source/Runtime/Apple/AGXRHI/Private/AGXRHIPrivate.h
@@ -252,10 +252,6 @@ FORCEINLINE MTLLoadAction GetMetalRTLoadAction(ERenderTargetLoadAction LoadActio
 
 MTLPrimitiveType AGXTranslatePrimitiveType(uint32 PrimitiveType);
 
-#if PLATFORM_MAC
-mtlpp::PrimitiveTopologyClass AGXTranslatePrimitiveTopology(uint32 PrimitiveType);
-#endif
-
 mtlpp::PixelFormat AGXToSRGBFormat(mtlpp::PixelFormat LinMTLFormat);
 
 uint8 AGXGetMetalPixelFormatKey(mtlpp::PixelFormat Format);
