AGXRHI: Replace mtlpp::PixelFormat with MTLPixelFormat

#rnx

diff --git a/Engine/Source/Runtime/Apple/AGXRHI/Private/AGXCommandEncoder.cpp b/Engine/Source/Runtime/Apple/AGXRHI/Private/AGXCommandEncoder.cpp
--- a/Engine/Source/Runtime/Apple/AGXRHI/Private/AGXCommandEncoder.cpp
+++ b/Engine/Source/Runtime/Apple/AGXRHI/Private/AGXCommandEncoder.cpp
@@ -1053,9 +1053,9 @@ void FAGXCommandEncoder::SetShaderTexture(mtlpp::FunctionType FunctionType, FAGX
 	{
 		uint8 Swizzle[4] = {0,0,0,0};
 		assert(sizeof(Swizzle) == sizeof(uint32));
-		if (Texture.GetPixelFormat() == mtlpp::PixelFormat::X32_Stencil8
+		if ((MTLPixelFormat)Texture.GetPixelFormat() == MTLPixelFormatX32_Stencil8
 #if PLATFORM_MAC
-		 ||	Texture.GetPixelFormat() == mtlpp::PixelFormat::X24_Stencil8
+		 ||	(MTLPixelFormat)Texture.GetPixelFormat() == MTLPixelFormatX24_Stencil8
 #endif
 		)
 		{
diff --git a/Engine/Source/Runtime/Apple/AGXRHI/Private/AGXLLM.cpp b/Engine/Source/Runtime/Apple/AGXRHI/Private/AGXLLM.cpp
--- a/Engine/Source/Runtime/Apple/AGXRHI/Private/AGXLLM.cpp
+++ b/Engine/Source/Runtime/Apple/AGXRHI/Private/AGXLLM.cpp
@@ -70,56 +70,56 @@ APPLE_PLATFORM_OBJECT_ALLOC_OVERRIDES(FAGXDeallocHandler)
 }
 @end
 
-static mtlpp::PixelFormat FromSRGBFormat(mtlpp::PixelFormat Format)
+static MTLPixelFormat FromSRGBFormat(MTLPixelFormat Format)
 {
-	mtlpp::PixelFormat MTLFormat = Format;
+	MTLPixelFormat MTLFormat = Format;
 	
 	switch (Format)
 	{
-		case mtlpp::PixelFormat::RGBA8Unorm_sRGB:
-			MTLFormat = mtlpp::PixelFormat::RGBA8Unorm;
+		case MTLPixelFormatRGBA8Unorm_sRGB:
+			MTLFormat = MTLPixelFormatRGBA8Unorm;
 			break;
-		case mtlpp::PixelFormat::BGRA8Unorm_sRGB:
-			MTLFormat = mtlpp::PixelFormat::BGRA8Unorm;
+		case MTLPixelFormatBGRA8Unorm_sRGB:
+			MTLFormat = MTLPixelFormatBGRA8Unorm;
 			break;
 #if PLATFORM_MAC
-		case mtlpp::PixelFormat::BC1_RGBA_sRGB:
-			MTLFormat = mtlpp::PixelFormat::BC1_RGBA;
+		case MTLPixelFormatBC1_RGBA_sRGB:
+			MTLFormat = MTLPixelFormatBC1_RGBA;
 			break;
-		case mtlpp::PixelFormat::BC2_RGBA_sRGB:
-			MTLFormat = mtlpp::PixelFormat::BC2_RGBA;
+		case MTLPixelFormatBC2_RGBA_sRGB:
+			MTLFormat = MTLPixelFormatBC2_RGBA;
 			break;
-		case mtlpp::PixelFormat::BC3_RGBA_sRGB:
-			MTLFormat = mtlpp::PixelFormat::BC3_RGBA;
+		case MTLPixelFormatBC3_RGBA_sRGB:
+			MTLFormat = MTLPixelFormatBC3_RGBA;
 			break;
-		case mtlpp::PixelFormat::BC7_RGBAUnorm_sRGB:
-			MTLFormat = mtlpp::PixelFormat::BC7_RGBAUnorm;
+		case MTLPixelFormatBC7_RGBAUnorm_sRGB:
+			MTLFormat = MTLPixelFormatBC7_RGBAUnorm;
 			break;
 #endif //PLATFORM_MAC
 #if PLATFORM_IOS
-		case mtlpp::PixelFormat::R8Unorm_sRGB:
-			MTLFormat = mtlpp::PixelFormat::R8Unorm;
+		case MTLPixelFormatR8Unorm_sRGB:
+			MTLFormat = MTLPixelFormatR8Unorm;
 			break;
-		case mtlpp::PixelFormat::PVRTC_RGBA_2BPP_sRGB:
-			MTLFormat = mtlpp::PixelFormat::PVRTC_RGBA_2BPP;
+		case MTLPixelFormatPVRTC_RGBA_2BPP_sRGB:
+			MTLFormat = MTLPixelFormatPVRTC_RGBA_2BPP;
 			break;
-		case mtlpp::PixelFormat::PVRTC_RGBA_4BPP_sRGB:
-			MTLFormat = mtlpp::PixelFormat::PVRTC_RGBA_4BPP;
+		case MTLPixelFormatPVRTC_RGBA_4BPP_sRGB:
+			MTLFormat = MTLPixelFormatPVRTC_RGBA_4BPP;
 			break;
-		case mtlpp::PixelFormat::ASTC_4x4_sRGB:
-			MTLFormat = mtlpp::PixelFormat::ASTC_4x4_LDR;
+		case MTLPixelFormatASTC_4x4_sRGB:
+			MTLFormat = MTLPixelFormatASTC_4x4_LDR;
 			break;
-		case mtlpp::PixelFormat::ASTC_6x6_sRGB:
-			MTLFormat = mtlpp::PixelFormat::ASTC_6x6_LDR;
+		case MTLPixelFormatASTC_6x6_sRGB:
+			MTLFormat = MTLPixelFormatASTC_6x6_LDR;
 			break;
-		case mtlpp::PixelFormat::ASTC_8x8_sRGB:
-			MTLFormat = mtlpp::PixelFormat::ASTC_8x8_LDR;
+		case MTLPixelFormatASTC_8x8_sRGB:
+			MTLFormat = MTLPixelFormatASTC_8x8_LDR;
 			break;
-		case mtlpp::PixelFormat::ASTC_10x10_sRGB:
-			MTLFormat = mtlpp::PixelFormat::ASTC_10x10_LDR;
+		case MTLPixelFormatASTC_10x10_sRGB:
+			MTLFormat = MTLPixelFormatASTC_10x10_LDR;
 			break;
-		case mtlpp::PixelFormat::ASTC_12x12_sRGB:
-			MTLFormat = mtlpp::PixelFormat::ASTC_12x12_LDR;
+		case MTLPixelFormatASTC_12x12_sRGB:
+			MTLFormat = MTLPixelFormatASTC_12x12_LDR;
 			break;
 #endif //PLATFORM_IOS
 		default:
@@ -129,12 +129,12 @@ static mtlpp::PixelFormat FromSRGBFormat(mtlpp::PixelFormat Format)
 	return MTLFormat;
 }
 
-static EPixelFormat MetalToRHIPixelFormat(mtlpp::PixelFormat Format)
+static EPixelFormat MetalToRHIPixelFormat(MTLPixelFormat Format)
 {
 	Format = FromSRGBFormat(Format);
 	for (uint32 i = 0; i < PF_MAX; i++)
 	{
-		if((mtlpp::PixelFormat)GPixelFormats[i].PlatformFormat == Format)
+		if((MTLPixelFormat)GPixelFormats[i].PlatformFormat == Format)
 		{
 			return (EPixelFormat)i;
 		}
diff --git a/Engine/Source/Runtime/Apple/AGXRHI/Private/AGXPipeline.cpp b/Engine/Source/Runtime/Apple/AGXRHI/Private/AGXPipeline.cpp
--- a/Engine/Source/Runtime/Apple/AGXRHI/Private/AGXPipeline.cpp
+++ b/Engine/Source/Runtime/Apple/AGXRHI/Private/AGXPipeline.cpp
@@ -149,14 +149,14 @@ struct FAGXGraphicsPipelineKey
 			EPixelFormat TargetFormat = (EPixelFormat)Init.RenderTargetFormats[i];
 			if (TargetFormat == PF_Unknown) { continue; }
 
-			mtlpp::PixelFormat MetalFormat = (mtlpp::PixelFormat)GPixelFormats[TargetFormat].PlatformFormat;
+			MTLPixelFormat MetalFormat = (MTLPixelFormat)GPixelFormats[TargetFormat].PlatformFormat;
 			ETextureCreateFlags Flags = Init.RenderTargetFlags[i];
 			if (EnumHasAnyFlags(Flags, TexCreate_SRGB))
 			{
 #if PLATFORM_MAC // Expand as R8_sRGB is iOS only.
-				if (MetalFormat == mtlpp::PixelFormat::R8Unorm)
+				if (MetalFormat == MTLPixelFormatR8Unorm)
 				{
-					MetalFormat = mtlpp::PixelFormat::RGBA8Unorm;
+					MetalFormat = MTLPixelFormatRGBA8Unorm;
 				}
 #endif
 				MetalFormat = AGXToSRGBFormat(MetalFormat);
@@ -175,21 +175,21 @@ struct FAGXGraphicsPipelineKey
 		{
 			case PF_DepthStencil:
 			{
-				mtlpp::PixelFormat MetalFormat = (mtlpp::PixelFormat)GPixelFormats[PF_DepthStencil].PlatformFormat;
+				MTLPixelFormat MetalFormat = (MTLPixelFormat)GPixelFormats[PF_DepthStencil].PlatformFormat;
 				if (Init.DepthTargetLoadAction != ERenderTargetLoadAction::ENoAction || Init.DepthTargetStoreAction != ERenderTargetStoreAction::ENoAction)
 				{
 					DepthFormatKey = AGXGetMetalPixelFormatKey(MetalFormat);
 				}
 				if (Init.StencilTargetLoadAction != ERenderTargetLoadAction::ENoAction || Init.StencilTargetStoreAction != ERenderTargetStoreAction::ENoAction)
 				{
-					StencilFormatKey = AGXGetMetalPixelFormatKey(mtlpp::PixelFormat::Stencil8);
+					StencilFormatKey = AGXGetMetalPixelFormatKey(MTLPixelFormatStencil8);
 				}
 				bHasActiveTargets |= true;
 				break;
 			}
 			case PF_ShadowDepth:
 			{
-				DepthFormatKey = AGXGetMetalPixelFormatKey((mtlpp::PixelFormat)GPixelFormats[PF_ShadowDepth].PlatformFormat);
+				DepthFormatKey = AGXGetMetalPixelFormatKey((MTLPixelFormat)GPixelFormats[PF_ShadowDepth].PlatformFormat);
 				bHasActiveTargets |= true;
 				break;
 			}
@@ -204,7 +204,7 @@ struct FAGXGraphicsPipelineKey
 		FAGXPixelShader* PixelShader = (FAGXPixelShader*)Init.BoundShaderState.PixelShaderRHI;
 		if ( PixelShader && ( ( PixelShader->Bindings.InOutMask.IsFieldEnabled(CrossCompiler::FShaderBindingInOutMask::DepthStencilMaskIndex) && DepthFormatKey == 0 ) || (bHasActiveTargets == false && PixelShader->Bindings.NumUAVs > 0) ) )
 		{
-			mtlpp::PixelFormat MetalFormat = (mtlpp::PixelFormat)GPixelFormats[PF_DepthStencil].PlatformFormat;
+			MTLPixelFormat MetalFormat = (MTLPixelFormat)GPixelFormats[PF_DepthStencil].PlatformFormat;
 			DepthFormatKey = AGXGetMetalPixelFormatKey(MetalFormat);
 		}
 		
@@ -567,7 +567,7 @@ static bool ConfigureRenderPipelineDescriptor(MTLRenderPipelineDescriptor* Rende
 				MetalFormat = MTLPixelFormatRGBA8Unorm;
 			}
 #endif
-			MetalFormat = (MTLPixelFormat)AGXToSRGBFormat((mtlpp::PixelFormat)MetalFormat);
+			MetalFormat = (MTLPixelFormat)AGXToSRGBFormat((MTLPixelFormat)MetalFormat);
 		}
 		
 		MTLRenderPipelineColorAttachmentDescriptor* Attachment = [ColorAttachments objectAtIndexedSubscript:ActiveTargetIndex];
diff --git a/Engine/Source/Runtime/Apple/AGXRHI/Private/AGXRHI.cpp b/Engine/Source/Runtime/Apple/AGXRHI/Private/AGXRHI.cpp
--- a/Engine/Source/Runtime/Apple/AGXRHI/Private/AGXRHI.cpp
+++ b/Engine/Source/Runtime/Apple/AGXRHI/Private/AGXRHI.cpp
@@ -628,153 +628,153 @@ FAGXDynamicRHI::FAGXDynamicRHI(ERHIFeatureLevel::Type RequestedFeatureLevel)
 #if METAL_DEBUG_OPTIONS
 	FMemory::Memset(GAGXBufferFormats, 255);
 #endif
-	GAGXBufferFormats[PF_Unknown              ] = { mtlpp::PixelFormat::Invalid, (uint8)EMetalBufferFormat::Unknown };
-	GAGXBufferFormats[PF_A32B32G32R32F        ] = { mtlpp::PixelFormat::RGBA32Float, (uint8)EMetalBufferFormat::RGBA32Float };
-	GAGXBufferFormats[PF_B8G8R8A8             ] = { mtlpp::PixelFormat::RGBA8Unorm, (uint8)EMetalBufferFormat::RGBA8Unorm }; // mtlpp::PixelFormat::BGRA8Unorm/EMetalBufferFormat::BGRA8Unorm,  < We don't support this as a vertex-format so we have code to swizzle in the shader
-	GAGXBufferFormats[PF_G8                   ] = { mtlpp::PixelFormat::R8Unorm, (uint8)EMetalBufferFormat::R8Unorm };
-	GAGXBufferFormats[PF_G16                  ] = { mtlpp::PixelFormat::R16Unorm, (uint8)EMetalBufferFormat::R16Unorm };
-	GAGXBufferFormats[PF_DXT1                 ] = { mtlpp::PixelFormat::Invalid, (uint8)EMetalBufferFormat::Unknown };
-	GAGXBufferFormats[PF_DXT3                 ] = { mtlpp::PixelFormat::Invalid, (uint8)EMetalBufferFormat::Unknown };
-	GAGXBufferFormats[PF_DXT5                 ] = { mtlpp::PixelFormat::Invalid, (uint8)EMetalBufferFormat::Unknown };
-	GAGXBufferFormats[PF_UYVY                 ] = { mtlpp::PixelFormat::Invalid, (uint8)EMetalBufferFormat::Unknown };
-	GAGXBufferFormats[PF_FloatRGB             ] = { mtlpp::PixelFormat::RG11B10Float, (uint8)EMetalBufferFormat::RG11B10Half };
-	GAGXBufferFormats[PF_FloatRGBA            ] = { mtlpp::PixelFormat::RGBA16Float, (uint8)EMetalBufferFormat::RGBA16Half };
-	GAGXBufferFormats[PF_DepthStencil         ] = { mtlpp::PixelFormat::Invalid, (uint8)EMetalBufferFormat::Unknown };
-	GAGXBufferFormats[PF_ShadowDepth          ] = { mtlpp::PixelFormat::Invalid, (uint8)EMetalBufferFormat::Unknown };
-	GAGXBufferFormats[PF_R32_FLOAT            ] = { mtlpp::PixelFormat::R32Float, (uint8)EMetalBufferFormat::R32Float };
-	GAGXBufferFormats[PF_G16R16               ] = { mtlpp::PixelFormat::RG16Unorm, (uint8)EMetalBufferFormat::RG16Unorm };
-	GAGXBufferFormats[PF_G16R16F              ] = { mtlpp::PixelFormat::RG16Float, (uint8)EMetalBufferFormat::RG16Half };
-	GAGXBufferFormats[PF_G16R16F_FILTER       ] = { mtlpp::PixelFormat::RG16Float, (uint8)EMetalBufferFormat::RG16Half };
-	GAGXBufferFormats[PF_G32R32F              ] = { mtlpp::PixelFormat::RG32Float, (uint8)EMetalBufferFormat::RG32Float };
-	GAGXBufferFormats[PF_A2B10G10R10          ] = { mtlpp::PixelFormat::RGB10A2Unorm, (uint8)EMetalBufferFormat::RGB10A2Unorm };
-	GAGXBufferFormats[PF_A16B16G16R16         ] = { mtlpp::PixelFormat::RGBA16Unorm, (uint8)EMetalBufferFormat::RGBA16Half };
-	GAGXBufferFormats[PF_D24                  ] = { mtlpp::PixelFormat::Invalid, (uint8)EMetalBufferFormat::Unknown };
-	GAGXBufferFormats[PF_R16F                 ] = { mtlpp::PixelFormat::R16Float, (uint8)EMetalBufferFormat::RG16Half };
-	GAGXBufferFormats[PF_R16F_FILTER          ] = { mtlpp::PixelFormat::R16Float, (uint8)EMetalBufferFormat::RG16Half };
-	GAGXBufferFormats[PF_BC5                  ] = { mtlpp::PixelFormat::Invalid, (uint8)EMetalBufferFormat::Unknown };
-	GAGXBufferFormats[PF_V8U8                 ] = { mtlpp::PixelFormat::RG8Snorm, (uint8)EMetalBufferFormat::RG8Unorm };
-	GAGXBufferFormats[PF_A1                   ] = { mtlpp::PixelFormat::Invalid, (uint8)EMetalBufferFormat::Unknown };
-	GAGXBufferFormats[PF_FloatR11G11B10       ] = { mtlpp::PixelFormat::RG11B10Float, (uint8)EMetalBufferFormat::RG11B10Half }; // < May not work on tvOS
-	GAGXBufferFormats[PF_A8                   ] = { mtlpp::PixelFormat::A8Unorm, (uint8)EMetalBufferFormat::R8Unorm };
-	GAGXBufferFormats[PF_R32_UINT             ] = { mtlpp::PixelFormat::R32Uint, (uint8)EMetalBufferFormat::R32Uint };
-	GAGXBufferFormats[PF_R32_SINT             ] = { mtlpp::PixelFormat::R32Sint, (uint8)EMetalBufferFormat::R32Sint };
-	GAGXBufferFormats[PF_PVRTC2               ] = { mtlpp::PixelFormat::Invalid, (uint8)EMetalBufferFormat::Unknown };
-	GAGXBufferFormats[PF_PVRTC4               ] = { mtlpp::PixelFormat::Invalid, (uint8)EMetalBufferFormat::Unknown };
-	GAGXBufferFormats[PF_R16_UINT             ] = { mtlpp::PixelFormat::R16Uint, (uint8)EMetalBufferFormat::R16Uint };
-	GAGXBufferFormats[PF_R16_SINT             ] = { mtlpp::PixelFormat::R16Sint, (uint8)EMetalBufferFormat::R16Sint };
-	GAGXBufferFormats[PF_R16G16B16A16_UINT    ] = { mtlpp::PixelFormat::RGBA16Uint, (uint8)EMetalBufferFormat::RGBA16Uint };
-	GAGXBufferFormats[PF_R16G16B16A16_SINT    ] = { mtlpp::PixelFormat::RGBA16Sint, (uint8)EMetalBufferFormat::RGBA16Sint };
-	GAGXBufferFormats[PF_R5G6B5_UNORM         ] = { mtlpp::PixelFormat::Invalid, (uint8)EMetalBufferFormat::R5G6B5Unorm };
-	GAGXBufferFormats[PF_B5G5R5A1_UNORM       ] = { mtlpp::PixelFormat::Invalid, (uint8)EMetalBufferFormat::B5G5R5A1Unorm };
-	GAGXBufferFormats[PF_R8G8B8A8             ] = { mtlpp::PixelFormat::RGBA8Unorm, (uint8)EMetalBufferFormat::RGBA8Unorm };
-	GAGXBufferFormats[PF_A8R8G8B8             ] = { mtlpp::PixelFormat::RGBA8Unorm, (uint8)EMetalBufferFormat::RGBA8Unorm }; // mtlpp::PixelFormat::BGRA8Unorm/EMetalBufferFormat::BGRA8Unorm,  < We don't support this as a vertex-format so we have code to swizzle in the shader
-	GAGXBufferFormats[PF_BC4                  ] = { mtlpp::PixelFormat::Invalid, (uint8)EMetalBufferFormat::Unknown };
-	GAGXBufferFormats[PF_R8G8                 ] = { mtlpp::PixelFormat::RG8Unorm, (uint8)EMetalBufferFormat::RG8Unorm };
-	GAGXBufferFormats[PF_ATC_RGB              ] = { mtlpp::PixelFormat::Invalid, (uint8)EMetalBufferFormat::Unknown };
-	GAGXBufferFormats[PF_ATC_RGBA_E           ] = { mtlpp::PixelFormat::Invalid, (uint8)EMetalBufferFormat::Unknown };
-	GAGXBufferFormats[PF_ATC_RGBA_I           ] = { mtlpp::PixelFormat::Invalid, (uint8)EMetalBufferFormat::Unknown };
-	GAGXBufferFormats[PF_X24_G8               ] = { mtlpp::PixelFormat::Invalid, (uint8)EMetalBufferFormat::Unknown };
-	GAGXBufferFormats[PF_ETC1                 ] = { mtlpp::PixelFormat::Invalid, (uint8)EMetalBufferFormat::Unknown };
-	GAGXBufferFormats[PF_ETC2_RGB             ] = { mtlpp::PixelFormat::Invalid, (uint8)EMetalBufferFormat::Unknown };
-	GAGXBufferFormats[PF_ETC2_RGBA            ] = { mtlpp::PixelFormat::Invalid, (uint8)EMetalBufferFormat::Unknown };
-	GAGXBufferFormats[PF_R32G32B32A32_UINT    ] = { mtlpp::PixelFormat::RGBA32Uint, (uint8)EMetalBufferFormat::RGBA32Uint };
-	GAGXBufferFormats[PF_R16G16_UINT          ] = { mtlpp::PixelFormat::RG16Uint, (uint8)EMetalBufferFormat::RG16Uint };
-	GAGXBufferFormats[PF_R32G32_UINT          ] = { mtlpp::PixelFormat::RG32Uint, (uint8)EMetalBufferFormat::RG32Uint };
-	GAGXBufferFormats[PF_ASTC_4x4             ] = { mtlpp::PixelFormat::Invalid, (uint8)EMetalBufferFormat::Unknown };
-	GAGXBufferFormats[PF_ASTC_6x6             ] = { mtlpp::PixelFormat::Invalid, (uint8)EMetalBufferFormat::Unknown };
-	GAGXBufferFormats[PF_ASTC_8x8             ] = { mtlpp::PixelFormat::Invalid, (uint8)EMetalBufferFormat::Unknown };
-	GAGXBufferFormats[PF_ASTC_10x10           ] = { mtlpp::PixelFormat::Invalid, (uint8)EMetalBufferFormat::Unknown };
-	GAGXBufferFormats[PF_ASTC_12x12           ] = { mtlpp::PixelFormat::Invalid, (uint8)EMetalBufferFormat::Unknown };
-	GAGXBufferFormats[PF_ASTC_4x4_HDR         ] = { mtlpp::PixelFormat::Invalid, (uint8)EMetalBufferFormat::Unknown };
-	GAGXBufferFormats[PF_ASTC_6x6_HDR         ] = { mtlpp::PixelFormat::Invalid, (uint8)EMetalBufferFormat::Unknown };
-	GAGXBufferFormats[PF_ASTC_8x8_HDR         ] = { mtlpp::PixelFormat::Invalid, (uint8)EMetalBufferFormat::Unknown };
-	GAGXBufferFormats[PF_ASTC_10x10_HDR       ] = { mtlpp::PixelFormat::Invalid, (uint8)EMetalBufferFormat::Unknown };
-	GAGXBufferFormats[PF_ASTC_12x12_HDR       ] = { mtlpp::PixelFormat::Invalid, (uint8)EMetalBufferFormat::Unknown };
-	GAGXBufferFormats[PF_BC6H                 ] = { mtlpp::PixelFormat::Invalid, (uint8)EMetalBufferFormat::Unknown };
-	GAGXBufferFormats[PF_BC7                  ] = { mtlpp::PixelFormat::Invalid, (uint8)EMetalBufferFormat::Unknown };
-	GAGXBufferFormats[PF_R8_UINT              ] = { mtlpp::PixelFormat::R8Uint, (uint8)EMetalBufferFormat::R8Uint };
-	GAGXBufferFormats[PF_R8                   ] = { mtlpp::PixelFormat::R8Unorm, (uint8)EMetalBufferFormat::R8Unorm };
-	GAGXBufferFormats[PF_L8                   ] = { mtlpp::PixelFormat::Invalid, (uint8)EMetalBufferFormat::R8Unorm };
-	GAGXBufferFormats[PF_XGXR8                ] = { mtlpp::PixelFormat::Invalid, (uint8)EMetalBufferFormat::Unknown };
-	GAGXBufferFormats[PF_R8G8B8A8_UINT        ] = { mtlpp::PixelFormat::RGBA8Uint, (uint8)EMetalBufferFormat::RGBA8Uint };
-	GAGXBufferFormats[PF_R8G8B8A8_SNORM       ] = { mtlpp::PixelFormat::RGBA8Snorm, (uint8)EMetalBufferFormat::RGBA8Snorm };
-	GAGXBufferFormats[PF_R16G16B16A16_UNORM   ] = { mtlpp::PixelFormat::RGBA16Unorm, (uint8)EMetalBufferFormat::RGBA16Unorm };
-	GAGXBufferFormats[PF_R16G16B16A16_SNORM   ] = { mtlpp::PixelFormat::RGBA16Snorm, (uint8)EMetalBufferFormat::RGBA16Snorm };
-	GAGXBufferFormats[PF_PLATFORM_HDR_0       ] = { mtlpp::PixelFormat::Invalid, (uint8)EMetalBufferFormat::Unknown };
-	GAGXBufferFormats[PF_PLATFORM_HDR_1       ] = { mtlpp::PixelFormat::Invalid, (uint8)EMetalBufferFormat::Unknown };
-	GAGXBufferFormats[PF_PLATFORM_HDR_2       ] = { mtlpp::PixelFormat::Invalid, (uint8)EMetalBufferFormat::Unknown };
-	GAGXBufferFormats[PF_NV12                 ] = { mtlpp::PixelFormat::Invalid, (uint8)EMetalBufferFormat::Unknown };
+	GAGXBufferFormats[PF_Unknown              ] = { MTLPixelFormatInvalid, (uint8)EMetalBufferFormat::Unknown };
+	GAGXBufferFormats[PF_A32B32G32R32F        ] = { MTLPixelFormatRGBA32Float, (uint8)EMetalBufferFormat::RGBA32Float };
+	GAGXBufferFormats[PF_B8G8R8A8             ] = { MTLPixelFormatRGBA8Unorm, (uint8)EMetalBufferFormat::RGBA8Unorm }; // MTLPixelFormatBGRA8Unorm/EMetalBufferFormat::BGRA8Unorm,  < We don't support this as a vertex-format so we have code to swizzle in the shader
+	GAGXBufferFormats[PF_G8                   ] = { MTLPixelFormatR8Unorm, (uint8)EMetalBufferFormat::R8Unorm };
+	GAGXBufferFormats[PF_G16                  ] = { MTLPixelFormatR16Unorm, (uint8)EMetalBufferFormat::R16Unorm };
+	GAGXBufferFormats[PF_DXT1                 ] = { MTLPixelFormatInvalid, (uint8)EMetalBufferFormat::Unknown };
+	GAGXBufferFormats[PF_DXT3                 ] = { MTLPixelFormatInvalid, (uint8)EMetalBufferFormat::Unknown };
+	GAGXBufferFormats[PF_DXT5                 ] = { MTLPixelFormatInvalid, (uint8)EMetalBufferFormat::Unknown };
+	GAGXBufferFormats[PF_UYVY                 ] = { MTLPixelFormatInvalid, (uint8)EMetalBufferFormat::Unknown };
+	GAGXBufferFormats[PF_FloatRGB             ] = { MTLPixelFormatRG11B10Float, (uint8)EMetalBufferFormat::RG11B10Half };
+	GAGXBufferFormats[PF_FloatRGBA            ] = { MTLPixelFormatRGBA16Float, (uint8)EMetalBufferFormat::RGBA16Half };
+	GAGXBufferFormats[PF_DepthStencil         ] = { MTLPixelFormatInvalid, (uint8)EMetalBufferFormat::Unknown };
+	GAGXBufferFormats[PF_ShadowDepth          ] = { MTLPixelFormatInvalid, (uint8)EMetalBufferFormat::Unknown };
+	GAGXBufferFormats[PF_R32_FLOAT            ] = { MTLPixelFormatR32Float, (uint8)EMetalBufferFormat::R32Float };
+	GAGXBufferFormats[PF_G16R16               ] = { MTLPixelFormatRG16Unorm, (uint8)EMetalBufferFormat::RG16Unorm };
+	GAGXBufferFormats[PF_G16R16F              ] = { MTLPixelFormatRG16Float, (uint8)EMetalBufferFormat::RG16Half };
+	GAGXBufferFormats[PF_G16R16F_FILTER       ] = { MTLPixelFormatRG16Float, (uint8)EMetalBufferFormat::RG16Half };
+	GAGXBufferFormats[PF_G32R32F              ] = { MTLPixelFormatRG32Float, (uint8)EMetalBufferFormat::RG32Float };
+	GAGXBufferFormats[PF_A2B10G10R10          ] = { MTLPixelFormatRGB10A2Unorm, (uint8)EMetalBufferFormat::RGB10A2Unorm };
+	GAGXBufferFormats[PF_A16B16G16R16         ] = { MTLPixelFormatRGBA16Unorm, (uint8)EMetalBufferFormat::RGBA16Half };
+	GAGXBufferFormats[PF_D24                  ] = { MTLPixelFormatInvalid, (uint8)EMetalBufferFormat::Unknown };
+	GAGXBufferFormats[PF_R16F                 ] = { MTLPixelFormatR16Float, (uint8)EMetalBufferFormat::RG16Half };
+	GAGXBufferFormats[PF_R16F_FILTER          ] = { MTLPixelFormatR16Float, (uint8)EMetalBufferFormat::RG16Half };
+	GAGXBufferFormats[PF_BC5                  ] = { MTLPixelFormatInvalid, (uint8)EMetalBufferFormat::Unknown };
+	GAGXBufferFormats[PF_V8U8                 ] = { MTLPixelFormatRG8Snorm, (uint8)EMetalBufferFormat::RG8Unorm };
+	GAGXBufferFormats[PF_A1                   ] = { MTLPixelFormatInvalid, (uint8)EMetalBufferFormat::Unknown };
+	GAGXBufferFormats[PF_FloatR11G11B10       ] = { MTLPixelFormatRG11B10Float, (uint8)EMetalBufferFormat::RG11B10Half }; // < May not work on tvOS
+	GAGXBufferFormats[PF_A8                   ] = { MTLPixelFormatA8Unorm, (uint8)EMetalBufferFormat::R8Unorm };
+	GAGXBufferFormats[PF_R32_UINT             ] = { MTLPixelFormatR32Uint, (uint8)EMetalBufferFormat::R32Uint };
+	GAGXBufferFormats[PF_R32_SINT             ] = { MTLPixelFormatR32Sint, (uint8)EMetalBufferFormat::R32Sint };
+	GAGXBufferFormats[PF_PVRTC2               ] = { MTLPixelFormatInvalid, (uint8)EMetalBufferFormat::Unknown };
+	GAGXBufferFormats[PF_PVRTC4               ] = { MTLPixelFormatInvalid, (uint8)EMetalBufferFormat::Unknown };
+	GAGXBufferFormats[PF_R16_UINT             ] = { MTLPixelFormatR16Uint, (uint8)EMetalBufferFormat::R16Uint };
+	GAGXBufferFormats[PF_R16_SINT             ] = { MTLPixelFormatR16Sint, (uint8)EMetalBufferFormat::R16Sint };
+	GAGXBufferFormats[PF_R16G16B16A16_UINT    ] = { MTLPixelFormatRGBA16Uint, (uint8)EMetalBufferFormat::RGBA16Uint };
+	GAGXBufferFormats[PF_R16G16B16A16_SINT    ] = { MTLPixelFormatRGBA16Sint, (uint8)EMetalBufferFormat::RGBA16Sint };
+	GAGXBufferFormats[PF_R5G6B5_UNORM         ] = { MTLPixelFormatInvalid, (uint8)EMetalBufferFormat::R5G6B5Unorm };
+	GAGXBufferFormats[PF_B5G5R5A1_UNORM       ] = { MTLPixelFormatInvalid, (uint8)EMetalBufferFormat::B5G5R5A1Unorm };
+	GAGXBufferFormats[PF_R8G8B8A8             ] = { MTLPixelFormatRGBA8Unorm, (uint8)EMetalBufferFormat::RGBA8Unorm };
+	GAGXBufferFormats[PF_A8R8G8B8             ] = { MTLPixelFormatRGBA8Unorm, (uint8)EMetalBufferFormat::RGBA8Unorm }; // MTLPixelFormatBGRA8Unorm/EMetalBufferFormat::BGRA8Unorm,  < We don't support this as a vertex-format so we have code to swizzle in the shader
+	GAGXBufferFormats[PF_BC4                  ] = { MTLPixelFormatInvalid, (uint8)EMetalBufferFormat::Unknown };
+	GAGXBufferFormats[PF_R8G8                 ] = { MTLPixelFormatRG8Unorm, (uint8)EMetalBufferFormat::RG8Unorm };
+	GAGXBufferFormats[PF_ATC_RGB              ] = { MTLPixelFormatInvalid, (uint8)EMetalBufferFormat::Unknown };
+	GAGXBufferFormats[PF_ATC_RGBA_E           ] = { MTLPixelFormatInvalid, (uint8)EMetalBufferFormat::Unknown };
+	GAGXBufferFormats[PF_ATC_RGBA_I           ] = { MTLPixelFormatInvalid, (uint8)EMetalBufferFormat::Unknown };
+	GAGXBufferFormats[PF_X24_G8               ] = { MTLPixelFormatInvalid, (uint8)EMetalBufferFormat::Unknown };
+	GAGXBufferFormats[PF_ETC1                 ] = { MTLPixelFormatInvalid, (uint8)EMetalBufferFormat::Unknown };
+	GAGXBufferFormats[PF_ETC2_RGB             ] = { MTLPixelFormatInvalid, (uint8)EMetalBufferFormat::Unknown };
+	GAGXBufferFormats[PF_ETC2_RGBA            ] = { MTLPixelFormatInvalid, (uint8)EMetalBufferFormat::Unknown };
+	GAGXBufferFormats[PF_R32G32B32A32_UINT    ] = { MTLPixelFormatRGBA32Uint, (uint8)EMetalBufferFormat::RGBA32Uint };
+	GAGXBufferFormats[PF_R16G16_UINT          ] = { MTLPixelFormatRG16Uint, (uint8)EMetalBufferFormat::RG16Uint };
+	GAGXBufferFormats[PF_R32G32_UINT          ] = { MTLPixelFormatRG32Uint, (uint8)EMetalBufferFormat::RG32Uint };
+	GAGXBufferFormats[PF_ASTC_4x4             ] = { MTLPixelFormatInvalid, (uint8)EMetalBufferFormat::Unknown };
+	GAGXBufferFormats[PF_ASTC_6x6             ] = { MTLPixelFormatInvalid, (uint8)EMetalBufferFormat::Unknown };
+	GAGXBufferFormats[PF_ASTC_8x8             ] = { MTLPixelFormatInvalid, (uint8)EMetalBufferFormat::Unknown };
+	GAGXBufferFormats[PF_ASTC_10x10           ] = { MTLPixelFormatInvalid, (uint8)EMetalBufferFormat::Unknown };
+	GAGXBufferFormats[PF_ASTC_12x12           ] = { MTLPixelFormatInvalid, (uint8)EMetalBufferFormat::Unknown };
+	GAGXBufferFormats[PF_ASTC_4x4_HDR         ] = { MTLPixelFormatInvalid, (uint8)EMetalBufferFormat::Unknown };
+	GAGXBufferFormats[PF_ASTC_6x6_HDR         ] = { MTLPixelFormatInvalid, (uint8)EMetalBufferFormat::Unknown };
+	GAGXBufferFormats[PF_ASTC_8x8_HDR         ] = { MTLPixelFormatInvalid, (uint8)EMetalBufferFormat::Unknown };
+	GAGXBufferFormats[PF_ASTC_10x10_HDR       ] = { MTLPixelFormatInvalid, (uint8)EMetalBufferFormat::Unknown };
+	GAGXBufferFormats[PF_ASTC_12x12_HDR       ] = { MTLPixelFormatInvalid, (uint8)EMetalBufferFormat::Unknown };
+	GAGXBufferFormats[PF_BC6H                 ] = { MTLPixelFormatInvalid, (uint8)EMetalBufferFormat::Unknown };
+	GAGXBufferFormats[PF_BC7                  ] = { MTLPixelFormatInvalid, (uint8)EMetalBufferFormat::Unknown };
+	GAGXBufferFormats[PF_R8_UINT              ] = { MTLPixelFormatR8Uint, (uint8)EMetalBufferFormat::R8Uint };
+	GAGXBufferFormats[PF_R8                   ] = { MTLPixelFormatR8Unorm, (uint8)EMetalBufferFormat::R8Unorm };
+	GAGXBufferFormats[PF_L8                   ] = { MTLPixelFormatInvalid, (uint8)EMetalBufferFormat::R8Unorm };
+	GAGXBufferFormats[PF_XGXR8                ] = { MTLPixelFormatInvalid, (uint8)EMetalBufferFormat::Unknown };
+	GAGXBufferFormats[PF_R8G8B8A8_UINT        ] = { MTLPixelFormatRGBA8Uint, (uint8)EMetalBufferFormat::RGBA8Uint };
+	GAGXBufferFormats[PF_R8G8B8A8_SNORM       ] = { MTLPixelFormatRGBA8Snorm, (uint8)EMetalBufferFormat::RGBA8Snorm };
+	GAGXBufferFormats[PF_R16G16B16A16_UNORM   ] = { MTLPixelFormatRGBA16Unorm, (uint8)EMetalBufferFormat::RGBA16Unorm };
+	GAGXBufferFormats[PF_R16G16B16A16_SNORM   ] = { MTLPixelFormatRGBA16Snorm, (uint8)EMetalBufferFormat::RGBA16Snorm };
+	GAGXBufferFormats[PF_PLATFORM_HDR_0       ] = { MTLPixelFormatInvalid, (uint8)EMetalBufferFormat::Unknown };
+	GAGXBufferFormats[PF_PLATFORM_HDR_1       ] = { MTLPixelFormatInvalid, (uint8)EMetalBufferFormat::Unknown };
+	GAGXBufferFormats[PF_PLATFORM_HDR_2       ] = { MTLPixelFormatInvalid, (uint8)EMetalBufferFormat::Unknown };
+	GAGXBufferFormats[PF_NV12                 ] = { MTLPixelFormatInvalid, (uint8)EMetalBufferFormat::Unknown };
 	
-	GAGXBufferFormats[PF_ETC2_R11_EAC         ] = { mtlpp::PixelFormat::Invalid, (uint8)EMetalBufferFormat::Unknown };
-	GAGXBufferFormats[PF_ETC2_RG11_EAC        ] = { mtlpp::PixelFormat::Invalid, (uint8)EMetalBufferFormat::Unknown };
+	GAGXBufferFormats[PF_ETC2_R11_EAC         ] = { MTLPixelFormatInvalid, (uint8)EMetalBufferFormat::Unknown };
+	GAGXBufferFormats[PF_ETC2_RG11_EAC        ] = { MTLPixelFormatInvalid, (uint8)EMetalBufferFormat::Unknown };
 
-	GAGXBufferFormats[PF_G16R16_SNORM         ] = { mtlpp::PixelFormat::RG16Snorm, (uint8)EMetalBufferFormat::RG16Snorm };
-	GAGXBufferFormats[PF_R8G8_UINT            ] = { mtlpp::PixelFormat::RG8Uint, (uint8)EMetalBufferFormat::RG8Uint };
-	GAGXBufferFormats[PF_R32G32B32_UINT       ] = { mtlpp::PixelFormat::Invalid, (uint8)EMetalBufferFormat::Unknown };
-	GAGXBufferFormats[PF_R32G32B32_SINT       ] = { mtlpp::PixelFormat::Invalid, (uint8)EMetalBufferFormat::Unknown };
-	GAGXBufferFormats[PF_R32G32B32F           ] = { mtlpp::PixelFormat::Invalid, (uint8)EMetalBufferFormat::Unknown };
-	GAGXBufferFormats[PF_R8_SINT              ] = { mtlpp::PixelFormat::R8Sint, (uint8)EMetalBufferFormat::R8Sint };
-	GAGXBufferFormats[PF_R64_UINT             ] = { mtlpp::PixelFormat::Invalid, (uint8)EMetalBufferFormat::Unknown };
+	GAGXBufferFormats[PF_G16R16_SNORM         ] = { MTLPixelFormatRG16Snorm, (uint8)EMetalBufferFormat::RG16Snorm };
+	GAGXBufferFormats[PF_R8G8_UINT            ] = { MTLPixelFormatRG8Uint, (uint8)EMetalBufferFormat::RG8Uint };
+	GAGXBufferFormats[PF_R32G32B32_UINT       ] = { MTLPixelFormatInvalid, (uint8)EMetalBufferFormat::Unknown };
+	GAGXBufferFormats[PF_R32G32B32_SINT       ] = { MTLPixelFormatInvalid, (uint8)EMetalBufferFormat::Unknown };
+	GAGXBufferFormats[PF_R32G32B32F           ] = { MTLPixelFormatInvalid, (uint8)EMetalBufferFormat::Unknown };
+	GAGXBufferFormats[PF_R8_SINT              ] = { MTLPixelFormatR8Sint, (uint8)EMetalBufferFormat::R8Sint };
+	GAGXBufferFormats[PF_R64_UINT             ] = { MTLPixelFormatInvalid, (uint8)EMetalBufferFormat::Unknown };
 
 	// Initialize the platform pixel format map.
-	GPixelFormats[PF_Unknown			].PlatformFormat	= (uint32)mtlpp::PixelFormat::Invalid;
-	GPixelFormats[PF_A32B32G32R32F		].PlatformFormat	= (uint32)mtlpp::PixelFormat::RGBA32Float;
-	GPixelFormats[PF_B8G8R8A8			].PlatformFormat	= (uint32)mtlpp::PixelFormat::BGRA8Unorm;
-	GPixelFormats[PF_G8					].PlatformFormat	= (uint32)mtlpp::PixelFormat::R8Unorm;
-	GPixelFormats[PF_G16				].PlatformFormat	= (uint32)mtlpp::PixelFormat::R16Unorm;
-	GPixelFormats[PF_R32G32B32A32_UINT	].PlatformFormat	= (uint32)mtlpp::PixelFormat::RGBA32Uint;
-	GPixelFormats[PF_R16G16_UINT		].PlatformFormat	= (uint32)mtlpp::PixelFormat::RG16Uint;
-	GPixelFormats[PF_R32G32_UINT		].PlatformFormat	= (uint32)mtlpp::PixelFormat::RG32Uint;
+	GPixelFormats[PF_Unknown			].PlatformFormat	= (uint32)MTLPixelFormatInvalid;
+	GPixelFormats[PF_A32B32G32R32F		].PlatformFormat	= (uint32)MTLPixelFormatRGBA32Float;
+	GPixelFormats[PF_B8G8R8A8			].PlatformFormat	= (uint32)MTLPixelFormatBGRA8Unorm;
+	GPixelFormats[PF_G8					].PlatformFormat	= (uint32)MTLPixelFormatR8Unorm;
+	GPixelFormats[PF_G16				].PlatformFormat	= (uint32)MTLPixelFormatR16Unorm;
+	GPixelFormats[PF_R32G32B32A32_UINT	].PlatformFormat	= (uint32)MTLPixelFormatRGBA32Uint;
+	GPixelFormats[PF_R16G16_UINT		].PlatformFormat	= (uint32)MTLPixelFormatRG16Uint;
+	GPixelFormats[PF_R32G32_UINT		].PlatformFormat	= (uint32)MTLPixelFormatRG32Uint;
 		
 #if PLATFORM_IOS
-    GPixelFormats[PF_DXT1				].PlatformFormat	= (uint32)mtlpp::PixelFormat::Invalid;
+    GPixelFormats[PF_DXT1				].PlatformFormat	= (uint32)MTLPixelFormatInvalid;
 	GPixelFormats[PF_DXT1				].Supported			= false;
-    GPixelFormats[PF_DXT3				].PlatformFormat	= (uint32)mtlpp::PixelFormat::Invalid;
+    GPixelFormats[PF_DXT3				].PlatformFormat	= (uint32)MTLPixelFormatInvalid;
 	GPixelFormats[PF_DXT3				].Supported			= false;
-    GPixelFormats[PF_DXT5				].PlatformFormat	= (uint32)mtlpp::PixelFormat::Invalid;
+    GPixelFormats[PF_DXT5				].PlatformFormat	= (uint32)MTLPixelFormatInvalid;
 	GPixelFormats[PF_DXT5				].Supported			= false;
-	GPixelFormats[PF_BC5				].PlatformFormat	= (uint32)mtlpp::PixelFormat::Invalid;
+	GPixelFormats[PF_BC5				].PlatformFormat	= (uint32)MTLPixelFormatInvalid;
 	GPixelFormats[PF_BC5				].Supported			= false;
-	GPixelFormats[PF_PVRTC2				].PlatformFormat	= (uint32)mtlpp::PixelFormat::PVRTC_RGBA_2BPP;
+	GPixelFormats[PF_PVRTC2				].PlatformFormat	= (uint32)MTLPixelFormatPVRTC_RGBA_2BPP;
 	GPixelFormats[PF_PVRTC2				].Supported			= true;
-	GPixelFormats[PF_PVRTC4				].PlatformFormat	= (uint32)mtlpp::PixelFormat::PVRTC_RGBA_4BPP;
+	GPixelFormats[PF_PVRTC4				].PlatformFormat	= (uint32)MTLPixelFormatPVRTC_RGBA_4BPP;
 	GPixelFormats[PF_PVRTC4				].Supported			= true;
-	GPixelFormats[PF_PVRTC4				].PlatformFormat	= (uint32)mtlpp::PixelFormat::PVRTC_RGBA_4BPP;
+	GPixelFormats[PF_PVRTC4				].PlatformFormat	= (uint32)MTLPixelFormatPVRTC_RGBA_4BPP;
 	GPixelFormats[PF_PVRTC4				].Supported			= true;
-	GPixelFormats[PF_ASTC_4x4			].PlatformFormat	= (uint32)mtlpp::PixelFormat::ASTC_4x4_LDR;
+	GPixelFormats[PF_ASTC_4x4			].PlatformFormat	= (uint32)MTLPixelFormatASTC_4x4_LDR;
 	GPixelFormats[PF_ASTC_4x4			].Supported			= true;
-	GPixelFormats[PF_ASTC_6x6			].PlatformFormat	= (uint32)mtlpp::PixelFormat::ASTC_6x6_LDR;
+	GPixelFormats[PF_ASTC_6x6			].PlatformFormat	= (uint32)MTLPixelFormatASTC_6x6_LDR;
 	GPixelFormats[PF_ASTC_6x6			].Supported			= true;
-	GPixelFormats[PF_ASTC_8x8			].PlatformFormat	= (uint32)mtlpp::PixelFormat::ASTC_8x8_LDR;
+	GPixelFormats[PF_ASTC_8x8			].PlatformFormat	= (uint32)MTLPixelFormatASTC_8x8_LDR;
 	GPixelFormats[PF_ASTC_8x8			].Supported			= true;
-	GPixelFormats[PF_ASTC_10x10			].PlatformFormat	= (uint32)mtlpp::PixelFormat::ASTC_10x10_LDR;
+	GPixelFormats[PF_ASTC_10x10			].PlatformFormat	= (uint32)MTLPixelFormatASTC_10x10_LDR;
 	GPixelFormats[PF_ASTC_10x10			].Supported			= true;
-	GPixelFormats[PF_ASTC_12x12			].PlatformFormat	= (uint32)mtlpp::PixelFormat::ASTC_12x12_LDR;
+	GPixelFormats[PF_ASTC_12x12			].PlatformFormat	= (uint32)MTLPixelFormatASTC_12x12_LDR;
 	GPixelFormats[PF_ASTC_12x12			].Supported			= true;
 
 #if !PLATFORM_TVOS
 	if ([GMtlDevice supportsFamily:MTLGPUFamilyApple6])
 	{
-		GPixelFormats[PF_ASTC_4x4_HDR].PlatformFormat		= (uint32)mtlpp::PixelFormat::ASTC_4x4_HDR;
+		GPixelFormats[PF_ASTC_4x4_HDR].PlatformFormat		= (uint32)MTLPixelFormatASTC_4x4_HDR;
 		GPixelFormats[PF_ASTC_4x4_HDR].Supported			= true;
-		GPixelFormats[PF_ASTC_6x6_HDR].PlatformFormat		= (uint32)mtlpp::PixelFormat::ASTC_6x6_HDR;
+		GPixelFormats[PF_ASTC_6x6_HDR].PlatformFormat		= (uint32)MTLPixelFormatASTC_6x6_HDR;
 		GPixelFormats[PF_ASTC_6x6_HDR].Supported			= true;
-		GPixelFormats[PF_ASTC_8x8_HDR].PlatformFormat		= (uint32)mtlpp::PixelFormat::ASTC_8x8_HDR;
+		GPixelFormats[PF_ASTC_8x8_HDR].PlatformFormat		= (uint32)MTLPixelFormatASTC_8x8_HDR;
 		GPixelFormats[PF_ASTC_8x8_HDR].Supported			= true;
-		GPixelFormats[PF_ASTC_10x10_HDR].PlatformFormat		= (uint32)mtlpp::PixelFormat::ASTC_10x10_HDR;
+		GPixelFormats[PF_ASTC_10x10_HDR].PlatformFormat		= (uint32)MTLPixelFormatASTC_10x10_HDR;
 		GPixelFormats[PF_ASTC_10x10_HDR].Supported			= true;
-		GPixelFormats[PF_ASTC_12x12_HDR].PlatformFormat		= (uint32)mtlpp::PixelFormat::ASTC_12x12_HDR;
+		GPixelFormats[PF_ASTC_12x12_HDR].PlatformFormat		= (uint32)MTLPixelFormatASTC_12x12_HDR;
 		GPixelFormats[PF_ASTC_12x12_HDR].Supported			= true;
 	}
 #endif
 	// used with virtual textures
-	GPixelFormats[PF_ETC2_RGB	  		].PlatformFormat	= (uint32)mtlpp::PixelFormat::ETC2_RGB8;
+	GPixelFormats[PF_ETC2_RGB	  		].PlatformFormat	= (uint32)MTLPixelFormatETC2_RGB8;
 	GPixelFormats[PF_ETC2_RGB			].Supported			= true;
-	GPixelFormats[PF_ETC2_RGBA	  		].PlatformFormat	= (uint32)mtlpp::PixelFormat::EAC_RGBA8;
+	GPixelFormats[PF_ETC2_RGBA	  		].PlatformFormat	= (uint32)MTLPixelFormatEAC_RGBA8;
 	GPixelFormats[PF_ETC2_RGBA			].Supported			= true;
-	GPixelFormats[PF_ETC2_R11_EAC	  	].PlatformFormat	= (uint32)mtlpp::PixelFormat::EAC_R11Unorm;
+	GPixelFormats[PF_ETC2_R11_EAC	  	].PlatformFormat	= (uint32)MTLPixelFormatEAC_R11Unorm;
 	GPixelFormats[PF_ETC2_R11_EAC		].Supported			= true;
-	GPixelFormats[PF_ETC2_RG11_EAC		].PlatformFormat	= (uint32)mtlpp::PixelFormat::EAC_RG11Unorm;
+	GPixelFormats[PF_ETC2_RG11_EAC		].PlatformFormat	= (uint32)MTLPixelFormatEAC_RG11Unorm;
 	GPixelFormats[PF_ETC2_RG11_EAC		].Supported			= true;
 
 	// IOS HDR format is BGR10_XR (32bits, 3 components)
@@ -783,7 +783,7 @@ FAGXDynamicRHI::FAGXDynamicRHI(ERHIFeatureLevel::Type RequestedFeatureLevel)
 	GPixelFormats[PF_PLATFORM_HDR_0		].BlockSizeZ		= 1;
 	GPixelFormats[PF_PLATFORM_HDR_0		].BlockBytes		= 4;
 	GPixelFormats[PF_PLATFORM_HDR_0		].NumComponents		= 3;
-	GPixelFormats[PF_PLATFORM_HDR_0		].PlatformFormat	= (uint32)mtlpp::PixelFormat::BGR10_XR_sRGB;
+	GPixelFormats[PF_PLATFORM_HDR_0		].PlatformFormat	= (uint32)MTLPixelFormatBGR10_XR_sRGB;
 	GPixelFormats[PF_PLATFORM_HDR_0		].Supported			= GRHISupportsHDROutput;
 		
 #if PLATFORM_TVOS
@@ -792,43 +792,43 @@ FAGXDynamicRHI::FAGXDynamicRHI(ERHIFeatureLevel::Type RequestedFeatureLevel)
 	if (![GMtlDevice supportsFeatureSet:MTLFeatureSet_iOS_GPUFamily3_v2])
 #endif
 	{
-		GPixelFormats[PF_FloatRGB			].PlatformFormat 	= (uint32)mtlpp::PixelFormat::RGBA16Float;
+		GPixelFormats[PF_FloatRGB			].PlatformFormat 	= (uint32)MTLPixelFormatRGBA16Float;
 		GPixelFormats[PF_FloatRGBA			].BlockBytes		= 8;
-		GPixelFormats[PF_FloatR11G11B10		].PlatformFormat	= (uint32)mtlpp::PixelFormat::RGBA16Float;
+		GPixelFormats[PF_FloatR11G11B10		].PlatformFormat	= (uint32)MTLPixelFormatRGBA16Float;
 		GPixelFormats[PF_FloatR11G11B10		].BlockBytes		= 8;
 		GPixelFormats[PF_FloatR11G11B10		].Supported			= true;
 	}
 	else
 	{
-		GPixelFormats[PF_FloatRGB			].PlatformFormat	= (uint32)mtlpp::PixelFormat::RG11B10Float;
+		GPixelFormats[PF_FloatRGB			].PlatformFormat	= (uint32)MTLPixelFormatRG11B10Float;
 		GPixelFormats[PF_FloatRGB			].BlockBytes		= 4;
-		GPixelFormats[PF_FloatR11G11B10		].PlatformFormat	= (uint32)mtlpp::PixelFormat::RG11B10Float;
+		GPixelFormats[PF_FloatR11G11B10		].PlatformFormat	= (uint32)MTLPixelFormatRG11B10Float;
 		GPixelFormats[PF_FloatR11G11B10		].BlockBytes		= 4;
 		GPixelFormats[PF_FloatR11G11B10		].Supported			= true;
 	}
 	
-		GPixelFormats[PF_DepthStencil		].PlatformFormat	= (uint32)mtlpp::PixelFormat::Depth32Float_Stencil8;
+		GPixelFormats[PF_DepthStencil		].PlatformFormat	= (uint32)MTLPixelFormatDepth32Float_Stencil8;
 		GPixelFormats[PF_DepthStencil		].BlockBytes		= 4;
 
 	GPixelFormats[PF_DepthStencil		].Supported			= true;
-	GPixelFormats[PF_ShadowDepth		].PlatformFormat	= (uint32)mtlpp::PixelFormat::Depth32Float;
+	GPixelFormats[PF_ShadowDepth		].PlatformFormat	= (uint32)MTLPixelFormatDepth32Float;
 	GPixelFormats[PF_ShadowDepth		].BlockBytes		= 4;
 	GPixelFormats[PF_ShadowDepth		].Supported			= true;
 		
-	GPixelFormats[PF_BC5				].PlatformFormat	= (uint32)mtlpp::PixelFormat::Invalid;
-	GPixelFormats[PF_R5G6B5_UNORM		].PlatformFormat	= (uint32)mtlpp::PixelFormat::B5G6R5Unorm;
+	GPixelFormats[PF_BC5				].PlatformFormat	= (uint32)MTLPixelFormatInvalid;
+	GPixelFormats[PF_R5G6B5_UNORM		].PlatformFormat	= (uint32)MTLPixelFormatB5G6R5Unorm;
 	GPixelFormats[PF_R5G6B5_UNORM       ].Supported         = true;
-	GPixelFormats[PF_B5G5R5A1_UNORM     ].PlatformFormat    = (uint32)mtlpp::PixelFormat::BGR5A1Unorm;
+	GPixelFormats[PF_B5G5R5A1_UNORM     ].PlatformFormat    = (uint32)MTLPixelFormatBGR5A1Unorm;
 	GPixelFormats[PF_B5G5R5A1_UNORM     ].Supported         = true;
 #else
-    GPixelFormats[PF_DXT1				].PlatformFormat	= (uint32)mtlpp::PixelFormat::BC1_RGBA;
-    GPixelFormats[PF_DXT3				].PlatformFormat	= (uint32)mtlpp::PixelFormat::BC2_RGBA;
-    GPixelFormats[PF_DXT5				].PlatformFormat	= (uint32)mtlpp::PixelFormat::BC3_RGBA;
+    GPixelFormats[PF_DXT1				].PlatformFormat	= (uint32)MTLPixelFormatBC1_RGBA;
+    GPixelFormats[PF_DXT3				].PlatformFormat	= (uint32)MTLPixelFormatBC2_RGBA;
+    GPixelFormats[PF_DXT5				].PlatformFormat	= (uint32)MTLPixelFormatBC3_RGBA;
 	
-	GPixelFormats[PF_FloatRGB			].PlatformFormat	= (uint32)mtlpp::PixelFormat::RG11B10Float;
+	GPixelFormats[PF_FloatRGB			].PlatformFormat	= (uint32)MTLPixelFormatRG11B10Float;
 	GPixelFormats[PF_FloatRGB			].BlockBytes		= 4;
 	
-	GPixelFormats[PF_FloatR11G11B10		].PlatformFormat	= (uint32)mtlpp::PixelFormat::RG11B10Float;
+	GPixelFormats[PF_FloatR11G11B10		].PlatformFormat	= (uint32)MTLPixelFormatRG11B10Float;
 	GPixelFormats[PF_FloatR11G11B10		].BlockBytes		= 4;
 	GPixelFormats[PF_FloatR11G11B10		].Supported			= true;
 	
@@ -838,101 +838,101 @@ FAGXDynamicRHI::FAGXDynamicRHI(ERHIFeatureLevel::Type RequestedFeatureLevel)
 	GPixelFormats[PF_PLATFORM_HDR_0		].BlockSizeZ		= 1;
 	GPixelFormats[PF_PLATFORM_HDR_0		].BlockBytes		= 8;
 	GPixelFormats[PF_PLATFORM_HDR_0		].NumComponents		= 4;
-	GPixelFormats[PF_PLATFORM_HDR_0		].PlatformFormat	= (uint32)mtlpp::PixelFormat::RGBA16Float;
+	GPixelFormats[PF_PLATFORM_HDR_0		].PlatformFormat	= (uint32)MTLPixelFormatRGBA16Float;
 	GPixelFormats[PF_PLATFORM_HDR_0		].Supported			= GRHISupportsHDROutput;
 		
 	// Use Depth28_Stencil8 when it is available for consistency
 	if(bSupportsD24S8)
 	{
-		GPixelFormats[PF_DepthStencil	].PlatformFormat	= (uint32)mtlpp::PixelFormat::Depth24Unorm_Stencil8;
+		GPixelFormats[PF_DepthStencil	].PlatformFormat	= (uint32)MTLPixelFormatDepth24Unorm_Stencil8;
 	}
 	else
 	{
-		GPixelFormats[PF_DepthStencil	].PlatformFormat	= (uint32)mtlpp::PixelFormat::Depth32Float_Stencil8;
+		GPixelFormats[PF_DepthStencil	].PlatformFormat	= (uint32)MTLPixelFormatDepth32Float_Stencil8;
 	}
 	GPixelFormats[PF_DepthStencil		].BlockBytes		= 4;
 	GPixelFormats[PF_DepthStencil		].Supported			= true;
 	if (bSupportsD16)
 	{
-		GPixelFormats[PF_ShadowDepth		].PlatformFormat	= (uint32)mtlpp::PixelFormat::Depth16Unorm;
+		GPixelFormats[PF_ShadowDepth		].PlatformFormat	= (uint32)MTLPixelFormatDepth16Unorm;
 		GPixelFormats[PF_ShadowDepth		].BlockBytes		= 2;
 	}
 	else
 	{
-		GPixelFormats[PF_ShadowDepth		].PlatformFormat	= (uint32)mtlpp::PixelFormat::Depth32Float;
+		GPixelFormats[PF_ShadowDepth		].PlatformFormat	= (uint32)MTLPixelFormatDepth32Float;
 		GPixelFormats[PF_ShadowDepth		].BlockBytes		= 4;
 	}
 	GPixelFormats[PF_ShadowDepth		].Supported			= true;
 	if(bSupportsD24S8)
 	{
-		GPixelFormats[PF_D24			].PlatformFormat	= (uint32)mtlpp::PixelFormat::Depth24Unorm_Stencil8;
+		GPixelFormats[PF_D24			].PlatformFormat	= (uint32)MTLPixelFormatDepth24Unorm_Stencil8;
 	}
 	else
 	{
-		GPixelFormats[PF_D24			].PlatformFormat	= (uint32)mtlpp::PixelFormat::Depth32Float;
+		GPixelFormats[PF_D24			].PlatformFormat	= (uint32)MTLPixelFormatDepth32Float;
 	}
 	GPixelFormats[PF_D24				].Supported			= true;
 	GPixelFormats[PF_BC4				].Supported			= true;
-	GPixelFormats[PF_BC4				].PlatformFormat	= (uint32)mtlpp::PixelFormat::BC4_RUnorm;
+	GPixelFormats[PF_BC4				].PlatformFormat	= (uint32)MTLPixelFormatBC4_RUnorm;
 	GPixelFormats[PF_BC5				].Supported			= true;
-	GPixelFormats[PF_BC5				].PlatformFormat	= (uint32)mtlpp::PixelFormat::BC5_RGUnorm;
+	GPixelFormats[PF_BC5				].PlatformFormat	= (uint32)MTLPixelFormatBC5_RGUnorm;
 	GPixelFormats[PF_BC6H				].Supported			= true;
-	GPixelFormats[PF_BC6H               ].PlatformFormat	= (uint32)mtlpp::PixelFormat::BC6H_RGBUfloat;
+	GPixelFormats[PF_BC6H               ].PlatformFormat	= (uint32)MTLPixelFormatBC6H_RGBUfloat;
 	GPixelFormats[PF_BC7				].Supported			= true;
-	GPixelFormats[PF_BC7				].PlatformFormat	= (uint32)mtlpp::PixelFormat::BC7_RGBAUnorm;
-	GPixelFormats[PF_R5G6B5_UNORM		].PlatformFormat	= (uint32)mtlpp::PixelFormat::Invalid;
-	GPixelFormats[PF_B5G5R5A1_UNORM		].PlatformFormat	= (uint32)mtlpp::PixelFormat::Invalid;
+	GPixelFormats[PF_BC7				].PlatformFormat	= (uint32)MTLPixelFormatBC7_RGBAUnorm;
+	GPixelFormats[PF_R5G6B5_UNORM		].PlatformFormat	= (uint32)MTLPixelFormatInvalid;
+	GPixelFormats[PF_B5G5R5A1_UNORM		].PlatformFormat	= (uint32)MTLPixelFormatInvalid;
 #endif
-	GPixelFormats[PF_UYVY				].PlatformFormat	= (uint32)mtlpp::PixelFormat::Invalid;
-	GPixelFormats[PF_FloatRGBA			].PlatformFormat	= (uint32)mtlpp::PixelFormat::RGBA16Float;
+	GPixelFormats[PF_UYVY				].PlatformFormat	= (uint32)MTLPixelFormatInvalid;
+	GPixelFormats[PF_FloatRGBA			].PlatformFormat	= (uint32)MTLPixelFormatRGBA16Float;
 	GPixelFormats[PF_FloatRGBA			].BlockBytes		= 8;
-    GPixelFormats[PF_X24_G8				].PlatformFormat	= (uint32)mtlpp::PixelFormat::Stencil8;
+    GPixelFormats[PF_X24_G8				].PlatformFormat	= (uint32)MTLPixelFormatStencil8;
     GPixelFormats[PF_X24_G8				].BlockBytes		= 1;
-	GPixelFormats[PF_R32_FLOAT			].PlatformFormat	= (uint32)mtlpp::PixelFormat::R32Float;
-	GPixelFormats[PF_G16R16				].PlatformFormat	= (uint32)mtlpp::PixelFormat::RG16Unorm;
+	GPixelFormats[PF_R32_FLOAT			].PlatformFormat	= (uint32)MTLPixelFormatR32Float;
+	GPixelFormats[PF_G16R16				].PlatformFormat	= (uint32)MTLPixelFormatRG16Unorm;
 	GPixelFormats[PF_G16R16				].Supported			= true;
-	GPixelFormats[PF_G16R16F			].PlatformFormat	= (uint32)mtlpp::PixelFormat::RG16Float;
-	GPixelFormats[PF_G16R16F_FILTER		].PlatformFormat	= (uint32)mtlpp::PixelFormat::RG16Float;
-	GPixelFormats[PF_G32R32F			].PlatformFormat	= (uint32)mtlpp::PixelFormat::RG32Float;
-	GPixelFormats[PF_A2B10G10R10		].PlatformFormat    = (uint32)mtlpp::PixelFormat::RGB10A2Unorm;
-	GPixelFormats[PF_A16B16G16R16		].PlatformFormat    = (uint32)mtlpp::PixelFormat::RGBA16Unorm;
-	GPixelFormats[PF_R16F				].PlatformFormat	= (uint32)mtlpp::PixelFormat::R16Float;
-	GPixelFormats[PF_R16F_FILTER		].PlatformFormat	= (uint32)mtlpp::PixelFormat::R16Float;
-	GPixelFormats[PF_V8U8				].PlatformFormat	= (uint32)mtlpp::PixelFormat::RG8Snorm;
-	GPixelFormats[PF_A1					].PlatformFormat	= (uint32)mtlpp::PixelFormat::Invalid;
+	GPixelFormats[PF_G16R16F			].PlatformFormat	= (uint32)MTLPixelFormatRG16Float;
+	GPixelFormats[PF_G16R16F_FILTER		].PlatformFormat	= (uint32)MTLPixelFormatRG16Float;
+	GPixelFormats[PF_G32R32F			].PlatformFormat	= (uint32)MTLPixelFormatRG32Float;
+	GPixelFormats[PF_A2B10G10R10		].PlatformFormat    = (uint32)MTLPixelFormatRGB10A2Unorm;
+	GPixelFormats[PF_A16B16G16R16		].PlatformFormat    = (uint32)MTLPixelFormatRGBA16Unorm;
+	GPixelFormats[PF_R16F				].PlatformFormat	= (uint32)MTLPixelFormatR16Float;
+	GPixelFormats[PF_R16F_FILTER		].PlatformFormat	= (uint32)MTLPixelFormatR16Float;
+	GPixelFormats[PF_V8U8				].PlatformFormat	= (uint32)MTLPixelFormatRG8Snorm;
+	GPixelFormats[PF_A1					].PlatformFormat	= (uint32)MTLPixelFormatInvalid;
 	// A8 does not allow writes in Metal. So we will fake it with R8.
 	// If you change this you must also change the swizzle pattern in Platform.ush
 	// See Texture2DSample_A8 in Common.ush and A8_SAMPLE_MASK in Platform.ush
-	GPixelFormats[PF_A8					].PlatformFormat	= (uint32)mtlpp::PixelFormat::R8Unorm;
-	GPixelFormats[PF_R32_UINT			].PlatformFormat	= (uint32)mtlpp::PixelFormat::R32Uint;
-	GPixelFormats[PF_R32_SINT			].PlatformFormat	= (uint32)mtlpp::PixelFormat::R32Sint;
-	GPixelFormats[PF_R16G16B16A16_UINT	].PlatformFormat	= (uint32)mtlpp::PixelFormat::RGBA16Uint;
-	GPixelFormats[PF_R16G16B16A16_SINT	].PlatformFormat	= (uint32)mtlpp::PixelFormat::RGBA16Sint;
-	GPixelFormats[PF_R8G8B8A8			].PlatformFormat	= (uint32)mtlpp::PixelFormat::RGBA8Unorm;
-	GPixelFormats[PF_R8G8B8A8_UINT		].PlatformFormat	= (uint32)mtlpp::PixelFormat::RGBA8Uint;
-	GPixelFormats[PF_R8G8B8A8_SNORM		].PlatformFormat	= (uint32)mtlpp::PixelFormat::RGBA8Snorm;
-	GPixelFormats[PF_R8G8				].PlatformFormat	= (uint32)mtlpp::PixelFormat::RG8Unorm;
-	GPixelFormats[PF_R16_SINT			].PlatformFormat	= (uint32)mtlpp::PixelFormat::R16Sint;
-	GPixelFormats[PF_R16_UINT			].PlatformFormat	= (uint32)mtlpp::PixelFormat::R16Uint;
-	GPixelFormats[PF_R8_UINT			].PlatformFormat	= (uint32)mtlpp::PixelFormat::R8Uint;
-	GPixelFormats[PF_R8					].PlatformFormat	= (uint32)mtlpp::PixelFormat::R8Unorm;
-
-	GPixelFormats[PF_R16G16B16A16_UNORM ].PlatformFormat	= (uint32)mtlpp::PixelFormat::RGBA16Unorm;
-	GPixelFormats[PF_R16G16B16A16_SNORM ].PlatformFormat	= (uint32)mtlpp::PixelFormat::RGBA16Snorm;
-
-	GPixelFormats[PF_NV12				].PlatformFormat	= (uint32)mtlpp::PixelFormat::Invalid;
+	GPixelFormats[PF_A8					].PlatformFormat	= (uint32)MTLPixelFormatR8Unorm;
+	GPixelFormats[PF_R32_UINT			].PlatformFormat	= (uint32)MTLPixelFormatR32Uint;
+	GPixelFormats[PF_R32_SINT			].PlatformFormat	= (uint32)MTLPixelFormatR32Sint;
+	GPixelFormats[PF_R16G16B16A16_UINT	].PlatformFormat	= (uint32)MTLPixelFormatRGBA16Uint;
+	GPixelFormats[PF_R16G16B16A16_SINT	].PlatformFormat	= (uint32)MTLPixelFormatRGBA16Sint;
+	GPixelFormats[PF_R8G8B8A8			].PlatformFormat	= (uint32)MTLPixelFormatRGBA8Unorm;
+	GPixelFormats[PF_R8G8B8A8_UINT		].PlatformFormat	= (uint32)MTLPixelFormatRGBA8Uint;
+	GPixelFormats[PF_R8G8B8A8_SNORM		].PlatformFormat	= (uint32)MTLPixelFormatRGBA8Snorm;
+	GPixelFormats[PF_R8G8				].PlatformFormat	= (uint32)MTLPixelFormatRG8Unorm;
+	GPixelFormats[PF_R16_SINT			].PlatformFormat	= (uint32)MTLPixelFormatR16Sint;
+	GPixelFormats[PF_R16_UINT			].PlatformFormat	= (uint32)MTLPixelFormatR16Uint;
+	GPixelFormats[PF_R8_UINT			].PlatformFormat	= (uint32)MTLPixelFormatR8Uint;
+	GPixelFormats[PF_R8					].PlatformFormat	= (uint32)MTLPixelFormatR8Unorm;
+
+	GPixelFormats[PF_R16G16B16A16_UNORM ].PlatformFormat	= (uint32)MTLPixelFormatRGBA16Unorm;
+	GPixelFormats[PF_R16G16B16A16_SNORM ].PlatformFormat	= (uint32)MTLPixelFormatRGBA16Snorm;
+
+	GPixelFormats[PF_NV12				].PlatformFormat	= (uint32)MTLPixelFormatInvalid;
 	GPixelFormats[PF_NV12				].Supported			= false;
 	
-	GPixelFormats[PF_G16R16_SNORM		].PlatformFormat	= (uint32)mtlpp::PixelFormat::RG16Snorm;
-	GPixelFormats[PF_R8G8_UINT			].PlatformFormat	= (uint32)mtlpp::PixelFormat::RG8Uint;
-	GPixelFormats[PF_R32G32B32_UINT		].PlatformFormat	= (uint32)mtlpp::PixelFormat::Invalid;
+	GPixelFormats[PF_G16R16_SNORM		].PlatformFormat	= (uint32)MTLPixelFormatRG16Snorm;
+	GPixelFormats[PF_R8G8_UINT			].PlatformFormat	= (uint32)MTLPixelFormatRG8Uint;
+	GPixelFormats[PF_R32G32B32_UINT		].PlatformFormat	= (uint32)MTLPixelFormatInvalid;
 	GPixelFormats[PF_R32G32B32_UINT		].Supported			= false;
-	GPixelFormats[PF_R32G32B32_SINT		].PlatformFormat	= (uint32)mtlpp::PixelFormat::Invalid;
+	GPixelFormats[PF_R32G32B32_SINT		].PlatformFormat	= (uint32)MTLPixelFormatInvalid;
 	GPixelFormats[PF_R32G32B32_SINT		].Supported			= false;
-	GPixelFormats[PF_R32G32B32F			].PlatformFormat	= (uint32)mtlpp::PixelFormat::Invalid;
+	GPixelFormats[PF_R32G32B32F			].PlatformFormat	= (uint32)MTLPixelFormatInvalid;
 	GPixelFormats[PF_R32G32B32F			].Supported			= false;
-	GPixelFormats[PF_R8_SINT			].PlatformFormat	= (uint32)mtlpp::PixelFormat::R8Sint;
-	GPixelFormats[PF_R64_UINT			].PlatformFormat	= (uint32)mtlpp::PixelFormat::Invalid;
+	GPixelFormats[PF_R8_SINT			].PlatformFormat	= (uint32)MTLPixelFormatR8Sint;
+	GPixelFormats[PF_R64_UINT			].PlatformFormat	= (uint32)MTLPixelFormatInvalid;
 	GPixelFormats[PF_R64_UINT			].Supported			= false;
 
 #if METAL_DEBUG_OPTIONS
diff --git a/Engine/Source/Runtime/Apple/AGXRHI/Private/AGXRHIPrivate.h b/Engine/Source/Runtime/Apple/AGXRHI/Private/AGXRHIPrivate.h
--- a/Engine/Source/Runtime/Apple/AGXRHI/Private/AGXRHIPrivate.h
+++ b/Engine/Source/Runtime/Apple/AGXRHI/Private/AGXRHIPrivate.h
@@ -139,7 +139,7 @@ extern bool GAGXSupportsCaptureManager;
 struct FAGXBufferFormat
 {
 	// Valid linear texture pixel formats - potentially different than the actual texture formats
-	mtlpp::PixelFormat LinearTextureFormat;
+	MTLPixelFormat LinearTextureFormat;
 	// Metal buffer data types for manual ALU format conversions
 	uint8 DataFormat;
 };
@@ -252,9 +252,9 @@ FORCEINLINE MTLLoadAction GetMetalRTLoadAction(ERenderTargetLoadAction LoadActio
 
 MTLPrimitiveType AGXTranslatePrimitiveType(uint32 PrimitiveType);
 
-mtlpp::PixelFormat AGXToSRGBFormat(mtlpp::PixelFormat LinMTLFormat);
+MTLPixelFormat AGXToSRGBFormat(MTLPixelFormat LinMTLFormat);
 
-uint8 AGXGetMetalPixelFormatKey(mtlpp::PixelFormat Format);
+uint8 AGXGetMetalPixelFormatKey(MTLPixelFormat Format);
 
 template<typename TRHIType>
 static FORCEINLINE typename TAGXResourceTraits<TRHIType>::TConcreteType* ResourceCast(TRHIType* Resource)
diff --git a/Engine/Source/Runtime/Apple/AGXRHI/Private/AGXRenderTarget.cpp b/Engine/Source/Runtime/Apple/AGXRHI/Private/AGXRenderTarget.cpp
--- a/Engine/Source/Runtime/Apple/AGXRHI/Private/AGXRenderTarget.cpp
+++ b/Engine/Source/Runtime/Apple/AGXRHI/Private/AGXRenderTarget.cpp
@@ -290,7 +290,7 @@ void FAGXDynamicRHI::RHIReadSurfaceData(FRHITexture* TextureRHI, FIntRect Rect,
 #else
 			mtlpp::StorageMode StorageMode = mtlpp::StorageMode::Shared;
 #endif
-			mtlpp::PixelFormat MetalFormat = (mtlpp::PixelFormat)GPixelFormats[Surface->GetDesc().Format].PlatformFormat;
+			MTLPixelFormat MetalFormat = (MTLPixelFormat)GPixelFormats[Surface->GetDesc().Format].PlatformFormat;
 			mtlpp::TextureDescriptor Desc;
 			Desc.SetTextureType(Texture.GetTextureType());
 			Desc.SetPixelFormat(Texture.GetPixelFormat());
diff --git a/Engine/Source/Runtime/Apple/AGXRHI/Private/AGXStateCache.cpp b/Engine/Source/Runtime/Apple/AGXRHI/Private/AGXStateCache.cpp
--- a/Engine/Source/Runtime/Apple/AGXRHI/Private/AGXStateCache.cpp
+++ b/Engine/Source/Runtime/Apple/AGXRHI/Private/AGXStateCache.cpp
@@ -747,21 +747,21 @@ bool FAGXStateCache::SetRenderPassInfo(FRHIRenderPassInfo const& InRenderTargets
 				case PF_DepthStencil:
 				case PF_D24:
 				{
-					mtlpp::PixelFormat DepthStencilFormat = Surface.Texture ? (mtlpp::PixelFormat)Surface.Texture.GetPixelFormat() : mtlpp::PixelFormat::Invalid;
+					MTLPixelFormat DepthStencilFormat = Surface.Texture ? (MTLPixelFormat)Surface.Texture.GetPixelFormat() : MTLPixelFormatInvalid;
 					
 					switch(DepthStencilFormat)
 					{
-						case mtlpp::PixelFormat::Depth32Float:
+						case MTLPixelFormatDepth32Float:
 							StencilTexture =  nil;
 							break;
-						case mtlpp::PixelFormat::Stencil8:
+						case MTLPixelFormatStencil8:
 							StencilTexture = DepthTexture;
 							break;
-						case mtlpp::PixelFormat::Depth32Float_Stencil8:
+						case MTLPixelFormatDepth32Float_Stencil8:
 							StencilTexture = DepthTexture;
 							break;
 #if PLATFORM_MAC
-						case mtlpp::PixelFormat::Depth24Unorm_Stencil8:
+						case MTLPixelFormatDepth24Unorm_Stencil8:
 							StencilTexture = DepthTexture;
 							break;
 #endif
@@ -791,7 +791,7 @@ bool FAGXStateCache::SetRenderPassInfo(FRHIRenderPassInfo const& InRenderTargets
 				DepthClearValue = 1.0f;
 			}
 
-			bool const bCombinedDepthStencilUsingStencil = (DepthTexture && (mtlpp::PixelFormat)DepthTexture.GetPixelFormat() != mtlpp::PixelFormat::Depth32Float && RenderPassInfo.DepthStencilRenderTarget.ExclusiveDepthStencil.IsUsingStencil());
+			bool const bCombinedDepthStencilUsingStencil = (DepthTexture && (MTLPixelFormat)DepthTexture.GetPixelFormat() != MTLPixelFormatDepth32Float && RenderPassInfo.DepthStencilRenderTarget.ExclusiveDepthStencil.IsUsingStencil());
 			bool const bUsingDepth = (RenderPassInfo.DepthStencilRenderTarget.ExclusiveDepthStencil.IsUsingDepth() || (bCombinedDepthStencilUsingStencil));
 			if (DepthTexture && bUsingDepth)
 			{
@@ -882,7 +882,7 @@ bool FAGXStateCache::SetRenderPassInfo(FRHIRenderPassInfo const& InRenderTargets
             //if we're dealing with a samplecount mismatch we just bail on stencil entirely as stencil
             //doesn't have an autoresolve target to use.
 			
-			bool const bCombinedDepthStencilUsingDepth = (StencilTexture && StencilTexture.GetPixelFormat() != mtlpp::PixelFormat::Stencil8 && RenderPassInfo.DepthStencilRenderTarget.ExclusiveDepthStencil.IsUsingDepth());
+			bool const bCombinedDepthStencilUsingDepth = (StencilTexture && (MTLPixelFormat)StencilTexture.GetPixelFormat() != MTLPixelFormatStencil8 && RenderPassInfo.DepthStencilRenderTarget.ExclusiveDepthStencil.IsUsingDepth());
 			bool const bUsingStencil = RenderPassInfo.DepthStencilRenderTarget.ExclusiveDepthStencil.IsUsingStencil() || (bCombinedDepthStencilUsingDepth);
 			if (StencilTexture && bUsingStencil)
 			{
diff --git a/Engine/Source/Runtime/Apple/AGXRHI/Private/AGXTexture.cpp b/Engine/Source/Runtime/Apple/AGXRHI/Private/AGXTexture.cpp
--- a/Engine/Source/Runtime/Apple/AGXRHI/Private/AGXTexture.cpp
+++ b/Engine/Source/Runtime/Apple/AGXRHI/Private/AGXTexture.cpp
@@ -151,27 +151,27 @@ void AGXSafeReleaseMetalTexture(FAGXSurface* Surface, FAGXTexture& Texture)
 }
 
 #if !PLATFORM_IOS
-static mtlpp::PixelFormat AGX_ToSRGBFormat_NonAppleMacGPU(mtlpp::PixelFormat MTLFormat)
+static MTLPixelFormat AGX_ToSRGBFormat_NonAppleMacGPU(MTLPixelFormat MTLFormat)
 {
 	switch (MTLFormat)
 	{
-		case mtlpp::PixelFormat::RGBA8Unorm:
-			MTLFormat = mtlpp::PixelFormat::RGBA8Unorm_sRGB;
+		case MTLPixelFormatRGBA8Unorm:
+			MTLFormat = MTLPixelFormatRGBA8Unorm_sRGB;
 			break;
-		case mtlpp::PixelFormat::BGRA8Unorm:
-			MTLFormat = mtlpp::PixelFormat::BGRA8Unorm_sRGB;
+		case MTLPixelFormatBGRA8Unorm:
+			MTLFormat = MTLPixelFormatBGRA8Unorm_sRGB;
 			break;
-		case mtlpp::PixelFormat::BC1_RGBA:
-			MTLFormat = mtlpp::PixelFormat::BC1_RGBA_sRGB;
+		case MTLPixelFormatBC1_RGBA:
+			MTLFormat = MTLPixelFormatBC1_RGBA_sRGB;
 			break;
-		case mtlpp::PixelFormat::BC2_RGBA:
-			MTLFormat = mtlpp::PixelFormat::BC2_RGBA_sRGB;
+		case MTLPixelFormatBC2_RGBA:
+			MTLFormat = MTLPixelFormatBC2_RGBA_sRGB;
 			break;
-		case mtlpp::PixelFormat::BC3_RGBA:
-			MTLFormat = mtlpp::PixelFormat::BC3_RGBA_sRGB;
+		case MTLPixelFormatBC3_RGBA:
+			MTLFormat = MTLPixelFormatBC3_RGBA_sRGB;
 			break;
-		case mtlpp::PixelFormat::BC7_RGBAUnorm:
-			MTLFormat = mtlpp::PixelFormat::BC7_RGBAUnorm_sRGB;
+		case MTLPixelFormatBC7_RGBAUnorm:
+			MTLFormat = MTLPixelFormatBC7_RGBAUnorm_sRGB;
 			break;
 		default:
 			break;
@@ -180,53 +180,53 @@ static mtlpp::PixelFormat AGX_ToSRGBFormat_NonAppleMacGPU(mtlpp::PixelFormat MTL
 }
 #endif
 
-static mtlpp::PixelFormat AGX_ToSRGBFormat_AppleGPU(mtlpp::PixelFormat MTLFormat)
+static MTLPixelFormat AGX_ToSRGBFormat_AppleGPU(MTLPixelFormat MTLFormat)
 {
 	switch (MTLFormat)
 	{
-		case mtlpp::PixelFormat::RGBA8Unorm:
-			MTLFormat = mtlpp::PixelFormat::RGBA8Unorm_sRGB;
+		case MTLPixelFormatRGBA8Unorm:
+			MTLFormat = MTLPixelFormatRGBA8Unorm_sRGB;
 			break;
-		case mtlpp::PixelFormat::BGRA8Unorm:
-			MTLFormat = mtlpp::PixelFormat::BGRA8Unorm_sRGB;
+		case MTLPixelFormatBGRA8Unorm:
+			MTLFormat = MTLPixelFormatBGRA8Unorm_sRGB;
 			break;
-		case mtlpp::PixelFormat::R8Unorm:
-			MTLFormat = mtlpp::PixelFormat::R8Unorm_sRGB;
+		case MTLPixelFormatR8Unorm:
+			MTLFormat = MTLPixelFormatR8Unorm_sRGB;
 			break;
-		case mtlpp::PixelFormat::PVRTC_RGBA_2BPP:
-			MTLFormat = mtlpp::PixelFormat::PVRTC_RGBA_2BPP_sRGB;
+		case MTLPixelFormatPVRTC_RGBA_2BPP:
+			MTLFormat = MTLPixelFormatPVRTC_RGBA_2BPP_sRGB;
 			break;
-		case mtlpp::PixelFormat::PVRTC_RGBA_4BPP:
-			MTLFormat = mtlpp::PixelFormat::PVRTC_RGBA_4BPP_sRGB;
+		case MTLPixelFormatPVRTC_RGBA_4BPP:
+			MTLFormat = MTLPixelFormatPVRTC_RGBA_4BPP_sRGB;
 			break;
-		case mtlpp::PixelFormat::ASTC_4x4_LDR:
-			MTLFormat = mtlpp::PixelFormat::ASTC_4x4_sRGB;
+		case MTLPixelFormatASTC_4x4_LDR:
+			MTLFormat = MTLPixelFormatASTC_4x4_sRGB;
 			break;
-		case mtlpp::PixelFormat::ASTC_6x6_LDR:
-			MTLFormat = mtlpp::PixelFormat::ASTC_6x6_sRGB;
+		case MTLPixelFormatASTC_6x6_LDR:
+			MTLFormat = MTLPixelFormatASTC_6x6_sRGB;
 			break;
-		case mtlpp::PixelFormat::ASTC_8x8_LDR:
-			MTLFormat = mtlpp::PixelFormat::ASTC_8x8_sRGB;
+		case MTLPixelFormatASTC_8x8_LDR:
+			MTLFormat = MTLPixelFormatASTC_8x8_sRGB;
 			break;
-		case mtlpp::PixelFormat::ASTC_10x10_LDR:
-			MTLFormat = mtlpp::PixelFormat::ASTC_10x10_sRGB;
+		case MTLPixelFormatASTC_10x10_LDR:
+			MTLFormat = MTLPixelFormatASTC_10x10_sRGB;
 			break;
-		case mtlpp::PixelFormat::ASTC_12x12_LDR:
-			MTLFormat = mtlpp::PixelFormat::ASTC_12x12_sRGB;
+		case MTLPixelFormatASTC_12x12_LDR:
+			MTLFormat = MTLPixelFormatASTC_12x12_sRGB;
 			break;
 #if PLATFORM_MAC
 		// Fix for Apple silicon M1 macs that can support BC pixel formats even though they are Apple family GPUs.
-		case mtlpp::PixelFormat::BC1_RGBA:
-			MTLFormat = mtlpp::PixelFormat::BC1_RGBA_sRGB;
+		case MTLPixelFormatBC1_RGBA:
+			MTLFormat = MTLPixelFormatBC1_RGBA_sRGB;
 			break;
-		case mtlpp::PixelFormat::BC2_RGBA:
-			MTLFormat = mtlpp::PixelFormat::BC2_RGBA_sRGB;
+		case MTLPixelFormatBC2_RGBA:
+			MTLFormat = MTLPixelFormatBC2_RGBA_sRGB;
 			break;
-		case mtlpp::PixelFormat::BC3_RGBA:
-			MTLFormat = mtlpp::PixelFormat::BC3_RGBA_sRGB;
+		case MTLPixelFormatBC3_RGBA:
+			MTLFormat = MTLPixelFormatBC3_RGBA_sRGB;
 			break;
-		case mtlpp::PixelFormat::BC7_RGBAUnorm:
-			MTLFormat = mtlpp::PixelFormat::BC7_RGBAUnorm_sRGB;
+		case MTLPixelFormatBC7_RGBAUnorm:
+			MTLFormat = MTLPixelFormatBC7_RGBAUnorm_sRGB;
 			break;
 #endif
 		default:
@@ -235,7 +235,7 @@ static mtlpp::PixelFormat AGX_ToSRGBFormat_AppleGPU(mtlpp::PixelFormat MTLFormat
 	return MTLFormat;
 }
 
-mtlpp::PixelFormat AGXToSRGBFormat(mtlpp::PixelFormat MTLFormat)
+MTLPixelFormat AGXToSRGBFormat(MTLPixelFormat MTLFormat)
 {
 	if([GMtlDevice supportsFamily:MTLGPUFamilyApple1])
 	{
@@ -323,15 +323,15 @@ void FAGXSurface::MakeAliasable(void)
 	// TODO
 }
 
-uint8 AGXGetMetalPixelFormatKey(mtlpp::PixelFormat Format)
+uint8 AGXGetMetalPixelFormatKey(MTLPixelFormat Format)
 {
 	struct FAGXPixelFormatKeyMap
 	{
 		FRWLock Mutex;
-		uint8 NextKey = 1; // 0 is reserved for mtlpp::PixelFormat::Invalid
+		uint8 NextKey = 1; // 0 is reserved for MTLPixelFormatInvalid
 		TMap<uint64, uint8> Map;
 
-		uint8 Get(mtlpp::PixelFormat Format)
+		uint8 Get(MTLPixelFormat Format)
 		{
 			FRWScopeLock Lock(Mutex, SLT_ReadOnly);
 			uint8* Key = Map.Find((uint64)Format);
@@ -352,12 +352,12 @@ uint8 AGXGetMetalPixelFormatKey(mtlpp::PixelFormat Format)
 		FAGXPixelFormatKeyMap()
 		{
 			// Add depth stencil formats first, so we don't have to use 6 bits for them in the pipeline hash
-			Get(mtlpp::PixelFormat::Depth32Float);
-			Get(mtlpp::PixelFormat::Stencil8);
-			Get(mtlpp::PixelFormat::Depth32Float_Stencil8);
+			Get(MTLPixelFormatDepth32Float);
+			Get(MTLPixelFormatStencil8);
+			Get(MTLPixelFormatDepth32Float_Stencil8);
 #if PLATFORM_MAC
-			Get(mtlpp::PixelFormat::Depth24Unorm_Stencil8);
-			Get(mtlpp::PixelFormat::Depth16Unorm);
+			Get(MTLPixelFormatDepth24Unorm_Stencil8);
+			Get(MTLPixelFormatDepth16Unorm);
 #endif
 		}
 	} static PixelFormatKeyMap;
@@ -368,7 +368,7 @@ uint8 AGXGetMetalPixelFormatKey(mtlpp::PixelFormat Format)
 FAGXTextureDesc::FAGXTextureDesc(FRHITextureDesc const& InDesc)
 	: bIsRenderTarget(IsRenderTarget(InDesc.Flags))
 {
-	MTLFormat = (mtlpp::PixelFormat)GPixelFormats[InDesc.Format].PlatformFormat;
+	MTLFormat = (MTLPixelFormat)GPixelFormats[InDesc.Format].PlatformFormat;
 	if (EnumHasAnyFlags(InDesc.Flags, TexCreate_SRGB))
 	{
 		MTLFormat = AGXToSRGBFormat(MTLFormat);
@@ -379,7 +379,7 @@ FAGXTextureDesc::FAGXTextureDesc(FRHITextureDesc const& InDesc)
 
 	if (InDesc.IsTextureCube())
 	{
-		Desc = mtlpp::TextureDescriptor::TextureCubeDescriptor(MTLFormat, InDesc.Extent.X, (InDesc.NumMips > 1));
+		Desc = mtlpp::TextureDescriptor::TextureCubeDescriptor((mtlpp::PixelFormat)MTLFormat, InDesc.Extent.X, (InDesc.NumMips > 1));
 	}
 	else if (InDesc.IsTexture3D())
 	{
@@ -387,14 +387,14 @@ FAGXTextureDesc::FAGXTextureDesc(FRHITextureDesc const& InDesc)
 		Desc.SetWidth(InDesc.Extent.X);
 		Desc.SetHeight(InDesc.Extent.Y);
 		Desc.SetDepth(InDesc.Depth);
-		Desc.SetPixelFormat(MTLFormat);
+		Desc.SetPixelFormat((mtlpp::PixelFormat)MTLFormat);
 		Desc.SetArrayLength(1);
 		Desc.SetMipmapLevelCount(1);
 		Desc.SetSampleCount(1);
 	}
 	else
 	{
-		Desc = mtlpp::TextureDescriptor::Texture2DDescriptor(MTLFormat, InDesc.Extent.X, InDesc.Extent.Y, (InDesc.NumMips > 1));
+		Desc = mtlpp::TextureDescriptor::Texture2DDescriptor((mtlpp::PixelFormat)MTLFormat, InDesc.Extent.X, InDesc.Extent.Y, (InDesc.NumMips > 1));
 		Desc.SetArrayLength(InDesc.ArraySize);
 	}
 
@@ -622,7 +622,7 @@ FAGXSurface::FAGXSurface(FAGXTextureCreateDesc const& CreateDesc)
 
 		if(bBufferCompatibleOption && (EnumHasAllFlags(CreateDesc.Flags, TexCreate_UAV | TexCreate_NoTiling) || EnumHasAllFlags(CreateDesc.Flags, TexCreate_AtomicCompatible)))
 		{
-			const uint32 MinimumByteAlignment = GMtlppDevice.GetMinimumLinearTextureAlignmentForPixelFormat(CreateDesc.MTLFormat);
+			const uint32 MinimumByteAlignment = GMtlppDevice.GetMinimumLinearTextureAlignmentForPixelFormat((mtlpp::PixelFormat)CreateDesc.MTLFormat);
 			const NSUInteger BytesPerRow = Align(CreateDesc.Desc.GetWidth() * GPixelFormats[CreateDesc.Format].BlockBytes, MinimumByteAlignment);
 
 			// Backing buffer resource options must match the texture we are going to create from it
@@ -898,7 +898,7 @@ id <MTLBuffer> FAGXSurface::AllocSurface(uint32 MipIndex, uint32 ArrayIndex, ERe
 		&& GetDesc().Dimension == ETextureDimension::Texture2D
 		&& EnumHasAnyFlags(GetDesc().Flags, TexCreate_SRGB)
 		&& LockMode == RLM_WriteOnly
-		&& Texture.GetPixelFormat() == mtlpp::PixelFormat::RGBA8Unorm_sRGB)
+		&& (MTLPixelFormat)Texture.GetPixelFormat() == MTLPixelFormatRGBA8Unorm_sRGB)
 	{
 		DestStride = FMath::Max<uint32>(GetDesc().Extent.X >> MipIndex, 1);
 	}
@@ -945,7 +945,7 @@ void FAGXSurface::UpdateSurfaceAndDestroySourceBuffer(id <MTLBuffer> SourceBuffe
 	if (   GetDesc().Format == PF_G8
 		&& GetDesc().Dimension == ETextureDimension::Texture2D
 		&& EnumHasAnyFlags(GetDesc().Flags, TexCreate_SRGB)
-		&& Texture.GetPixelFormat() == mtlpp::PixelFormat::RGBA8Unorm_sRGB)
+		&& (MTLPixelFormat)Texture.GetPixelFormat() == MTLPixelFormatRGBA8Unorm_sRGB)
 	{
 		TArray<uint8> Data;
 		uint8* ExpandedMem = (uint8*) SourceBuffer.contents;
@@ -981,7 +981,7 @@ void FAGXSurface::UpdateSurfaceAndDestroySourceBuffer(id <MTLBuffer> SourceBuffe
 		
 		mtlpp::BlitOption Options = mtlpp::BlitOption::None;
 #if !PLATFORM_MAC
-		if (Texture.GetPixelFormat() >= mtlpp::PixelFormat::PVRTC_RGB_2BPP && Texture.GetPixelFormat() <= mtlpp::PixelFormat::PVRTC_RGBA_4BPP_sRGB)
+		if ((MTLPixelFormat)Texture.GetPixelFormat() >= MTLPixelFormatPVRTC_RGB_2BPP && (MTLPixelFormat)Texture.GetPixelFormat() <= MTLPixelFormatPVRTC_RGBA_4BPP_sRGB)
 		{
 			Options = mtlpp::BlitOption::RowLinearPVRTC;
 		}
@@ -1033,7 +1033,7 @@ void FAGXSurface::UpdateSurfaceAndDestroySourceBuffer(id <MTLBuffer> SourceBuffe
 	else
 	{
 #if !PLATFORM_MAC
-		if (Texture.GetPixelFormat() >= mtlpp::PixelFormat::PVRTC_RGB_2BPP && Texture.GetPixelFormat() <= mtlpp::PixelFormat::PVRTC_RGBA_4BPP_sRGB) // @todo Calculate correct strides and byte-counts
+		if ((MTLPixelFormat)Texture.GetPixelFormat() >= MTLPixelFormatPVRTC_RGB_2BPP && (MTLPixelFormat)Texture.GetPixelFormat() <= MTLPixelFormatPVRTC_RGBA_4BPP_sRGB) // @todo Calculate correct strides and byte-counts
 		{
 			Stride = 0;
 			BytesPerImage = 0;
@@ -1163,7 +1163,7 @@ void* FAGXSurface::Lock(uint32 MipIndex, uint32 ArrayIndex, EResourceLockMode Lo
 			if (   GetDesc().Format == PF_G8
 				&& GetDesc().Dimension == ETextureDimension::Texture2D
 				&& EnumHasAnyFlags(GetDesc().Flags, TexCreate_SRGB)
-				&& Texture.GetPixelFormat() == mtlpp::PixelFormat::RGBA8Unorm_sRGB)
+				&& (MTLPixelFormat)Texture.GetPixelFormat() == MTLPixelFormatRGBA8Unorm_sRGB)
 			{
 				TArray<uint8> Data;
 				uint8* ExpandedMem = (uint8*)SourceData.GetContents();
@@ -1310,7 +1310,7 @@ uint32 FAGXSurface::GetMipSize(uint32 MipIndex, uint32* Stride, bool bSingleLaye
 		NumBlocksY = FMath::Max<uint32>(NumBlocksY, 2);
 	}
 #if PLATFORM_MAC
-	else if (PixelFormat == PF_G8 && EnumHasAnyFlags(GetDesc().Flags, TexCreate_SRGB) && Texture.GetPixelFormat() == mtlpp::PixelFormat::RGBA8Unorm_sRGB)
+	else if (PixelFormat == PF_G8 && EnumHasAnyFlags(GetDesc().Flags, TexCreate_SRGB) && (MTLPixelFormat)Texture.GetPixelFormat() == MTLPixelFormatRGBA8Unorm_sRGB)
 	{
 		// RGBA_sRGB is the closest match - so expand the data.
 		NumBlocksX *= 4;
@@ -1658,7 +1658,7 @@ static FAGXBuffer InternalCopyTexture2DUpdateRegion(FRHITexture2D* TextureRHI, c
 	// Expand R8_sRGB into RGBA8_sRGB for Mac.
 	if (   Texture->GetFormat() == PF_G8
 		&& EnumHasAnyFlags(Texture->GetFlags(), TexCreate_SRGB)
-		&& Texture->Texture.GetPixelFormat() == mtlpp::PixelFormat::RGBA8Unorm_sRGB)
+		&& (MTLPixelFormat)Texture->Texture.GetPixelFormat() == MTLPixelFormatRGBA8Unorm_sRGB)
 	{
 		const uint32 BufferSize = UpdateRegion.Height * UpdateRegion.Width * sizeof(uint32);
 		Buffer = GetAGXDeviceContext().CreatePooledBuffer(FAGXPooledBufferArgs(BufferSize, BUF_Dynamic, mtlpp::StorageMode::Shared));
@@ -1703,7 +1703,7 @@ static void InternalUpdateTexture2D(FAGXContext& Context, FRHITexture2D* Texture
 		
 		mtlpp::BlitOption Options = mtlpp::BlitOption::None;
 #if !PLATFORM_MAC
-		if (Tex.GetPixelFormat() >= mtlpp::PixelFormat::PVRTC_RGB_2BPP && Tex.GetPixelFormat() <= mtlpp::PixelFormat::PVRTC_RGBA_4BPP_sRGB)
+		if ((MTLPixelFormat)Tex.GetPixelFormat() >= MTLPixelFormatPVRTC_RGB_2BPP && (MTLPixelFormat)Tex.GetPixelFormat() <= MTLPixelFormatPVRTC_RGBA_4BPP_sRGB)
 		{
 			Options = mtlpp::BlitOption::RowLinearPVRTC;
 		}
@@ -1781,7 +1781,7 @@ void FAGXDynamicRHI::RHIUpdateTexture2D(FRHITexture2D* TextureRHI, uint32 MipInd
 		{
 #if PLATFORM_MAC
 			TArray<uint32> ExpandedData;
-			if (Texture->GetFormat() == PF_G8 && EnumHasAnyFlags(Texture->GetFlags(), TexCreate_SRGB) && Tex.GetPixelFormat() == mtlpp::PixelFormat::RGBA8Unorm_sRGB)
+			if (Texture->GetFormat() == PF_G8 && EnumHasAnyFlags(Texture->GetFlags(), TexCreate_SRGB) && (MTLPixelFormat)Tex.GetPixelFormat() == MTLPixelFormatRGBA8Unorm_sRGB)
 			{
 				ExpandedData.AddZeroed(UpdateRegion.Height * UpdateRegion.Width);
 				InternalExpandR8ToStandardRGBA((uint32*)ExpandedData.GetData(), UpdateRegion, SourcePitch, SourceData);
@@ -1833,7 +1833,7 @@ static void InternalUpdateTexture3D(FAGXContext& Context, FRHITexture3D* Texture
 		const uint32 BytesPerImage = SourceRowPitch * UpdateRegion.Height;
 		mtlpp::BlitOption Options = mtlpp::BlitOption::None;
 #if !PLATFORM_MAC
-		if (Tex.GetPixelFormat() >= mtlpp::PixelFormat::PVRTC_RGB_2BPP && Tex.GetPixelFormat() <= mtlpp::PixelFormat::PVRTC_RGBA_4BPP_sRGB)
+		if ((MTLPixelFormat)Tex.GetPixelFormat() >= MTLPixelFormatPVRTC_RGB_2BPP && (MTLPixelFormat)Tex.GetPixelFormat() <= MTLPixelFormatPVRTC_RGBA_4BPP_sRGB)
 		{
 			Options = mtlpp::BlitOption::RowLinearPVRTC;
 		}
@@ -1937,7 +1937,7 @@ void FAGXDynamicRHI::RHIUpdateTexture3D(FRHITexture3D* TextureRHI,uint32 MipInde
 		FAGXTexture Tex = Texture->Texture;
 		
 #if PLATFORM_MAC
-		checkf(!(Texture->GetFormat() == PF_G8 && EnumHasAnyFlags(Texture->GetFlags(), TexCreate_SRGB) && Tex.GetPixelFormat() == mtlpp::PixelFormat::RGBA8Unorm_sRGB), TEXT("AGXRHI does not support PF_G8_sRGB on 3D, array or cube textures as it requires manual, CPU-side expansion to RGBA8_sRGB which is expensive!"));
+		checkf(!(Texture->GetFormat() == PF_G8 && EnumHasAnyFlags(Texture->GetFlags(), TexCreate_SRGB) && (MTLPixelFormat)Tex.GetPixelFormat() == MTLPixelFormatRGBA8Unorm_sRGB), TEXT("AGXRHI does not support PF_G8_sRGB on 3D, array or cube textures as it requires manual, CPU-side expansion to RGBA8_sRGB which is expensive!"));
 #endif
 		if(Tex.GetStorageMode() == mtlpp::StorageMode::Private)
 		{
@@ -2112,7 +2112,7 @@ void FAGXRHICommandContext::RHICopyTexture(FRHITexture* SourceTextureRHI, FRHITe
 						
 						mtlpp::BlitOption Options = mtlpp::BlitOption::None;
 #if !PLATFORM_MAC
-						if (MetalSrcTexture->Texture.GetPixelFormat() >= mtlpp::PixelFormat::PVRTC_RGB_2BPP && MetalSrcTexture->Texture.GetPixelFormat() <= mtlpp::PixelFormat::PVRTC_RGBA_4BPP_sRGB)
+						if ((MTLPixelFormat)MetalSrcTexture->Texture.GetPixelFormat() >= MTLPixelFormatPVRTC_RGB_2BPP && (MTLPixelFormat)MetalSrcTexture->Texture.GetPixelFormat() <= MTLPixelFormatPVRTC_RGBA_4BPP_sRGB)
 						{
 							Options = mtlpp::BlitOption::RowLinearPVRTC;
 						}
diff --git a/Engine/Source/Runtime/Apple/AGXRHI/Private/AGXUAV.cpp b/Engine/Source/Runtime/Apple/AGXRHI/Private/AGXUAV.cpp
--- a/Engine/Source/Runtime/Apple/AGXRHI/Private/AGXUAV.cpp
+++ b/Engine/Source/Runtime/Apple/AGXRHI/Private/AGXUAV.cpp
@@ -104,26 +104,26 @@ FAGXResourceViewBase::FAGXResourceViewBase(
 		{
 			// Determine the appropriate metal format for the view.
 			// This format will be non-sRGB. We convert to sRGB below if required.
-			mtlpp::PixelFormat MetalFormat = (mtlpp::PixelFormat)GPixelFormats[Format].PlatformFormat;
+			MTLPixelFormat MetalFormat = (MTLPixelFormat)GPixelFormats[Format].PlatformFormat;
 
 			if (Format == PF_X24_G8)
 			{
 				// Stencil buffer view of a depth texture
 				check(SourceTexture->GetDesc().Format == PF_DepthStencil);
-				switch (SourceTexture->Texture.GetPixelFormat())
+				switch ((MTLPixelFormat)SourceTexture->Texture.GetPixelFormat())
 				{
 					default:
 						checkNoEntry();
 						break;
 
 #if PLATFORM_MAC
-					case mtlpp::PixelFormat::Depth24Unorm_Stencil8:
-						MetalFormat = mtlpp::PixelFormat::X24_Stencil8;
+					case MTLPixelFormatDepth24Unorm_Stencil8:
+						MetalFormat = MTLPixelFormatX24_Stencil8;
 						break;
 #endif
 
-					case mtlpp::PixelFormat::Depth32Float_Stencil8:
-						MetalFormat = mtlpp::PixelFormat::X32_Stencil8;
+					case MTLPixelFormatDepth32Float_Stencil8:
+						MetalFormat = MTLPixelFormatX32_Stencil8;
 						break;
 				}
 			}
@@ -136,9 +136,9 @@ FAGXResourceViewBase::FAGXResourceViewBase(
 					{
 #if PLATFORM_MAC
 						// R8Unorm has been expanded in the source surface for sRGBA support - we need to expand to RGBA to enable compatible texture format view for non apple silicon macs
-						if (Format == PF_G8 && SourceTexture->Texture.GetPixelFormat() == mtlpp::PixelFormat::RGBA8Unorm_sRGB)
+						if (Format == PF_G8 && (MTLPixelFormat)SourceTexture->Texture.GetPixelFormat() == MTLPixelFormatRGBA8Unorm_sRGB)
 						{
-							MetalFormat = mtlpp::PixelFormat::RGBA8Unorm;
+							MetalFormat = MTLPixelFormatRGBA8Unorm;
 						}
 #endif
 					}
@@ -154,7 +154,7 @@ FAGXResourceViewBase::FAGXResourceViewBase(
 			bool bUseSourceTex =
 				   MipLevel == 0
 				&& NumMips == SourceTexture->Texture.GetMipmapLevelCount()
-				&& MetalFormat == SourceTexture->Texture.GetPixelFormat();
+				&& MetalFormat == (MTLPixelFormat)SourceTexture->Texture.GetPixelFormat();
 
 			if (bUseSourceTex)
 			{
@@ -171,7 +171,7 @@ FAGXResourceViewBase::FAGXResourceViewBase(
 				ns::Range Slices(0, SourceTexture->Texture.GetArrayLength() * (SourceTexture->GetDesc().IsTextureCube() ? 6 : 1));
 
 				TextureView = SourceTexture->Texture.NewTextureView(
-					MetalFormat,
+					(mtlpp::PixelFormat)MetalFormat,
 					SourceTexture->Texture.GetTextureType(),
 					ns::Range(MipLevel, NumMips),
 					Slices);
diff --git a/Engine/Source/Runtime/Apple/AGXRHI/Private/AGXVertexBuffer.cpp b/Engine/Source/Runtime/Apple/AGXRHI/Private/AGXVertexBuffer.cpp
--- a/Engine/Source/Runtime/Apple/AGXRHI/Private/AGXVertexBuffer.cpp
+++ b/Engine/Source/Runtime/Apple/AGXRHI/Private/AGXVertexBuffer.cpp
@@ -340,7 +340,7 @@ void FAGXRHIBuffer::AllocLinearTextures(const LinearTextureMapKey& InLinearTextu
 		uint8 InFormat = InLinearTextureMapKey.Key;
 		const FAGXLinearTextureDescriptor& LinearTextureDesc = InLinearTextureMapKey.Value;
 		
-		mtlpp::PixelFormat MTLFormat = (mtlpp::PixelFormat)GAGXBufferFormats[InFormat].LinearTextureFormat;
+		MTLPixelFormat MTLFormat = (MTLPixelFormat)GAGXBufferFormats[InFormat].LinearTextureFormat;
 		
 		mtlpp::TextureDescriptor Desc;
 		NSUInteger TexUsage = mtlpp::TextureUsage::Unknown;
@@ -354,7 +354,7 @@ void FAGXRHIBuffer::AllocLinearTextures(const LinearTextureMapKey& InLinearTextu
 		}
 		
 		uint32 BytesPerElement = (0 == LinearTextureDesc.BytesPerElement) ? GPixelFormats[InFormat].BlockBytes : LinearTextureDesc.BytesPerElement;
-		if (MTLFormat == mtlpp::PixelFormat::RG11B10Float && MTLFormat != (mtlpp::PixelFormat)GPixelFormats[InFormat].PlatformFormat)
+		if (MTLFormat == MTLPixelFormatRG11B10Float && MTLFormat != (MTLPixelFormat)GPixelFormats[InFormat].PlatformFormat)
 		{
 			BytesPerElement = 4;
 		}
@@ -373,7 +373,7 @@ void FAGXRHIBuffer::AllocLinearTextures(const LinearTextureMapKey& InLinearTextu
 		if (FAGXCommandQueue::SupportsFeature(EAGXFeaturesTextureBuffers))
 		{
 			NSUInteger Options = ((NSUInteger) Mode) << mtlpp::ResourceStorageModeShift;
-			Desc = mtlpp::TextureDescriptor::TextureBufferDescriptor(MTLFormat, NumElements, mtlpp::ResourceOptions(Options), mtlpp::TextureUsage(TexUsage));
+			Desc = mtlpp::TextureDescriptor::TextureBufferDescriptor((mtlpp::PixelFormat)MTLFormat, NumElements, mtlpp::ResourceOptions(Options), mtlpp::TextureUsage(TexUsage));
 			Desc.SetAllowGPUOptimisedContents(false);
 		}
 		else
@@ -422,7 +422,7 @@ void FAGXRHIBuffer::AllocLinearTextures(const LinearTextureMapKey& InLinearTextu
 			check(RowBytes % MinimumByteAlignment == 0);
 			check((RowBytes * Height) + Offset <= Length);
 
-			Desc = mtlpp::TextureDescriptor::Texture2DDescriptor(MTLFormat, Width, Height, NO);
+			Desc = mtlpp::TextureDescriptor::Texture2DDescriptor((mtlpp::PixelFormat)MTLFormat, Width, Height, NO);
 			Desc.SetStorageMode(Mode);
 			Desc.SetCpuCacheMode(CurrentBuffer.GetCpuCacheMode());
 			Desc.SetUsage((mtlpp::TextureUsage)TexUsage);
@@ -434,7 +434,7 @@ void FAGXRHIBuffer::AllocLinearTextures(const LinearTextureMapKey& InLinearTextu
 			FAGXTexture NewTexture = MTLPP_VALIDATE(mtlpp::Buffer, Buffer, AGXSafeGetRuntimeDebuggingLevel() >= EAGXDebugLevelValidation, NewTexture(Desc, Offset, RowBytes));
 			METAL_FATAL_ASSERT(NewTexture, TEXT("Failed to create linear texture, desc %s from buffer %s"), *FString([Desc description]), *FString([Buffer description]));
 			
-			check(GAGXBufferFormats[InFormat].LinearTextureFormat == mtlpp::PixelFormat::RG11B10Float || GAGXBufferFormats[InFormat].LinearTextureFormat == (mtlpp::PixelFormat)NewTexture.GetPixelFormat());
+			check(GAGXBufferFormats[InFormat].LinearTextureFormat == MTLPixelFormatRG11B10Float || GAGXBufferFormats[InFormat].LinearTextureFormat == (MTLPixelFormat)NewTexture.GetPixelFormat());
 			Backing.Views.Add(InLinearTextureMapKey, NewTexture);
 		}
 	}
@@ -478,7 +478,7 @@ struct FAGXRHICommandCreateLinearTexture : public FRHICommand<FAGXRHICommandCrea
 
 void FAGXRHIBuffer::CreateLinearTexture(EPixelFormat InFormat, FRHIResource* InParent, const FAGXLinearTextureDescriptor* InLinearTextureDescriptor)
 {
-	if (EnumHasAnyFlags(Usage, BUF_UnorderedAccess|BUF_ShaderResource) && GAGXBufferFormats[InFormat].LinearTextureFormat != mtlpp::PixelFormat::Invalid)
+	if (EnumHasAnyFlags(Usage, BUF_UnorderedAccess|BUF_ShaderResource) && GAGXBufferFormats[InFormat].LinearTextureFormat != MTLPixelFormatInvalid)
 	{
 		if (IsRunningRHIInSeparateThread() && !IsInRHIThread() && !FRHICommandListExecutor::GetImmediateCommandList().Bypass())
 		{
@@ -504,7 +504,7 @@ void FAGXRHIBuffer::CreateLinearTexture(EPixelFormat InFormat, FRHIResource* InP
 ns::AutoReleased<FAGXTexture> FAGXRHIBuffer::GetLinearTexture(EPixelFormat InFormat, const FAGXLinearTextureDescriptor* InLinearTextureDescriptor)
 {
 	ns::AutoReleased<FAGXTexture> Texture;
-	if (EnumHasAnyFlags(Usage, BUF_UnorderedAccess|BUF_ShaderResource) && GAGXBufferFormats[InFormat].LinearTextureFormat != mtlpp::PixelFormat::Invalid)
+	if (EnumHasAnyFlags(Usage, BUF_UnorderedAccess|BUF_ShaderResource) && GAGXBufferFormats[InFormat].LinearTextureFormat != MTLPixelFormatInvalid)
 	{
 		LinearTextureMapKey MapKey = (InLinearTextureDescriptor != nullptr) ? LinearTextureMapKey(InFormat, *InLinearTextureDescriptor) : LinearTextureMapKey(InFormat, FAGXLinearTextureDescriptor());
 
diff --git a/Engine/Source/Runtime/Apple/AGXRHI/Private/AGXViewport.cpp b/Engine/Source/Runtime/Apple/AGXRHI/Private/AGXViewport.cpp
--- a/Engine/Source/Runtime/Apple/AGXRHI/Private/AGXViewport.cpp
+++ b/Engine/Source/Runtime/Apple/AGXRHI/Private/AGXViewport.cpp
@@ -178,7 +178,7 @@ void FAGXViewport::Resize(uint32 InSizeX, uint32 InSizeY, bool bInIsFullscreen,E
 		Format = PF_B8G8R8A8;
 	}
 	
-	mtlpp::PixelFormat MetalFormat = (mtlpp::PixelFormat)GPixelFormats[Format].PlatformFormat;
+	MTLPixelFormat MetalFormat = (MTLPixelFormat)GPixelFormats[Format].PlatformFormat;
 	
 	if (IsValidRef(BackBuffer[Index]) && Format != BackBuffer[Index]->GetFormat())
 	{
@@ -205,9 +205,9 @@ void FAGXViewport::Resize(uint32 InSizeX, uint32 InSizeY, bool bInIsFullscreen,E
 		
 		MetalLayer.drawableSize = CGSizeMake(InSizeX, InSizeY);
 		
-		if (MetalFormat != (mtlpp::PixelFormat)MetalLayer.pixelFormat)
+		if (MetalFormat != MetalLayer.pixelFormat)
 		{
-			MetalLayer.pixelFormat = (MTLPixelFormat)MetalFormat;
+			MetalLayer.pixelFormat = MetalFormat;
 		}
 		
 		if (bUseHDR != MetalLayer.wantsExtendedDynamicRangeContent)
@@ -227,9 +227,9 @@ void FAGXViewport::Resize(uint32 InSizeX, uint32 InSizeY, bool bInIsFullscreen,E
 		FIOSView* IOSView = AppDelegate.IOSView;
 		CAMetalLayer* MetalLayer = (CAMetalLayer*) IOSView.layer;
 		
-		if (MetalFormat != (mtlpp::PixelFormat) MetalLayer.pixelFormat)
+		if (MetalFormat != MetalLayer.pixelFormat)
 		{
-			MetalLayer.pixelFormat = (MTLPixelFormat) MetalFormat;
+			MetalLayer.pixelFormat = MetalFormat;
 		}
 		
 		[IOSView UpdateRenderWidth:InSizeX andHeight:InSizeY];
diff --git a/Engine/Source/Runtime/Apple/AGXRHI/Public/AGXResources.h b/Engine/Source/Runtime/Apple/AGXRHI/Public/AGXResources.h
--- a/Engine/Source/Runtime/Apple/AGXRHI/Public/AGXResources.h
+++ b/Engine/Source/Runtime/Apple/AGXRHI/Public/AGXResources.h
@@ -138,7 +138,7 @@ struct FAGXTextureDesc
 	FAGXTextureDesc(FRHITextureDesc const& InDesc);
 	
 	mtlpp::TextureDescriptor Desc;
-	mtlpp::PixelFormat MTLFormat;
+	MTLPixelFormat MTLFormat;
 	bool bMemoryless = false;
 	bool bIsRenderTarget = false;
 	uint8 FormatKey = 0;
