AGXRHI: Remove/replace mtlpp::ResourceOptions with native Metal type

- Also, remove PermittedOptions and unnecessary option filtering.
- Also, remove texture buffer conditions, fallback and feature flag. All Metal platforms now support this and it is by far the preferred way of doing things.

#rnx

diff --git a/Engine/Source/Runtime/Apple/AGXRHI/Private/AGXBuffer.cpp b/Engine/Source/Runtime/Apple/AGXRHI/Private/AGXBuffer.cpp
--- a/Engine/Source/Runtime/Apple/AGXRHI/Private/AGXBuffer.cpp
+++ b/Engine/Source/Runtime/Apple/AGXRHI/Private/AGXBuffer.cpp
@@ -164,22 +164,21 @@ void FAGXBuffer::SetOwner(class FAGXRHIBuffer* Owner, bool bIsSwap)
 	}
 }
 
-FAGXSubBufferHeap::FAGXSubBufferHeap(NSUInteger Size, NSUInteger Alignment, mtlpp::ResourceOptions Options, FCriticalSection& InPoolMutex)
+FAGXSubBufferHeap::FAGXSubBufferHeap(NSUInteger Size, NSUInteger Alignment, MTLResourceOptions Options, FCriticalSection& InPoolMutex)
 : PoolMutex(InPoolMutex)
 , OutstandingAllocs(0)
 , MinAlign(Alignment)
 , UsedSize(0)
 {
-	Options = (mtlpp::ResourceOptions)FAGXCommandQueue::GetCompatibleResourceOptions(Options);
 	NSUInteger FullSize = Align(Size, Alignment);
 	METAL_GPUPROFILE(FAGXScopedCPUStats CPUStat(FString::Printf(TEXT("AllocBuffer: %llu, %llu"), FullSize, Options)));
 	
-	mtlpp::StorageMode Storage = (mtlpp::StorageMode)((Options & mtlpp::ResourceStorageModeMask) >> mtlpp::ResourceStorageModeShift);
 #if PLATFORM_MAC
-	check(Storage != mtlpp::StorageMode::Managed /* Managed memory cannot be safely suballocated! When you overwrite existing data the GPU buffer is immediately disposed of! */);
+	 // Managed resources cannot be safely suballocated.  When you overwrite existing data the GPU buffer is immediately disposed of.
+	check(((Options & MTLResourceStorageModeMask) >> MTLResourceStorageModeShift) != MTLStorageModeManaged);
 #endif
 
-	ParentBuffer = MTLPP_VALIDATE(mtlpp::Device, GMtlppDevice, AGXSafeGetRuntimeDebuggingLevel() >= EAGXDebugLevelValidation, NewBuffer(FullSize, Options));
+	ParentBuffer = MTLPP_VALIDATE(mtlpp::Device, GMtlppDevice, AGXSafeGetRuntimeDebuggingLevel() >= EAGXDebugLevelValidation, NewBuffer(FullSize, mtlpp::ResourceOptions(Options)));
 	check(ParentBuffer.GetPtr());
 	check(ParentBuffer.GetLength() >= FullSize);
 #if STATS || ENABLE_LOW_LEVEL_MEM_TRACKER
@@ -297,16 +296,6 @@ ns::String FAGXSubBufferHeap::GetLabel() const
 	return ParentBuffer.GetLabel();
 }
 
-mtlpp::StorageMode FAGXSubBufferHeap::GetStorageMode() const
-{
-	return ParentBuffer.GetStorageMode();
-}
-
-mtlpp::CpuCacheMode FAGXSubBufferHeap::GetCpuCacheMode() const
-{
-	return ParentBuffer.GetCpuCacheMode();
-}
-
 NSUInteger FAGXSubBufferHeap::GetSize() const
 {
 	return ParentBuffer.GetLength();
@@ -404,26 +393,24 @@ FAGXBuffer FAGXSubBufferHeap::NewBuffer(NSUInteger length)
 	return Result;
 }
 
-mtlpp::PurgeableState FAGXSubBufferHeap::SetPurgeableState(mtlpp::PurgeableState state)
+MTLPurgeableState FAGXSubBufferHeap::SetPurgeableState(MTLPurgeableState State)
 {
-	return ParentBuffer.SetPurgeableState(state);
+	return [ParentBuffer.GetPtr() setPurgeableState:State];
 }
 
 #pragma mark --
 
-FAGXSubBufferLinear::FAGXSubBufferLinear(NSUInteger Size, NSUInteger Alignment, mtlpp::ResourceOptions Options, FCriticalSection& InPoolMutex)
+FAGXSubBufferLinear::FAGXSubBufferLinear(NSUInteger Size, NSUInteger Alignment, MTLResourceOptions Options, FCriticalSection& InPoolMutex)
 : PoolMutex(InPoolMutex)
 , MinAlign(Alignment)
 , WriteHead(0)
 , UsedSize(0)
 , FreedSize(0)
 {
-	Options = (mtlpp::ResourceOptions)FAGXCommandQueue::GetCompatibleResourceOptions(Options);
 	NSUInteger FullSize = Align(Size, Alignment);
 	METAL_GPUPROFILE(FAGXScopedCPUStats CPUStat(FString::Printf(TEXT("AllocBuffer: %llu, %llu"), FullSize, Options)));
 	
-	mtlpp::StorageMode Storage = (mtlpp::StorageMode)((Options & mtlpp::ResourceStorageModeMask) >> mtlpp::ResourceStorageModeShift);
-	ParentBuffer = MTLPP_VALIDATE(mtlpp::Device, GMtlppDevice, AGXSafeGetRuntimeDebuggingLevel() >= EAGXDebugLevelValidation, NewBuffer(FullSize, Options));
+	ParentBuffer = MTLPP_VALIDATE(mtlpp::Device, GMtlppDevice, AGXSafeGetRuntimeDebuggingLevel() >= EAGXDebugLevelValidation, NewBuffer(FullSize, mtlpp::ResourceOptions(Options)));
 	check(ParentBuffer.GetPtr());
 	check(ParentBuffer.GetLength() >= FullSize);
 #if STATS || ENABLE_LOW_LEVEL_MEM_TRACKER
@@ -470,16 +457,6 @@ ns::String FAGXSubBufferLinear::GetLabel() const
 	return ParentBuffer.GetLabel();
 }
 
-mtlpp::StorageMode FAGXSubBufferLinear::GetStorageMode() const
-{
-	return ParentBuffer.GetStorageMode();
-}
-
-mtlpp::CpuCacheMode FAGXSubBufferLinear::GetCpuCacheMode() const
-{
-	return ParentBuffer.GetCpuCacheMode();
-}
-
 NSUInteger FAGXSubBufferLinear::GetSize() const
 {
 	return ParentBuffer.GetLength();
@@ -531,29 +508,28 @@ FAGXBuffer FAGXSubBufferLinear::NewBuffer(NSUInteger length)
 	return Result;
 }
 
-mtlpp::PurgeableState FAGXSubBufferLinear::SetPurgeableState(mtlpp::PurgeableState state)
+MTLPurgeableState FAGXSubBufferLinear::SetPurgeableState(MTLPurgeableState State)
 {
-	return ParentBuffer.SetPurgeableState(state);
+	return [ParentBuffer.GetPtr() setPurgeableState:State];
 }
 
 #pragma mark --
 
-FAGXSubBufferMagazine::FAGXSubBufferMagazine(NSUInteger Size, NSUInteger ChunkSize, mtlpp::ResourceOptions Options)
+FAGXSubBufferMagazine::FAGXSubBufferMagazine(NSUInteger Size, NSUInteger ChunkSize, MTLResourceOptions Options)
 : MinAlign(ChunkSize)
 , BlockSize(ChunkSize)
 , OutstandingAllocs(0)
 , UsedSize(0)
 {
-	Options = (mtlpp::ResourceOptions)FAGXCommandQueue::GetCompatibleResourceOptions(Options);
-	mtlpp::StorageMode Storage = (mtlpp::StorageMode)((Options & mtlpp::ResourceStorageModeMask) >> mtlpp::ResourceStorageModeShift);
 	NSUInteger FullSize = Align(Size, MinAlign);
 	METAL_GPUPROFILE(FAGXScopedCPUStats CPUStat(FString::Printf(TEXT("AllocBuffer: %llu, %llu"), FullSize, Options)));
 	
 #if PLATFORM_MAC
-	check(Storage != mtlpp::StorageMode::Managed /* Managed memory cannot be safely suballocated! When you overwrite existing data the GPU buffer is immediately disposed of! */);
+	 // Managed resources cannot be safely suballocated.  When you overwrite existing data the GPU buffer is immediately disposed of.
+	check(((Options & MTLResourceStorageModeMask) >> MTLResourceStorageModeShift) != MTLStorageModeManaged);
 #endif
 
-	ParentBuffer = MTLPP_VALIDATE(mtlpp::Device, GMtlppDevice, AGXSafeGetRuntimeDebuggingLevel() >= EAGXDebugLevelValidation, NewBuffer(FullSize, Options));
+	ParentBuffer = MTLPP_VALIDATE(mtlpp::Device, GMtlppDevice, AGXSafeGetRuntimeDebuggingLevel() >= EAGXDebugLevelValidation, NewBuffer(FullSize, mtlpp::ResourceOptions(Options)));
 	check(ParentBuffer.GetPtr());
 	check(ParentBuffer.GetLength() >= FullSize);
 	METAL_FATAL_ASSERT(ParentBuffer, TEXT("Failed to create heap of size %u and resource options %u"), Size, (uint32)Options);
@@ -600,16 +576,6 @@ ns::String FAGXSubBufferMagazine::GetLabel() const
 	return ParentBuffer.GetLabel();
 }
 
-mtlpp::StorageMode FAGXSubBufferMagazine::GetStorageMode() const
-{
-	return ParentBuffer.GetStorageMode();
-}
-
-mtlpp::CpuCacheMode FAGXSubBufferMagazine::GetCpuCacheMode() const
-{
-	return ParentBuffer.GetCpuCacheMode();
-}
-
 NSUInteger FAGXSubBufferMagazine::GetSize() const
 {
 	return ParentBuffer.GetLength();
@@ -666,9 +632,9 @@ FAGXBuffer FAGXSubBufferMagazine::NewBuffer()
 	return Result;
 }
 
-mtlpp::PurgeableState FAGXSubBufferMagazine::SetPurgeableState(mtlpp::PurgeableState state)
+MTLPurgeableState FAGXSubBufferMagazine::SetPurgeableState(MTLPurgeableState State)
 {
-	return ParentBuffer.SetPurgeableState(state);
+	return [ParentBuffer.GetPtr() setPurgeableState:State];
 }
 
 FAGXRingBufferRef::FAGXRingBufferRef(FAGXBuffer Buf)
@@ -684,7 +650,7 @@ FAGXRingBufferRef::~FAGXRingBufferRef()
 	AGXSafeReleaseMetalBuffer(Buffer);
 }
 
-FAGXSubBufferRing::FAGXSubBufferRing(NSUInteger Size, NSUInteger Alignment, mtlpp::ResourceOptions InOptions)
+FAGXSubBufferRing::FAGXSubBufferRing(NSUInteger Size, NSUInteger Alignment, MTLResourceOptions InOptions)
 : LastFrameChange(0)
 , InitialSize(Align(Size, Alignment))
 , MinAlign(Alignment)
@@ -693,10 +659,10 @@ FAGXSubBufferRing::FAGXSubBufferRing(NSUInteger Size, NSUInteger Alignment, mtlp
 , WriteHead(0)
 , BufferSize(0)
 , Options(InOptions)
-, Storage((mtlpp::StorageMode)((Options & mtlpp::ResourceStorageModeMask) >> mtlpp::ResourceStorageModeShift))
 {
-	Options = (mtlpp::ResourceOptions)FAGXCommandQueue::GetCompatibleResourceOptions(Options);
-	check(Storage != mtlpp::StorageMode::Private /* Private memory requires command-buffers and encoders to properly marshal! */);
+	// Private resources require command buffers and encoders to properly marshal.
+	check(0 == (Options & MTLResourceStorageModePrivate));
+
 	FMemory::Memzero(FrameSize);
 }
 
@@ -704,14 +670,6 @@ FAGXSubBufferRing::~FAGXSubBufferRing()
 {
 }
 
-mtlpp::StorageMode FAGXSubBufferRing::GetStorageMode() const
-{
-	return Buffer.IsValid() ? Buffer->Buffer.GetStorageMode() : Storage;
-}
-mtlpp::CpuCacheMode FAGXSubBufferRing::GetCpuCacheMode() const
-{
-	return Buffer.IsValid() ? Buffer->Buffer.GetCpuCacheMode() : ((mtlpp::CpuCacheMode)((Options & mtlpp::ResourceCpuCacheModeMask) >> mtlpp::ResourceCpuCacheModeShift));
-}
 NSUInteger FAGXSubBufferRing::GetSize() const
 {
 	return Buffer.IsValid() ? Buffer->Buffer.GetLength() : InitialSize;
@@ -750,7 +708,7 @@ FAGXBuffer FAGXSubBufferRing::NewBuffer(NSUInteger Size, uint32 Alignment)
 			return NewBuffer;
 		}
 #if PLATFORM_MAC
-		else if (Storage == mtlpp::StorageMode::Managed)
+		else if (0 != (Options & MTLResourceStorageModeManaged))
 		{
 			Submit();
 			Buffer = MakeShared<FAGXRingBufferRef, ESPMode::ThreadSafe>(GetAGXDeviceContext().GetResourceHeap().CreateBuffer(BufferSize, MinAlign, BUF_Dynamic, Options, true));
@@ -828,7 +786,7 @@ void FAGXSubBufferRing::Submit()
 	if (Buffer.IsValid() && WriteHead != SubmitHead)
 	{
 #if PLATFORM_MAC
-		if (Storage == mtlpp::StorageMode::Managed)
+		if (0 != (Options & MTLResourceStorageModeManaged))
 		{
 			check(SubmitHead < WriteHead);
 			ns::Range ModifiedRange(SubmitHead, Align(WriteHead - SubmitHead, MinAlign));
@@ -845,7 +803,7 @@ void FAGXSubBufferRing::Commit(mtlpp::CommandBuffer& CmdBuf)
 	if (Buffer.IsValid() && WriteHead != CommitHead)
 	{
 #if PLATFORM_MAC
-		check(Storage != mtlpp::StorageMode::Managed || CommitHead < WriteHead);
+		check((0 == (Options & MTLResourceStorageModeManaged)) || CommitHead < WriteHead);
 #endif
 		Submit();
 		
@@ -926,8 +884,11 @@ uint32 FAGXBufferPoolPolicyData::GetPoolBucketIndex(CreationArguments Args)
 	
 	check( Size <= BucketSizes[Lower] );
 	check( (Lower == 0 ) || ( Size > BucketSizes[Lower-1] ) );
-	
-	return (int32)Args.Storage * NumPoolBucketSizes + Lower;
+
+	NSUInteger Storage = (Args.Options & MTLResourceStorageModeMask) >> MTLResourceStorageModeShift;
+	check(Storage < NumResourceStorageModes);
+
+	return (int32)(Storage * NumPoolBucketSizes + Lower);
 }
 
 uint32 FAGXBufferPoolPolicyData::GetPoolBucketSize(uint32 Bucket)
@@ -940,13 +901,9 @@ uint32 FAGXBufferPoolPolicyData::GetPoolBucketSize(uint32 Bucket)
 FAGXBuffer FAGXBufferPoolPolicyData::CreateResource(CreationArguments Args)
 {
 	uint32 BufferSize = GetPoolBucketSize(GetPoolBucketIndex(Args));
-	
-	NSUInteger CpuCacheMode = (NSUInteger)Args.CpuCacheMode << mtlpp::ResourceCpuCacheModeShift;
-	NSUInteger StorageMode = (NSUInteger)Args.Storage << mtlpp::ResourceStorageModeShift;
-	mtlpp::ResourceOptions ResourceOptions = FAGXCommandQueue::GetCompatibleResourceOptions(mtlpp::ResourceOptions(CpuCacheMode | StorageMode | mtlpp::ResourceOptions::HazardTrackingModeUntracked));
-	
-	METAL_GPUPROFILE(FAGXScopedCPUStats CPUStat(FString::Printf(TEXT("AllocBuffer: %llu, %llu"), BufferSize, ResourceOptions)));
-	FAGXBuffer NewBuf(MTLPP_VALIDATE(mtlpp::Device, GMtlppDevice, AGXSafeGetRuntimeDebuggingLevel() >= EAGXDebugLevelValidation, NewBuffer(BufferSize, ResourceOptions), true));
+
+	METAL_GPUPROFILE(FAGXScopedCPUStats CPUStat(FString::Printf(TEXT("AllocBuffer: %llu, %llu"), BufferSize, Args.Options)));
+	FAGXBuffer NewBuf(MTLPP_VALIDATE(mtlpp::Device, GMtlppDevice, AGXSafeGetRuntimeDebuggingLevel() >= EAGXDebugLevelValidation, NewBuffer(BufferSize, mtlpp::ResourceOptions(Args.Options)), true));
 
 #if STATS || ENABLE_LOW_LEVEL_MEM_TRACKER
 	AGXLLM::LogAllocBuffer(NewBuf);
@@ -958,7 +915,7 @@ FAGXBuffer FAGXBufferPoolPolicyData::CreateResource(CreationArguments Args)
 
 FAGXBufferPoolPolicyData::CreationArguments FAGXBufferPoolPolicyData::GetCreationArguments(FAGXBuffer const& Resource)
 {
-	return FAGXBufferPoolPolicyData::CreationArguments(Resource.GetLength(), BUF_None, Resource.GetStorageMode(), Resource.GetCpuCacheMode());
+	return FAGXBufferPoolPolicyData::CreationArguments(Resource.GetLength(), BUF_None, [Resource.GetPtr() resourceOptions]);
 }
 
 void FAGXBufferPoolPolicyData::FreeResource(FAGXBuffer& Resource)
@@ -1030,7 +987,7 @@ void FAGXTexturePool::ReleaseTexture(FAGXTexture& Texture)
 	Descriptor.mipmapLevelCount = Texture.GetMipmapLevelCount();
 	Descriptor.sampleCount = Texture.GetSampleCount();
 	Descriptor.arrayLength = Texture.GetArrayLength();
-	Descriptor.resourceOptions = ((NSUInteger)Texture.GetStorageMode() << mtlpp::ResourceStorageModeShift) | ((NSUInteger)Texture.GetCpuCacheMode() << mtlpp::ResourceCpuCacheModeShift);
+	Descriptor.resourceOptions = [Texture.GetPtr() resourceOptions];
 	Descriptor.usage = Texture.GetUsage();
 	Descriptor.freedFrame = GFrameNumberRenderThread;
 	
@@ -1164,7 +1121,7 @@ FAGXResourceHeap::TextureHeapSize FAGXResourceHeap::TextureSizeToIndex(uint32 Si
 	return (TextureHeapSize)Lower;
 }
 
-FAGXBuffer FAGXResourceHeap::CreateBuffer(uint32 Size, uint32 Alignment, EBufferUsageFlags Flags, mtlpp::ResourceOptions Options, bool bForceUnique)
+FAGXBuffer FAGXResourceHeap::CreateBuffer(uint32 Size, uint32 Alignment, EBufferUsageFlags Flags, MTLResourceOptions Options, bool bForceUnique)
 {
 	LLM_SCOPE_METAL(ELLMTagAGX::Buffers);
 	LLM_PLATFORM_SCOPE_METAL(ELLMTagAGX::Buffers);
@@ -1176,18 +1133,17 @@ FAGXBuffer FAGXResourceHeap::CreateBuffer(uint32 Size, uint32 Alignment, EBuffer
 	
 	FAGXBuffer Buffer;
 	uint32 BlockSize = Align(Size, Alignment);
-	mtlpp::StorageMode StorageMode = (mtlpp::StorageMode)(((NSUInteger)Options & mtlpp::ResourceStorageModeMask) >> mtlpp::ResourceStorageModeShift);
-	mtlpp::CpuCacheMode CpuMode = (mtlpp::CpuCacheMode)(((NSUInteger)Options & mtlpp::ResourceCpuCacheModeMask) >> mtlpp::ResourceCpuCacheModeShift);
 	
 	// Write combined should be on a case by case basis
-	check(CpuMode == mtlpp::CpuCacheMode::DefaultCache);
+	check(((Options & MTLResourceCPUCacheModeMask) >> MTLResourceCPUCacheModeShift) == MTLCPUCacheModeDefaultCache);
 	
 	if (BlockSize <= 33554432)
 	{
+		MTLStorageMode StorageMode = (MTLStorageMode)((Options & MTLResourceStorageModeMask) >> MTLResourceStorageModeShift);
 		switch (StorageMode)
 		{
 	#if PLATFORM_MAC
-			case mtlpp::StorageMode::Managed:
+			case MTLStorageModeManaged:
 			{
 				// TextureBuffers must be 1024 aligned.
 				check(Alignment == 256 || Alignment == 1024);
@@ -1207,7 +1163,7 @@ FAGXBuffer FAGXResourceHeap::CreateBuffer(uint32 Size, uint32 Alignment, EBuffer
 				 	}
 				 	if (!Found)
 				 	{
-				 		Found = new FAGXSubBufferLinear(HeapAllocSizes[NumHeapSizes - 1], BufferOffsetAlignment, mtlpp::ResourceOptions((NSUInteger)Options & (mtlpp::ResourceStorageModeMask|mtlpp::ResourceHazardTrackingModeMask)), Mutex);
+				 		Found = new FAGXSubBufferLinear(HeapAllocSizes[NumHeapSizes - 1], BufferOffsetAlignment, Options, Mutex);
 				 		ManagedSubHeaps.Add(Found);
 				 	}
 				 	check(Found);
@@ -1216,7 +1172,7 @@ FAGXBuffer FAGXResourceHeap::CreateBuffer(uint32 Size, uint32 Alignment, EBuffer
 				 }
 				 else
 				 {
-					Buffer = ManagedBuffers.CreatePooledResource(FAGXPooledBufferArgs(BlockSize, Flags, StorageMode, CpuMode));
+					Buffer = ManagedBuffers.CreatePooledResource(FAGXPooledBufferArgs(BlockSize, Flags, Options));
 					if (GAGXResourcePurgeInPool)
 					{
 						Buffer.SetPurgeableState(mtlpp::PurgeableState::NonVolatile);
@@ -1228,10 +1184,10 @@ FAGXBuffer FAGXResourceHeap::CreateBuffer(uint32 Size, uint32 Alignment, EBuffer
 				break;
 			}
 	#endif
-			case mtlpp::StorageMode::Private:
-			case mtlpp::StorageMode::Shared:
+			case MTLStorageModePrivate:
+			case MTLStorageModeShared:
 			{
-				AllocTypes Storage = StorageMode != mtlpp::StorageMode::Private ? AllocShared : AllocPrivate;
+				AllocTypes Storage = StorageMode != MTLStorageModePrivate ? AllocShared : AllocPrivate;
 				check(Alignment == 16 || Alignment == 64 || Alignment == 256 || Alignment == 1024);
 
 				static bool bSupportsPrivateBufferSubAllocation = FAGXCommandQueue::SupportsFeature(EAGXFeaturesPrivateBufferSubAllocation);
@@ -1254,7 +1210,7 @@ FAGXBuffer FAGXResourceHeap::CreateBuffer(uint32 Size, uint32 Alignment, EBuffer
 					
 					if (!Found)
 					{
-						Found = new FAGXSubBufferMagazine(MagazineAllocSizes[i], MagazineSizes[i], mtlpp::ResourceOptions((NSUInteger)Options & (mtlpp::ResourceStorageModeMask|mtlpp::ResourceHazardTrackingModeMask)));
+						Found = new FAGXSubBufferMagazine(MagazineAllocSizes[i], MagazineSizes[i], Options);
 						SmallBuffers[Usage][Storage][i].Add(Found);
 					}
 					check(Found);
@@ -1282,7 +1238,7 @@ FAGXBuffer FAGXResourceHeap::CreateBuffer(uint32 Size, uint32 Alignment, EBuffer
 					if (!Found)
 					{
 						uint32 MinAlign = PLATFORM_MAC ? 1024 : 64;
-						Found = new FAGXSubBufferHeap(HeapAllocSizes[i], MinAlign, mtlpp::ResourceOptions((NSUInteger)Options & (mtlpp::ResourceStorageModeMask|mtlpp::ResourceHazardTrackingModeMask)), Mutex);
+						Found = new FAGXSubBufferHeap(HeapAllocSizes[i], MinAlign, Options, Mutex);
 						BufferHeaps[Usage][Storage][i].Add(Found);
 					}
 					check(Found);
@@ -1293,7 +1249,7 @@ FAGXBuffer FAGXResourceHeap::CreateBuffer(uint32 Size, uint32 Alignment, EBuffer
 				else
 				{
 					FScopeLock Lock(&Mutex);
-					Buffer = Buffers[Storage].CreatePooledResource(FAGXPooledBufferArgs(BlockSize, Flags, StorageMode, CpuMode));
+					Buffer = Buffers[Storage].CreatePooledResource(FAGXPooledBufferArgs(BlockSize, Flags, Options));
 					if (GAGXResourcePurgeInPool)
 					{
                    		Buffer.SetPurgeableState(mtlpp::PurgeableState::NonVolatile);
@@ -1314,7 +1270,7 @@ FAGXBuffer FAGXResourceHeap::CreateBuffer(uint32 Size, uint32 Alignment, EBuffer
 	else
 	{
 		METAL_GPUPROFILE(FAGXScopedCPUStats CPUStat(FString::Printf(TEXT("AllocBuffer: %llu, %llu"), BlockSize, Options)));
-		mtlpp::Buffer MtlppBuffer([GMtlDevice newBufferWithLength:BlockSize options:(MTLResourceOptions)Options], nullptr, ns::Ownership::Assign);
+		mtlpp::Buffer MtlppBuffer([GMtlDevice newBufferWithLength:BlockSize options:Options], nullptr, ns::Ownership::Assign);
 		Buffer = FAGXBuffer(MoveTemp(MtlppBuffer), false);
 #if STATS || ENABLE_LOW_LEVEL_MEM_TRACKER
 		AGXLLM::LogAllocBuffer(Buffer);
@@ -1322,7 +1278,7 @@ FAGXBuffer FAGXResourceHeap::CreateBuffer(uint32 Size, uint32 Alignment, EBuffer
 		INC_MEMORY_STAT_BY(STAT_AGXDeviceBufferMemory, Buffer.GetLength());
 	}
 	
-	if (GAGXBufferZeroFill && Buffer.GetStorageMode() != mtlpp::StorageMode::Private)
+	if (GAGXBufferZeroFill && [Buffer.GetPtr() storageMode] != MTLStorageModePrivate)
 	{
 		FMemory::Memset(((uint8*)Buffer.GetContents()), 0, Buffer.GetLength());
 	}
@@ -1334,7 +1290,7 @@ FAGXBuffer FAGXResourceHeap::CreateBuffer(uint32 Size, uint32 Alignment, EBuffer
 
 void FAGXResourceHeap::ReleaseBuffer(FAGXBuffer& Buffer)
 {
-	mtlpp::StorageMode StorageMode = Buffer.GetStorageMode();
+	MTLStorageMode StorageMode = [Buffer.GetPtr() storageMode];
 	if (Buffer.IsPooled())
 	{
 		FScopeLock Lock(&Mutex);
@@ -1351,14 +1307,14 @@ void FAGXResourceHeap::ReleaseBuffer(FAGXBuffer& Buffer)
 		switch (StorageMode)
 		{
 	#if PLATFORM_MAC
-			case mtlpp::StorageMode::Managed:
+			case MTLStorageModeManaged:
 			{
 				ManagedBuffers.ReleasePooledResource(Buffer);
 				break;
 			}
 	#endif
-			case mtlpp::StorageMode::Private:
-			case mtlpp::StorageMode::Shared:
+			case MTLStorageModePrivate:
+			case MTLStorageModeShared:
 			{
 				Buffers[(NSUInteger)StorageMode].ReleasePooledResource(Buffer);
 				break;
diff --git a/Engine/Source/Runtime/Apple/AGXRHI/Private/AGXBuffer.h b/Engine/Source/Runtime/Apple/AGXRHI/Private/AGXBuffer.h
--- a/Engine/Source/Runtime/Apple/AGXRHI/Private/AGXBuffer.h
+++ b/Engine/Source/Runtime/Apple/AGXRHI/Private/AGXBuffer.h
@@ -11,20 +11,21 @@
 
 struct FAGXPooledBufferArgs
 {
-    FAGXPooledBufferArgs() : Size(0), Flags(BUF_None), Storage(mtlpp::StorageMode::Shared), CpuCacheMode(mtlpp::CpuCacheMode::DefaultCache) {}
+	static const MTLResourceOptions SharedStorageResourceOptions  = MTLResourceCPUCacheModeDefaultCache | MTLResourceStorageModeShared  | MTLResourceHazardTrackingModeDefault;
+	static const MTLResourceOptions PrivateStorageResourceOptions = MTLResourceCPUCacheModeDefaultCache | MTLResourceStorageModePrivate | MTLResourceHazardTrackingModeDefault;
+
+    FAGXPooledBufferArgs() : Size(0), Flags(BUF_None), Options(0) {}
 	
-    FAGXPooledBufferArgs(uint32 InSize, EBufferUsageFlags InFlags, mtlpp::StorageMode InStorage, mtlpp::CpuCacheMode InCpuCacheMode = mtlpp::CpuCacheMode::DefaultCache)
+    FAGXPooledBufferArgs(uint32 InSize, EBufferUsageFlags InFlags, MTLResourceOptions InOptions)
 	: Size(InSize)
     , Flags(InFlags)
-	, Storage(InStorage)
-	, CpuCacheMode(InCpuCacheMode)
+	, Options(InOptions)
 	{
 	}
 	
 	uint32 Size;
 	EBufferUsageFlags Flags;
-	mtlpp::StorageMode Storage;
-	mtlpp::CpuCacheMode CpuCacheMode;
+	MTLResourceOptions Options;
 };
 
 class FAGXSubBufferHeap
@@ -32,25 +33,24 @@ class FAGXSubBufferHeap
     friend class FAGXResourceHeap;
     
 public:
-	FAGXSubBufferHeap(NSUInteger Size, NSUInteger Alignment, mtlpp::ResourceOptions, FCriticalSection& PoolMutex);
+	FAGXSubBufferHeap(NSUInteger Size, NSUInteger Alignment, MTLResourceOptions Options, FCriticalSection& PoolMutex);
 	~FAGXSubBufferHeap();
-	
-	ns::String   GetLabel() const;
-    mtlpp::StorageMode  GetStorageMode() const;
-    mtlpp::CpuCacheMode GetCpuCacheMode() const;
-    NSUInteger     GetSize() const;
-    NSUInteger     GetUsedSize() const;
-	NSUInteger	 MaxAvailableSize() const;
-	int64     NumCurrentAllocations() const;
-    bool     CanAllocateSize(NSUInteger Size) const;
 
-    void SetLabel(const ns::String& label);
-	
-    FAGXBuffer NewBuffer(NSUInteger length);
-    mtlpp::PurgeableState SetPurgeableState(mtlpp::PurgeableState state);
-	void FreeRange(ns::Range const& Range);
+	ns::String			GetLabel() const;
+    MTLResourceOptions	GetResourceOptions() const;
+    NSUInteger			GetSize() const;
+    NSUInteger			GetUsedSize() const;
+	NSUInteger			MaxAvailableSize() const;
+	int64				NumCurrentAllocations() const;
+    bool				CanAllocateSize(NSUInteger Size) const;
+
+    void				SetLabel(const ns::String& label);
+
+    FAGXBuffer			NewBuffer(NSUInteger length);
+    MTLPurgeableState	SetPurgeableState(MTLPurgeableState State);
+	void				FreeRange(ns::Range const& Range);
 
-    void SetOwner(ns::Range const& Range, FAGXRHIBuffer* Owner, bool bIsSwap);
+    void				SetOwner(ns::Range const& Range, FAGXRHIBuffer* Owner, bool bIsSwap);
 
 private:
     struct Allocation
@@ -72,21 +72,20 @@ private:
 class FAGXSubBufferLinear
 {
 public:
-	FAGXSubBufferLinear(NSUInteger Size, NSUInteger Alignment, mtlpp::ResourceOptions, FCriticalSection& PoolMutex);
+	FAGXSubBufferLinear(NSUInteger Size, NSUInteger Alignment, MTLResourceOptions Options, FCriticalSection& PoolMutex);
 	~FAGXSubBufferLinear();
 	
-	ns::String   GetLabel() const;
-	mtlpp::StorageMode  GetStorageMode() const;
-	mtlpp::CpuCacheMode GetCpuCacheMode() const;
-	NSUInteger     GetSize() const;
-	NSUInteger     GetUsedSize() const;
-	bool	 CanAllocateSize(NSUInteger Size) const;
+	ns::String			GetLabel() const;
+	MTLResourceOptions	GetResourceOptions() const;
+	NSUInteger			GetSize() const;
+	NSUInteger			GetUsedSize() const;
+	bool				CanAllocateSize(NSUInteger Size) const;
 
-	void SetLabel(const ns::String& label);
+	void				SetLabel(const ns::String& label);
 	
-	FAGXBuffer NewBuffer(NSUInteger length);
-	mtlpp::PurgeableState SetPurgeableState(mtlpp::PurgeableState state);
-	void FreeRange(ns::Range const& Range);
+	FAGXBuffer			NewBuffer(NSUInteger length);
+	MTLPurgeableState	SetPurgeableState(MTLPurgeableState State);
+	void				FreeRange(ns::Range const& Range);
 	
 private:
 	FCriticalSection& PoolMutex;
@@ -100,23 +99,22 @@ private:
 class FAGXSubBufferMagazine
 {
 public:
-	FAGXSubBufferMagazine(NSUInteger Size, NSUInteger ChunkSize, mtlpp::ResourceOptions);
+	FAGXSubBufferMagazine(NSUInteger Size, NSUInteger ChunkSize, MTLResourceOptions Options);
 	~FAGXSubBufferMagazine();
 	
-	ns::String   GetLabel() const;
-    mtlpp::StorageMode  GetStorageMode() const;
-    mtlpp::CpuCacheMode GetCpuCacheMode() const;
-    NSUInteger     GetSize() const;
-    NSUInteger     GetUsedSize() const;
-	NSUInteger	 GetFreeSize() const;
-	int64     NumCurrentAllocations() const;
-    bool     CanAllocateSize(NSUInteger Size) const;
+	ns::String			GetLabel() const;
+    MTLResourceOptions	GetResourceOptions() const;
+    NSUInteger			GetSize() const;
+    NSUInteger			GetUsedSize() const;
+	NSUInteger			GetFreeSize() const;
+	int64				NumCurrentAllocations() const;
+    bool				CanAllocateSize(NSUInteger Size) const;
 
-    void SetLabel(const ns::String& label);
-	void FreeRange(ns::Range const& Range);
+    void				SetLabel(const ns::String& label);
+	void				FreeRange(ns::Range const& Range);
 
-    FAGXBuffer NewBuffer();
-    mtlpp::PurgeableState SetPurgeableState(mtlpp::PurgeableState state);
+    FAGXBuffer			NewBuffer();
+    MTLPurgeableState	SetPurgeableState(MTLPurgeableState State);
 
 private:
 	NSUInteger MinAlign;
@@ -143,23 +141,22 @@ class FAGXResourceHeap;
 class FAGXSubBufferRing
 {
 public:
-	FAGXSubBufferRing(NSUInteger Size, NSUInteger Alignment, mtlpp::ResourceOptions Options);
+	FAGXSubBufferRing(NSUInteger InSize, NSUInteger InAlignment, MTLResourceOptions InOptions);
 	~FAGXSubBufferRing();
 	
-	mtlpp::StorageMode  GetStorageMode() const;
-	mtlpp::CpuCacheMode GetCpuCacheMode() const;
-	NSUInteger     GetSize() const;
+	MTLResourceOptions	GetResourceOptions() const;
+	NSUInteger			GetSize() const;
 	
-	FAGXBuffer NewBuffer(NSUInteger Size, uint32 Alignment);
+	FAGXBuffer			NewBuffer(NSUInteger InSize, uint32 InAlignment);
 	
 	/** Tries to shrink the ring-buffer back toward its initial size, but not smaller. */
-	void Shrink();
+	void				Shrink();
 	
 	/** Submits all outstanding writes to the GPU, coalescing the updates into a single contiguous range. */
-	void Submit();
+	void				Submit();
 	
 	/** Commits a completion handler to the cmd-buffer to release the processed range */
-	void Commit(mtlpp::CommandBuffer& CmdBuffer);
+	void				Commit(mtlpp::CommandBuffer& CmdBuffer);
 	
 private:
 	NSUInteger FrameSize[10];
@@ -170,8 +167,7 @@ private:
 	NSUInteger SubmitHead;
 	NSUInteger WriteHead;
 	NSUInteger BufferSize;
-	mtlpp::ResourceOptions Options;
-	mtlpp::StorageMode Storage;
+	MTLResourceOptions Options;
 	TSharedPtr<FAGXRingBufferRef, ESPMode::ThreadSafe> Buffer;
 	TArray<ns::Range> AllocatedRanges;
 };
@@ -414,7 +410,7 @@ public:
 	
 	void Init(FAGXCommandQueue& Queue);
 	
-    FAGXBuffer CreateBuffer(uint32 Size, uint32 Alignment, EBufferUsageFlags Flags, mtlpp::ResourceOptions Options, bool bForceUnique = false);
+    FAGXBuffer CreateBuffer(uint32 InSize, uint32 InAlignment, EBufferUsageFlags InFlags, MTLResourceOptions InOptions, bool bForceUnique = false);
 	FAGXTexture CreateTexture(mtlpp::TextureDescriptor Desc, FAGXSurface* Surface);
 	
 	void ReleaseBuffer(FAGXBuffer& Buffer);
diff --git a/Engine/Source/Runtime/Apple/AGXRHI/Private/AGXCommandEncoder.cpp b/Engine/Source/Runtime/Apple/AGXRHI/Private/AGXCommandEncoder.cpp
--- a/Engine/Source/Runtime/Apple/AGXRHI/Private/AGXCommandEncoder.cpp
+++ b/Engine/Source/Runtime/Apple/AGXRHI/Private/AGXCommandEncoder.cpp
@@ -23,7 +23,7 @@ extern int32 GAGXBufferScribble;
 FAGXCommandEncoder::FAGXCommandEncoder(FAGXCommandList& CmdList, EAGXCommandEncoderType InType)
 : CommandList(CmdList)
 , bSupportsMetalFeaturesSetBytes(CmdList.GetCommandQueue().SupportsFeature(EAGXFeaturesSetBytes))
-, RingBuffer(EncoderRingBufferSize, BufferOffsetAlignment, FAGXCommandQueue::GetCompatibleResourceOptions((mtlpp::ResourceOptions)(mtlpp::ResourceOptions::HazardTrackingModeUntracked | BUFFER_RESOURCE_STORAGE_MANAGED)))
+, RingBuffer(EncoderRingBufferSize, BufferOffsetAlignment, (MTLResourceCPUCacheModeDefaultCache | BUFFER_RESOURCE_STORAGE_MANAGED | MTLResourceHazardTrackingModeDefault))
 , RenderPassDesc(nil)
 , ParallelRenderCommandEncoder(nil)
 , RenderCommandEncoder(nil)
diff --git a/Engine/Source/Runtime/Apple/AGXRHI/Private/AGXCommandQueue.cpp b/Engine/Source/Runtime/Apple/AGXRHI/Private/AGXCommandQueue.cpp
--- a/Engine/Source/Runtime/Apple/AGXRHI/Private/AGXCommandQueue.cpp
+++ b/Engine/Source/Runtime/Apple/AGXRHI/Private/AGXCommandQueue.cpp
@@ -12,13 +12,9 @@
 #include "Misc/ConfigCacheIni.h"
 #include "command_buffer.hpp"
 
-#pragma mark - Private C++ Statics -
-NSUInteger FAGXCommandQueue::PermittedOptions = 0;
 uint64 FAGXCommandQueue::Features = 0;
 extern mtlpp::VertexFormat GAGXFColorVertexFormat;
 
-#pragma mark - Public C++ Boilerplate -
-
 FAGXCommandQueue::FAGXCommandQueue(uint32 const MaxNumCommandBuffers /* = 0 */)
 : ParallelCommandLists(0)
 , RuntimeDebuggingLevel(EAGXDebugLevelOff)
@@ -54,7 +50,7 @@ FAGXCommandQueue::FAGXCommandQueue(uint32 const MaxNumCommandBuffers /* = 0 */)
 		Features |= EAGXFeaturesCountingQueries | EAGXFeaturesBaseVertexInstance | EAGXFeaturesIndirectBuffer | EAGXFeaturesMSAADepthResolve | EAGXFeaturesMSAAStoreAndResolve;
 	}
 
-	Features |= EAGXFeaturesPrivateBufferSubAllocation | EAGXFeaturesGPUCaptureManager | EAGXFeaturesBufferSubAllocation | EAGXFeaturesParallelRenderEncoders | EAGXFeaturesPipelineBufferMutability | EAGXFeaturesMaxThreadsPerThreadgroup | EAGXFeaturesTextureBuffers;
+	Features |= EAGXFeaturesPrivateBufferSubAllocation | EAGXFeaturesGPUCaptureManager | EAGXFeaturesBufferSubAllocation | EAGXFeaturesParallelRenderEncoders | EAGXFeaturesPipelineBufferMutability | EAGXFeaturesMaxThreadsPerThreadgroup;
 	GAGXFColorVertexFormat = mtlpp::VertexFormat::UChar4Normalized_BGRA;
 
 #else
@@ -83,7 +79,6 @@ FAGXCommandQueue::FAGXCommandQueue(uint32 const MaxNumCommandBuffers /* = 0 */)
 		| EAGXFeaturesPipelineBufferMutability;
 
 	Features |= EAGXFeaturesMaxThreadsPerThreadgroup;
-	Features |= EAGXFeaturesTextureBuffers;
 
 	if ([GMtlDevice supportsFeatureSet : MTLFeatureSet_iOS_GPUFamily4_v1])
 	{
@@ -119,8 +114,6 @@ FAGXCommandQueue::FAGXCommandQueue(uint32 const MaxNumCommandBuffers /* = 0 */)
 		{
 			Features |= EAGXFeaturesParallelRenderEncoders;
 		}
-
-		Features |= EAGXFeaturesTextureBuffers;
     }
     else if ([[GMtlDevice name] rangeOfString:@"Nvidia" options:NSCaseInsensitiveSearch].location != NSNotFound)
 	{
@@ -147,17 +140,6 @@ FAGXCommandQueue::FAGXCommandQueue(uint32 const MaxNumCommandBuffers /* = 0 */)
 	{
 		Features |= EAGXFeaturesGPUTrace;
 	}
-
-	PermittedOptions = 0;
-	PermittedOptions |= mtlpp::ResourceOptions::CpuCacheModeDefaultCache;
-	PermittedOptions |= mtlpp::ResourceOptions::CpuCacheModeWriteCombined;
-	PermittedOptions |= mtlpp::ResourceOptions::StorageModeShared;
-	PermittedOptions |= mtlpp::ResourceOptions::StorageModePrivate;
-#if PLATFORM_MAC
-	PermittedOptions |= mtlpp::ResourceOptions::StorageModeManaged;
-#else
-	PermittedOptions |= mtlpp::ResourceOptions::StorageModeMemoryless;
-#endif
 }
 
 FAGXCommandQueue::~FAGXCommandQueue(void)
@@ -169,8 +151,6 @@ FAGXCommandQueue::~FAGXCommandQueue(void)
 	}
 }
 	
-#pragma mark - Public Command Buffer Mutators -
-
 mtlpp::CommandBuffer FAGXCommandQueue::CreateCommandBuffer(void)
 {
 #if PLATFORM_MAC
@@ -247,22 +227,6 @@ void FAGXCommandQueue::GetCommittedCommandBufferFences(TArray<mtlpp::CommandBuff
 	}
 }
 
-#pragma mark - Public Command Queue Accessors -
-
-mtlpp::ResourceOptions FAGXCommandQueue::GetCompatibleResourceOptions(mtlpp::ResourceOptions Options)
-{
-	NSUInteger NewOptions = (Options & PermittedOptions);
-#if PLATFORM_IOS // Swizzle Managed to Shared for iOS - we can do this as they are equivalent, unlike Shared -> Managed on Mac.
-	if ((Options & (1 /*mtlpp::StorageMode::Managed*/ << mtlpp::ResourceStorageModeShift)))
-	{
-		NewOptions |= mtlpp::ResourceOptions::StorageModeShared;
-	}
-#endif
-	return (mtlpp::ResourceOptions)NewOptions;
-}
-
-#pragma mark - Public Debug Support -
-
 void FAGXCommandQueue::InsertDebugCaptureBoundary(void)
 {
 	PRAGMA_DISABLE_DEPRECATION_WARNINGS
diff --git a/Engine/Source/Runtime/Apple/AGXRHI/Private/AGXCommandQueue.h b/Engine/Source/Runtime/Apple/AGXRHI/Private/AGXCommandQueue.h
--- a/Engine/Source/Runtime/Apple/AGXRHI/Private/AGXCommandQueue.h
+++ b/Engine/Source/Runtime/Apple/AGXRHI/Private/AGXCommandQueue.h
@@ -60,7 +60,7 @@ typedef NS_OPTIONS(uint64, EAGXFeatures)
 	/** Supports private buffer sub-allocation */
 	EAGXFeaturesPrivateBufferSubAllocation = 1llu << 20llu,
 	/** Supports texture buffers */
-	EAGXFeaturesTextureBuffers = 1llu << 21llu,
+	EAGXFeaturesTextureBuffers_REMOVED = 1llu << 21llu,
 	/** Supports max. compute threads per threadgroup */
 	EAGXFeaturesMaxThreadsPerThreadgroup = 1llu << 22llu,
 	/** Supports parallel render encoders */
@@ -123,9 +123,6 @@ public:
 	
 #pragma mark - Public Command Queue Accessors -
 	
-	/** Converts a Metal v1.1+ resource option to something valid on the current version. */
-	static mtlpp::ResourceOptions GetCompatibleResourceOptions(mtlpp::ResourceOptions Options);
-	
 	/**
 	 * @param InFeature A specific Metal feature to check for.
 	 * @returns True if the requested feature is supported, else false.
@@ -157,6 +154,5 @@ private:
 	TLockFreePointerListLIFO<mtlpp::CommandBufferFence> CommandBufferFences;
 	uint64 ParallelCommandLists;
 	int32 RuntimeDebuggingLevel;
-	static NSUInteger PermittedOptions;
 	static uint64 Features;
 };
diff --git a/Engine/Source/Runtime/Apple/AGXRHI/Private/AGXContext.cpp b/Engine/Source/Runtime/Apple/AGXRHI/Private/AGXContext.cpp
--- a/Engine/Source/Runtime/Apple/AGXRHI/Private/AGXContext.cpp
+++ b/Engine/Source/Runtime/Apple/AGXRHI/Private/AGXContext.cpp
@@ -429,11 +429,11 @@ void FAGXDeviceContext::BeginFrame()
 void FAGXDeviceContext::ScribbleBuffer(FAGXBuffer& Buffer)
 {
 	static uint8 Fill = 0;
-	if (Buffer.GetStorageMode() != mtlpp::StorageMode::Private)
+	if ([Buffer.GetPtr() storageMode] != MTLStorageModePrivate)
 	{
 		FMemory::Memset(Buffer.GetContents(), Fill++, Buffer.GetLength());
 #if PLATFORM_MAC
-		if (Buffer.GetStorageMode() == mtlpp::StorageMode::Managed)
+		if ([Buffer.GetPtr() storageMode] == MTLStorageModeManaged)
 		{
 			Buffer.DidModify(ns::Range(0, Buffer.GetLength()));
 		}
@@ -688,7 +688,7 @@ void FAGXDeviceContext::ReleaseTexture(FAGXTexture& Texture)
 	{
 		check(Texture);
 		FreeListMutex.Lock();
-        if (Texture.GetStorageMode() == mtlpp::StorageMode::Private)
+        if ([Texture.GetPtr() storageMode] == MTLStorageModePrivate)
         {
             Heap.ReleaseTexture(nullptr, Texture);
 			
@@ -722,8 +722,6 @@ FAGXTexture FAGXDeviceContext::CreateTexture(FAGXSurface* Surface, mtlpp::Textur
 
 FAGXBuffer FAGXDeviceContext::CreatePooledBuffer(FAGXPooledBufferArgs const& Args)
 {
-	NSUInteger CpuResourceOption = ((NSUInteger)Args.CpuCacheMode) << mtlpp::ResourceCpuCacheModeShift;
-	
 	uint32 RequestedBufferOffsetAlignment = BufferOffsetAlignment;
 	
 	if(EnumHasAnyFlags(Args.Flags, BUF_UnorderedAccess | BUF_ShaderResource))
@@ -733,7 +731,7 @@ FAGXBuffer FAGXDeviceContext::CreatePooledBuffer(FAGXPooledBufferArgs const& Arg
 		RequestedBufferOffsetAlignment = BufferBackedLinearTextureOffsetAlignment;
 	}
 	
-    FAGXBuffer Buffer = Heap.CreateBuffer(Args.Size, RequestedBufferOffsetAlignment, Args.Flags, FAGXCommandQueue::GetCompatibleResourceOptions((mtlpp::ResourceOptions)(CpuResourceOption | mtlpp::ResourceOptions::HazardTrackingModeUntracked | ((NSUInteger)Args.Storage << mtlpp::ResourceStorageModeShift))));
+    FAGXBuffer Buffer = Heap.CreateBuffer(Args.Size, RequestedBufferOffsetAlignment, Args.Flags, Args.Options);
 	check(Buffer && Buffer.GetPtr());
 #if METAL_DEBUG_OPTIONS
 	if (GAGXResourcePurgeOnDelete)
diff --git a/Engine/Source/Runtime/Apple/AGXRHI/Private/AGXLLM.cpp b/Engine/Source/Runtime/Apple/AGXRHI/Private/AGXLLM.cpp
--- a/Engine/Source/Runtime/Apple/AGXRHI/Private/AGXLLM.cpp
+++ b/Engine/Source/Runtime/Apple/AGXRHI/Private/AGXLLM.cpp
@@ -155,7 +155,7 @@ void AGXLLM::LogAllocTexture(mtlpp::TextureDescriptor const& Desc, mtlpp::Textur
 	uint64 Size = static_cast<uint64>(SizeAlign.size);
 	
 #if PLATFORM_IOS
-	bool bMemoryless = (Texture.GetStorageMode() == mtlpp::StorageMode::Memoryless);
+	bool bMemoryless = ([Texture.GetPtr() storageMode] == MTLStorageModeMemoryless);
 	if (!bMemoryless)
 #endif
 	{
diff --git a/Engine/Source/Runtime/Apple/AGXRHI/Private/AGXRHIPrivate.h b/Engine/Source/Runtime/Apple/AGXRHI/Private/AGXRHIPrivate.h
--- a/Engine/Source/Runtime/Apple/AGXRHI/Private/AGXRHIPrivate.h
+++ b/Engine/Source/Runtime/Apple/AGXRHI/Private/AGXRHIPrivate.h
@@ -105,20 +105,14 @@ struct FAGXDebugShaderResourceMask
 	FAGXSamplerMask SamplerMask;
 };
 
-#define BUFFER_CACHE_MODE mtlpp::ResourceOptions::CpuCacheModeDefaultCache
-
 #if PLATFORM_MAC
-#define BUFFER_MANAGED_MEM mtlpp::ResourceOptions::StorageModeManaged
-#define BUFFER_STORAGE_MODE mtlpp::StorageMode::Managed
-#define BUFFER_RESOURCE_STORAGE_MANAGED mtlpp::ResourceOptions::StorageModeManaged
-#define BUFFER_DYNAMIC_REALLOC BUF_AnyDynamic
+#define BUFFER_STORAGE_MODE MTLStorageModeManaged
+#define BUFFER_RESOURCE_STORAGE_MANAGED MTLResourceStorageModeManaged
 // How many possible vertex streams are allowed
 const uint32 MaxMetalStreams = 31;
 #else
-#define BUFFER_MANAGED_MEM 0
-#define BUFFER_STORAGE_MODE mtlpp::StorageMode::Shared
-#define BUFFER_RESOURCE_STORAGE_MANAGED mtlpp::ResourceOptions::StorageModeShared
-#define BUFFER_DYNAMIC_REALLOC BUF_AnyDynamic
+#define BUFFER_STORAGE_MODE MTLStorageModeShared
+#define BUFFER_RESOURCE_STORAGE_MANAGED MTLResourceStorageModeShared
 // How many possible vertex streams are allowed
 const uint32 MaxMetalStreams = 30;
 #endif
diff --git a/Engine/Source/Runtime/Apple/AGXRHI/Private/AGXRHIRenderQuery.cpp b/Engine/Source/Runtime/Apple/AGXRHI/Private/AGXRHIRenderQuery.cpp
--- a/Engine/Source/Runtime/Apple/AGXRHI/Private/AGXRHIRenderQuery.cpp
+++ b/Engine/Source/Runtime/Apple/AGXRHI/Private/AGXRHIRenderQuery.cpp
@@ -92,7 +92,7 @@ void FAGXQueryBufferPool::Allocate(FAGXQueryResult& NewQuery)
 
 FAGXQueryBuffer* FAGXQueryBufferPool::GetCurrentQueryBuffer()
 {
-	if (!IsValidRef(CurrentBuffer) || (CurrentBuffer->Buffer.GetStorageMode() != mtlpp::StorageMode::Shared && CurrentBuffer->WriteOffset > 0))
+	if (!IsValidRef(CurrentBuffer) || (([CurrentBuffer->Buffer.GetPtr() storageMode] != MTLStorageModeShared) && (CurrentBuffer->WriteOffset > 0)))
 	{
 		FAGXBuffer Buffer;
 		if (Buffers.Num())
@@ -104,8 +104,8 @@ FAGXQueryBuffer* FAGXQueryBufferPool::GetCurrentQueryBuffer()
 			LLM_SCOPE_METAL(ELLMTagAGX::Buffers);
 			LLM_PLATFORM_SCOPE_METAL(ELLMTagAGX::Buffers);
 
-			METAL_GPUPROFILE(FAGXScopedCPUStats CPUStat(FString::Printf(TEXT("AllocBuffer: %llu, %llu"), EQueryBufferMaxSize, mtlpp::ResourceOptions::StorageModeShared)));
-			Buffer = GetAGXDeviceContext().GetResourceHeap().CreateBuffer(EQueryBufferMaxSize, 16, BUF_Dynamic, FAGXCommandQueue::GetCompatibleResourceOptions((mtlpp::ResourceOptions)(BUFFER_CACHE_MODE | mtlpp::ResourceOptions::HazardTrackingModeUntracked | mtlpp::ResourceOptions::StorageModeShared)), true);
+			METAL_GPUPROFILE(FAGXScopedCPUStats CPUStat(FString::Printf(TEXT("AllocBuffer: %llu, %llu"), EQueryBufferMaxSize, MTLResourceStorageModeShared)));
+			Buffer = GetAGXDeviceContext().GetResourceHeap().CreateBuffer(EQueryBufferMaxSize, 16, BUF_Dynamic, (MTLResourceCPUCacheModeDefaultCache | MTLResourceStorageModeShared | MTLResourceHazardTrackingModeDefault), true);
 			FMemory::Memzero((((uint8*)Buffer.GetContents())), EQueryBufferMaxSize);
 
 #if STATS || ENABLE_LOW_LEVEL_MEM_TRACKER
diff --git a/Engine/Source/Runtime/Apple/AGXRHI/Private/AGXRenderTarget.cpp b/Engine/Source/Runtime/Apple/AGXRHI/Private/AGXRenderTarget.cpp
--- a/Engine/Source/Runtime/Apple/AGXRHI/Private/AGXRenderTarget.cpp
+++ b/Engine/Source/Runtime/Apple/AGXRHI/Private/AGXRenderTarget.cpp
@@ -283,12 +283,12 @@ void FAGXDynamicRHI::RHIReadSurfaceData(FRHITexture* TextureRHI, FIntRect Rect,
 		SCOPE_CYCLE_COUNTER(STAT_AGXTexturePageOffTime);
 		
 		FAGXTexture TempTexture = nil;
-		if (Texture.GetStorageMode() == mtlpp::StorageMode::Private)
+		if ([Texture.GetPtr() storageMode] == MTLStorageModePrivate)
 		{
 #if PLATFORM_MAC
-			mtlpp::StorageMode StorageMode = mtlpp::StorageMode::Managed;
+			MTLResourceOptions ResourceStorageMode = MTLResourceStorageModeManaged;
 #else
-			mtlpp::StorageMode StorageMode = mtlpp::StorageMode::Shared;
+			MTLResourceOptions ResourceStorageMode = MTLResourceStorageModeShared;
 #endif
 			MTLPixelFormat MetalFormat = (MTLPixelFormat)GPixelFormats[Surface->GetDesc().Format].PlatformFormat;
 			mtlpp::TextureDescriptor Desc;
@@ -301,11 +301,8 @@ void FAGXDynamicRHI::RHIReadSurfaceData(FRHITexture* TextureRHI, FIntRect Rect,
 			Desc.SetSampleCount(Texture.GetSampleCount());
 			Desc.SetArrayLength(Texture.GetArrayLength());
 			
-			mtlpp::ResourceOptions GeneralResourceOption = (mtlpp::ResourceOptions)FAGXCommandQueue::GetCompatibleResourceOptions(mtlpp::ResourceOptions(((NSUInteger)Texture.GetCpuCacheMode() << mtlpp::ResourceCpuCacheModeShift) | ((NSUInteger)StorageMode << mtlpp::ResourceStorageModeShift) | mtlpp::ResourceOptions::HazardTrackingModeUntracked));
-			Desc.SetResourceOptions(GeneralResourceOption);
+			Desc.SetResourceOptions(mtlpp::ResourceOptions(([Texture.GetPtr() resourceOptions] & ~MTLResourceStorageModeMask) | ResourceStorageMode));
 			
-			Desc.SetCpuCacheMode(Texture.GetCpuCacheMode());
-			Desc.SetStorageMode(StorageMode);
 			Desc.SetUsage(Texture.GetUsage());
 			
 			TempTexture = GMtlppDevice.NewTexture(Desc);
@@ -316,7 +313,7 @@ void FAGXDynamicRHI::RHIReadSurfaceData(FRHITexture* TextureRHI, FIntRect Rect,
 			Region = mtlpp::Region(0, 0, SizeX, SizeY);
 		}
 #if PLATFORM_MAC
-		if(Texture.GetStorageMode() == mtlpp::StorageMode::Managed)
+		if ([Texture.GetPtr() storageMode] == MTLStorageModeManaged)
 		{
 			// Synchronise the texture with the CPU
 			ImmediateContext.Context->SynchronizeTexture(Texture, 0, InFlags.GetMip());
@@ -348,7 +345,7 @@ void FAGXDynamicRHI::RHIReadSurfaceData(FRHITexture* TextureRHI, FIntRect Rect,
 		const uint32 Alignment = PLATFORM_MAC ? 1u : 64u; // Mac permits natural row alignment (tightly-packed) but iOS does not.
 		const uint32 AlignedStride = ((Stride - 1) & ~(Alignment - 1)) + Alignment;
 		const uint32 BytesPerImage = AlignedStride * SizeY;
-		FAGXBuffer Buffer = ((FAGXDeviceContext*)ImmediateContext.Context)->CreatePooledBuffer(FAGXPooledBufferArgs(BytesPerImage, BUF_Dynamic, mtlpp::StorageMode::Shared));
+		FAGXBuffer Buffer = ((FAGXDeviceContext*)ImmediateContext.Context)->CreatePooledBuffer(FAGXPooledBufferArgs(BytesPerImage, BUF_Dynamic, FAGXPooledBufferArgs::SharedStorageResourceOptions));
 		{
 			// Synchronise the texture with the CPU
 			SCOPE_CYCLE_COUNTER(STAT_AGXTexturePageOffTime);
@@ -395,7 +392,7 @@ void FAGXDynamicRHI::RHIMapStagingSurface(FRHITexture* TextureRHI, FRHIGPUFence*
     OutData = Surface->Lock(0, 0, RLM_ReadOnly, Stride);
     
 #if PLATFORM_MAC
-	Surface->GPUReadback = (Surface->Texture.GetPtr() && Surface->Texture.GetStorageMode() == mtlpp::StorageMode::Managed) ? FAGXSurface::EAGXGPUReadbackFlags::ReadbackRequested : 0;
+	Surface->GPUReadback = (Surface->Texture.GetPtr() && [Surface->Texture.GetPtr() storageMode] == MTLStorageModeManaged) ? FAGXSurface::EAGXGPUReadbackFlags::ReadbackRequested : 0;
 #endif
 	}
 }
@@ -452,7 +449,7 @@ void FAGXDynamicRHI::RHIReadSurfaceFloatData(FRHITexture* TextureRHI, FIntRect R
 	const uint32 AlignedStride = ((Stride - 1) & ~(Alignment - 1)) + Alignment;
 	const uint32 BytesPerImage = AlignedStride  * SizeY;
 	int32 FloatBGRADataSize = BytesPerImage;
-	FAGXBuffer Buffer = ((FAGXDeviceContext*)ImmediateContext.Context)->CreatePooledBuffer(FAGXPooledBufferArgs(FloatBGRADataSize, BUF_Dynamic, mtlpp::StorageMode::Shared));
+	FAGXBuffer Buffer = ((FAGXDeviceContext*)ImmediateContext.Context)->CreatePooledBuffer(FAGXPooledBufferArgs(FloatBGRADataSize, BUF_Dynamic, FAGXPooledBufferArgs::SharedStorageResourceOptions));
 	{
 		// Synchronise the texture with the CPU
 		SCOPE_CYCLE_COUNTER(STAT_AGXTexturePageOffTime);
@@ -518,7 +515,7 @@ void FAGXDynamicRHI::RHIRead3DSurfaceFloatData(FRHITexture* TextureRHI,FIntRect
 	const uint32 AlignedStride = ((Stride - 1) & ~(Alignment - 1)) + Alignment;
 	const uint32 BytesPerImage = AlignedStride  * SizeY;
 	int32 FloatBGRADataSize = BytesPerImage * SizeZ;
-	FAGXBuffer Buffer = ((FAGXDeviceContext*)ImmediateContext.Context)->CreatePooledBuffer(FAGXPooledBufferArgs(FloatBGRADataSize, BUF_Dynamic, mtlpp::StorageMode::Shared));
+	FAGXBuffer Buffer = ((FAGXDeviceContext*)ImmediateContext.Context)->CreatePooledBuffer(FAGXPooledBufferArgs(FloatBGRADataSize, BUF_Dynamic, FAGXPooledBufferArgs::SharedStorageResourceOptions));
 	{
 		// Synchronise the texture with the CPU
 		SCOPE_CYCLE_COUNTER(STAT_AGXTexturePageOffTime);
diff --git a/Engine/Source/Runtime/Apple/AGXRHI/Private/AGXStateCache.cpp b/Engine/Source/Runtime/Apple/AGXRHI/Private/AGXStateCache.cpp
--- a/Engine/Source/Runtime/Apple/AGXRHI/Private/AGXStateCache.cpp
+++ b/Engine/Source/Runtime/Apple/AGXRHI/Private/AGXStateCache.cpp
@@ -576,7 +576,7 @@ bool FAGXStateCache::SetRenderPassInfo(FRHIRenderPassInfo const& InRenderTargets
 #if PLATFORM_IOS
 				bUseResolvedTexture = (
 					Surface.MSAATexture && 
-					Surface.MSAATexture.GetStorageMode() == mtlpp::StorageMode::Memoryless && 
+					[Surface.MSAATexture.GetPtr() storageMode] == MTLStorageModeMemoryless && 
 					HighLevelLoadAction == ERenderTargetLoadAction::ELoad);
 #endif
 				
@@ -584,7 +584,7 @@ bool FAGXStateCache::SetRenderPassInfo(FRHIRenderPassInfo const& InRenderTargets
 				if (Surface.MSAATexture && !bUseResolvedTexture)
 				{
 #if PLATFORM_IOS
-					if (Surface.MSAATexture.GetStorageMode() == mtlpp::StorageMode::Memoryless)
+					if ([Surface.MSAATexture.GetPtr() storageMode] == MTLStorageModeMemoryless)
 					{
 						bMemoryless = true;
 						HighLevelLoadAction = ERenderTargetLoadAction::EClear;
@@ -600,7 +600,7 @@ bool FAGXStateCache::SetRenderPassInfo(FRHIRenderPassInfo const& InRenderTargets
 				else
 				{
 #if PLATFORM_IOS
-					if (Surface.Texture.GetStorageMode() == mtlpp::StorageMode::Memoryless)
+					if ([Surface.Texture.GetPtr() storageMode] == MTLStorageModeMemoryless)
 					{
 						bMemoryless = true;
 						HighLevelStoreAction = ERenderTargetStoreAction::ENoAction;
@@ -825,7 +825,7 @@ bool FAGXStateCache::SetRenderPassInfo(FRHIRenderPassInfo const& InRenderTargets
 				const bool bSupportsMSAADepthResolve = GetAGXDeviceContext().SupportsFeature(EAGXFeaturesMSAADepthResolve);
 				bool bDepthTextureMemoryless = false;
 #if PLATFORM_IOS
-				bDepthTextureMemoryless = DepthTexture.GetStorageMode() == mtlpp::StorageMode::Memoryless;
+				bDepthTextureMemoryless = [DepthTexture.GetPtr() storageMode] == MTLStorageModeMemoryless;
 				if (bDepthTextureMemoryless)
 				{
 					[DepthAttachment setLoadAction:MTLLoadActionClear];
@@ -908,7 +908,7 @@ bool FAGXStateCache::SetRenderPassInfo(FRHIRenderPassInfo const& InRenderTargets
 				
 				bool bStencilMemoryless = false;
 #if PLATFORM_IOS
-				if (StencilTexture.GetStorageMode() == mtlpp::StorageMode::Memoryless)
+				if ([StencilTexture.GetPtr() storageMode] == MTLStorageModeMemoryless)
 				{
 					bStencilMemoryless = true;
 					HighLevelStoreAction = ERenderTargetStoreAction::ENoAction;
@@ -1285,7 +1285,7 @@ bool FAGXStateCache::NeedsToSetRenderTarget(const FRHIRenderPassInfo& InRenderPa
 					
 #if PLATFORM_IOS
 					FAGXTexture& Tex = Surface.MSAATexture ? Surface.MSAATexture : Surface.Texture;
-					if (Tex.GetStorageMode() == mtlpp::StorageMode::Memoryless)
+					if ([Tex.GetPtr() storageMode] == MTLStorageModeMemoryless)
 					{
 						HighLevelStoreAction = ERenderTargetStoreAction::ENoAction;
 					}
@@ -1855,7 +1855,7 @@ void FAGXStateCache::SetRenderStoreActions(FAGXCommandEncoder& CommandEncoder, b
 void FAGXStateCache::FlushVisibilityResults(FAGXCommandEncoder& CommandEncoder)
 {
 #if PLATFORM_MAC
-	if(VisibilityResults && VisibilityResults->Buffer && VisibilityResults->Buffer.GetStorageMode() == mtlpp::StorageMode::Managed && VisibilityWritten && CommandEncoder.IsRenderCommandEncoderActive())
+	if (VisibilityResults && VisibilityResults->Buffer && [VisibilityResults->Buffer.GetPtr() storageMode] == MTLStorageModeManaged && VisibilityWritten && CommandEncoder.IsRenderCommandEncoderActive())
 	{
 		CommandEncoder.EndEncoding();
 		
diff --git a/Engine/Source/Runtime/Apple/AGXRHI/Private/AGXTexture.cpp b/Engine/Source/Runtime/Apple/AGXRHI/Private/AGXTexture.cpp
--- a/Engine/Source/Runtime/Apple/AGXRHI/Private/AGXTexture.cpp
+++ b/Engine/Source/Runtime/Apple/AGXRHI/Private/AGXTexture.cpp
@@ -257,7 +257,7 @@ void FAGXSurface::PrepareTextureView()
 	mtlpp::TextureUsage Usage = (mtlpp::TextureUsage)Texture.GetUsage();
 	bool bMemoryLess = false;
 #if PLATFORM_IOS
-	bMemoryLess = (Texture.GetStorageMode() == mtlpp::StorageMode::Memoryless);
+	bMemoryLess = ([Texture storageMode] == MTLStorageModeMemoryless);
 #endif
 	if(!(Usage & mtlpp::TextureUsage::PixelFormatView) && !bMemoryLess)
 	{
@@ -305,12 +305,7 @@ FAGXTexture FAGXSurface::Reallocate(FAGXTexture InTexture, mtlpp::TextureUsage U
 	Desc.SetMipmapLevelCount(InTexture.GetMipmapLevelCount());
 	Desc.SetSampleCount(InTexture.GetSampleCount());
 	Desc.SetArrayLength(InTexture.GetArrayLength());
-	
-	static mtlpp::ResourceOptions GeneralResourceOption = (mtlpp::ResourceOptions)FAGXCommandQueue::GetCompatibleResourceOptions(mtlpp::ResourceOptions::HazardTrackingModeUntracked);
-	
-	Desc.SetResourceOptions(mtlpp::ResourceOptions(((NSUInteger)InTexture.GetCpuCacheMode() << mtlpp::ResourceCpuCacheModeShift) | ((NSUInteger)Texture.GetStorageMode() << mtlpp::ResourceStorageModeShift) | GeneralResourceOption));
-	Desc.SetCpuCacheMode(InTexture.GetCpuCacheMode());
-	Desc.SetStorageMode(InTexture.GetStorageMode());
+	Desc.SetResourceOptions(mtlpp::ResourceOptions([InTexture.GetPtr() resourceOptions]));
 	Desc.SetUsage(mtlpp::TextureUsage(InTexture.GetUsage() | UsageModifier));
 	
 	FAGXTexture NewTex = GetAGXDeviceContext().CreateTexture(this, Desc);
@@ -422,62 +417,48 @@ FAGXTextureDesc::FAGXTextureDesc(FRHITextureDesc const& InDesc)
 	Desc.SetMipmapLevelCount(InDesc.NumMips);
 
 	{
-		Desc.SetUsage(ConvertFlagsToUsage(InDesc.Flags));
-
+		MTLResourceOptions ResourceStorageMode;
 		if (EnumHasAnyFlags(InDesc.Flags, TexCreate_CPUReadback) && !EnumHasAnyFlags(InDesc.Flags, TexCreate_RenderTargetable | TexCreate_DepthStencilTargetable | TexCreate_FastVRAM))
 		{
-			Desc.SetCpuCacheMode(mtlpp::CpuCacheMode::DefaultCache);
 #if PLATFORM_MAC
-			Desc.SetStorageMode(mtlpp::StorageMode::Managed);
-			Desc.SetResourceOptions((mtlpp::ResourceOptions)(mtlpp::ResourceOptions::CpuCacheModeDefaultCache | mtlpp::ResourceOptions::StorageModeManaged));
+			ResourceStorageMode = MTLResourceStorageModeManaged;
 #else
-			Desc.SetStorageMode(mtlpp::StorageMode::Shared);
-			Desc.SetResourceOptions((mtlpp::ResourceOptions)(mtlpp::ResourceOptions::CpuCacheModeDefaultCache | mtlpp::ResourceOptions::StorageModeShared));
+			ResourceStorageMode = MTLResourceStorageModeShared;
 #endif
 		}
 		else if (EnumHasAnyFlags(InDesc.Flags, TexCreate_NoTiling) && !EnumHasAnyFlags(InDesc.Flags, TexCreate_FastVRAM | TexCreate_DepthStencilTargetable | TexCreate_RenderTargetable | TexCreate_UAV))
 		{
-			Desc.SetCpuCacheMode(mtlpp::CpuCacheMode::DefaultCache);
 #if PLATFORM_MAC
-			Desc.SetStorageMode(mtlpp::StorageMode::Managed);
-			Desc.SetResourceOptions((mtlpp::ResourceOptions)(mtlpp::ResourceOptions::CpuCacheModeDefaultCache | mtlpp::ResourceOptions::StorageModeManaged));
+			ResourceStorageMode = MTLResourceStorageModeManaged;
 #else
-			Desc.SetStorageMode(mtlpp::StorageMode::Shared);
-			Desc.SetResourceOptions((mtlpp::ResourceOptions)(mtlpp::ResourceOptions::CpuCacheModeDefaultCache | mtlpp::ResourceOptions::StorageModeShared));
+			ResourceStorageMode = MTLResourceStorageModeShared;
 #endif
 		}
 		else if (EnumHasAnyFlags(InDesc.Flags, TexCreate_RenderTargetable | TexCreate_DepthStencilTargetable | TexCreate_ResolveTargetable | TexCreate_DepthStencilResolveTarget))
 		{
 			check(!(InDesc.Flags & TexCreate_CPUReadback));
-			Desc.SetCpuCacheMode(mtlpp::CpuCacheMode::DefaultCache);
 #if PLATFORM_MAC
-			Desc.SetStorageMode(mtlpp::StorageMode::Private);
-			Desc.SetResourceOptions((mtlpp::ResourceOptions)(mtlpp::ResourceOptions::CpuCacheModeDefaultCache | mtlpp::ResourceOptions::StorageModePrivate));
+			ResourceStorageMode = MTLResourceStorageModePrivate;
 #else
 			if (GAGXForceIOSTexturesShared)
 			{
-				Desc.SetStorageMode(mtlpp::StorageMode::Shared);
-				Desc.SetResourceOptions((mtlpp::ResourceOptions)(mtlpp::ResourceOptions::CpuCacheModeDefaultCache | mtlpp::ResourceOptions::StorageModeShared));
+				ResourceStorageMode = MTLResourceStorageModeShared;
 			}
 			else
 			{
-				Desc.SetStorageMode(mtlpp::StorageMode::Private);
-				Desc.SetResourceOptions((mtlpp::ResourceOptions)(mtlpp::ResourceOptions::CpuCacheModeDefaultCache | mtlpp::ResourceOptions::StorageModePrivate));
+				ResourceStorageMode = MTLResourceStorageModePrivate;
 			}
 #endif
 		}
 		else
 		{
 			check(!(InDesc.Flags & TexCreate_CPUReadback));
-			Desc.SetCpuCacheMode(mtlpp::CpuCacheMode::DefaultCache);
 #if PLATFORM_MAC
-			Desc.SetStorageMode(mtlpp::StorageMode::Private);
-			Desc.SetResourceOptions((mtlpp::ResourceOptions)(mtlpp::ResourceOptions::CpuCacheModeDefaultCache | mtlpp::ResourceOptions::StorageModePrivate));
+			ResourceStorageMode = MTLResourceStorageModePrivate;
 #else
 			if (GAGXForceIOSTexturesShared)
 			{
-				Desc.SetStorageMode(mtlpp::StorageMode::Shared);
-				Desc.SetResourceOptions((mtlpp::ResourceOptions)(mtlpp::ResourceOptions::CpuCacheModeDefaultCache | mtlpp::ResourceOptions::StorageModeShared));
+				ResourceStorageMode = MTLResourceStorageModeShared;
 			}
 			// No private storage for PVRTC as it messes up the blit-encoder usage.
 			// note: this is set to always be on and will be re-addressed in a future release
@@ -485,13 +466,11 @@ FAGXTextureDesc::FAGXTextureDesc(FRHITextureDesc const& InDesc)
 			{
 				if (IsPixelFormatPVRTCCompressed(InDesc.Format))
 				{
-					Desc.SetStorageMode(mtlpp::StorageMode::Shared);
-					Desc.SetResourceOptions((mtlpp::ResourceOptions)(mtlpp::ResourceOptions::CpuCacheModeDefaultCache | mtlpp::ResourceOptions::StorageModeShared));
+					ResourceStorageMode = MTLResourceStorageModeShared;
 				}
 				else
 				{
-					Desc.SetStorageMode(mtlpp::StorageMode::Private);
-					Desc.SetResourceOptions((mtlpp::ResourceOptions)(mtlpp::ResourceOptions::CpuCacheModeDefaultCache | mtlpp::ResourceOptions::StorageModePrivate));
+					ResourceStorageMode = MTLResourceStorageModePrivate;
 				}
 			}
 #endif
@@ -503,13 +482,21 @@ FAGXTextureDesc::FAGXTextureDesc(FRHITextureDesc const& InDesc)
 			ensure(EnumHasAnyFlags(InDesc.Flags, (TexCreate_RenderTargetable | TexCreate_DepthStencilTargetable)));
 			ensure(!EnumHasAnyFlags(InDesc.Flags, (TexCreate_CPUReadback | TexCreate_CPUWritable)));
 			ensure(!EnumHasAnyFlags(InDesc.Flags, TexCreate_UAV));
-			Desc.SetStorageMode(mtlpp::StorageMode::Memoryless);
-			Desc.SetResourceOptions((mtlpp::ResourceOptions)(mtlpp::ResourceOptions::CpuCacheModeDefaultCache | mtlpp::ResourceOptions::StorageModeMemoryless));
+			ResourceStorageMode = MTLResourceStorageModeMemoryless;
+		}
+
+		if (!FParse::Param(FCommandLine::Get(), TEXT("nomsaa")) && InDesc.NumSamples > 1)
+		{
+			if (GMaxRHIFeatureLevel < ERHIFeatureLevel::SM5)
+			{
+				ResourceStorageMode = MTLResourceStorageModeMemoryless;
+				bMemoryless = true;
+			}
 		}
 #endif
 
-		static mtlpp::ResourceOptions GeneralResourceOption = FAGXCommandQueue::GetCompatibleResourceOptions(mtlpp::ResourceOptions::HazardTrackingModeUntracked);
-		Desc.SetResourceOptions((mtlpp::ResourceOptions)(Desc.GetResourceOptions() | GeneralResourceOption));
+		Desc.SetResourceOptions(mtlpp::ResourceOptions(MTLResourceCPUCacheModeDefaultCache | ResourceStorageMode | MTLResourceHazardTrackingModeDefault));
+		Desc.SetUsage(ConvertFlagsToUsage(InDesc.Flags));
 	}
 	
 	if (!FParse::Param(FCommandLine::Get(), TEXT("nomsaa")))
@@ -529,15 +516,6 @@ FAGXTextureDesc::FAGXTextureDesc(FRHITextureDesc const& InDesc)
 			{
 				Desc.SetSampleCount(InDesc.NumSamples);
 			}
-
-#if PLATFORM_IOS
-			if (GMaxRHIFeatureLevel < ERHIFeatureLevel::SM5)
-			{
-				bMemoryless = true;
-				Desc.SetStorageMode(mtlpp::StorageMode::Memoryless);
-				Desc.SetResourceOptions(mtlpp::ResourceOptions::StorageModeMemoryless);
-			}
-#endif
 		}
 	}
 }
@@ -596,8 +574,7 @@ FAGXSurface::FAGXSurface(FAGXTextureCreateDesc const& CreateDesc)
 				CFRetain(ImageSurfaceRef);
 
 				mtlpp::TextureDescriptor DescCopy([CreateDesc.Desc.GetPtr() copy], ns::Ownership::Assign);
-				DescCopy.SetStorageMode(mtlpp::StorageMode::Managed);
-				DescCopy.SetResourceOptions((mtlpp::ResourceOptions)((DescCopy.GetResourceOptions() & ~(mtlpp::ResourceStorageModeMask)) | mtlpp::ResourceOptions::StorageModeManaged));
+				[DescCopy.GetPtr() setResourceOptions:(([DescCopy.GetPtr() resourceOptions] & ~MTLResourceStorageModeMask) | MTLResourceStorageModeManaged)];
 
 				ns::IOSurface MtlppIOSurface((IOSurfaceRef)ImageSurfaceRef, ns::Ownership::AutoRelease);
 				Texture = MTLPP_VALIDATE(mtlpp::Device, GMtlppDevice, AGXSafeGetRuntimeDebuggingLevel() >= EAGXDebugLevelValidation, NewTextureWithDescriptor(DescCopy, MtlppIOSurface, 0));
@@ -626,7 +603,7 @@ FAGXSurface::FAGXSurface(FAGXTextureCreateDesc const& CreateDesc)
 			const NSUInteger BytesPerRow = Align(CreateDesc.Desc.GetWidth() * GPixelFormats[CreateDesc.Format].BlockBytes, MinimumByteAlignment);
 
 			// Backing buffer resource options must match the texture we are going to create from it
-			FAGXPooledBufferArgs Args(BytesPerRow * CreateDesc.Desc.GetHeight(), BUF_Dynamic, mtlpp::StorageMode::Private, CreateDesc.Desc.GetCpuCacheMode());
+			FAGXPooledBufferArgs Args(BytesPerRow * CreateDesc.Desc.GetHeight(), BUF_Dynamic, FAGXPooledBufferArgs::PrivateStorageResourceOptions);
 			FAGXBuffer Buffer = GetAGXDeviceContext().CreatePooledBuffer(Args);
 
 			Texture = Buffer.NewTexture(CreateDesc.Desc, 0, BytesPerRow);
@@ -967,7 +944,7 @@ void FAGXSurface::UpdateSurfaceAndDestroySourceBuffer(id <MTLBuffer> SourceBuffe
 	}
 #endif
 	
-	if(Texture.GetStorageMode() == mtlpp::StorageMode::Private)
+	if ([Texture.GetPtr() storageMode] == MTLStorageModePrivate)
 	{
 		SCOPED_AUTORELEASE_POOL;
 		
@@ -1091,7 +1068,7 @@ void* FAGXSurface::Lock(uint32 MipIndex, uint32 ArrayIndex, EResourceLockMode Lo
 				);
 			}
 			
-			if (Texture.GetStorageMode() == mtlpp::StorageMode::Private)
+			if ([Texture.GetPtr() storageMode] == MTLStorageModePrivate)
 			{
 				// If we are running with command lists or the RHI thread is enabled we have to execute GFX commands in that context.
 				auto CopyTexToBuf =
@@ -1661,7 +1638,7 @@ static FAGXBuffer InternalCopyTexture2DUpdateRegion(FRHITexture2D* TextureRHI, c
 		&& (MTLPixelFormat)Texture->Texture.GetPixelFormat() == MTLPixelFormatRGBA8Unorm_sRGB)
 	{
 		const uint32 BufferSize = UpdateRegion.Height * UpdateRegion.Width * sizeof(uint32);
-		Buffer = GetAGXDeviceContext().CreatePooledBuffer(FAGXPooledBufferArgs(BufferSize, BUF_Dynamic, mtlpp::StorageMode::Shared));
+		Buffer = GetAGXDeviceContext().CreatePooledBuffer(FAGXPooledBufferArgs(BufferSize, BUF_Dynamic, FAGXPooledBufferArgs::SharedStorageResourceOptions));
 		InternalExpandR8ToStandardRGBA((uint32*)Buffer.GetContents(), UpdateRegion, InOutSourcePitch, SourceData);
 	}
 
@@ -1669,7 +1646,7 @@ static FAGXBuffer InternalCopyTexture2DUpdateRegion(FRHITexture2D* TextureRHI, c
 #endif
 	{
 		const uint32 BufferSize = UpdateRegion.Height * InSourcePitch;
-		Buffer = GetAGXDeviceContext().CreatePooledBuffer(FAGXPooledBufferArgs(BufferSize, BUF_Dynamic, mtlpp::StorageMode::Shared));
+		Buffer = GetAGXDeviceContext().CreatePooledBuffer(FAGXPooledBufferArgs(BufferSize, BUF_Dynamic, FAGXPooledBufferArgs::SharedStorageResourceOptions));
 
 		// Limit copy to line by line by update region pitch otherwise we can go off the end of source data on the last row
 		uint8* pDestRow = (uint8*)Buffer.GetContents();
@@ -1695,7 +1672,7 @@ static void InternalUpdateTexture2D(FAGXContext& Context, FRHITexture2D* Texture
 	
 	mtlpp::Region Region = mtlpp::Region(UpdateRegion.DestX, UpdateRegion.DestY, UpdateRegion.Width, UpdateRegion.Height);
 	
-	if(Tex.GetStorageMode() == mtlpp::StorageMode::Private)
+	if ([Tex.GetPtr() storageMode] == MTLStorageModePrivate)
 	{
 		SCOPED_AUTORELEASE_POOL;
 		
@@ -1769,7 +1746,7 @@ void FAGXDynamicRHI::RHIUpdateTexture2D(FRHITexture2D* TextureRHI, uint32 MipInd
 	{
 		FAGXSurface* Texture = ResourceCast(TextureRHI);
 		FAGXTexture Tex = Texture->Texture;
-		bool const bUseIntermediateMetalBuffer = Tex.GetStorageMode() == mtlpp::StorageMode::Private;
+		bool const bUseIntermediateMetalBuffer = ([Tex.GetPtr() storageMode] == MTLStorageModePrivate);
 		
 		if(bUseIntermediateMetalBuffer)
 		{
@@ -1828,7 +1805,7 @@ static void InternalUpdateTexture3D(FAGXContext& Context, FRHITexture3D* Texture
 	
 	mtlpp::Region Region = mtlpp::Region(UpdateRegion.DestX, UpdateRegion.DestY, UpdateRegion.DestZ, UpdateRegion.Width, UpdateRegion.Height, UpdateRegion.Depth);
 	
-	if(Tex.GetStorageMode() == mtlpp::StorageMode::Private)
+	if ([Tex.GetPtr() storageMode] == MTLStorageModePrivate)
 	{
 		const uint32 BytesPerImage = SourceRowPitch * UpdateRegion.Height;
 		mtlpp::BlitOption Options = mtlpp::BlitOption::None;
@@ -1873,7 +1850,7 @@ struct FAGXDynamicRHIUpdateTexture3DCommand final : public FRHICommand<FAGXDynam
 		FAGXTexture Tex = Texture->Texture;
 		const uint32 BufferSize = UpdateRegion.Height * UpdateRegion.Depth* SourceRowPitch;
 		
-		Buffer = GetAGXDeviceContext().CreatePooledBuffer(FAGXPooledBufferArgs(BufferSize, BUF_Dynamic, mtlpp::StorageMode::Shared));
+		Buffer = GetAGXDeviceContext().CreatePooledBuffer(FAGXPooledBufferArgs(BufferSize, BUF_Dynamic, FAGXPooledBufferArgs::SharedStorageResourceOptions));
 		InternalCopyTexture3DUpdateRegionData(Texture, UpdateRegion, SourceRowPitch, SourceDepthPitch, SourceData, (uint8*)Buffer.GetContents());
 	}
 	
@@ -1939,12 +1916,12 @@ void FAGXDynamicRHI::RHIUpdateTexture3D(FRHITexture3D* TextureRHI,uint32 MipInde
 #if PLATFORM_MAC
 		checkf(!(Texture->GetFormat() == PF_G8 && EnumHasAnyFlags(Texture->GetFlags(), TexCreate_SRGB) && (MTLPixelFormat)Tex.GetPixelFormat() == MTLPixelFormatRGBA8Unorm_sRGB), TEXT("AGXRHI does not support PF_G8_sRGB on 3D, array or cube textures as it requires manual, CPU-side expansion to RGBA8_sRGB which is expensive!"));
 #endif
-		if(Tex.GetStorageMode() == mtlpp::StorageMode::Private)
+		if ([Tex.GetPtr() storageMode] == MTLStorageModePrivate)
 		{
 			SCOPED_AUTORELEASE_POOL;
 
 			const uint32 BufferSize = UpdateRegion.Height * UpdateRegion.Depth * SourceRowPitch;
-			FAGXBuffer IntermediateBuffer = GetAGXDeviceContext().CreatePooledBuffer(FAGXPooledBufferArgs(BufferSize, BUF_Dynamic, mtlpp::StorageMode::Shared));
+			FAGXBuffer IntermediateBuffer = GetAGXDeviceContext().CreatePooledBuffer(FAGXPooledBufferArgs(BufferSize, BUF_Dynamic, FAGXPooledBufferArgs::SharedStorageResourceOptions));
 			InternalCopyTexture3DUpdateRegionData(TextureRHI, UpdateRegion, SourceRowPitch, SourceDepthPitch, SourceData, (uint8*)IntermediateBuffer.GetContents());
 			InternalUpdateTexture3D(ImmediateContext.GetInternalContext(), TextureRHI, MipIndex, UpdateRegion, SourceRowPitch, SourceDepthPitch, IntermediateBuffer);
 			GetAGXDeviceContext().ReleaseBuffer(IntermediateBuffer);
@@ -2106,7 +2083,7 @@ void FAGXRHICommandContext::RHICopyTexture(FRHITexture* SourceTextureRHI, FRHITe
 						const uint32 BytesPerImage = AlignedStride *  SourceSize.height;
 						const uint32 DataSize = BytesPerImage * SourceSize.depth;
 						
-						FAGXBuffer Buffer = GetAGXDeviceContext().CreatePooledBuffer(FAGXPooledBufferArgs(DataSize, BUF_Dynamic, mtlpp::StorageMode::Shared));
+						FAGXBuffer Buffer = GetAGXDeviceContext().CreatePooledBuffer(FAGXPooledBufferArgs(DataSize, BUF_Dynamic, FAGXPooledBufferArgs::SharedStorageResourceOptions));
 						
 						check(Buffer);
 						
diff --git a/Engine/Source/Runtime/Apple/AGXRHI/Private/AGXUAV.cpp b/Engine/Source/Runtime/Apple/AGXRHI/Private/AGXUAV.cpp
--- a/Engine/Source/Runtime/Apple/AGXRHI/Private/AGXUAV.cpp
+++ b/Engine/Source/Runtime/Apple/AGXRHI/Private/AGXUAV.cpp
@@ -99,7 +99,7 @@ FAGXResourceViewBase::FAGXResourceViewBase(
 		// TODO: Apple Silicon supports memoryless
 #if PLATFORM_IOS
 		// Memoryless targets can't have texture views (SRVs or UAVs)
-		if (SourceTexture->Texture.GetStorageMode() != mtlpp::StorageMode::Memoryless)
+		if ([SourceTexture->Texture.GetPtr() storageMode] != MTLStorageModeMemoryless)
 #endif
 		{
 			// Determine the appropriate metal format for the view.
@@ -448,7 +448,7 @@ void FAGXRHICommandContext::ClearUAVWithBlitEncoder(FRHIUnorderedAccessView* Uno
 	check(Type != EAGXRHIClearUAVType::VertexBuffer || EnumHasAnyFlags(SourceBuffer->GetUsage(), BUF_ByteAddressBuffer));
 
 	uint32 AlignedSize = Align(Size, BufferOffsetAlignment);
-	FAGXPooledBufferArgs Args(AlignedSize, BUF_Dynamic, mtlpp::StorageMode::Shared);
+	FAGXPooledBufferArgs Args(AlignedSize, BUF_Dynamic, FAGXPooledBufferArgs::SharedStorageResourceOptions);
 	FAGXBuffer Temp = GetAGXDeviceContext().CreatePooledBuffer(Args);
 	uint32* ContentBytes = (uint32*)Temp.GetContents();
 	for (uint32 Element = 0; Element < (AlignedSize >> 2); ++Element)
@@ -597,7 +597,7 @@ void FAGXRHICommandContext::RHICopyToStagingBuffer(FRHIBuffer* SourceBufferRHI,
 			{
 				AGXSafeReleaseMetalBuffer(ReadbackBuffer);
 			}
-			FAGXPooledBufferArgs ArgsCPU(NumBytes, BUF_Dynamic, mtlpp::StorageMode::Shared);
+			FAGXPooledBufferArgs ArgsCPU(NumBytes, BUF_Dynamic, FAGXPooledBufferArgs::SharedStorageResourceOptions);
 			ReadbackBuffer = GetAGXDeviceContext().CreatePooledBuffer(ArgsCPU);
 		}
 
diff --git a/Engine/Source/Runtime/Apple/AGXRHI/Private/AGXVertexBuffer.cpp b/Engine/Source/Runtime/Apple/AGXRHI/Private/AGXVertexBuffer.cpp
--- a/Engine/Source/Runtime/Apple/AGXRHI/Private/AGXVertexBuffer.cpp
+++ b/Engine/Source/Runtime/Apple/AGXRHI/Private/AGXVertexBuffer.cpp
@@ -126,7 +126,7 @@ FAGXRHIBuffer::FAGXRHIBuffer(uint32 InSize, EBufferUsageFlags InUsage, EAGXBuffe
 , Size(InSize)
 , Usage(InUsage)
 , AgxUsage(InAgxUsage)
-, Mode(BUFFER_STORAGE_MODE)
+, StorageMode(BUFFER_STORAGE_MODE)
 , Type(InType)
 {
 	// No life-time usage information? Enforce Dynamic.
@@ -142,8 +142,8 @@ FAGXRHIBuffer::FAGXRHIBuffer(uint32 InSize, EBufferUsageFlags InUsage, EAGXBuffe
 	
 	check(bIsStatic ^ bIsDynamic ^ bIsVolatile);
 
-	Mode = UsePrivateMemory() ? mtlpp::StorageMode::Private : BUFFER_STORAGE_MODE;
-	Mode = CanUsePrivateMemory() ? mtlpp::StorageMode::Private : Mode;
+	StorageMode = UsePrivateMemory() ? MTLStorageModePrivate : BUFFER_STORAGE_MODE;
+	StorageMode = CanUsePrivateMemory() ? MTLStorageModePrivate : StorageMode;
 	
 	if (InSize)
 	{
@@ -160,53 +160,6 @@ FAGXRHIBuffer::FAGXRHIBuffer(uint32 InSize, EBufferUsageFlags InUsage, EAGXBuffe
 		{
 			uint32 AllocSize = Size;
 			
-			if (EnumHasAnyFlags(AgxUsage, EAGXBufferUsage::LinearTex) && !FAGXCommandQueue::SupportsFeature(EAGXFeaturesTextureBuffers))
-			{
-				if (EnumHasAnyFlags(Usage, BUF_UnorderedAccess))
-				{
-					// Padding for write flushing when not using linear texture bindings for buffers
-					AllocSize = Align(AllocSize + 512, 1024);
-				}
-				
-				if (bWantsView)
-				{
-					uint32 NumElements = AllocSize;
-					uint32 SizeX = NumElements;
-					uint32 SizeY = 1;
-					uint32 Dimension = GMaxTextureDimensions;
-					while (SizeX > GMaxTextureDimensions)
-					{
-						while((NumElements % Dimension) != 0)
-						{
-							check(Dimension >= 1);
-							Dimension = (Dimension >> 1);
-						}
-						SizeX = Dimension;
-						SizeY = NumElements / Dimension;
-						if(SizeY > GMaxTextureDimensions)
-						{
-							Dimension <<= 1;
-							checkf(SizeX <= GMaxTextureDimensions, TEXT("Calculated width %u is greater than maximum permitted %d when converting buffer of size %u to a 2D texture."), Dimension, (int32)GMaxTextureDimensions, AllocSize);
-							if(Dimension <= GMaxTextureDimensions)
-							{
-								AllocSize = Align(Size, Dimension);
-								NumElements = AllocSize;
-								SizeX = NumElements;
-							}
-							else
-							{
-								// We don't know the Pixel Format and so the bytes per element for the potential linear texture
-								// Use max texture dimension as the align to be a worst case rather than crashing
-								AllocSize = Align(Size, GMaxTextureDimensions);
-								break;
-							}
-						}
-					}
-					
-					AllocSize = Align(AllocSize, 1024);
-				}
-			}
-			
 			// Static buffers will never be discarded. You can update them directly.
 			if(bIsStatic)
 			{
@@ -236,24 +189,22 @@ FAGXRHIBuffer::FAGXRHIBuffer(uint32 InSize, EBufferUsageFlags InUsage, EAGXBuffe
 			{
 				FAGXBufferAndViews& Backing = BufferPool[i];
 
+				MTLResourceOptions Options = MTLResourceCPUCacheModeDefaultCache | (StorageMode << MTLResourceStorageModeShift) | MTLResourceHazardTrackingModeDefault;
 #if METAL_POOL_BUFFER_BACKING
-				FAGXPooledBufferArgs ArgsCPU(AllocSize, Usage, Mode);
+				FAGXPooledBufferArgs ArgsCPU(AllocSize, Usage, Options);
 				Backing.Buffer = GetAGXDeviceContext().CreatePooledBuffer(ArgsCPU);
 				Backing.Buffer.SetOwner(nullptr, false);
 #else
-				
-				NSUInteger Options = (((NSUInteger) Mode) << mtlpp::ResourceStorageModeShift);
-				
 				METAL_GPUPROFILE(FAGXScopedCPUStats CPUStat(FString::Printf(TEXT("AllocBuffer: %llu, %llu"), AllocSize, Options)));
 				// Allocate one.
-				Backing.Buffer = FAGXBuffer(MTLPP_VALIDATE(mtlpp::Device, GMtlppDevice, AGXSafeGetRuntimeDebuggingLevel() >= EAGXDebugLevelValidation, NewBuffer(AllocSize, (mtlpp::ResourceOptions) Options)), false);
+				Backing.Buffer = FAGXBuffer(MTLPP_VALIDATE(mtlpp::Device, GMtlppDevice, AGXSafeGetRuntimeDebuggingLevel() >= EAGXDebugLevelValidation, NewBuffer(AllocSize, mtlpp::ResourceOptions(Options))), false);
 				
 #if STATS || ENABLE_LOW_LEVEL_MEM_TRACKER
 				AGXLLM::LogAllocBuffer(Backing.Buffer);
 #endif
 				INC_MEMORY_STAT_BY(STAT_AGXDeviceBufferMemory, Backing.Buffer.GetLength());
 				
-				if (GAGXBufferZeroFill && Mode != mtlpp::StorageMode::Private)
+				if (GAGXBufferZeroFill && StorageMode != MTLStorageModePrivate)
 				{
 					FMemory::Memset(((uint8*)Backing.Buffer.GetContents()), 0, Backing.Buffer.GetLength());
 				}
@@ -277,7 +228,7 @@ FAGXRHIBuffer::FAGXRHIBuffer(uint32 InSize, EBufferUsageFlags InUsage, EAGXBuffe
 			{
 				check(Backing.Buffer);
 				check(AllocSize <= Backing.Buffer.GetLength());
-				check(Backing.Buffer.GetStorageMode() == Mode);
+				check([Backing.Buffer.GetPtr() storageMode] == StorageMode);
 				check(Backing.Views.Num() == 0);
 			}
 		}
@@ -317,12 +268,12 @@ FAGXRHIBuffer::~FAGXRHIBuffer()
 void FAGXRHIBuffer::AllocTransferBuffer(bool bOnRHIThread, uint32 InSize, EResourceLockMode LockMode)
 {
 	check(!TransferBuffer);
-	FAGXPooledBufferArgs ArgsCPU(InSize, BUF_Dynamic, mtlpp::StorageMode::Shared);
+	FAGXPooledBufferArgs ArgsCPU(InSize, BUF_Dynamic, FAGXPooledBufferArgs::SharedStorageResourceOptions);
 	TransferBuffer = GetAGXDeviceContext().CreatePooledBuffer(ArgsCPU);
 	TransferBuffer.SetOwner(nullptr, false);
 	check(TransferBuffer && TransferBuffer.GetPtr());
 	METAL_INC_DWORD_STAT_BY(Type, MemAlloc, InSize, Usage);
-	METAL_FATAL_ASSERT(TransferBuffer, TEXT("Failed to create buffer of size %u and storage mode %u"), InSize, (uint32)mtlpp::StorageMode::Shared);
+	METAL_FATAL_ASSERT(TransferBuffer, TEXT("Failed to create buffer of size %u and storage mode %u"), InSize, (uint32)MTLStorageModeShared);
 }
 
 void FAGXRHIBuffer::AllocLinearTextures(const LinearTextureMapKey& InLinearTextureMapKey)
@@ -370,63 +321,11 @@ void FAGXRHIBuffer::AllocLinearTextures(const LinearTextureMapKey& InLinearTextu
 
 		uint32 RowBytes = NumElements * BytesPerElement;
 
-		if (FAGXCommandQueue::SupportsFeature(EAGXFeaturesTextureBuffers))
 		{
-			NSUInteger Options = ((NSUInteger) Mode) << mtlpp::ResourceStorageModeShift;
+			MTLResourceOptions Options = MTLResourceCPUCacheModeDefaultCache | (StorageMode << MTLResourceStorageModeShift) | MTLResourceHazardTrackingModeDefault;
 			Desc = mtlpp::TextureDescriptor::TextureBufferDescriptor((mtlpp::PixelFormat)MTLFormat, NumElements, mtlpp::ResourceOptions(Options), mtlpp::TextureUsage(TexUsage));
 			Desc.SetAllowGPUOptimisedContents(false);
 		}
-		else
-		{
-			uint32 Width = NumElements;
-			uint32 Height = 1;
-
-			if (NumElements > GMaxTextureDimensions)
-			{
-				uint32 Dimension = GMaxTextureDimensions;
-				while ((NumElements % Dimension) != 0)
-				{
-					check(Dimension >= 1);
-					Dimension = (Dimension >> 1);
-				}
-
-				Width = Dimension;
-				Height = NumElements / Dimension;
-
-				// If we're just trying to fit as many elements as we can into
-				// the available buffer space, we can trim some padding at the
-				// end of the buffer in order to create widest possible linear
-				// texture that will fit.
-				if ((UINT_MAX == LinearTextureDesc.NumElements) && (Height > GMaxTextureDimensions))
-				{
-					Width = GMaxTextureDimensions;
-					Height = 1;
-
-					while ((Width * Height) < NumElements)
-					{
-						Height <<= 1;
-					}
-
-					while ((Width * Height) > NumElements)
-					{
-						Height -= 1;
-					}
-				}
-
-				checkf(Width <= GMaxTextureDimensions, TEXT("Calculated width %u is greater than maximum permitted %d when converting buffer of size %llu with element stride %u to a 2D texture with %u elements."), Width, (int32)GMaxTextureDimensions, Length, BytesPerElement, NumElements);
-				checkf(Height <= GMaxTextureDimensions, TEXT("Calculated height %u is greater than maximum permitted %d when converting buffer of size %llu with element stride %u to a 2D texture with %u elements."), Height, (int32)GMaxTextureDimensions, Length, BytesPerElement, NumElements);
-			}
-
-			RowBytes = Width * BytesPerElement;
-
-			check(RowBytes % MinimumByteAlignment == 0);
-			check((RowBytes * Height) + Offset <= Length);
-
-			Desc = mtlpp::TextureDescriptor::Texture2DDescriptor((mtlpp::PixelFormat)MTLFormat, Width, Height, NO);
-			Desc.SetStorageMode(Mode);
-			Desc.SetCpuCacheMode(CurrentBuffer.GetCpuCacheMode());
-			Desc.SetUsage((mtlpp::TextureUsage)TexUsage);
-		}
 
 		for(FAGXBufferAndViews& Backing : BufferPool)
 		{
@@ -560,7 +459,7 @@ void* FAGXRHIBuffer::Lock(bool bIsOnRHIThread, EResourceLockMode InLockMode, uin
 
 		// Use transfer buffer for writing into 'Static' buffers as they could be in use by GPU atm
 		// Initialization of 'Static' buffers still uses direct copy when possible
-		const bool bUseTransferBuffer = (Mode == mtlpp::StorageMode::Private || (Mode == mtlpp::StorageMode::Shared && bIsStatic));
+		const bool bUseTransferBuffer = (StorageMode == MTLStorageModePrivate || (StorageMode == MTLStorageModeShared && bIsStatic));
 		if(bUseTransferBuffer)
 		{
 			FAGXFrameAllocator::AllocationEntry TempBacking = GetAGXDeviceContext().GetTransferAllocator()->AcquireSpace(Len);
@@ -581,7 +480,7 @@ void* FAGXRHIBuffer::Lock(bool bIsOnRHIThread, EResourceLockMode InLockMode, uin
 		// assumes offset is 0 for reads.
 		check(Offset == 0);
 		
-		if(Mode == mtlpp::StorageMode::Private)
+		if (StorageMode == MTLStorageModePrivate)
 		{
 			SCOPE_CYCLE_COUNTER(STAT_AGXBufferPageOffTime);
 			AllocTransferBuffer(true, Len, RLM_WriteOnly);
@@ -595,8 +494,8 @@ void* FAGXRHIBuffer::Lock(bool bIsOnRHIThread, EResourceLockMode InLockMode, uin
 			
 			ReturnPointer = TransferBuffer.GetContents();
 		}
-		#if PLATFORM_MAC
-		else if(Mode == mtlpp::StorageMode::Managed)
+#if PLATFORM_MAC
+		else if (StorageMode == MTLStorageModeManaged)
 		{
 			SCOPE_CYCLE_COUNTER(STAT_AGXBufferPageOffTime);
 			
@@ -608,7 +507,7 @@ void* FAGXRHIBuffer::Lock(bool bIsOnRHIThread, EResourceLockMode InLockMode, uin
 			
 			ReturnPointer = GetCurrentBackingInternal().Buffer.GetContents();
 		}
-		#endif
+#endif
 		else
 		{
 			// Shared
@@ -639,7 +538,7 @@ void FAGXRHIBuffer::Unlock()
 {
 	check(MetalIsSafeToUseRHIThreadResources());
 
-	if(!Data)
+	if (!Data)
 	{
 		FAGXBufferAndViews& Backing = GetCurrentBackingInternal();
 		FAGXBuffer& CurrentBuffer = Backing.Buffer;
@@ -649,7 +548,7 @@ void FAGXRHIBuffer::Unlock()
 		const bool bWriteLock = CurrentLockMode == RLM_WriteOnly;
 		const bool bIsStatic = EnumHasAnyFlags(Usage, BUF_Static);
 		
-		if(bWriteLock)
+		if (bWriteLock)
 		{
 			check(!TransferBuffer);
 			check(LockOffset == 0);
@@ -657,15 +556,15 @@ void FAGXRHIBuffer::Unlock()
 			
 			// Use transfer buffer for writing into 'Static' buffers as they could be in use by GPU atm
 			// Initialization of 'Static' buffers still uses direct copy when possible
-			const bool bUseTransferBuffer = (Mode == mtlpp::StorageMode::Private || (Mode == mtlpp::StorageMode::Shared && bIsStatic));
-			if(bUseTransferBuffer)
+			const bool bUseTransferBuffer = (StorageMode == MTLStorageModePrivate || (StorageMode == MTLStorageModeShared && bIsStatic));
+			if (bUseTransferBuffer)
 			{
 				FAGXFrameAllocator::AllocationEntry Entry = GetAGXDeviceContext().FetchAndRemoveLock(this);
 				FAGXBuffer Transfer = Entry.Backing;
 				GetAGXDeviceContext().AsyncCopyFromBufferToBuffer(Transfer, Entry.Offset, CurrentBuffer, 0, LockSize);
 			}
 #if PLATFORM_MAC
-			else if(Mode == mtlpp::StorageMode::Managed)
+			else if (StorageMode == MTLStorageModeManaged)
 			{
 				if (GAGXBufferZeroFill)
 					MTLPP_VALIDATE(mtlpp::Buffer, CurrentBuffer, AGXSafeGetRuntimeDebuggingLevel() >= EAGXDebugLevelValidation, DidModify(ns::Range(0, CurrentBuffer.GetLength())));
@@ -676,15 +575,15 @@ void FAGXRHIBuffer::Unlock()
 			else
 			{
 				// shared buffers are always mapped so nothing happens
-				check(Mode == mtlpp::StorageMode::Shared);
+				check(StorageMode == MTLStorageModeShared);
 			}
 		}
 		else
 		{
 			check(CurrentLockMode == RLM_ReadOnly);
-			if(TransferBuffer)
+			if (TransferBuffer)
 			{
-				check(Mode == mtlpp::StorageMode::Private);
+				check(StorageMode == MTLStorageModePrivate);
 				AGXSafeReleaseMetalBuffer(TransferBuffer);
 				TransferBuffer = nil;
 			}
@@ -758,7 +657,7 @@ FBufferRHIRef FAGXDynamicRHI::RHICreateBuffer(uint32 Size, EBufferUsageFlags Usa
 		// Discard the resource array's contents.
 		CreateInfo.ResourceArray->Discard();
 	}
-	else if (Buffer->Mode == mtlpp::StorageMode::Private)
+	else if (Buffer->StorageMode == MTLStorageModePrivate)
 	{
 		check (!Buffer->TransferBuffer);
 
@@ -772,7 +671,7 @@ FBufferRHIRef FAGXDynamicRHI::RHICreateBuffer(uint32 Size, EBufferUsageFlags Usa
 		}
 	}
 #if PLATFORM_MAC
-	else if (GAGXBufferZeroFill && Buffer->Mode == mtlpp::StorageMode::Managed)
+	else if (GAGXBufferZeroFill && Buffer->StorageMode == MTLStorageModeManaged)
 	{
 		for(FAGXRHIBuffer::FAGXBufferAndViews& Backing : Buffer->BufferPool)
 		{
@@ -821,7 +720,7 @@ void FAGXDynamicRHI::RHICopyBuffer(FRHIBuffer* SourceBufferRHI, FRHIBuffer* Dest
 		}
 		else if (TheDstBuffer)
 		{
-			FAGXPooledBufferArgs ArgsCPU(SrcBuffer->GetSize(), BUF_Dynamic, mtlpp::StorageMode::Shared);
+			FAGXPooledBufferArgs ArgsCPU(SrcBuffer->GetSize(), BUF_Dynamic, FAGXPooledBufferArgs::SharedStorageResourceOptions);
 			FAGXBuffer TempBuffer = GetAGXDeviceContext().CreatePooledBuffer(ArgsCPU);
 			FMemory::Memcpy(TempBuffer.GetContents(), SrcBuffer->Data->Data, SrcBuffer->GetSize());
 			GetAGXDeviceContext().CopyFromBufferToBuffer(TempBuffer, 0, TheDstBuffer, 0, FMath::Min(SrcBuffer->GetSize(), DstBuffer->GetSize()));
@@ -868,7 +767,7 @@ void FAGXRHIBuffer::Init_RenderThread(FRHICommandListImmediate& RHICmdList, uint
 		}
 		else
 		{
-			if (Mode == mtlpp::StorageMode::Private)
+			if (StorageMode == MTLStorageModePrivate)
 			{
 				if (RHICmdList.IsBottomOfPipe())
 				{
@@ -896,7 +795,7 @@ void FAGXRHIBuffer::Init_RenderThread(FRHICommandListImmediate& RHICmdList, uint
 				FAGXBuffer& TheBuffer = GetCurrentBufferInternal();
 				FMemory::Memcpy(TheBuffer.GetContents(), CreateInfo.ResourceArray->GetResourceData(), InSize);
 	#if PLATFORM_MAC
-				if(Mode == mtlpp::StorageMode::Managed)
+				if (StorageMode == MTLStorageModeManaged)
 				{
 					MTLPP_VALIDATE(mtlpp::Buffer, TheBuffer, AGXSafeGetRuntimeDebuggingLevel() >= EAGXDebugLevelValidation, DidModify(ns::Range(0, GAGXBufferZeroFill ? TheBuffer.GetLength() : InSize)));
 				}
diff --git a/Engine/Source/Runtime/Apple/AGXRHI/Private/Shaders/Types/AGXComputeShader.cpp b/Engine/Source/Runtime/Apple/AGXRHI/Private/Shaders/Types/AGXComputeShader.cpp
--- a/Engine/Source/Runtime/Apple/AGXRHI/Private/Shaders/Types/AGXComputeShader.cpp
+++ b/Engine/Source/Runtime/Apple/AGXRHI/Private/Shaders/Types/AGXComputeShader.cpp
@@ -46,10 +46,7 @@ FAGXShaderPipeline* FAGXComputeShader::GetPipeline()
 		MTLComputePipelineDescriptor* Descriptor = [[MTLComputePipelineDescriptor alloc] init];
 		[Descriptor setLabel:[Func name]];
 		[Descriptor setComputeFunction:Func];
-		if (FAGXCommandQueue::SupportsFeature(EAGXFeaturesTextureBuffers))
-		{
-			[Descriptor setMaxTotalThreadsPerThreadgroup:(NumThreadsX * NumThreadsY * NumThreadsZ)];
-		}
+		[Descriptor setMaxTotalThreadsPerThreadgroup:(NumThreadsX * NumThreadsY * NumThreadsZ)];
 
 		if (FAGXCommandQueue::SupportsFeature(EAGXFeaturesPipelineBufferMutability))
 		{
diff --git a/Engine/Source/Runtime/Apple/AGXRHI/Public/AGXResources.h b/Engine/Source/Runtime/Apple/AGXRHI/Public/AGXResources.h
--- a/Engine/Source/Runtime/Apple/AGXRHI/Public/AGXResources.h
+++ b/Engine/Source/Runtime/Apple/AGXRHI/Public/AGXResources.h
@@ -440,7 +440,7 @@ public:
 	EAGXBufferUsage AgxUsage;
 	
 	// Storage mode
-	mtlpp::StorageMode Mode;
+	MTLStorageMode StorageMode;
 	
 	// Resource type
 	ERHIResourceType Type;
