Fix memory leak of MTLTextureDescriptor on macOS/iOS when using the RHI thread and generating the MetalRHIBuffer linear texture view - add SCOPED_AUTORELEASE_POOL to FMetalRHIBuffer::CreateLinearTexture().

diff --git a/Engine/Source/Runtime/Apple/AGXRHI/Private/AGXVertexBuffer.cpp b/Engine/Source/Runtime/Apple/AGXRHI/Private/AGXVertexBuffer.cpp
--- a/Engine/Source/Runtime/Apple/AGXRHI/Private/AGXVertexBuffer.cpp
+++ b/Engine/Source/Runtime/Apple/AGXRHI/Private/AGXVertexBuffer.cpp
@@ -377,6 +377,8 @@ struct FAGXRHICommandCreateLinearTexture : public FRHICommand<FAGXRHICommandCrea
 
 void FAGXRHIBuffer::CreateLinearTexture(EPixelFormat InFormat, FRHIResource* InParent, const FAGXLinearTextureDescriptor* InLinearTextureDescriptor)
 {
+	SCOPED_AUTORELEASE_POOL;
+
 	if (EnumHasAnyFlags(Usage, BUF_UnorderedAccess|BUF_ShaderResource) && GAGXBufferFormats[InFormat].LinearTextureFormat != MTLPixelFormatInvalid)
 	{
 		if (IsRunningRHIInSeparateThread() && !IsInRHIThread() && !FRHICommandListExecutor::GetImmediateCommandList().Bypass())
diff --git a/Engine/Source/Runtime/Apple/MetalRHI/Private/MetalVertexBuffer.cpp b/Engine/Source/Runtime/Apple/MetalRHI/Private/MetalVertexBuffer.cpp
--- a/Engine/Source/Runtime/Apple/MetalRHI/Private/MetalVertexBuffer.cpp
+++ b/Engine/Source/Runtime/Apple/MetalRHI/Private/MetalVertexBuffer.cpp
@@ -485,6 +485,8 @@ struct FMetalRHICommandCreateLinearTexture : public FRHICommand<FMetalRHICommand
 
 void FMetalRHIBuffer::CreateLinearTexture(EPixelFormat InFormat, FRHIResource* InParent, const FMetalLinearTextureDescriptor* InLinearTextureDescriptor)
 {
+	SCOPED_AUTORELEASE_POOL;
+	
 	if (EnumHasAnyFlags(Usage, BUF_UnorderedAccess|BUF_ShaderResource) && GMetalBufferFormats[InFormat].LinearTextureFormat != mtlpp::PixelFormat::Invalid)
 	{
 		if (IsRunningRHIInSeparateThread() && !IsInRHIThread() && !FRHICommandListExecutor::GetImmediateCommandList().Bypass())
