Restore 2DArray mapping of Cube UAVs for Metal after a refactoring pass blew it away.

- We have to support older Apple GPU Family devices (< 4), so we cannot simply use the obvious mapping. This is a quick workaround to resolve a metal validation layer error while we implement a better solution that uses the default/obvious mapping on newer/supported devices.

#rnx

diff --git a/Engine/Source/Runtime/Apple/AGXRHI/Private/AGXUAV.cpp b/Engine/Source/Runtime/Apple/AGXRHI/Private/AGXUAV.cpp
--- a/Engine/Source/Runtime/Apple/AGXRHI/Private/AGXUAV.cpp
+++ b/Engine/Source/Runtime/Apple/AGXRHI/Private/AGXUAV.cpp
@@ -78,6 +78,7 @@ FAGXResourceViewBase::FAGXResourceViewBase(
 	, ERHITextureSRVOverrideSRGBType InSRGBOverride
 	, uint32 InFirstArraySlice
 	, uint32 InNumArraySlices
+	, bool bInUAV
 	)
 	: SourceTexture    (ResourceCast(InTexture))
 	, bTexture         (true)
@@ -154,7 +155,8 @@ FAGXResourceViewBase::FAGXResourceViewBase(
 			bool bUseSourceTex =
 				   MipLevel == 0
 				&& NumMips == SourceTexture->Texture.GetMipmapLevelCount()
-				&& MetalFormat == (MTLPixelFormat)SourceTexture->Texture.GetPixelFormat();
+				&& MetalFormat == (MTLPixelFormat)SourceTexture->Texture.GetPixelFormat()
+				&& !(bInUAV && SourceTexture->GetDesc().IsTextureCube());		// @todo: Remove this once Cube UAV supported for all Metal Devices
 
 			if (bUseSourceTex)
 			{
@@ -169,10 +171,13 @@ FAGXResourceViewBase::FAGXResourceViewBase(
 					SourceTexture->PrepareTextureView();
 
 				ns::Range Slices(0, SourceTexture->Texture.GetArrayLength() * (SourceTexture->GetDesc().IsTextureCube() ? 6 : 1));
+				
+				// @todo: Remove this type swizzle once Cube UAV supported for all Metal Devices - SRV seem to want to stay as cube but UAV are expected to be 2DArray
+				mtlpp::TextureType TextureType = bInUAV && SourceTexture->GetDesc().IsTextureCube() ? mtlpp::TextureType::Texture2DArray : SourceTexture->Texture.GetTextureType();
 
 				TextureView = SourceTexture->Texture.NewTextureView(
 					(mtlpp::PixelFormat)MetalFormat,
-					SourceTexture->Texture.GetTextureType(),
+					TextureType,
 					ns::Range(MipLevel, NumMips),
 					Slices);
 			}
diff --git a/Engine/Source/Runtime/Apple/AGXRHI/Public/AGXResources.h b/Engine/Source/Runtime/Apple/AGXRHI/Public/AGXResources.h
--- a/Engine/Source/Runtime/Apple/AGXRHI/Public/AGXResources.h
+++ b/Engine/Source/Runtime/Apple/AGXRHI/Public/AGXResources.h
@@ -506,6 +506,7 @@ protected:
 		, ERHITextureSRVOverrideSRGBType InSRGBOverride
 		, uint32 InFirstArraySlice
 		, uint32 InNumArraySlices
+		, bool bInUAV
 	);
 
 public:
@@ -565,6 +566,7 @@ public:
 			, CreateInfo.SRGBOverride
 			, CreateInfo.FirstArraySlice
 			, CreateInfo.NumArraySlices
+			, false // bInUAV
 		)
 	{}
 
@@ -598,6 +600,7 @@ public:
 			, ERHITextureSRVOverrideSRGBType::SRGBO_ForceDisable
 			, FirstArraySlice
 			, NumArraySlices
+			, true // bInUAV
 		)
 	{}
 
diff --git a/Engine/Source/Runtime/Apple/MetalRHI/Private/MetalUAV.cpp b/Engine/Source/Runtime/Apple/MetalRHI/Private/MetalUAV.cpp
--- a/Engine/Source/Runtime/Apple/MetalRHI/Private/MetalUAV.cpp
+++ b/Engine/Source/Runtime/Apple/MetalRHI/Private/MetalUAV.cpp
@@ -77,6 +77,7 @@ FMetalResourceViewBase::FMetalResourceViewBase(
 	, ERHITextureSRVOverrideSRGBType InSRGBOverride
 	, uint32 InFirstArraySlice
 	, uint32 InNumArraySlices
+	, bool bInUAV
 	)
 	: SourceTexture(ResourceCast(InTexture))
 	, bTexture(true)
@@ -144,7 +145,8 @@ FMetalResourceViewBase::FMetalResourceViewBase(
 			bool bUseSourceTex =
 				MipLevel == 0
 				&& NumMips == SourceTexture->Texture.GetMipmapLevelCount()
-				&& MetalFormat == SourceTexture->Texture.GetPixelFormat();
+				&& MetalFormat == SourceTexture->Texture.GetPixelFormat()
+				&& !(bInUAV && SourceTexture->GetDesc().IsTextureCube());		// @todo: Remove this once Cube UAV supported for all Metal Devices
 
 			if (bUseSourceTex)
 			{
@@ -159,10 +161,13 @@ FMetalResourceViewBase::FMetalResourceViewBase(
 					SourceTexture->PrepareTextureView();
 
 				ns::Range Slices(0, SourceTexture->Texture.GetArrayLength() * (SourceTexture->GetDesc().IsTextureCube() ? 6 : 1));
+				
+				// @todo: Remove this type swizzle once Cube UAV supported for all Metal Devices - SRV seem to want to stay as cube but UAV are expected to be 2DArray
+				mtlpp::TextureType TextureType = bInUAV && SourceTexture->GetDesc().IsTextureCube() ? mtlpp::TextureType::Texture2DArray : SourceTexture->Texture.GetTextureType();
 
 				TextureView = SourceTexture->Texture.NewTextureView(
 					MetalFormat,
-					SourceTexture->Texture.GetTextureType(),
+					TextureType,
 					ns::Range(MipLevel, NumMips),
 					Slices);
 			}
diff --git a/Engine/Source/Runtime/Apple/MetalRHI/Public/MetalResources.h b/Engine/Source/Runtime/Apple/MetalRHI/Public/MetalResources.h
--- a/Engine/Source/Runtime/Apple/MetalRHI/Public/MetalResources.h
+++ b/Engine/Source/Runtime/Apple/MetalRHI/Public/MetalResources.h
@@ -506,6 +506,7 @@ protected:
 		, ERHITextureSRVOverrideSRGBType InSRGBOverride
 		, uint32 InFirstArraySlice
 		, uint32 InNumArraySlices
+		, bool bInUAV
 	);
 
 public:
@@ -565,6 +566,7 @@ public:
 			, CreateInfo.SRGBOverride
 			, CreateInfo.FirstArraySlice
 			, CreateInfo.NumArraySlices
+			, false // bInUAV
 		)
 	{}
 
@@ -598,6 +600,7 @@ public:
 			, ERHITextureSRVOverrideSRGBType::SRGBO_ForceDisable
 			, FirstArraySlice
 			, NumArraySlices
+			, true // bInUAV
 		)
 	{}
 
