AGXRHI: Remove mtlpp::TextureUsage and mtlpp::PrugableState types

#rnx

diff --git a/Engine/Source/Runtime/Apple/AGXRHI/Private/AGXBuffer.cpp b/Engine/Source/Runtime/Apple/AGXRHI/Private/AGXBuffer.cpp
--- a/Engine/Source/Runtime/Apple/AGXRHI/Private/AGXBuffer.cpp
+++ b/Engine/Source/Runtime/Apple/AGXRHI/Private/AGXBuffer.cpp
@@ -947,9 +947,9 @@ FAGXTexture FAGXTexturePool::CreateTexture(mtlpp::TextureDescriptor Desc)
 	Descriptor.arrayLength = Desc.GetArrayLength();
 	Descriptor.resourceOptions = Desc.GetResourceOptions();
 	Descriptor.usage = Desc.GetUsage();
-	if (Descriptor.usage == mtlpp::TextureUsage::Unknown)
+	if (Descriptor.usage == MTLTextureUsageUnknown)
 	{
-		Descriptor.usage = (mtlpp::TextureUsage)(mtlpp::TextureUsage::ShaderRead | mtlpp::TextureUsage::ShaderWrite | mtlpp::TextureUsage::RenderTarget | mtlpp::TextureUsage::PixelFormatView);
+		Descriptor.usage = (MTLTextureUsageShaderRead | MTLTextureUsageShaderWrite | MTLTextureUsageRenderTarget | MTLTextureUsagePixelFormatView);
 	}
 	Descriptor.freedFrame = 0;
 
@@ -962,7 +962,7 @@ FAGXTexture FAGXTexturePool::CreateTexture(mtlpp::TextureDescriptor Desc)
 		Pool.Remove(Descriptor);
 		if (GAGXResourcePurgeInPool)
 		{
-            Texture.SetPurgeableState(mtlpp::PurgeableState::NonVolatile);
+            [Texture.GetPtr() setPurgeableState:MTLPurgeableStateNonVolatile];
         }
 	}
 	else
@@ -991,9 +991,9 @@ void FAGXTexturePool::ReleaseTexture(FAGXTexture& Texture)
 	Descriptor.usage = Texture.GetUsage();
 	Descriptor.freedFrame = GFrameNumberRenderThread;
 	
-    if (GAGXResourcePurgeInPool && Texture.SetPurgeableState(mtlpp::PurgeableState::KeepCurrent) == mtlpp::PurgeableState::NonVolatile)
+    if (GAGXResourcePurgeInPool && [Texture.GetPtr() setPurgeableState:MTLPurgeableStateKeepCurrent] == MTLPurgeableStateNonVolatile)
     {
-        Texture.SetPurgeableState(mtlpp::PurgeableState::Volatile);
+        [Texture.GetPtr() setPurgeableState:MTLPurgeableStateVolatile];
     }
 	
 	FScopeLock Lock(&PoolMutex);
@@ -1019,7 +1019,7 @@ void FAGXTexturePool::Drain(bool const bForce)
             {
 				if (GAGXResourcePurgeInPool && (GFrameNumberRenderThread - It->Key.freedFrame) >= PurgeAfterNumFrames)
 				{
-					It->Value.SetPurgeableState(mtlpp::PurgeableState::Empty);
+					[It->Value.GetPtr() setPurgeableState:MTLPurgeableStateEmpty];
 				}
             }
         }
@@ -1175,7 +1175,7 @@ FAGXBuffer FAGXResourceHeap::CreateBuffer(uint32 Size, uint32 Alignment, EBuffer
 					Buffer = ManagedBuffers.CreatePooledResource(FAGXPooledBufferArgs(BlockSize, Flags, Options));
 					if (GAGXResourcePurgeInPool)
 					{
-						Buffer.SetPurgeableState(mtlpp::PurgeableState::NonVolatile);
+						[Buffer.GetPtr() setPurgeableState:MTLPurgeableStateNonVolatile];
 					}
 					DEC_MEMORY_STAT_BY(STAT_AGXBufferUnusedMemory, Buffer.GetLength());
 					DEC_MEMORY_STAT_BY(STAT_AGXPooledBufferUnusedMemory, Buffer.GetLength());
@@ -1252,7 +1252,7 @@ FAGXBuffer FAGXResourceHeap::CreateBuffer(uint32 Size, uint32 Alignment, EBuffer
 					Buffer = Buffers[Storage].CreatePooledResource(FAGXPooledBufferArgs(BlockSize, Flags, Options));
 					if (GAGXResourcePurgeInPool)
 					{
-                   		Buffer.SetPurgeableState(mtlpp::PurgeableState::NonVolatile);
+                   		[Buffer.GetPtr() setPurgeableState:MTLPurgeableStateNonVolatile];
 					}
 					DEC_MEMORY_STAT_BY(STAT_AGXBufferUnusedMemory, Buffer.GetLength());
 					DEC_MEMORY_STAT_BY(STAT_AGXPooledBufferUnusedMemory, Buffer.GetLength());
@@ -1301,7 +1301,7 @@ void FAGXResourceHeap::ReleaseBuffer(FAGXBuffer& Buffer)
 		
 		if (GAGXResourcePurgeInPool)
 		{
-       		Buffer.SetPurgeableState(mtlpp::PurgeableState::Volatile);
+       		[Buffer.GetPtr() setPurgeableState:MTLPurgeableStateVolatile];
 		}
         
 		switch (StorageMode)
@@ -1338,7 +1338,7 @@ FAGXTexture FAGXResourceHeap::CreateTexture(mtlpp::TextureDescriptor Desc, FAGXS
 	LLM_SCOPE_METAL(ELLMTagAGX::Textures);
 	LLM_PLATFORM_SCOPE_METAL(ELLMTagAGX::Textures);
 	
-	if (Desc.GetUsage() & mtlpp::TextureUsage::RenderTarget)
+	if ([Desc.GetPtr() usage] & MTLTextureUsageRenderTarget)
 	{
 		LLM_PLATFORM_SCOPE_METAL(ELLMTagAGX::RenderTargets);
 		return TargetPool.CreateTexture(Desc);
@@ -1353,7 +1353,7 @@ void FAGXResourceHeap::ReleaseTexture(FAGXSurface* Surface, FAGXTexture& Texture
 {
 	if (Texture && !Texture.GetBuffer() && !Texture.GetParentTexture())
 	{
-        if (Texture.GetUsage() & mtlpp::TextureUsage::RenderTarget)
+        if ([Texture.GetPtr() usage] & MTLTextureUsageRenderTarget)
         {
            	TargetPool.ReleaseTexture(Texture);
         }
diff --git a/Engine/Source/Runtime/Apple/AGXRHI/Private/AGXContext.cpp b/Engine/Source/Runtime/Apple/AGXRHI/Private/AGXContext.cpp
--- a/Engine/Source/Runtime/Apple/AGXRHI/Private/AGXContext.cpp
+++ b/Engine/Source/Runtime/Apple/AGXRHI/Private/AGXContext.cpp
@@ -467,7 +467,7 @@ void FAGXDeviceContext::ClearFreeList()
 				}
 				if (GAGXResourcePurgeOnDelete && !Buffer.GetParentBuffer())
 				{
-					Buffer.SetPurgeableState(mtlpp::PurgeableState::Empty);
+					[Buffer.GetPtr() setPurgeableState:MTLPurgeableStateEmpty];
 				}
 #endif
 				Heap.ReleaseBuffer(Buffer);
@@ -479,7 +479,7 @@ void FAGXDeviceContext::ClearFreeList()
 #if METAL_DEBUG_OPTIONS
 					if (GAGXResourcePurgeOnDelete)
 					{
-						Texture.SetPurgeableState(mtlpp::PurgeableState::Empty);
+						[Texture.GetPtr() setPurgeableState:MTLPurgeableStateEmpty];
 					}
 #endif
 					Heap.ReleaseTexture(nullptr, Texture);
@@ -713,7 +713,7 @@ FAGXTexture FAGXDeviceContext::CreateTexture(FAGXSurface* Surface, mtlpp::Textur
 #if METAL_DEBUG_OPTIONS
 	if (GAGXResourcePurgeOnDelete)
 	{
-		Tex.SetPurgeableState(mtlpp::PurgeableState::NonVolatile);
+		[Tex.GetPtr() setPurgeableState:MTLPurgeableStateNonVolatile];
 	}
 #endif
 	
@@ -736,7 +736,7 @@ FAGXBuffer FAGXDeviceContext::CreatePooledBuffer(FAGXPooledBufferArgs const& Arg
 #if METAL_DEBUG_OPTIONS
 	if (GAGXResourcePurgeOnDelete)
 	{
-		Buffer.SetPurgeableState(mtlpp::PurgeableState::NonVolatile);
+		[Buffer.GetPtr() setPurgeableState:MTLPurgeableStateNonVolatile];
 	}
 #endif
 	
diff --git a/Engine/Source/Runtime/Apple/AGXRHI/Private/AGXLLM.cpp b/Engine/Source/Runtime/Apple/AGXRHI/Private/AGXLLM.cpp
--- a/Engine/Source/Runtime/Apple/AGXRHI/Private/AGXLLM.cpp
+++ b/Engine/Source/Runtime/Apple/AGXRHI/Private/AGXLLM.cpp
@@ -168,7 +168,7 @@ void AGXLLM::LogAllocTexture(mtlpp::TextureDescriptor const& Desc, mtlpp::Textur
 	{
 		LLM_SCOPED_PAUSE_TRACKING(ELLMAllocType::System);
 		
-		if (Desc.GetUsage() & mtlpp::TextureUsage::RenderTarget)
+		if ([Desc.GetPtr() usage] & MTLTextureUsageRenderTarget)
 		{
 			objc_setAssociatedObject(Texture.GetPtr(), (void*)&AGXLLM::LogAllocTexture,
 			[[[FAGXDeallocHandler alloc] initWithBlock:^{
diff --git a/Engine/Source/Runtime/Apple/AGXRHI/Private/AGXRenderTarget.cpp b/Engine/Source/Runtime/Apple/AGXRHI/Private/AGXRenderTarget.cpp
--- a/Engine/Source/Runtime/Apple/AGXRHI/Private/AGXRenderTarget.cpp
+++ b/Engine/Source/Runtime/Apple/AGXRHI/Private/AGXRenderTarget.cpp
@@ -128,7 +128,7 @@ void FAGXRHICommandContext::RHICopyToResolveTarget(FRHITexture* SourceTextureRHI
 		{
 			const FPixelFormatInfo& SourceFormatInfo = GPixelFormats[SourceDesc.Format];
 			const FPixelFormatInfo& DestFormatInfo = GPixelFormats[DestinationDesc.Format];
-			bool bUsingPixelFormatView = (Source->Texture.GetUsage() & mtlpp::TextureUsage::PixelFormatView) != 0;
+			bool bUsingPixelFormatView = ([Source->Texture.GetPtr() usage] & MTLTextureUsagePixelFormatView) != 0;
 			
 			// Attempt to Resolve with a texture view - source Texture doesn't have to be created with MTLTextureUsagePixelFormatView for these cases e.g:
 			// If we are resolving to/from sRGB linear color space within the same format OR using same bit length color format
diff --git a/Engine/Source/Runtime/Apple/AGXRHI/Private/AGXTexture.cpp b/Engine/Source/Runtime/Apple/AGXRHI/Private/AGXTexture.cpp
--- a/Engine/Source/Runtime/Apple/AGXRHI/Private/AGXTexture.cpp
+++ b/Engine/Source/Runtime/Apple/AGXRHI/Private/AGXTexture.cpp
@@ -38,38 +38,39 @@ static bool IsRenderTarget(ETextureCreateFlags Flags)
 	return EnumHasAnyFlags(Flags, TexCreate_RenderTargetable | TexCreate_ResolveTargetable | TexCreate_DepthStencilTargetable | TexCreate_DepthStencilResolveTarget);
 }
 
-static mtlpp::TextureUsage ConvertFlagsToUsage(ETextureCreateFlags Flags)
+static MTLTextureUsage ConvertFlagsToUsage(ETextureCreateFlags Flags)
 {
-	NSUInteger Usage = mtlpp::TextureUsage::Unknown;
-    if (EnumHasAnyFlags(Flags, TexCreate_ShaderResource|TexCreate_ResolveTargetable|TexCreate_DepthStencilTargetable))
+	MTLTextureUsage Usage = MTLTextureUsageUnknown;
+    if (EnumHasAnyFlags(Flags, TexCreate_ShaderResource | TexCreate_ResolveTargetable | TexCreate_DepthStencilTargetable))
 	{
-		Usage |= mtlpp::TextureUsage::ShaderRead;
-		Usage |= mtlpp::TextureUsage::PixelFormatView;
+		Usage |= MTLTextureUsageShaderRead;
+		Usage |= MTLTextureUsagePixelFormatView;
 	}
 	
 	if (EnumHasAnyFlags(Flags, TexCreate_UAV))
 	{
-		Usage |= mtlpp::TextureUsage::ShaderRead;
-		Usage |= mtlpp::TextureUsage::ShaderWrite;
-		Usage |= mtlpp::TextureUsage::PixelFormatView;
+		Usage |= MTLTextureUsageShaderRead;
+		Usage |= MTLTextureUsageShaderWrite;
+		Usage |= MTLTextureUsagePixelFormatView;
 	}
 	
 	// offline textures are normal shader read textures
 	if (EnumHasAnyFlags(Flags, TexCreate_OfflineProcessed))
 	{
-		Usage |= mtlpp::TextureUsage::ShaderRead;
+		Usage |= MTLTextureUsageShaderRead;
 	}
 
 	//if the high level is doing manual resolves then the textures specifically markes as resolve targets
 	//are likely to be used in a manual shader resolve by the high level and must be bindable as rendertargets.
 	const bool bSeparateResolveTargets = FAGXCommandQueue::SupportsSeparateMSAAAndResolveTarget();
 	const bool bResolveTarget = EnumHasAnyFlags(Flags, TexCreate_ResolveTargetable);
-	if (EnumHasAnyFlags(Flags, TexCreate_RenderTargetable|TexCreate_DepthStencilTargetable|TexCreate_DepthStencilResolveTarget) || (bResolveTarget && bSeparateResolveTargets))
+	if (EnumHasAnyFlags(Flags, TexCreate_RenderTargetable | TexCreate_DepthStencilTargetable | TexCreate_DepthStencilResolveTarget) || (bResolveTarget && bSeparateResolveTargets))
 	{
-		Usage |= mtlpp::TextureUsage::RenderTarget;
-		Usage |= mtlpp::TextureUsage::ShaderRead;
+		Usage |= MTLTextureUsageRenderTarget;
+		Usage |= MTLTextureUsageShaderRead;
 	}
-	return (mtlpp::TextureUsage)Usage;
+
+	return Usage;
 }
 
 static bool IsPixelFormatCompressed(EPixelFormat Format)
@@ -254,12 +255,12 @@ MTLPixelFormat AGXToSRGBFormat(MTLPixelFormat MTLFormat)
 void FAGXSurface::PrepareTextureView()
 {
 	// Recreate the texture to enable MTLTextureUsagePixelFormatView which must be off unless we definitely use this feature or we are throwing ~4% performance vs. Windows on the floor.
-	mtlpp::TextureUsage Usage = (mtlpp::TextureUsage)Texture.GetUsage();
+	MTLTextureUsage Usage = [Texture.GetPtr() usage];
 	bool bMemoryLess = false;
 #if PLATFORM_IOS
-	bMemoryLess = ([Texture storageMode] == MTLStorageModeMemoryless);
+	bMemoryLess = ([Texture.GetPtr() storageMode] == MTLStorageModeMemoryless);
 #endif
-	if(!(Usage & mtlpp::TextureUsage::PixelFormatView) && !bMemoryLess)
+	if (!(Usage & MTLTextureUsagePixelFormatView) && !bMemoryLess)
 	{
 		check(ImageSurfaceRef == nullptr);
 		
@@ -269,18 +270,18 @@ void FAGXSurface::PrepareTextureView()
 		if (MSAATexture && !bMSAATextureIsTexture)
 		{
 			FAGXTexture OldMSAATexture = MSAATexture;
-			MSAATexture = Reallocate(MSAATexture, mtlpp::TextureUsage::PixelFormatView);
+			MSAATexture = Reallocate(MSAATexture, MTLTextureUsagePixelFormatView);
 			AGXSafeReleaseMetalTexture(this, OldMSAATexture, ImageSurfaceRef != nullptr);
 		}
 		if (MSAAResolveTexture && !bMSAAResolveTextureIsTexture)
 		{
 			FAGXTexture OldMSAAResolveTexture = MSAAResolveTexture;
-			MSAAResolveTexture = Reallocate(MSAAResolveTexture, mtlpp::TextureUsage::PixelFormatView);
+			MSAAResolveTexture = Reallocate(MSAAResolveTexture, MTLTextureUsagePixelFormatView);
 			AGXSafeReleaseMetalTexture(this, OldMSAAResolveTexture, ImageSurfaceRef != nullptr);
 		}
 		
 		FAGXTexture OldTexture = Texture;
-		Texture = Reallocate(Texture, mtlpp::TextureUsage::PixelFormatView);
+		Texture = Reallocate(Texture, MTLTextureUsagePixelFormatView);
 		AGXSafeReleaseMetalTexture(this, OldTexture, ImageSurfaceRef != nullptr);
 		
 		if (bMSAATextureIsTexture)
@@ -294,7 +295,7 @@ void FAGXSurface::PrepareTextureView()
 	}
 }
 
-FAGXTexture FAGXSurface::Reallocate(FAGXTexture InTexture, mtlpp::TextureUsage UsageModifier)
+FAGXTexture FAGXSurface::Reallocate(FAGXTexture InTexture, MTLTextureUsage UsageModifier)
 {
 	mtlpp::TextureDescriptor Desc;
 	Desc.SetTextureType(InTexture.GetTextureType());
@@ -306,7 +307,7 @@ FAGXTexture FAGXSurface::Reallocate(FAGXTexture InTexture, mtlpp::TextureUsage U
 	Desc.SetSampleCount(InTexture.GetSampleCount());
 	Desc.SetArrayLength(InTexture.GetArrayLength());
 	Desc.SetResourceOptions(mtlpp::ResourceOptions([InTexture.GetPtr() resourceOptions]));
-	Desc.SetUsage(mtlpp::TextureUsage(InTexture.GetUsage() | UsageModifier));
+	Desc.SetUsage(mtlpp::TextureUsage([InTexture.GetPtr() usage] | UsageModifier));
 	
 	FAGXTexture NewTex = GetAGXDeviceContext().CreateTexture(this, Desc);
 	check(NewTex);
@@ -496,7 +497,7 @@ FAGXTextureDesc::FAGXTextureDesc(FRHITextureDesc const& InDesc)
 #endif
 
 		Desc.SetResourceOptions(mtlpp::ResourceOptions(MTLResourceCPUCacheModeDefaultCache | ResourceStorageMode | MTLResourceHazardTrackingModeDefault));
-		Desc.SetUsage(ConvertFlagsToUsage(InDesc.Flags));
+		Desc.SetUsage(mtlpp::TextureUsage(ConvertFlagsToUsage(InDesc.Flags)));
 	}
 	
 	if (!FParse::Param(FCommandLine::Get(), TEXT("nomsaa")))
@@ -2025,8 +2026,8 @@ void FAGXRHICommandContext::RHICopyTexture(FRHITexture* SourceTextureRHI, FRHITe
 
 			if (TextureFormatExactMatch)
 			{
-				mtlpp::TextureUsage Usage = MetalSrcTexture->Texture.GetUsage();
-				if (Usage & mtlpp::TextureUsage::PixelFormatView)
+				MTLTextureUsage Usage = [MetalSrcTexture->Texture.GetPtr() usage];
+				if (Usage & MTLTextureUsagePixelFormatView)
 				{
 					ns::Range Slices(0, MetalSrcTexture->Texture.GetArrayLength() * (MetalSrcTexture->GetDesc().IsTextureCube() ? 6 : 1));
 					if (MetalSrcTexture->Texture.GetPixelFormat() != MetalDestTexture->Texture.GetPixelFormat())
diff --git a/Engine/Source/Runtime/Apple/AGXRHI/Private/AGXUAV.cpp b/Engine/Source/Runtime/Apple/AGXRHI/Private/AGXUAV.cpp
--- a/Engine/Source/Runtime/Apple/AGXRHI/Private/AGXUAV.cpp
+++ b/Engine/Source/Runtime/Apple/AGXRHI/Private/AGXUAV.cpp
@@ -167,17 +167,17 @@ FAGXResourceViewBase::FAGXResourceViewBase(
 			{
 				// Recreate the texture to enable MTLTextureUsagePixelFormatView which must be off unless we definitely use this feature or we are throwing ~4% performance vs. Windows on the floor.
 				// @todo recreating resources like this will likely prevent us from making view creation multi-threaded.
-				if (!(SourceTexture->Texture.GetUsage() & mtlpp::TextureUsage::PixelFormatView))
+				if (!([SourceTexture->Texture.GetPtr() usage] & MTLTextureUsagePixelFormatView))
 					SourceTexture->PrepareTextureView();
 
 				ns::Range Slices(0, SourceTexture->Texture.GetArrayLength() * (SourceTexture->GetDesc().IsTextureCube() ? 6 : 1));
 				
 				// @todo: Remove this type swizzle once Cube UAV supported for all Metal Devices - SRV seem to want to stay as cube but UAV are expected to be 2DArray
-				mtlpp::TextureType TextureType = bInUAV && SourceTexture->GetDesc().IsTextureCube() ? mtlpp::TextureType::Texture2DArray : SourceTexture->Texture.GetTextureType();
+				MTLTextureType TextureType = bInUAV && SourceTexture->GetDesc().IsTextureCube() ? MTLTextureType2DArray : [SourceTexture->Texture.GetPtr() textureType];
 
 				TextureView = SourceTexture->Texture.NewTextureView(
-					(mtlpp::PixelFormat)MetalFormat,
-					TextureType,
+					mtlpp::PixelFormat(MetalFormat),
+					mtlpp::TextureType(TextureType),
 					ns::Range(MipLevel, NumMips),
 					Slices);
 			}
@@ -236,7 +236,7 @@ FUnorderedAccessViewRHIRef FAGXDynamicRHI::RHICreateUnorderedAccessView_RenderTh
 	// safe.
 	//
 
-	if (!(Surface->Texture.GetUsage() & mtlpp::TextureUsage::PixelFormatView))
+	if (!([Surface->Texture.GetPtr() usage] & MTLTextureUsagePixelFormatView))
 	{
 		FScopedRHIThreadStaller StallRHIThread(RHICmdList);
 		return this->RHICreateUnorderedAccessView(Texture, MipLevel, FirstArraySlice, NumArraySlices);
@@ -352,7 +352,7 @@ FShaderResourceViewRHIRef FAGXDynamicRHI::RHICreateShaderResourceView_RenderThre
 	// safe.
 	//
 	
-	if (!(Surface->Texture.GetUsage() & mtlpp::TextureUsage::PixelFormatView))
+	if (!([Surface->Texture.GetPtr() usage] & MTLTextureUsagePixelFormatView))
 	{
 		FScopedRHIThreadStaller StallRHIThread(RHICmdList);
 		return this->RHICreateShaderResourceView(Texture2DRHI, CreateInfo);
diff --git a/Engine/Source/Runtime/Apple/AGXRHI/Private/AGXVertexBuffer.cpp b/Engine/Source/Runtime/Apple/AGXRHI/Private/AGXVertexBuffer.cpp
--- a/Engine/Source/Runtime/Apple/AGXRHI/Private/AGXVertexBuffer.cpp
+++ b/Engine/Source/Runtime/Apple/AGXRHI/Private/AGXVertexBuffer.cpp
@@ -294,14 +294,14 @@ void FAGXRHIBuffer::AllocLinearTextures(const LinearTextureMapKey& InLinearTextu
 		MTLPixelFormat MTLFormat = (MTLPixelFormat)GAGXBufferFormats[InFormat].LinearTextureFormat;
 		
 		mtlpp::TextureDescriptor Desc;
-		NSUInteger TexUsage = mtlpp::TextureUsage::Unknown;
+		MTLTextureUsage TexUsage = MTLTextureUsageUnknown;
 		if (EnumHasAnyFlags(Usage, BUF_ShaderResource))
 		{
-			TexUsage |= mtlpp::TextureUsage::ShaderRead;
+			TexUsage |= MTLTextureUsageShaderRead;
 		}
 		if (EnumHasAnyFlags(Usage, BUF_UnorderedAccess))
 		{
-			TexUsage |= mtlpp::TextureUsage::ShaderWrite;
+			TexUsage |= MTLTextureUsageShaderWrite;
 		}
 		
 		uint32 BytesPerElement = (0 == LinearTextureDesc.BytesPerElement) ? GPixelFormats[InFormat].BlockBytes : LinearTextureDesc.BytesPerElement;
diff --git a/Engine/Source/Runtime/Apple/AGXRHI/Public/AGXResources.h b/Engine/Source/Runtime/Apple/AGXRHI/Public/AGXResources.h
--- a/Engine/Source/Runtime/Apple/AGXRHI/Public/AGXResources.h
+++ b/Engine/Source/Runtime/Apple/AGXRHI/Public/AGXResources.h
@@ -223,7 +223,7 @@ public:
 	FAGXTexture GetDrawableTexture();
 	ns::AutoReleased<FAGXTexture> GetCurrentTexture();
 
-	FAGXTexture Reallocate(FAGXTexture Texture, mtlpp::TextureUsage UsageModifier);
+	FAGXTexture Reallocate(FAGXTexture Texture, MTLTextureUsage UsageModifier);
 	void MakeAliasable(void);
 
 	uint8 const FormatKey;
