AGXRHI: Port CL 20023804

---

Added support for per-slice texture array UAV+SRV on some platforms

---

#rnx

diff --git a/Engine/Source/Runtime/Apple/AGXRHI/Private/AGXUAV.cpp b/Engine/Source/Runtime/Apple/AGXRHI/Private/AGXUAV.cpp
--- a/Engine/Source/Runtime/Apple/AGXRHI/Private/AGXUAV.cpp
+++ b/Engine/Source/Runtime/Apple/AGXRHI/Private/AGXUAV.cpp
@@ -90,9 +90,6 @@ FAGXResourceViewBase::FAGXResourceViewBase(
 	, Stride           (0)
 	, Offset           (0)
 {
-	// Slice selection of a texture array still need to be implemented on Metal
-	check(InFirstArraySlice == 0 && InNumArraySlices == 0);
-
 	if (SourceTexture)
 	{
 		SourceTexture->AddRef();
@@ -156,7 +153,9 @@ FAGXResourceViewBase::FAGXResourceViewBase(
 				   MipLevel == 0
 				&& NumMips == SourceTexture->Texture.GetMipmapLevelCount()
 				&& MetalFormat == (MTLPixelFormat)SourceTexture->Texture.GetPixelFormat()
-				&& !(bInUAV && SourceTexture->GetDesc().IsTextureCube());		// @todo: Remove this once Cube UAV supported for all Metal Devices
+				&& !(bInUAV && SourceTexture->GetDesc().IsTextureCube())
+				&& InFirstArraySlice == 0
+				&& InNumArraySlices == 0;		// @todo: Remove this once Cube UAV supported for all Metal Devices
 
 			if (bUseSourceTex)
 			{
@@ -168,18 +167,44 @@ FAGXResourceViewBase::FAGXResourceViewBase(
 				// Recreate the texture to enable MTLTextureUsagePixelFormatView which must be off unless we definitely use this feature or we are throwing ~4% performance vs. Windows on the floor.
 				// @todo recreating resources like this will likely prevent us from making view creation multi-threaded.
 				if (!([SourceTexture->Texture.GetPtr() usage] & MTLTextureUsagePixelFormatView))
+				{
 					SourceTexture->PrepareTextureView();
+				}
 
-				ns::Range Slices(0, SourceTexture->Texture.GetArrayLength() * (SourceTexture->GetDesc().IsTextureCube() ? 6 : 1));
+				const uint32 TextureSliceCount = SourceTexture->Texture.GetArrayLength();
+				const uint32 CubeSliceMultiplier = SourceTexture->GetDesc().IsTextureCube() ? 6 : 1;
+				const uint32 NumArraySlices = (InNumArraySlices > 0 ? InNumArraySlices : TextureSliceCount) * CubeSliceMultiplier;
 				
 				// @todo: Remove this type swizzle once Cube UAV supported for all Metal Devices - SRV seem to want to stay as cube but UAV are expected to be 2DArray
 				MTLTextureType TextureType = bInUAV && SourceTexture->GetDesc().IsTextureCube() ? MTLTextureType2DArray : [SourceTexture->Texture.GetPtr() textureType];
 
+				// Assume a texture view of 1 slice into a multislice texture wants to be the non-array texture type
+				// This doesn't really matter to Metal but will be very important when this texture is bound in the shader
+				if (InNumArraySlices == 1)
+				{
+					switch (TextureType)
+					{
+					case MTLTextureType2DArray:
+						TextureType = MTLTextureType2D;
+						break;
+					case MTLTextureTypeCubeArray:
+						TextureType = MTLTextureTypeCube;
+						break;
+					default:
+						// NOP
+						break;
+					}
+				}
+
 				TextureView = SourceTexture->Texture.NewTextureView(
 					mtlpp::PixelFormat(MetalFormat),
 					mtlpp::TextureType(TextureType),
 					ns::Range(MipLevel, NumMips),
-					Slices);
+					ns::Range(InFirstArraySlice, NumArraySlices));
+
+#if METAL_DEBUG_OPTIONS
+				TextureView.SetLabel([SourceTexture->Texture.GetLabel() stringByAppendingString:@"_TextureView"]);
+#endif
 			}
 		}
 	}
