AGXRHI: Fix null pointer dereference on compute shader reflection object.

#rnx

diff --git a/Engine/Source/Runtime/Apple/AGXRHI/Private/Shaders/Types/AGXComputeShader.cpp b/Engine/Source/Runtime/Apple/AGXRHI/Private/Shaders/Types/AGXComputeShader.cpp
--- a/Engine/Source/Runtime/Apple/AGXRHI/Private/Shaders/Types/AGXComputeShader.cpp
+++ b/Engine/Source/Runtime/Apple/AGXRHI/Private/Shaders/Types/AGXComputeShader.cpp
@@ -71,7 +71,7 @@ FAGXShaderPipeline* FAGXComputeShader::GetPipeline()
 
 		id<MTLComputePipelineState> Kernel = nil;
 #if METAL_DEBUG_OPTIONS
-		MTLAutoreleasedComputePipelineReflection* Reflection = nil;
+		MTLAutoreleasedComputePipelineReflection Reflection = nil;
 #endif
 
 		METAL_GPUPROFILE(FAGXScopedCPUStats CPUStat(FString::Printf(TEXT("NewComputePipeline: %d_%d"), SourceLen, SourceCRC)));
@@ -80,7 +80,7 @@ FAGXShaderPipeline* FAGXComputeShader::GetPipeline()
 		{
 			Kernel = [GMtlDevice newComputePipelineStateWithDescriptor:Descriptor
 															  options:(MTLPipelineOptionArgumentInfo | MTLPipelineOptionBufferTypeInfo)
-														   reflection:Reflection
+														   reflection:&Reflection
 														   		error:&Error];
 		}
 		else
@@ -101,7 +101,7 @@ FAGXShaderPipeline* FAGXComputeShader::GetPipeline()
 		Pipeline = [FAGXShaderPipeline new];
 		Pipeline->ComputePipelineState = Kernel;
 #if METAL_DEBUG_OPTIONS
-		Pipeline->ComputePipelineReflection = [*Reflection retain];
+		Pipeline->ComputePipelineReflection = [Reflection retain];
 		Pipeline->ComputeSource = GetSourceCode();
 #endif // METAL_DEBUG_OPTIONS
 		METAL_DEBUG_OPTION(FMemory::Memzero(Pipeline->ResourceMask, sizeof(Pipeline->ResourceMask)));
