AGXRHI: Convert FAGXViewport and FAGXSurface GetDrawableTexture() and GetCurrentTexture() return types to Metal native types.

#rnx

diff --git a/Engine/Source/Runtime/Apple/AGXRHI/Private/AGXTexture.cpp b/Engine/Source/Runtime/Apple/AGXRHI/Private/AGXTexture.cpp
--- a/Engine/Source/Runtime/Apple/AGXRHI/Private/AGXTexture.cpp
+++ b/Engine/Source/Runtime/Apple/AGXRHI/Private/AGXTexture.cpp
@@ -151,7 +151,7 @@ void AGXSafeReleaseMetalTexture(FAGXSurface* Surface, FAGXTexture& Texture)
 	}
 }
 
-#if !PLATFORM_IOS
+#if PLATFORM_MAC
 static MTLPixelFormat AGX_ToSRGBFormat_NonAppleMacGPU(MTLPixelFormat MTLFormat)
 {
 	switch (MTLFormat)
@@ -242,7 +242,7 @@ MTLPixelFormat AGXToSRGBFormat(MTLPixelFormat MTLFormat)
 	{
 		MTLFormat = AGX_ToSRGBFormat_AppleGPU(MTLFormat);
 	}
-#if !PLATFORM_IOS
+#if PLATFORM_MAC
 	else if([GMtlDevice supportsFamily:MTLGPUFamilyMac1])
 	{
 		MTLFormat = AGX_ToSRGBFormat_NonAppleMacGPU(MTLFormat);
@@ -1335,25 +1335,24 @@ uint32 FAGXSurface::GetNumFaces()
 	return GetDesc().Depth * GetDesc().ArraySize;
 }
 
-FAGXTexture FAGXSurface::GetDrawableTexture()
+void FAGXSurface::GetDrawableTexture()
 {
 	if (!Texture && EnumHasAnyFlags(GetDesc().Flags, TexCreate_Presentable))
 	{
 		check(Viewport);
 		Texture = Viewport->GetDrawableTexture(EAGXViewportAccessRHI);
 	}
-	return Texture;
 }
 
-ns::AutoReleased<FAGXTexture> FAGXSurface::GetCurrentTexture()
+id<MTLTexture> FAGXSurface::GetCurrentTexture()
 {
-	ns::AutoReleased<FAGXTexture> Tex;
+	id<MTLTexture> MtlTexture = nil;
 	if (Viewport && EnumHasAnyFlags(GetDesc().Flags, TexCreate_Presentable))
 	{
 		check(Viewport);
-		Tex = Viewport->GetCurrentTexture(EAGXViewportAccessRHI);
+		MtlTexture = Viewport->GetCurrentTexture(EAGXViewportAccessRHI);
 	}
-	return Tex;
+	return MtlTexture;
 }
 
 
diff --git a/Engine/Source/Runtime/Apple/AGXRHI/Private/AGXViewport.cpp b/Engine/Source/Runtime/Apple/AGXRHI/Private/AGXViewport.cpp
--- a/Engine/Source/Runtime/Apple/AGXRHI/Private/AGXViewport.cpp
+++ b/Engine/Source/Runtime/Apple/AGXRHI/Private/AGXViewport.cpp
@@ -32,12 +32,6 @@ static FAutoConsoleVariableRef CVarAGXEnablePresentPacing(
 
 #if PLATFORM_MAC
 
-// Quick way to disable availability warnings is to duplicate the definitions into a new type - gotta love ObjC dynamic-dispatch!
-@interface FCAMetalLayer : CALayer
-@property BOOL displaySyncEnabled;
-@property BOOL allowsNextDrawableTimeout;
-@end
-
 @implementation FAGXView
 
 - (id)initWithFrame:(NSRect)frameRect
@@ -70,7 +64,7 @@ FAGXViewport::FAGXViewport(void* WindowHandle, uint32 InSizeX, uint32 InSizeY, b
 	: Drawable{nil}
 	, BackBuffer{nullptr, nullptr}
 	, Mutex{}
-	, DrawableTextures{}
+	, DrawableTextures{nil, nil}
 	, DisplayID{0}
 	, Block{nullptr}
 	, FrameAvailable{0}
@@ -108,6 +102,8 @@ FAGXViewport::FAGXViewport(void* WindowHandle, uint32 InSizeX, uint32 InSizeY, b
 
 		[View setLayer:Layer];
 
+		[Layer release];
+
 		[Window setContentView:View];
 		[[Window standardWindowButton:NSWindowCloseButton] setAction:@selector(performClose:)];
 	}, NSDefaultRunLoopMode, true);
@@ -370,9 +366,10 @@ id<CAMetalDrawable> FAGXViewport::GetDrawable(EAGXViewportAccessFlag Accessor)
 	return Drawable;
 }
 
-FAGXTexture FAGXViewport::GetDrawableTexture(EAGXViewportAccessFlag Accessor)
+id<MTLTexture> FAGXViewport::GetDrawableTexture(EAGXViewportAccessFlag Accessor)
 {
 	id<CAMetalDrawable> CurrentDrawable = GetDrawable(Accessor);
+
 #if METAL_DEBUG_OPTIONS
 	@autoreleasepool
 	{
@@ -390,11 +387,13 @@ FAGXTexture FAGXViewport::GetDrawableTexture(EAGXViewportAccessFlag Accessor)
 		}
 	}
 #endif
-	DrawableTextures[Accessor] = CurrentDrawable.texture;
-	return CurrentDrawable.texture;
+
+	DrawableTextures[Accessor] = [CurrentDrawable texture];
+
+	return DrawableTextures[Accessor];
 }
 
-ns::AutoReleased<FAGXTexture> FAGXViewport::GetCurrentTexture(EAGXViewportAccessFlag Accessor)
+id<MTLTexture> FAGXViewport::GetCurrentTexture(EAGXViewportAccessFlag Accessor) const
 {
 	return DrawableTextures[Accessor];
 }
@@ -416,13 +415,6 @@ void FAGXViewport::ReleaseDrawable()
 	}
 }
 
-#if PLATFORM_MAC
-NSWindow* FAGXViewport::GetWindow() const
-{
-	return [View window];
-}
-#endif
-
 void FAGXViewport::Present(FAGXCommandQueue& CommandQueue, bool bLockToVsync)
 {
 	FScopeLock Lock(&Mutex);
@@ -433,9 +425,8 @@ void FAGXViewport::Present(FAGXCommandQueue& CommandQueue, bool bLockToVsync)
 	DisplayID = ScreenId.unsignedIntValue;
 	bIsLiveResize = View.inLiveResize;
 	{
-		FCAMetalLayer* CurrentLayer = (FCAMetalLayer*)[View layer];
-		static bool sVSyncSafe = FPlatformMisc::MacOSXVersionCompare(10,13,4) >= 0;
-		CurrentLayer.displaySyncEnabled = bLockToVsync || (!sVSyncSafe && !(IsRunningGame() && bIsFullScreen));
+		CAMetalLayer* CurrentLayer = (CAMetalLayer*)[View layer];
+		CurrentLayer.displaySyncEnabled = bLockToVsync;
 	}
 #endif
 	
diff --git a/Engine/Source/Runtime/Apple/AGXRHI/Public/AGXResources.h b/Engine/Source/Runtime/Apple/AGXRHI/Public/AGXResources.h
--- a/Engine/Source/Runtime/Apple/AGXRHI/Public/AGXResources.h
+++ b/Engine/Source/Runtime/Apple/AGXRHI/Public/AGXResources.h
@@ -220,8 +220,8 @@ public:
 	uint32 GetNumFaces();
 	
 	/** Gets the drawable texture if this is a back-buffer surface. */
-	FAGXTexture GetDrawableTexture();
-	ns::AutoReleased<FAGXTexture> GetCurrentTexture();
+	void GetDrawableTexture();
+	id<MTLTexture> GetCurrentTexture();
 
 	FAGXTexture Reallocate(FAGXTexture Texture, MTLTextureUsage UsageModifier);
 	void MakeAliasable(void);
diff --git a/Engine/Source/Runtime/Apple/AGXRHI/Public/AGXViewport.h b/Engine/Source/Runtime/Apple/AGXRHI/Public/AGXViewport.h
--- a/Engine/Source/Runtime/Apple/AGXRHI/Public/AGXViewport.h
+++ b/Engine/Source/Runtime/Apple/AGXRHI/Public/AGXViewport.h
@@ -38,8 +38,8 @@ public:
 	
 	TRefCountPtr<FAGXSurface> GetBackBuffer(EAGXViewportAccessFlag Accessor) const;
 	id<CAMetalDrawable> GetDrawable(EAGXViewportAccessFlag Accessor);
-	FAGXTexture GetDrawableTexture(EAGXViewportAccessFlag Accessor);
-	ns::AutoReleased<FAGXTexture> GetCurrentTexture(EAGXViewportAccessFlag Accessor);
+	id<MTLTexture> GetDrawableTexture(EAGXViewportAccessFlag Accessor);
+	id<MTLTexture> GetCurrentTexture(EAGXViewportAccessFlag Accessor) const;
 	void ReleaseDrawable(void);
 
 	// supports pulling the raw MTLTexture
@@ -47,8 +47,6 @@ public:
 	virtual void* GetNativeBackBufferRT() const override { return (const_cast<FAGXViewport *>(this))->GetDrawableTexture(EAGXViewportAccessRenderer); }
 	
 #if PLATFORM_MAC
-	NSWindow* GetWindow() const;
-	
 	virtual void SetCustomPresent(FRHICustomPresent* InCustomPresent) override
 	{
 		CustomPresent = InCustomPresent;
@@ -68,7 +66,7 @@ private:
 	TRefCountPtr<FAGXSurface> BackBuffer[2];
 	mutable FCriticalSection Mutex;
 	
-	ns::AutoReleased<FAGXTexture> DrawableTextures[2];
+	id<MTLTexture> DrawableTextures[2];
 	
 	uint32 DisplayID;
 	FAGXViewportPresentHandler Block;
