AGXRHI: Convert FAGXResourceViewBase::TextureView type to true Metal texture type

#rnx

diff --git a/Engine/Source/Runtime/Apple/AGXRHI/Private/AGXStateCache.cpp b/Engine/Source/Runtime/Apple/AGXRHI/Private/AGXStateCache.cpp
--- a/Engine/Source/Runtime/Apple/AGXRHI/Private/AGXStateCache.cpp
+++ b/Engine/Source/Runtime/Apple/AGXRHI/Private/AGXStateCache.cpp
@@ -1457,7 +1457,7 @@ void FAGXStateCache::SetShaderResourceView(FAGXContext* Context, EAGXShaderStage
 	{
 		if (SRV->bTexture)
 		{
-			FAGXTexture const& View = SRV->GetTextureView();
+			FAGXTexture View(SRV->GetTextureView());
 			if (View)
 			{
 				SetShaderTexture(ShaderStage, View, BindIndex, mtlpp::ResourceUsage(mtlpp::ResourceUsage::Read | mtlpp::ResourceUsage::Sample));
@@ -1519,7 +1519,7 @@ void FAGXStateCache::SetShaderUnorderedAccessView(EAGXShaderStages ShaderStage,
 		if (UAV->bTexture)
 		{
 			FAGXSurface* Surface = UAV->GetSourceTexture();
-			FAGXTexture const& View = UAV->GetTextureView();
+			FAGXTexture  View(UAV->GetTextureView());
 
 			if (View)
 			{
diff --git a/Engine/Source/Runtime/Apple/AGXRHI/Private/AGXUAV.cpp b/Engine/Source/Runtime/Apple/AGXRHI/Private/AGXUAV.cpp
--- a/Engine/Source/Runtime/Apple/AGXRHI/Private/AGXUAV.cpp
+++ b/Engine/Source/Runtime/Apple/AGXRHI/Private/AGXUAV.cpp
@@ -94,10 +94,12 @@ FAGXResourceViewBase::FAGXResourceViewBase(
 	{
 		SourceTexture->AddRef();
 
+		id<MTLTexture> SourceTextureInternal = SourceTexture->Texture.GetPtr();
+
 		// TODO: Apple Silicon supports memoryless
 #if PLATFORM_IOS
 		// Memoryless targets can't have texture views (SRVs or UAVs)
-		if ([SourceTexture->Texture.GetPtr() storageMode] != MTLStorageModeMemoryless)
+		if ([SourceTextureInternal storageMode] != MTLStorageModeMemoryless)
 #endif
 		{
 			// Determine the appropriate metal format for the view.
@@ -108,7 +110,7 @@ FAGXResourceViewBase::FAGXResourceViewBase(
 			{
 				// Stencil buffer view of a depth texture
 				check(SourceTexture->GetDesc().Format == PF_DepthStencil);
-				switch ((MTLPixelFormat)SourceTexture->Texture.GetPixelFormat())
+				switch ([SourceTextureInternal pixelFormat])
 				{
 					default:
 						checkNoEntry();
@@ -134,7 +136,7 @@ FAGXResourceViewBase::FAGXResourceViewBase(
 					{
 #if PLATFORM_MAC
 						// R8Unorm has been expanded in the source surface for sRGBA support - we need to expand to RGBA to enable compatible texture format view for non apple silicon macs
-						if (Format == PF_G8 && (MTLPixelFormat)SourceTexture->Texture.GetPixelFormat() == MTLPixelFormatRGBA8Unorm_sRGB)
+						if (Format == PF_G8 && [SourceTextureInternal pixelFormat] == MTLPixelFormatRGBA8Unorm_sRGB)
 						{
 							MetalFormat = MTLPixelFormatRGBA8Unorm;
 						}
@@ -151,8 +153,8 @@ FAGXResourceViewBase::FAGXResourceViewBase(
 			// We can use the source texture directly if the view's format / mip count etc matches.
 			bool bUseSourceTex =
 				   MipLevel == 0
-				&& NumMips == SourceTexture->Texture.GetMipmapLevelCount()
-				&& MetalFormat == (MTLPixelFormat)SourceTexture->Texture.GetPixelFormat()
+				&& NumMips == [SourceTextureInternal mipmapLevelCount]
+				&& MetalFormat == [SourceTextureInternal pixelFormat]
 				&& !(bInUAV && SourceTexture->GetDesc().IsTextureCube())
 				&& InFirstArraySlice == 0
 				&& InNumArraySlices == 0;		// @todo: Remove this once Cube UAV supported for all Metal Devices
@@ -160,23 +162,24 @@ FAGXResourceViewBase::FAGXResourceViewBase(
 			if (bUseSourceTex)
 			{
 				// SRV is exactly compatible with the original texture.
-				TextureView = SourceTexture->Texture;
+				TextureView = [SourceTextureInternal retain];
 			}
 			else
 			{
 				// Recreate the texture to enable MTLTextureUsagePixelFormatView which must be off unless we definitely use this feature or we are throwing ~4% performance vs. Windows on the floor.
 				// @todo recreating resources like this will likely prevent us from making view creation multi-threaded.
-				if (!([SourceTexture->Texture.GetPtr() usage] & MTLTextureUsagePixelFormatView))
+				if (!([SourceTextureInternal usage] & MTLTextureUsagePixelFormatView))
 				{
 					SourceTexture->PrepareTextureView();
+					SourceTextureInternal = SourceTexture->Texture.GetPtr();
 				}
 
-				const uint32 TextureSliceCount = SourceTexture->Texture.GetArrayLength();
+				const uint32 TextureSliceCount = [SourceTextureInternal arrayLength];
 				const uint32 CubeSliceMultiplier = SourceTexture->GetDesc().IsTextureCube() ? 6 : 1;
 				const uint32 NumArraySlices = (InNumArraySlices > 0 ? InNumArraySlices : TextureSliceCount) * CubeSliceMultiplier;
 				
 				// @todo: Remove this type swizzle once Cube UAV supported for all Metal Devices - SRV seem to want to stay as cube but UAV are expected to be 2DArray
-				MTLTextureType TextureType = bInUAV && SourceTexture->GetDesc().IsTextureCube() ? MTLTextureType2DArray : [SourceTexture->Texture.GetPtr() textureType];
+				MTLTextureType TextureType = bInUAV && SourceTexture->GetDesc().IsTextureCube() ? MTLTextureType2DArray : [SourceTextureInternal textureType];
 
 				// Assume a texture view of 1 slice into a multislice texture wants to be the non-array texture type
 				// This doesn't really matter to Metal but will be very important when this texture is bound in the shader
@@ -196,21 +199,20 @@ FAGXResourceViewBase::FAGXResourceViewBase(
 					}
 				}
 
-				TextureView = SourceTexture->Texture.NewTextureView(
-					mtlpp::PixelFormat(MetalFormat),
-					mtlpp::TextureType(TextureType),
-					ns::Range(MipLevel, NumMips),
-					ns::Range(InFirstArraySlice, NumArraySlices));
+				TextureView = [SourceTextureInternal newTextureViewWithPixelFormat:MetalFormat
+				                                                       textureType:TextureType
+																	        levels:NSMakeRange(MipLevel, NumMips)
+																			slices:NSMakeRange(InFirstArraySlice, NumArraySlices)];
 
 #if METAL_DEBUG_OPTIONS
-				TextureView.SetLabel([SourceTexture->Texture.GetLabel() stringByAppendingString:@"_TextureView"]);
+				[TextureView setLabel:[[SourceTextureInternal label] stringByAppendingString:@"_TextureView"]];
 #endif
 			}
 		}
 	}
 	else
 	{
-		TextureView = {};
+		TextureView = nil;
 	}
 }
 
@@ -220,6 +222,7 @@ FAGXResourceViewBase::~FAGXResourceViewBase()
 	{
 		if (SourceTexture)
 		{
+			[TextureView release];
 			SourceTexture->Release();
 		}
 	}
diff --git a/Engine/Source/Runtime/Apple/AGXRHI/Public/AGXResources.h b/Engine/Source/Runtime/Apple/AGXRHI/Public/AGXResources.h
--- a/Engine/Source/Runtime/Apple/AGXRHI/Public/AGXResources.h
+++ b/Engine/Source/Runtime/Apple/AGXRHI/Public/AGXResources.h
@@ -515,7 +515,7 @@ public:
 	inline FAGXResourceMultiBuffer* GetSourceBuffer () const { check(!bTexture); return SourceBuffer; }
 
 	inline FAGXSurface*             GetSourceTexture() const { check(bTexture); return SourceTexture; }
-	inline FAGXTexture const&       GetTextureView  () const { check(bTexture); return TextureView;   }
+	inline id<MTLTexture>           GetTextureView  () const { check(bTexture); return TextureView;   }
 
 private:
 	// Needed for RHIUpdateShaderResourceView
@@ -528,7 +528,7 @@ private:
 	};
 
 	TUniquePtr<FAGXLinearTextureDescriptor> LinearTextureDesc = nullptr;
-	FAGXTexture TextureView = nullptr;
+	id<MTLTexture> TextureView = nil;
 
 public:
 	uint8 const bTexture : 1;
