Metal RHI, AGXRHI: Fix ASAN report of global-buffer-overrun

#rnx

diff --git a/Engine/Source/Runtime/Apple/AGXRHI/Private/AGXStateCache.cpp b/Engine/Source/Runtime/Apple/AGXRHI/Private/AGXStateCache.cpp
--- a/Engine/Source/Runtime/Apple/AGXRHI/Private/AGXStateCache.cpp
+++ b/Engine/Source/Runtime/Apple/AGXRHI/Private/AGXStateCache.cpp
@@ -1535,7 +1535,7 @@ void FAGXStateCache::SetShaderUnorderedAccessView(EAGXShaderStages ShaderStage,
 					FAGXBuffer Buffer(Surface->Texture.GetBuffer(), false);
 					const uint32 BufferOffset = Surface->Texture.GetBufferOffset();
 					const uint32 BufferSize = Surface->Texture.GetBuffer().GetLength();
-					SetShaderBuffer(ShaderStage, Buffer, nullptr, BufferOffset, BufferSize, BindIndex, mtlpp::ResourceUsage(mtlpp::ResourceUsage::Read | mtlpp::ResourceUsage::Write), PF_MAX, ElementsPerRow);
+					SetShaderBuffer(ShaderStage, Buffer, nullptr, BufferOffset, BufferSize, BindIndex, mtlpp::ResourceUsage(mtlpp::ResourceUsage::Read | mtlpp::ResourceUsage::Write), static_cast<EPixelFormat>(UAV->Format), ElementsPerRow);
 				}
 			}
 			else
diff --git a/Engine/Source/Runtime/Apple/MetalRHI/Private/MetalStateCache.cpp b/Engine/Source/Runtime/Apple/MetalRHI/Private/MetalStateCache.cpp
--- a/Engine/Source/Runtime/Apple/MetalRHI/Private/MetalStateCache.cpp
+++ b/Engine/Source/Runtime/Apple/MetalRHI/Private/MetalStateCache.cpp
@@ -1535,7 +1535,7 @@ void FMetalStateCache::SetShaderUnorderedAccessView(EMetalShaderStages ShaderSta
 					FMetalBuffer Buffer(Surface->Texture.GetBuffer(), false);
 					const uint32 BufferOffset = Surface->Texture.GetBufferOffset();
 					const uint32 BufferSize = Surface->Texture.GetBuffer().GetLength();
-					SetShaderBuffer(ShaderStage, Buffer, nullptr, BufferOffset, BufferSize, BindIndex, mtlpp::ResourceUsage(mtlpp::ResourceUsage::Read | mtlpp::ResourceUsage::Write), PF_MAX, ElementsPerRow);
+					SetShaderBuffer(ShaderStage, Buffer, nullptr, BufferOffset, BufferSize, BindIndex, mtlpp::ResourceUsage(mtlpp::ResourceUsage::Read | mtlpp::ResourceUsage::Write), static_cast<EPixelFormat>(UAV->Format), ElementsPerRow);
 				}
 			}
 			else
