Removing the long-form create functions and constructors from FRHITextureCreateDesc in favor of the shorter versions. Fixing Mac RHI code.

diff --git a/Engine/Source/Runtime/Apple/AGXRHI/Private/AGXViewport.cpp b/Engine/Source/Runtime/Apple/AGXRHI/Private/AGXViewport.cpp
--- a/Engine/Source/Runtime/Apple/AGXRHI/Private/AGXViewport.cpp
+++ b/Engine/Source/Runtime/Apple/AGXRHI/Private/AGXViewport.cpp
@@ -238,17 +238,15 @@ void FAGXViewport::Resize(uint32 InSizeX, uint32 InSizeY, bool bInIsFullscreen,E
 		TRefCountPtr<FAGXSurface> NewBackBuffer;
 		TRefCountPtr<FAGXSurface> DoubleBuffer;
 
-		ETextureCreateFlags Flags = GAGXSupportsIntermediateBackBuffer
-			? TexCreate_RenderTargetable
-			: TexCreate_RenderTargetable | TexCreate_Presentable;
-
-		FRHITextureCreateDesc CreateDesc = FRHITextureCreateDesc::Create2D(
-			  TEXT("BackBuffer")
-			, { (int32)InSizeX, (int32)InSizeY }
-			, Format
-			, FClearValueBinding::None
-			, Flags
-		);
+		FRHITextureCreateDesc CreateDesc =
+			FRHITextureCreateDesc::Create2D(TEXT("BackBuffer"), InSizeX, InSizeY, Format)
+			.SetClearValue(FClearValueBinding::None)
+			.SetFlags(ETextureCreateFlags::RenderTargetable);
+		
+		if (!GAGXSupportsIntermediateBackBuffer)
+		{
+			CreateDesc.AddFlags(ETextureCreateFlags::Presentable);
+		}
 
 		NewBackBuffer = new FAGXSurface(CreateDesc);
 		NewBackBuffer->Viewport = this;
diff --git a/Engine/Source/Runtime/Apple/MetalRHI/Private/MetalViewport.cpp b/Engine/Source/Runtime/Apple/MetalRHI/Private/MetalViewport.cpp
--- a/Engine/Source/Runtime/Apple/MetalRHI/Private/MetalViewport.cpp
+++ b/Engine/Source/Runtime/Apple/MetalRHI/Private/MetalViewport.cpp
@@ -242,17 +242,15 @@ void FMetalViewport::Resize(uint32 InSizeX, uint32 InSizeY, bool bInIsFullscreen
 		TRefCountPtr<FMetalSurface> NewBackBuffer;
 		TRefCountPtr<FMetalSurface> DoubleBuffer;
 
-		ETextureCreateFlags Flags = GMetalSupportsIntermediateBackBuffer
-			? TexCreate_RenderTargetable
-			: TexCreate_RenderTargetable | TexCreate_Presentable;
-
-		FRHITextureCreateDesc CreateDesc = FRHITextureCreateDesc::Create2D(
-			  TEXT("BackBuffer")
-			, { (int32)InSizeX, (int32)InSizeY }
-			, Format
-			, FClearValueBinding::None
-			, Flags
-		);
+		FRHITextureCreateDesc CreateDesc =
+			FRHITextureCreateDesc::Create2D(TEXT("BackBuffer"), InSizeX, InSizeY, Format)
+			.SetClearValue(FClearValueBinding::None)
+			.SetFlags(ETextureCreateFlags::RenderTargetable);
+		
+		if (!GMetalSupportsIntermediateBackBuffer)
+		{
+			CreateDesc.AddFlags(ETextureCreateFlags::Presentable);
+		}
 
 		NewBackBuffer = new FMetalSurface(CreateDesc);
 		NewBackBuffer->Viewport = this;
