Adding missing SetInitialState calls to the Mac internal viewport textures.

diff --git a/Engine/Source/Runtime/Apple/AGXRHI/Private/AGXViewport.cpp b/Engine/Source/Runtime/Apple/AGXRHI/Private/AGXViewport.cpp
--- a/Engine/Source/Runtime/Apple/AGXRHI/Private/AGXViewport.cpp
+++ b/Engine/Source/Runtime/Apple/AGXRHI/Private/AGXViewport.cpp
@@ -247,6 +247,8 @@ void FAGXViewport::Resize(uint32 InSizeX, uint32 InSizeY, bool bInIsFullscreen,E
 		{
 			CreateDesc.AddFlags(ETextureCreateFlags::Presentable);
 		}
+		
+		CreateDesc.SetInitialState(RHIGetDefaultResourceState(CreateDesc.Flags, false));
 
 		NewBackBuffer = new FAGXSurface(CreateDesc);
 		NewBackBuffer->Viewport = this;
diff --git a/Engine/Source/Runtime/Apple/MetalRHI/Private/MetalViewport.cpp b/Engine/Source/Runtime/Apple/MetalRHI/Private/MetalViewport.cpp
--- a/Engine/Source/Runtime/Apple/MetalRHI/Private/MetalViewport.cpp
+++ b/Engine/Source/Runtime/Apple/MetalRHI/Private/MetalViewport.cpp
@@ -251,6 +251,8 @@ void FMetalViewport::Resize(uint32 InSizeX, uint32 InSizeY, bool bInIsFullscreen
 		{
 			CreateDesc.AddFlags(ETextureCreateFlags::Presentable);
 		}
+		
+		CreateDesc.SetInitialState(RHIGetDefaultResourceState(CreateDesc.Flags, false));
 
 		NewBackBuffer = new FMetalSurface(CreateDesc);
 		NewBackBuffer->Viewport = this;
