Implemented dedicated MSAA resolve code path for platforms, bypassing RHICopyToResolveTarget.

diff --git a/Engine/Source/Runtime/Apple/AGXRHI/Private/AGXRHIContext.cpp b/Engine/Source/Runtime/Apple/AGXRHI/Private/AGXRHIContext.cpp
--- a/Engine/Source/Runtime/Apple/AGXRHI/Private/AGXRHIContext.cpp
+++ b/Engine/Source/Runtime/Apple/AGXRHI/Private/AGXRHIContext.cpp
@@ -3,7 +3,6 @@
 #include "AGXRHIPrivate.h"
 #include "AGXRHIRenderQuery.h"
 #include "AGXCommandBufferFence.h"
-#include "RHICore.h"
 
 TGlobalResource<TBoundShaderStateHistory<10000>> FAGXRHICommandContext::BoundShaderStateHistory;
 
@@ -159,7 +158,69 @@ void FAGXRHICommandContext::RHIEndRenderPass()
 		RHIEndOcclusionQueryBatch();
 	}
 
-	UE::RHICore::ResolveRenderPassTargets(*this, RenderPassInfo);
+	UE::RHICore::ResolveRenderPassTargets(RenderPassInfo, [this](UE::RHICore::FResolveTextureInfo Info)
+	{
+		ResolveTexture(Info);
+	});
+}
+
+void FAGXRHICommandContext::ResolveTexture(UE::RHICore::FResolveTextureInfo Info)
+{
+	@autoreleasepool {
+	FAGXSurface* Source = AGXGetMetalSurfaceFromRHITexture(Info.SourceTexture);
+	FAGXSurface* Destination = AGXGetMetalSurfaceFromRHITexture(Info.DestTexture);
+
+	const FRHITextureDesc& SourceDesc = Source->GetDesc();
+	const FRHITextureDesc& DestinationDesc = Destination->GetDesc();
+
+	const bool bDepthStencil = SourceDesc.Format == PF_DepthStencil;
+	const bool bSupportsMSAADepthResolve = GetAGXDeviceContext().SupportsFeature(EAGXFeaturesMSAADepthResolve);
+	const bool bSupportsMSAAStoreAndResolve = GetAGXDeviceContext().SupportsFeature(EAGXFeaturesMSAAStoreAndResolve);
+	// Resolve required - Device must support this - Using Shader for resolve not supported amd NumSamples should be 1
+	check((!bDepthStencil && bSupportsMSAAStoreAndResolve) || (bDepthStencil && bSupportsMSAADepthResolve));
+
+	mtlpp::Origin Origin(0, 0, 0);
+	mtlpp::Size Size(0, 0, 1);
+
+	if (Info.ResolveRect.IsValid())
+	{
+		Origin.x    = Info.ResolveRect.X1;
+		Origin.y    = Info.ResolveRect.Y1;
+		Size.width  = Info.ResolveRect.X2 - Info.ResolveRect.X1;
+		Size.height = Info.ResolveRect.Y2 - Info.ResolveRect.Y1;
+	}
+	else
+	{
+		Size.width  = FMath::Max<uint32>(1, SourceDesc.Extent.X >> Info.MipLevel);
+		Size.height = FMath::Max<uint32>(1, SourceDesc.Extent.Y >> Info.MipLevel);
+	}
+
+	if (Profiler)
+	{
+		Profiler->RegisterGPUWork();
+	}
+
+	int32 ArraySliceBegin = Info.ArraySlice;
+	int32 ArraySliceEnd   = Info.ArraySlice + 1;
+
+	if (Info.ArraySlice < 0)
+	{
+		ArraySliceBegin = 0;
+		ArraySliceEnd   = SourceDesc.ArraySize;
+	}
+
+	for (int32 ArraySlice = ArraySliceBegin; ArraySlice < ArraySliceEnd; ArraySlice++)
+	{
+		Context->CopyFromTextureToTexture(Source->MSAAResolveTexture, ArraySlice, Info.MipLevel, Origin, Size, Destination->Texture, ArraySlice, Info.MipLevel, Origin);
+
+#if PLATFORM_MAC
+		if ((Destination->GPUReadback & FAGXSurface::EAGXGPUReadbackFlags::ReadbackRequested) != 0)
+		{
+			Context->GetCurrentRenderPass().SynchronizeTexture(Destination->Texture, ArraySlice, Info.MipLevel);
+		}
+#endif
+	}
+	}
 }
 
 void FAGXRHICommandContext::RHINextSubpass()
diff --git a/Engine/Source/Runtime/Apple/AGXRHI/Public/AGXRHIContext.h b/Engine/Source/Runtime/Apple/AGXRHI/Public/AGXRHIContext.h
--- a/Engine/Source/Runtime/Apple/AGXRHI/Public/AGXRHIContext.h
+++ b/Engine/Source/Runtime/Apple/AGXRHI/Public/AGXRHIContext.h
@@ -6,6 +6,7 @@
 #include "AGXState.h"
 #include "AGXResources.h"
 #include "AGXViewport.h"
+#include "RHICore.h"
 
 #define UE_METAL_RHI_SUPPORT_CLEAR_UAV_WITH_BLIT_ENCODER	1
 
@@ -237,6 +238,8 @@ private:
 #endif // UE_METAL_RHI_SUPPORT_CLEAR_UAV_WITH_BLIT_ENCODER
 	void ClearUAV(TRHICommandList_RecursiveHazardous<FAGXRHICommandContext>& RHICmdList, FAGXUnorderedAccessView* UnorderedAccessView, const void* ClearValue, bool bFloat);
 
+	void ResolveTexture(UE::RHICore::FResolveTextureInfo Info);
+
 	void RHIClearMRT(bool bClearColor, int32 NumClearColors, const FLinearColor* ColorArray, bool bClearDepth, float Depth, bool bClearStencil, uint32 Stencil);
 };
 
diff --git a/Engine/Source/Runtime/Apple/MetalRHI/Private/MetalRHIContext.cpp b/Engine/Source/Runtime/Apple/MetalRHI/Private/MetalRHIContext.cpp
--- a/Engine/Source/Runtime/Apple/MetalRHI/Private/MetalRHIContext.cpp
+++ b/Engine/Source/Runtime/Apple/MetalRHI/Private/MetalRHIContext.cpp
@@ -3,7 +3,6 @@
 #include "MetalRHIPrivate.h"
 #include "MetalRHIRenderQuery.h"
 #include "MetalCommandBufferFence.h"
-#include "RHICore.h"
 
 TGlobalResource<TBoundShaderStateHistory<10000>> FMetalRHICommandContext::BoundShaderStateHistory;
 
@@ -170,7 +169,69 @@ void FMetalRHICommandContext::RHIEndRenderPass()
 		RHIEndOcclusionQueryBatch();
 	}
 
-	UE::RHICore::ResolveRenderPassTargets(*this, RenderPassInfo);
+	UE::RHICore::ResolveRenderPassTargets(RenderPassInfo, [this](UE::RHICore::FResolveTextureInfo Info)
+	{
+		ResolveTexture(Info);
+	});
+}
+
+void FMetalRHICommandContext::ResolveTexture(UE::RHICore::FResolveTextureInfo Info)
+{
+	@autoreleasepool{
+	FMetalSurface* Source = GetMetalSurfaceFromRHITexture(Info.SourceTexture);
+	FMetalSurface* Destination = GetMetalSurfaceFromRHITexture(Info.DestTexture);
+
+	const FRHITextureDesc& SourceDesc = Source->GetDesc();
+	const FRHITextureDesc& DestinationDesc = Destination->GetDesc();
+
+	const bool bDepthStencil = SourceDesc.Format == PF_DepthStencil;
+	const bool bSupportsMSAADepthResolve = GetMetalDeviceContext().SupportsFeature(EMetalFeaturesMSAADepthResolve);
+	const bool bSupportsMSAAStoreAndResolve = GetMetalDeviceContext().SupportsFeature(EMetalFeaturesMSAAStoreAndResolve);
+	// Resolve required - Device must support this - Using Shader for resolve not supported amd NumSamples should be 1
+	check((!bDepthStencil && bSupportsMSAAStoreAndResolve) || (bDepthStencil && bSupportsMSAADepthResolve));
+
+	mtlpp::Origin Origin(0, 0, 0);
+	mtlpp::Size Size(0, 0, 1);
+
+	if (Info.ResolveRect.IsValid())
+	{
+		Origin.x    = Info.ResolveRect.X1;
+		Origin.y    = Info.ResolveRect.Y1;
+		Size.width  = Info.ResolveRect.X2 - Info.ResolveRect.X1;
+		Size.height = Info.ResolveRect.Y2 - Info.ResolveRect.Y1;
+	}
+	else
+	{
+		Size.width  = FMath::Max<uint32>(1, SourceDesc.Extent.X >> Info.MipLevel);
+		Size.height = FMath::Max<uint32>(1, SourceDesc.Extent.Y >> Info.MipLevel);
+	}
+
+	if (Profiler)
+	{
+		Profiler->RegisterGPUWork();
+	}
+
+	int32 ArraySliceBegin = Info.ArraySlice;
+	int32 ArraySliceEnd   = Info.ArraySlice + 1;
+
+	if (Info.ArraySlice < 0)
+	{
+		ArraySliceBegin = 0;
+		ArraySliceEnd   = SourceDesc.ArraySize;
+	}
+
+	for (int32 ArraySlice = ArraySliceBegin; ArraySlice < ArraySliceEnd; ArraySlice++)
+	{
+		Context->CopyFromTextureToTexture(Source->MSAAResolveTexture, ArraySlice, Info.MipLevel, Origin, Size, Destination->Texture, ArraySlice, Info.MipLevel, Origin);
+
+#if PLATFORM_MAC
+		if ((Destination->GPUReadback & FMetalSurface::EMetalGPUReadbackFlags::ReadbackRequested) != 0)
+		{
+			Context->GetCurrentRenderPass().SynchronizeTexture(Destination->Texture, ArraySlice, Info.MipLevel);
+		}
+#endif
+	}
+	}
 }
 
 void FMetalRHICommandContext::RHINextSubpass()
diff --git a/Engine/Source/Runtime/Apple/MetalRHI/Public/MetalRHIContext.h b/Engine/Source/Runtime/Apple/MetalRHI/Public/MetalRHIContext.h
--- a/Engine/Source/Runtime/Apple/MetalRHI/Public/MetalRHIContext.h
+++ b/Engine/Source/Runtime/Apple/MetalRHI/Public/MetalRHIContext.h
@@ -7,6 +7,7 @@
 #include "MetalState.h"
 #include "MetalResources.h"
 #include "MetalViewport.h"
+#include "RHICore.h"
 
 #define UE_METAL_RHI_SUPPORT_CLEAR_UAV_WITH_BLIT_ENCODER	1
 
@@ -230,6 +231,8 @@ protected:
 	template <typename TRHIShader>
 	void ApplyStaticUniformBuffers(TRHIShader* Shader);
 
+	void ResolveTexture(UE::RHICore::FResolveTextureInfo Info);
+
 	TArray<FRHIUniformBuffer*> GlobalUniformBuffers;
 
 private:
diff --git a/Engine/Source/Runtime/OpenGLDrv/Private/OpenGLRenderTarget.cpp b/Engine/Source/Runtime/OpenGLDrv/Private/OpenGLRenderTarget.cpp
--- a/Engine/Source/Runtime/OpenGLDrv/Private/OpenGLRenderTarget.cpp
+++ b/Engine/Source/Runtime/OpenGLDrv/Private/OpenGLRenderTarget.cpp
@@ -1164,8 +1164,6 @@ void FOpenGLDynamicRHI::RHIEndRenderPass()
 		EndOcclusionQueryBatch();
 	}
 
-	UE::RHICore::ResolveRenderPassTargets(*this, RenderPassInfo);
-
 	// Drop depth and stencil to avoid export
 	if (RenderPassInfo.DepthStencilRenderTarget.DepthStencilTarget)
 	{
diff --git a/Engine/Source/Runtime/RHICore/Private/RHICore.cpp b/Engine/Source/Runtime/RHICore/Private/RHICore.cpp
--- a/Engine/Source/Runtime/RHICore/Private/RHICore.cpp
+++ b/Engine/Source/Runtime/RHICore/Private/RHICore.cpp
@@ -11,55 +11,6 @@ namespace UE
 namespace RHICore
 {
 
-void ResolveRenderPassTargets(IRHICommandContext& Context, const FRHIRenderPassInfo& Info)
-{
-	const auto ResolveTexture = [&](FRHITexture* Target, FRHITexture* Resolve, uint8 MipIndex, int32 ArraySlice)
-	{
-		if (!Target || !Resolve || Target == Resolve)
-		{
-			return;
-		}
-
-		const FRHITextureDesc& TargetDesc  = Target->GetDesc();
-		const FRHITextureDesc& ResolveDesc = Resolve->GetDesc();
-
-		check(TargetDesc.IsTextureCube() == ResolveDesc.IsTextureCube());
-		check(TargetDesc.IsMultisample() && !ResolveDesc.IsMultisample());
-		check(!TargetDesc.IsTextureArray() || ArraySlice >= 0);
-
-		ArraySlice = FMath::Max(ArraySlice, 0);
-
-		int32 CubeFaceIndex = 0;
-
-		if (TargetDesc.IsTextureCube())
-		{
-			CubeFaceIndex = ArraySlice % CubeFace_MAX;
-			ArraySlice    = ArraySlice / CubeFace_MAX;
-		}
-
-		FResolveParams Params;
-		Params.CubeFace          = (ECubeFace)CubeFaceIndex;
-		Params.Rect              = Info.ResolveRect;
-		Params.DestRect          = Info.ResolveRect;
-		Params.MipIndex          = MipIndex;
-		Params.SourceAccessFinal = ERHIAccess::RTV;
-		Params.DestAccessFinal   = ERHIAccess::ResolveDst;
-		Params.SourceArrayIndex  = ArraySlice;
-		Params.DestArrayIndex    = ArraySlice;
-
-		Context.RHICopyToResolveTarget(Target, Resolve, Params);
-	};
-
-	for (int32 Index = 0; Index < MaxSimultaneousRenderTargets; ++Index)
-	{
-		const auto& RTV = Info.ColorRenderTargets[Index];
-		ResolveTexture(RTV.RenderTarget, RTV.ResolveTarget, RTV.MipIndex, RTV.ArraySlice);
-	}
-
-	const auto& DSV = Info.DepthStencilRenderTarget;
-	ResolveTexture(DSV.DepthStencilTarget, DSV.ResolveTarget, 0, 0);
-}
-
 void ResolveRenderPassTargets(const FRHIRenderPassInfo& RenderPassInfo, TFunction<void(FResolveTextureInfo)> ResolveFunction)
 {
 	const auto ResolveTexture = [&](FResolveTextureInfo ResolveInfo)
diff --git a/Engine/Source/Runtime/RHICore/Public/RHICore.h b/Engine/Source/Runtime/RHICore/Public/RHICore.h
--- a/Engine/Source/Runtime/RHICore/Public/RHICore.h
+++ b/Engine/Source/Runtime/RHICore/Public/RHICore.h
@@ -9,16 +9,8 @@
 
 RHICORE_API DECLARE_LOG_CATEGORY_EXTERN(LogRHICore, Log, VeryVerbose);
 
-class IRHICommandContext;
-
-struct FRHIRenderPassInfo;
-
-namespace UE
+namespace UE::RHICore
 {
-namespace RHICore
-{
-
-RHICORE_API void ResolveRenderPassTargets(IRHICommandContext& Context, const FRHIRenderPassInfo& Info);
 
 struct FResolveTextureInfo
 {
@@ -45,5 +37,4 @@ struct FResolveTextureInfo
 
 RHICORE_API void ResolveRenderPassTargets(const FRHIRenderPassInfo& Info, TFunction<void(FResolveTextureInfo)> Function);
 
-} //! RHICore
-} //! UE
\ No newline at end of file
+} //! UE::RHICore
\ No newline at end of file
diff --git a/Engine/Source/Runtime/Windows/D3D11RHI/Private/D3D11RHIPrivate.h b/Engine/Source/Runtime/Windows/D3D11RHI/Private/D3D11RHIPrivate.h
--- a/Engine/Source/Runtime/Windows/D3D11RHI/Private/D3D11RHIPrivate.h
+++ b/Engine/Source/Runtime/Windows/D3D11RHI/Private/D3D11RHIPrivate.h
@@ -658,6 +658,8 @@ public:
 	virtual void RHIBeginRenderPass(const FRHIRenderPassInfo& InInfo, const TCHAR* InName) final override;
 	virtual void RHIEndRenderPass() final override;
 
+	void ResolveTexture(UE::RHICore::FResolveTextureInfo Info);
+
 	virtual void RHICalibrateTimers() override;
 
 	// ID3D11DynamicRHI interface
diff --git a/Engine/Source/Runtime/Windows/D3D11RHI/Private/D3D11RenderTarget.cpp b/Engine/Source/Runtime/Windows/D3D11RHI/Private/D3D11RenderTarget.cpp
--- a/Engine/Source/Runtime/Windows/D3D11RHI/Private/D3D11RenderTarget.cpp
+++ b/Engine/Source/Runtime/Windows/D3D11RHI/Private/D3D11RenderTarget.cpp
@@ -1136,9 +1136,65 @@ void FD3D11DynamicRHI::RHIEndRenderPass()
 		RHIEndOcclusionQueryBatch();
 	}
 
-	UE::RHICore::ResolveRenderPassTargets(*this, RenderPassInfo);
+	UE::RHICore::ResolveRenderPassTargets(RenderPassInfo, [this](UE::RHICore::FResolveTextureInfo Info)
+	{
+		ResolveTexture(Info);
+	});
 
 	FRHIRenderTargetView RTV(nullptr, ERenderTargetLoadAction::ENoAction);
 	FRHIDepthRenderTargetView DepthRTV(nullptr, ERenderTargetLoadAction::ENoAction, ERenderTargetStoreAction::ENoAction);
 	SetRenderTargets(1, &RTV, &DepthRTV);
 }
+
+void FD3D11DynamicRHI::ResolveTexture(UE::RHICore::FResolveTextureInfo Info)
+{
+	GPUProfilingData.RegisterGPUWork();
+
+	FD3D11Texture* SourceTexture      = GetD3D11TextureFromRHITexture(Info.SourceTexture);
+	const FRHITextureDesc& SourceDesc = SourceTexture->GetDesc();
+
+	FD3D11Texture* DestTexture        = GetD3D11TextureFromRHITexture(Info.DestTexture);
+	const FRHITextureDesc& DestDesc   = DestTexture->GetDesc();
+
+	if (SourceDesc.Format == PF_DepthStencil)
+	{
+		FRHICommandList_RecursiveHazardous RHICmdList(this);
+
+		D3D11_TEXTURE2D_DESC ResolveTargetDesc;
+		DestTexture->GetD3D11Texture2D()->GetDesc(&ResolveTargetDesc);
+
+		ResolveTextureUsingShader<FResolveDepthPS>(
+			RHICmdList,
+			SourceTexture,
+			DestTexture,
+			DestTexture->GetRenderTargetView(0, -1),
+			DestTexture->GetDepthStencilView(FExclusiveDepthStencil::DepthWrite_StencilWrite),
+			ResolveTargetDesc,
+			GetDefaultRect(Info.ResolveRect, SourceDesc.Extent.X, SourceDesc.Extent.Y),
+			GetDefaultRect(Info.ResolveRect, DestDesc.Extent.X, DestDesc.Extent.Y),
+			Direct3DDeviceIMContext,
+			FDummyResolveParameter()
+		);
+	}
+	else
+	{
+		const DXGI_FORMAT DestFormatTypeless = ConvertTypelessToUnorm((DXGI_FORMAT)GPixelFormats[DestDesc.Format].PlatformFormat);
+
+		int32 ArraySliceBegin = Info.ArraySlice;
+		int32 ArraySliceEnd   = Info.ArraySlice + 1;
+
+		if (Info.ArraySlice < 0)
+		{
+			ArraySliceBegin = 0;
+			ArraySliceEnd   = SourceDesc.ArraySize;
+		}
+
+		for (int32 ArraySlice = ArraySliceBegin; ArraySlice < ArraySliceEnd; ArraySlice++)
+		{
+			int32 DestSubresource   = D3D11CalcSubresource(Info.MipLevel, ArraySlice, DestDesc.ArraySize);
+			int32 SourceSubresource = D3D11CalcSubresource(Info.MipLevel, ArraySlice, SourceDesc.ArraySize);
+
+			Direct3DDeviceIMContext->ResolveSubresource(DestTexture->GetResource(), DestSubresource, SourceTexture->GetResource(), SourceSubresource, DestFormatTypeless);
+		}
+	}
+}
\ No newline at end of file
