Metal RHI: Add missing pixel format conversion (A8R8G8B8)

#rnx

diff --git a/Engine/Source/Runtime/Apple/AGXRHI/Private/AGXRHI.cpp b/Engine/Source/Runtime/Apple/AGXRHI/Private/AGXRHI.cpp
--- a/Engine/Source/Runtime/Apple/AGXRHI/Private/AGXRHI.cpp
+++ b/Engine/Source/Runtime/Apple/AGXRHI/Private/AGXRHI.cpp
@@ -907,6 +907,7 @@ FAGXDynamicRHI::FAGXDynamicRHI(ERHIFeatureLevel::Type RequestedFeatureLevel)
 	GPixelFormats[PF_R16G16B16A16_UINT	].PlatformFormat	= (uint32)MTLPixelFormatRGBA16Uint;
 	GPixelFormats[PF_R16G16B16A16_SINT	].PlatformFormat	= (uint32)MTLPixelFormatRGBA16Sint;
 	GPixelFormats[PF_R8G8B8A8			].PlatformFormat	= (uint32)MTLPixelFormatRGBA8Unorm;
+	GPixelFormats[PF_A8R8G8B8			].PlatformFormat	= (uint32)MTLPixelFormatRGBA8Unorm;
 	GPixelFormats[PF_R8G8B8A8_UINT		].PlatformFormat	= (uint32)MTLPixelFormatRGBA8Uint;
 	GPixelFormats[PF_R8G8B8A8_SNORM		].PlatformFormat	= (uint32)MTLPixelFormatRGBA8Snorm;
 	GPixelFormats[PF_R8G8				].PlatformFormat	= (uint32)MTLPixelFormatRG8Unorm;
diff --git a/Engine/Source/Runtime/Apple/MetalRHI/Private/MetalRHI.cpp b/Engine/Source/Runtime/Apple/MetalRHI/Private/MetalRHI.cpp
--- a/Engine/Source/Runtime/Apple/MetalRHI/Private/MetalRHI.cpp
+++ b/Engine/Source/Runtime/Apple/MetalRHI/Private/MetalRHI.cpp
@@ -918,6 +918,7 @@ FMetalDynamicRHI::FMetalDynamicRHI(ERHIFeatureLevel::Type RequestedFeatureLevel)
 	GPixelFormats[PF_R16G16B16A16_UINT	].PlatformFormat	= (uint32)mtlpp::PixelFormat::RGBA16Uint;
 	GPixelFormats[PF_R16G16B16A16_SINT	].PlatformFormat	= (uint32)mtlpp::PixelFormat::RGBA16Sint;
 	GPixelFormats[PF_R8G8B8A8			].PlatformFormat	= (uint32)mtlpp::PixelFormat::RGBA8Unorm;
+    GPixelFormats[PF_A8R8G8B8           ].PlatformFormat    = (uint32)mtlpp::PixelFormat::RGBA8Unorm;
 	GPixelFormats[PF_R8G8B8A8_UINT		].PlatformFormat	= (uint32)mtlpp::PixelFormat::RGBA8Uint;
 	GPixelFormats[PF_R8G8B8A8_SNORM		].PlatformFormat	= (uint32)mtlpp::PixelFormat::RGBA8Snorm;
 	GPixelFormats[PF_R8G8				].PlatformFormat	= (uint32)mtlpp::PixelFormat::RG8Unorm;
