Speculative fix for premature release of texture view.   UAV texture view is now an FMetalTexture which is getting deleted at the end of the destructor - it was an FMetalSurface which handled deferred deletion - manually add safe destruction of texture view in the new FMetalResourceViewBase class destructor.

#rnx

diff --git a/Engine/Source/Runtime/Apple/AGXRHI/Private/AGXUAV.cpp b/Engine/Source/Runtime/Apple/AGXRHI/Private/AGXUAV.cpp
--- a/Engine/Source/Runtime/Apple/AGXRHI/Private/AGXUAV.cpp
+++ b/Engine/Source/Runtime/Apple/AGXRHI/Private/AGXUAV.cpp
@@ -155,9 +155,9 @@ FAGXResourceViewBase::FAGXResourceViewBase(
 				   MipLevel == 0
 				&& NumMips == [SourceTextureInternal mipmapLevelCount]
 				&& MetalFormat == [SourceTextureInternal pixelFormat]
-				&& !(bInUAV && SourceTexture->GetDesc().IsTextureCube())
+				&& !(bInUAV && SourceTexture->GetDesc().IsTextureCube())	// @todo: Remove this once Cube UAV supported for all Metal Devices
 				&& InFirstArraySlice == 0
-				&& InNumArraySlices == 0;		// @todo: Remove this once Cube UAV supported for all Metal Devices
+				&& InNumArraySlices == 0;
 
 			if (bUseSourceTex)
 			{
@@ -218,11 +218,17 @@ FAGXResourceViewBase::FAGXResourceViewBase(
 
 FAGXResourceViewBase::~FAGXResourceViewBase()
 {
+	if (TextureView)
+	{
+		FAGXTexture TmpWrapper(TextureView);
+		AGXSafeReleaseMetalTexture(TmpWrapper);
+		TextureView = nil;
+	}
+	
 	if (bTexture)
 	{
 		if (SourceTexture)
 		{
-			[TextureView release];
 			SourceTexture->Release();
 		}
 	}
diff --git a/Engine/Source/Runtime/Apple/MetalRHI/Private/MetalUAV.cpp b/Engine/Source/Runtime/Apple/MetalRHI/Private/MetalUAV.cpp
--- a/Engine/Source/Runtime/Apple/MetalRHI/Private/MetalUAV.cpp
+++ b/Engine/Source/Runtime/Apple/MetalRHI/Private/MetalUAV.cpp
@@ -207,6 +207,12 @@ FMetalResourceViewBase::FMetalResourceViewBase(
 
 FMetalResourceViewBase::~FMetalResourceViewBase()
 {
+	if (TextureView)
+	{
+		SafeReleaseMetalTexture(TextureView);
+		TextureView = nil;
+	}
+	
 	if (bTexture)
 	{
 		if (SourceTexture)
