Fix for Texture3D update region with compressed pixel formats - account for format info blocksize when performing line copies to support compressed formats.   Apply the same logic for Texture2D.  Some other minor tidy ups in this area to simplify buffer handling.  Change intermediate buffers from Dynamic to Static - they are just temp surfaces and only update once we don't need to pay the additional cost of 3 backing buffers.

#rnx

diff --git a/Engine/Source/Runtime/Apple/AGXRHI/Private/AGXTexture.cpp b/Engine/Source/Runtime/Apple/AGXRHI/Private/AGXTexture.cpp
--- a/Engine/Source/Runtime/Apple/AGXRHI/Private/AGXTexture.cpp
+++ b/Engine/Source/Runtime/Apple/AGXRHI/Private/AGXTexture.cpp
@@ -1623,11 +1623,9 @@ static void InternalExpandR8ToStandardRGBA(uint32* pDest, const struct FUpdateTe
 }
 #endif
 
-static FAGXBuffer InternalCopyTexture2DUpdateRegion(FRHITexture2D* TextureRHI, const struct FUpdateTextureRegion2D& UpdateRegion, uint32& InOutSourcePitch, const uint8* SourceData)
-{	
-	const uint32 InSourcePitch = InOutSourcePitch;
-
-	FAGXBuffer Buffer;
+static FAGXBuffer Internal_CreateBufferAndCopyTexture2DUpdateRegionData(FRHITexture2D* TextureRHI, const struct FUpdateTextureRegion2D& UpdateRegion, uint32& InOutSourcePitch, const uint8* SourceData)
+{
+	FAGXBuffer OutBuffer;
 
 	FAGXSurface* Texture = ResourceCast(TextureRHI);	
 
@@ -1637,32 +1635,35 @@ static FAGXBuffer InternalCopyTexture2DUpdateRegion(FRHITexture2D* TextureRHI, c
 		&& EnumHasAnyFlags(Texture->GetFlags(), TexCreate_SRGB)
 		&& (MTLPixelFormat)Texture->Texture.GetPixelFormat() == MTLPixelFormatRGBA8Unorm_sRGB)
 	{
-		const uint32 BufferSize = UpdateRegion.Height * UpdateRegion.Width * sizeof(uint32);
-		Buffer = GetAGXDeviceContext().CreatePooledBuffer(FAGXPooledBufferArgs(BufferSize, BUF_Dynamic, FAGXPooledBufferArgs::SharedStorageResourceOptions));
-		InternalExpandR8ToStandardRGBA((uint32*)Buffer.GetContents(), UpdateRegion, InOutSourcePitch, SourceData);
+		const uint32 ExpandedBufferSize = UpdateRegion.Height * UpdateRegion.Width * sizeof(uint32);
+		OutBuffer = GetAGXDeviceContext().CreatePooledBuffer(FAGXPooledBufferArgs(ExpandedBufferSize, BUF_Static, FAGXPooledBufferArgs::SharedStorageResourceOptions));
+		InternalExpandR8ToStandardRGBA((uint32*)OutBuffer.GetContents(), UpdateRegion, InOutSourcePitch, SourceData);
 	}
-
-	if(Buffer.GetPtr() == nil)
+	else
 #endif
 	{
-		const uint32 BufferSize = UpdateRegion.Height * InSourcePitch;
-		Buffer = GetAGXDeviceContext().CreatePooledBuffer(FAGXPooledBufferArgs(BufferSize, BUF_Dynamic, FAGXPooledBufferArgs::SharedStorageResourceOptions));
+		const FPixelFormatInfo& FormatInfo = GPixelFormats[TextureRHI->GetFormat()];
+		
+		const uint32 BufferSize = UpdateRegion.Height * InOutSourcePitch;
+		OutBuffer = GetAGXDeviceContext().CreatePooledBuffer(FAGXPooledBufferArgs(BufferSize, BUF_Static, FAGXPooledBufferArgs::SharedStorageResourceOptions));
 
-		// Limit copy to line by line by update region pitch otherwise we can go off the end of source data on the last row
-		uint8* pDestRow = (uint8*)Buffer.GetContents();
-		uint8* pSourceRow = (uint8*)SourceData;
-		uint32 CopyPitch = FMath::DivideAndRoundUp(UpdateRegion.Width, (uint32)GPixelFormats[Texture->GetFormat()].BlockSizeX) * GPixelFormats[Texture->GetFormat()].BlockBytes;
-		check(CopyPitch <= InSourcePitch);
+		uint32 CopyPitch = FMath::DivideAndRoundUp(UpdateRegion.Width, (uint32)FormatInfo.BlockSizeX) * FormatInfo.BlockBytes;
+		check(CopyPitch <= InOutSourcePitch);
 		
-		for (uint32 i = 0;i < UpdateRegion.Height;++i)
+		uint8* pDestRow = (uint8*)OutBuffer.GetContents();
+		uint8* pSourceRow = (uint8*)SourceData;
+		const uint32 NumRows = UpdateRegion.Height / (uint32)FormatInfo.BlockSizeY;
+
+		// Limit copy to line by line by update region pitch otherwise we can go off the end of source data on the last row
+		for (uint32 i = 0;i < NumRows;++i)
 		{
 			FMemory::Memcpy(pDestRow, pSourceRow, CopyPitch);
-			pSourceRow += InSourcePitch;
-			pDestRow += InSourcePitch;
+			pSourceRow += InOutSourcePitch;
+			pDestRow += InOutSourcePitch;
 		}
 	}
 
-	return Buffer;
+	return OutBuffer;
 }
 
 static void InternalUpdateTexture2D(FAGXContext& Context, FRHITexture2D* TextureRHI, uint32 MipIndex, FUpdateTextureRegion2D const& UpdateRegion, uint32 SourcePitch, FAGXBuffer Buffer)
@@ -1676,7 +1677,9 @@ static void InternalUpdateTexture2D(FAGXContext& Context, FRHITexture2D* Texture
 	{
 		SCOPED_AUTORELEASE_POOL;
 		
-		uint32 BytesPerImage = SourcePitch * UpdateRegion.Height;
+		const FPixelFormatInfo& FormatInfo = GPixelFormats[TextureRHI->GetFormat()];
+		const uint32 NumRows = UpdateRegion.Height / (uint32)FormatInfo.BlockSizeY;
+		uint32 BytesPerImage = SourcePitch * NumRows;
 		
 		mtlpp::BlitOption Options = mtlpp::BlitOption::None;
 #if !PLATFORM_MAC
@@ -1707,14 +1710,14 @@ struct FAGXRHICommandUpdateTexture2D final : public FRHICommand<FAGXRHICommandUp
 	uint32 SourcePitch;
 	FAGXBuffer SourceBuffer;
 
-	FORCEINLINE_DEBUGGABLE FAGXRHICommandUpdateTexture2D(FAGXContext& InContext, FRHITexture2D* InTexture, uint32 InMipIndex, FUpdateTextureRegion2D InUpdateRegion, uint32 InSourcePitch, FAGXBuffer InSourceBuffer)
+	FORCEINLINE_DEBUGGABLE FAGXRHICommandUpdateTexture2D(FAGXContext& InContext, FRHITexture2D* InTexture, uint32 InMipIndex, FUpdateTextureRegion2D InUpdateRegion, uint32 InSourcePitch, const uint8* SourceData)
 	: Context(InContext)
 	, Texture(InTexture)
 	, MipIndex(InMipIndex)
 	, UpdateRegion(InUpdateRegion)
 	, SourcePitch(InSourcePitch)
-	, SourceBuffer(InSourceBuffer)
 	{
+		SourceBuffer = Internal_CreateBufferAndCopyTexture2DUpdateRegionData(Texture, UpdateRegion, SourcePitch, SourceData);
 	}
 	
 	void Execute(FRHICommandListBase& CmdList)
@@ -1734,8 +1737,7 @@ void FAGXDynamicRHI::UpdateTexture2D_RenderThread(class FRHICommandListImmediate
 		}
 		else
 		{
-			FAGXBuffer Buffer = InternalCopyTexture2DUpdateRegion(Texture, UpdateRegion, SourcePitch, SourceData);
-			new (RHICmdList.AllocCommand<FAGXRHICommandUpdateTexture2D>()) FAGXRHICommandUpdateTexture2D(ImmediateContext.GetInternalContext(), Texture, MipIndex, UpdateRegion, SourcePitch, Buffer);
+			new (RHICmdList.AllocCommand<FAGXRHICommandUpdateTexture2D>()) FAGXRHICommandUpdateTexture2D(ImmediateContext.GetInternalContext(), Texture, MipIndex, UpdateRegion, SourcePitch, SourceData);
 		}
 	}
 }
@@ -1746,11 +1748,10 @@ void FAGXDynamicRHI::RHIUpdateTexture2D(FRHITexture2D* TextureRHI, uint32 MipInd
 	{
 		FAGXSurface* Texture = ResourceCast(TextureRHI);
 		FAGXTexture Tex = Texture->Texture;
-		bool const bUseIntermediateMetalBuffer = ([Tex.GetPtr() storageMode] == MTLStorageModePrivate);
 		
-		if(bUseIntermediateMetalBuffer)
+		if([Tex.GetPtr() storageMode] == MTLStorageModePrivate)
 		{
-			FAGXBuffer Buffer = InternalCopyTexture2DUpdateRegion(TextureRHI, UpdateRegion, SourcePitch, SourceData);
+			FAGXBuffer Buffer = Internal_CreateBufferAndCopyTexture2DUpdateRegionData(TextureRHI, UpdateRegion, SourcePitch, SourceData);
 			InternalUpdateTexture2D(ImmediateContext.GetInternalContext(), TextureRHI, MipIndex, UpdateRegion, SourcePitch, Buffer);
 			GetAGXDeviceContext().ReleaseBuffer(Buffer);
 		}
@@ -1775,26 +1776,37 @@ void FAGXDynamicRHI::RHIUpdateTexture2D(FRHITexture2D* TextureRHI, uint32 MipInd
 	}
 }
 
-static void InternalCopyTexture3DUpdateRegionData(FRHITexture3D* TextureRHI, const struct FUpdateTextureRegion3D& UpdateRegion, uint32 SourceRowPitch, uint32 SourceDepthPitch, const uint8* SourceData, uint8* DestData)
+static FAGXBuffer Internal_CreateBufferAndCopyTexture3DUpdateRegionData(FRHITexture3D* TextureRHI, const struct FUpdateTextureRegion3D& UpdateRegion, uint32 SourceRowPitch, uint32 SourceDepthPitch, const uint8* SourceData)
 {
-	// Perform safe line copy
-	FAGXSurface* Texture = ResourceCast(TextureRHI);	
+	FAGXSurface* Texture = ResourceCast(TextureRHI);
 	
-	uint8* pDest = DestData;
-	uint32 CopyPitch = FMath::DivideAndRoundUp(UpdateRegion.Width, (uint32)GPixelFormats[Texture->GetFormat()].BlockSizeX) * GPixelFormats[Texture->GetFormat()].BlockBytes;
+	const uint32 BufferSize = SourceDepthPitch * UpdateRegion.Depth;
+	FAGXBuffer OutBuffer = GetAGXDeviceContext().CreatePooledBuffer(FAGXPooledBufferArgs(BufferSize, BUF_Static, FAGXPooledBufferArgs::SharedStorageResourceOptions));
 
-	for (uint32 i = 0; i < UpdateRegion.Depth;++i)
+	const FPixelFormatInfo& FormatInfo = GPixelFormats[TextureRHI->GetFormat()];
+	uint32 CopyPitch = FMath::DivideAndRoundUp(UpdateRegion.Width, (uint32)FormatInfo.BlockSizeX) * FormatInfo.BlockBytes;
+	
+	check(FormatInfo.BlockSizeZ == 1);
+	check(CopyPitch <= SourceRowPitch);
+	
+	uint8_t* DestData = (uint8_t*)OutBuffer.GetContents();
+	const uint32 NumRows = UpdateRegion.Height / (uint32)FormatInfo.BlockSizeY;
+	
+	// Perform safe line copy
+	for (uint32 i = 0;i < UpdateRegion.Depth;++i)
 	{
 		const uint8* pSourceRowData = SourceData + (SourceDepthPitch * i);
-		uint8* pDestRowData = pDest + (SourceDepthPitch * i);
+		uint8* pDestRowData = DestData + (SourceDepthPitch * i);
 
-		for (uint32 j = 0;j < UpdateRegion.Height;++j)
+		for (uint32 j = 0;j < NumRows;++j)
 		{
 			FMemory::Memcpy(pDestRowData, pSourceRowData, CopyPitch);
 			pSourceRowData += SourceRowPitch;
 			pDestRowData += SourceRowPitch;
 		}
 	}
+	
+	return OutBuffer;
 }
 
 
@@ -1807,7 +1819,10 @@ static void InternalUpdateTexture3D(FAGXContext& Context, FRHITexture3D* Texture
 	
 	if ([Tex.GetPtr() storageMode] == MTLStorageModePrivate)
 	{
-		const uint32 BytesPerImage = SourceRowPitch * UpdateRegion.Height;
+		const FPixelFormatInfo& FormatInfo = GPixelFormats[TextureRHI->GetFormat()];
+		const uint32 NumRows = UpdateRegion.Height / (uint32)FormatInfo.BlockSizeY;
+		const uint32 BytesPerImage = SourceRowPitch * NumRows;
+
 		mtlpp::BlitOption Options = mtlpp::BlitOption::None;
 #if !PLATFORM_MAC
 		if ((MTLPixelFormat)Tex.GetPixelFormat() >= MTLPixelFormatPVRTC_RGB_2BPP && (MTLPixelFormat)Tex.GetPixelFormat() <= MTLPixelFormatPVRTC_RGBA_4BPP_sRGB)
@@ -1846,12 +1861,7 @@ struct FAGXDynamicRHIUpdateTexture3DCommand final : public FRHICommand<FAGXDynam
 	, SourceRowPitch(InSourceRowPitch)
 	, SourceDepthPitch(InSourceDepthPitch)
 	{
-		FAGXSurface* Texture = ResourceCast(TextureRHI);
-		FAGXTexture Tex = Texture->Texture;
-		const uint32 BufferSize = UpdateRegion.Height * UpdateRegion.Depth* SourceRowPitch;
-		
-		Buffer = GetAGXDeviceContext().CreatePooledBuffer(FAGXPooledBufferArgs(BufferSize, BUF_Dynamic, FAGXPooledBufferArgs::SharedStorageResourceOptions));
-		InternalCopyTexture3DUpdateRegionData(Texture, UpdateRegion, SourceRowPitch, SourceDepthPitch, SourceData, (uint8*)Buffer.GetContents());
+		Buffer = Internal_CreateBufferAndCopyTexture3DUpdateRegionData(TextureRHI, UpdateRegion, SourceRowPitch, SourceDepthPitch, SourceData);
 	}
 	
 	void Execute(FRHICommandListBase& CmdList)
@@ -1919,10 +1929,7 @@ void FAGXDynamicRHI::RHIUpdateTexture3D(FRHITexture3D* TextureRHI,uint32 MipInde
 		if ([Tex.GetPtr() storageMode] == MTLStorageModePrivate)
 		{
 			SCOPED_AUTORELEASE_POOL;
-
-			const uint32 BufferSize = UpdateRegion.Height * UpdateRegion.Depth * SourceRowPitch;
-			FAGXBuffer IntermediateBuffer = GetAGXDeviceContext().CreatePooledBuffer(FAGXPooledBufferArgs(BufferSize, BUF_Dynamic, FAGXPooledBufferArgs::SharedStorageResourceOptions));
-			InternalCopyTexture3DUpdateRegionData(TextureRHI, UpdateRegion, SourceRowPitch, SourceDepthPitch, SourceData, (uint8*)IntermediateBuffer.GetContents());
+			FAGXBuffer IntermediateBuffer = Internal_CreateBufferAndCopyTexture3DUpdateRegionData(TextureRHI, UpdateRegion, SourceRowPitch, SourceDepthPitch, SourceData);
 			InternalUpdateTexture3D(ImmediateContext.GetInternalContext(), TextureRHI, MipIndex, UpdateRegion, SourceRowPitch, SourceDepthPitch, IntermediateBuffer);
 			GetAGXDeviceContext().ReleaseBuffer(IntermediateBuffer);
 		}
diff --git a/Engine/Source/Runtime/Apple/MetalRHI/Private/MetalTexture.cpp b/Engine/Source/Runtime/Apple/MetalRHI/Private/MetalTexture.cpp
--- a/Engine/Source/Runtime/Apple/MetalRHI/Private/MetalTexture.cpp
+++ b/Engine/Source/Runtime/Apple/MetalRHI/Private/MetalTexture.cpp
@@ -1680,46 +1680,47 @@ static void InternalExpandR8ToStandardRGBA(uint32* pDest, const struct FUpdateTe
 }
 #endif
 
-static FMetalBuffer InternalCopyTexture2DUpdateRegion(FRHITexture2D* TextureRHI, const struct FUpdateTextureRegion2D& UpdateRegion, uint32& InOutSourcePitch, const uint8* SourceData)
-{	
-	const uint32 InSourcePitch = InOutSourcePitch;
-
-	FMetalBuffer Buffer;
+static FMetalBuffer Internal_CreateBufferAndCopyTexture2DUpdateRegionData(FRHITexture2D* TextureRHI, const struct FUpdateTextureRegion2D& UpdateRegion, uint32& InOutSourcePitch, const uint8* SourceData)
+{
+	FMetalBuffer OutBuffer;
 
 	FMetalSurface* Texture = ResourceCast(TextureRHI);
 
 #if PLATFORM_MAC
 	// Expand R8_sRGB into RGBA8_sRGB for non Apple Silicon Mac.
-	if (   Texture->GetFormat() == PF_G8 
-		&& EnumHasAnyFlags(Texture->GetFlags(), TexCreate_SRGB) 
+	if (   Texture->GetFormat() == PF_G8
+		&& EnumHasAnyFlags(Texture->GetFlags(), TexCreate_SRGB)
 		&& Texture->Texture.GetPixelFormat() == mtlpp::PixelFormat::RGBA8Unorm_sRGB)
 	{
-		const uint32 BufferSize = UpdateRegion.Height * UpdateRegion.Width * sizeof(uint32);		
-		Buffer = GetMetalDeviceContext().CreatePooledBuffer(FMetalPooledBufferArgs(GetMetalDeviceContext().GetDevice(), BufferSize, BUF_Dynamic, mtlpp::StorageMode::Shared));
-		InternalExpandR8ToStandardRGBA((uint32*)Buffer.GetContents(), UpdateRegion, InOutSourcePitch, SourceData);
+		const uint32 ExpandedBufferSize = UpdateRegion.Height * UpdateRegion.Width * sizeof(uint32);
+		OutBuffer = GetMetalDeviceContext().CreatePooledBuffer(FMetalPooledBufferArgs(GetMetalDeviceContext().GetDevice(), ExpandedBufferSize, BUF_Static, mtlpp::StorageMode::Shared));
+		InternalExpandR8ToStandardRGBA((uint32*)OutBuffer.GetContents(), UpdateRegion, InOutSourcePitch, SourceData);
 	}
-
-	if(Buffer.GetPtr() == nil)
+	else
 #endif
 	{
-		const uint32 BufferSize = UpdateRegion.Height * InSourcePitch;
-		Buffer = GetMetalDeviceContext().CreatePooledBuffer(FMetalPooledBufferArgs(GetMetalDeviceContext().GetDevice(), BufferSize, BUF_Dynamic, mtlpp::StorageMode::Shared));
+		const FPixelFormatInfo& FormatInfo = GPixelFormats[TextureRHI->GetFormat()];
+		
+		const uint32 BufferSize = UpdateRegion.Height * InOutSourcePitch;
+		OutBuffer = GetMetalDeviceContext().CreatePooledBuffer(FMetalPooledBufferArgs(GetMetalDeviceContext().GetDevice(), BufferSize, BUF_Static, mtlpp::StorageMode::Shared));
 
-		// Limit copy to line by line by update region pitch otherwise we can go off the end of source data on the last row
-		uint8* pDestRow = (uint8*)Buffer.GetContents();
+		uint32 CopyPitch = FMath::DivideAndRoundUp(UpdateRegion.Width, (uint32)FormatInfo.BlockSizeX) * FormatInfo.BlockBytes;
+		check(CopyPitch <= InOutSourcePitch);
+		
+		uint8* pDestRow = (uint8*)OutBuffer.GetContents();
 		uint8* pSourceRow = (uint8*)SourceData;
-		uint32 CopyPitch = FMath::DivideAndRoundUp(UpdateRegion.Width, (uint32)GPixelFormats[Texture->GetFormat()].BlockSizeX) * GPixelFormats[Texture->GetFormat()].BlockBytes;
-		check(CopyPitch <= InSourcePitch);
+		const uint32 NumRows = UpdateRegion.Height / (uint32)FormatInfo.BlockSizeY;
 		
-		for (uint32 i = 0;i < UpdateRegion.Height;++i)
+		// Limit copy to line by line by update region pitch otherwise we can go off the end of source data on the last row
+		for (uint32 i = 0;i < NumRows;++i)
 		{
 			FMemory::Memcpy(pDestRow, pSourceRow, CopyPitch);
-			pSourceRow += InSourcePitch;
-			pDestRow += InSourcePitch;
+			pSourceRow += InOutSourcePitch;
+			pDestRow += InOutSourcePitch;
 		}
 	}
 
-	return Buffer;
+	return OutBuffer;
 }
 
 static void InternalUpdateTexture2D(FMetalContext& Context, FRHITexture2D* TextureRHI, uint32 MipIndex, FUpdateTextureRegion2D const& UpdateRegion, uint32 SourcePitch, FMetalBuffer Buffer)
@@ -1733,7 +1734,9 @@ static void InternalUpdateTexture2D(FMetalContext& Context, FRHITexture2D* Textu
 	{
 		SCOPED_AUTORELEASE_POOL;
 		
-		uint32 BytesPerImage = SourcePitch * UpdateRegion.Height;
+		const FPixelFormatInfo& FormatInfo = GPixelFormats[TextureRHI->GetFormat()];
+		const uint32 NumRows = UpdateRegion.Height / (uint32)FormatInfo.BlockSizeY;
+		uint32 BytesPerImage = SourcePitch * NumRows;
 		
 		mtlpp::BlitOption Options = mtlpp::BlitOption::None;
 #if !PLATFORM_MAC
@@ -1764,14 +1767,14 @@ struct FMetalRHICommandUpdateTexture2D final : public FRHICommand<FMetalRHIComma
 	uint32 SourcePitch;
 	FMetalBuffer SourceBuffer;
 
-	FORCEINLINE_DEBUGGABLE FMetalRHICommandUpdateTexture2D(FMetalContext& InContext, FRHITexture2D* InTexture, uint32 InMipIndex, FUpdateTextureRegion2D InUpdateRegion, uint32 InSourcePitch, FMetalBuffer InSourceBuffer)
+	FORCEINLINE_DEBUGGABLE FMetalRHICommandUpdateTexture2D(FMetalContext& InContext, FRHITexture2D* InTexture, uint32 InMipIndex, FUpdateTextureRegion2D InUpdateRegion, uint32 InSourcePitch, const uint8* SourceData)
 	: Context(InContext)
 	, Texture(InTexture)
 	, MipIndex(InMipIndex)
 	, UpdateRegion(InUpdateRegion)
 	, SourcePitch(InSourcePitch)
-	, SourceBuffer(InSourceBuffer)
 	{
+		SourceBuffer = Internal_CreateBufferAndCopyTexture2DUpdateRegionData(Texture, UpdateRegion, SourcePitch, SourceData);
 	}
 	
 	void Execute(FRHICommandListBase& CmdList)
@@ -1791,8 +1794,7 @@ void FMetalDynamicRHI::UpdateTexture2D_RenderThread(class FRHICommandListImmedia
 		}
 		else
 		{
-			FMetalBuffer Buffer = InternalCopyTexture2DUpdateRegion(Texture, UpdateRegion, SourcePitch, SourceData);
-			new (RHICmdList.AllocCommand<FMetalRHICommandUpdateTexture2D>()) FMetalRHICommandUpdateTexture2D(ImmediateContext.GetInternalContext(), Texture, MipIndex, UpdateRegion, SourcePitch, Buffer);
+			new (RHICmdList.AllocCommand<FMetalRHICommandUpdateTexture2D>()) FMetalRHICommandUpdateTexture2D(ImmediateContext.GetInternalContext(), Texture, MipIndex, UpdateRegion, SourcePitch, SourceData);
 		}
 	}
 }
@@ -1803,11 +1805,10 @@ void FMetalDynamicRHI::RHIUpdateTexture2D(FRHITexture2D* TextureRHI, uint32 MipI
 	{
 		FMetalSurface* Texture = ResourceCast(TextureRHI);
 		FMetalTexture Tex = Texture->Texture;
-		bool const bUseIntermediateMetalBuffer = Tex.GetStorageMode() == mtlpp::StorageMode::Private;
 		
-		if(bUseIntermediateMetalBuffer)
+		if(Tex.GetStorageMode() == mtlpp::StorageMode::Private)
 		{
-			FMetalBuffer Buffer = InternalCopyTexture2DUpdateRegion(TextureRHI, UpdateRegion, SourcePitch, SourceData);
+			FMetalBuffer Buffer = Internal_CreateBufferAndCopyTexture2DUpdateRegionData(TextureRHI, UpdateRegion, SourcePitch, SourceData);
 			InternalUpdateTexture2D(ImmediateContext.GetInternalContext(), TextureRHI, MipIndex, UpdateRegion, SourcePitch, Buffer);
 			GetMetalDeviceContext().ReleaseBuffer(Buffer);
 		}
@@ -1833,26 +1834,37 @@ void FMetalDynamicRHI::RHIUpdateTexture2D(FRHITexture2D* TextureRHI, uint32 MipI
 	}
 }
 
-static void InternalCopyTexture3DUpdateRegionData(FRHITexture3D* TextureRHI, const struct FUpdateTextureRegion3D& UpdateRegion, uint32 SourceRowPitch, uint32 SourceDepthPitch, const uint8* SourceData, uint8* DestData)
+static FMetalBuffer Internal_CreateBufferAndCopyTexture3DUpdateRegionData(FRHITexture3D* TextureRHI, const struct FUpdateTextureRegion3D& UpdateRegion, uint32 SourceRowPitch, uint32 SourceDepthPitch, const uint8* SourceData)
 {
-	// Perform safe line copy
-	FMetalSurface* Texture = ResourceCast(TextureRHI);	
+	FMetalSurface* Texture = ResourceCast(TextureRHI);
 	
-	uint8* pDest = DestData;
-	uint32 CopyPitch = FMath::DivideAndRoundUp(UpdateRegion.Width, (uint32)GPixelFormats[Texture->GetFormat()].BlockSizeX) * GPixelFormats[Texture->GetFormat()].BlockBytes;
+	const uint32 BufferSize = SourceDepthPitch * UpdateRegion.Depth;
+	FMetalBuffer OutBuffer = GetMetalDeviceContext().CreatePooledBuffer(FMetalPooledBufferArgs(GetMetalDeviceContext().GetDevice(), BufferSize, BUF_Static, mtlpp::StorageMode::Shared));
 
-	for (uint32 i = 0; i < UpdateRegion.Depth;++i)
+	const FPixelFormatInfo& FormatInfo = GPixelFormats[TextureRHI->GetFormat()];
+	uint32 CopyPitch = FMath::DivideAndRoundUp(UpdateRegion.Width, (uint32)FormatInfo.BlockSizeX) * FormatInfo.BlockBytes;
+	
+	check(FormatInfo.BlockSizeZ == 1);
+	check(CopyPitch <= SourceRowPitch);
+		
+	uint8_t* DestData = (uint8_t*)OutBuffer.GetContents();
+	const uint32 NumRows = UpdateRegion.Height / (uint32)FormatInfo.BlockSizeY;
+		
+	// Perform safe line copy
+	for (uint32 i = 0;i < UpdateRegion.Depth;++i)
 	{
 		const uint8* pSourceRowData = SourceData + (SourceDepthPitch * i);
-		uint8* pDestRowData = pDest + (SourceDepthPitch * i);
+		uint8* pDestRowData = DestData + (SourceDepthPitch * i);
 
-		for (uint32 j = 0;j < UpdateRegion.Height;++j)
+		for (uint32 j = 0;j < NumRows;++j)
 		{
 			FMemory::Memcpy(pDestRowData, pSourceRowData, CopyPitch);
 			pSourceRowData += SourceRowPitch;
 			pDestRowData += SourceRowPitch;
 		}
 	}
+	
+	return OutBuffer;
 }
 
 
@@ -1865,7 +1877,10 @@ static void InternalUpdateTexture3D(FMetalContext& Context, FRHITexture3D* Textu
 	
 	if(Tex.GetStorageMode() == mtlpp::StorageMode::Private)
 	{
-		const uint32 BytesPerImage = SourceRowPitch * UpdateRegion.Height;
+		const FPixelFormatInfo& FormatInfo = GPixelFormats[TextureRHI->GetFormat()];
+		const uint32 NumRows = UpdateRegion.Height / (uint32)FormatInfo.BlockSizeY;
+		const uint32 BytesPerImage = SourceRowPitch * NumRows;
+		
 		mtlpp::BlitOption Options = mtlpp::BlitOption::None;
 #if !PLATFORM_MAC
 		if (Tex.GetPixelFormat() >= mtlpp::PixelFormat::PVRTC_RGB_2BPP && Tex.GetPixelFormat() <= mtlpp::PixelFormat::PVRTC_RGBA_4BPP_sRGB)
@@ -1904,12 +1919,7 @@ struct FMetalDynamicRHIUpdateTexture3DCommand final : public FRHICommand<FMetalD
 	, SourceRowPitch(InSourceRowPitch)
 	, SourceDepthPitch(InSourceDepthPitch)
 	{
-		FMetalSurface* Texture = ResourceCast(TextureRHI);
-		FMetalTexture Tex = Texture->Texture;
-		const uint32 BufferSize = UpdateRegion.Height * UpdateRegion.Depth* SourceRowPitch;
-		
-		Buffer = GetMetalDeviceContext().CreatePooledBuffer(FMetalPooledBufferArgs(GetMetalDeviceContext().GetDevice(), BufferSize, BUF_Dynamic, mtlpp::StorageMode::Shared));
-		InternalCopyTexture3DUpdateRegionData(Texture, UpdateRegion, SourceRowPitch, SourceDepthPitch, SourceData, (uint8*)Buffer.GetContents());
+		Buffer = Internal_CreateBufferAndCopyTexture3DUpdateRegionData(TextureRHI, UpdateRegion, SourceRowPitch, SourceDepthPitch, SourceData);
 	}
 	
 	void Execute(FRHICommandListBase& CmdList)
@@ -1978,9 +1988,7 @@ void FMetalDynamicRHI::RHIUpdateTexture3D(FRHITexture3D* TextureRHI,uint32 MipIn
 		{
 			SCOPED_AUTORELEASE_POOL;
 
-			const uint32 BufferSize = UpdateRegion.Height * UpdateRegion.Depth * SourceRowPitch;
-			FMetalBuffer IntermediateBuffer = GetMetalDeviceContext().CreatePooledBuffer(FMetalPooledBufferArgs(GetMetalDeviceContext().GetDevice(), BufferSize, BUF_Dynamic, mtlpp::StorageMode::Shared));
-			InternalCopyTexture3DUpdateRegionData(TextureRHI, UpdateRegion, SourceRowPitch, SourceDepthPitch, SourceData, (uint8*)IntermediateBuffer.GetContents());
+			FMetalBuffer IntermediateBuffer = Internal_CreateBufferAndCopyTexture3DUpdateRegionData(TextureRHI, UpdateRegion, SourceRowPitch, SourceDepthPitch, SourceData);
 			InternalUpdateTexture3D(ImmediateContext.GetInternalContext(), TextureRHI, MipIndex, UpdateRegion, SourceRowPitch, SourceDepthPitch, IntermediateBuffer);
 			GetMetalDeviceContext().ReleaseBuffer(IntermediateBuffer);
 		}
