Add array index to the lock tracker lookup value.  Fix clashes during multiple array locks on a texture array.  Remove BufferLocks variable from Metal/AGXSurface class as it doesn't appear to be doing anything.

#rnx

diff --git a/Engine/Source/Runtime/Apple/AGXRHI/Private/AGXTexture.cpp b/Engine/Source/Runtime/Apple/AGXRHI/Private/AGXTexture.cpp
--- a/Engine/Source/Runtime/Apple/AGXRHI/Private/AGXTexture.cpp
+++ b/Engine/Source/Runtime/Apple/AGXRHI/Private/AGXTexture.cpp
@@ -252,6 +252,13 @@ MTLPixelFormat AGXToSRGBFormat(MTLPixelFormat MTLFormat)
 	return MTLFormat;
 }
 
+static inline uint32 ComputeLockIndex(uint32 MipIndex, uint32 ArrayIndex)
+{
+	check(MipIndex < MAX_uint16);
+	check(ArrayIndex < MAX_uint16);
+	return (MipIndex & MAX_uint16) | ((ArrayIndex & MAX_uint16) << 16);
+}
+
 void FAGXSurface::PrepareTextureView()
 {
 	// Recreate the texture to enable MTLTextureUsagePixelFormatView which must be off unless we definitely use this feature or we are throwing ~4% performance vs. Windows on the floor.
@@ -527,7 +534,6 @@ FAGXSurface::FAGXSurface(FAGXTextureCreateDesc const& CreateDesc)
 	, Texture           (nil)
 	, MSAATexture       (nil)
 	, MSAAResolveTexture(nil)
-	, BufferLocks       (0)
 	, TotalTextureSize  (0)
 	, Viewport          (nullptr)
 	, ImageSurfaceRef   (nullptr)
@@ -868,7 +874,8 @@ id <MTLBuffer> FAGXSurface::AllocSurface(uint32 MipIndex, uint32 ArrayIndex, ERe
 	Buffer.label = @"Temporary Surface Backing";
 	
 	// Note: while the lock is active, this map owns the backing store.
-	GRHILockTracker.Lock(this, Buffer, MipIndex, 0, LockMode, false);
+	const uint32 LockIndex = ComputeLockIndex(MipIndex, ArrayIndex);
+	GRHILockTracker.Lock(this, Buffer, LockIndex, MipBytes, LockMode, false);
 	
 #if PLATFORM_MAC
 	// Expand R8_sRGB into RGBA8_sRGB for Mac.
@@ -1164,7 +1171,6 @@ void* FAGXSurface::Lock(uint32 MipIndex, uint32 ArrayIndex, EResourceLockMode Lo
 		}
 		case RLM_WriteOnly:
 		{
-			BufferLocks |= 1 << MipIndex;
 			break;
 		}
 		default:
@@ -1178,11 +1184,11 @@ void* FAGXSurface::Lock(uint32 MipIndex, uint32 ArrayIndex, EResourceLockMode Lo
 void FAGXSurface::Unlock(uint32 MipIndex, uint32 ArrayIndex, bool bTryAsync)
 {
 	check(IsInRenderingThread());
-	BufferLocks &= ~(1 << MipIndex);
 	
-	FRHILockTracker::FLockParams Params = GRHILockTracker.Unlock(this, MipIndex);
-	id <MTLBuffer> SourceData = (id <MTLBuffer>) Params.Buffer;
+	const uint32 LockIndex = ComputeLockIndex(MipIndex, ArrayIndex);
+	FRHILockTracker::FLockParams Params = GRHILockTracker.Unlock(this, LockIndex);
 	
+	id <MTLBuffer> SourceData = (id <MTLBuffer>) Params.Buffer;
 	if(bTryAsync)
 	{
 		AsyncUnlock(SourceData, MipIndex, ArrayIndex);
@@ -1212,7 +1218,7 @@ void* FAGXSurface::AsyncLock(class FRHICommandListImmediate& RHICmdList, uint32
 	}
 	else
 	{
-		id <MTLBuffer> Buffer = AllocSurface(MipIndex, 0, LockMode, DestStride);
+		id <MTLBuffer> Buffer = AllocSurface(MipIndex, ArrayIndex, LockMode, DestStride);
 		check(Buffer);
 		
 		BufferData = Buffer.contents;
diff --git a/Engine/Source/Runtime/Apple/AGXRHI/Public/AGXResources.h b/Engine/Source/Runtime/Apple/AGXRHI/Public/AGXResources.h
--- a/Engine/Source/Runtime/Apple/AGXRHI/Public/AGXResources.h
+++ b/Engine/Source/Runtime/Apple/AGXRHI/Public/AGXResources.h
@@ -251,8 +251,6 @@ public:
 		ReadbackRequestedAndComplete 	= ReadbackRequested | ReadbackFenceComplete
 	};
 
-	uint32 BufferLocks;
-
 	// how much memory is allocated for this texture
 	uint64 TotalTextureSize;
 	
diff --git a/Engine/Source/Runtime/Apple/MetalRHI/Private/MetalTexture.cpp b/Engine/Source/Runtime/Apple/MetalRHI/Private/MetalTexture.cpp
--- a/Engine/Source/Runtime/Apple/MetalRHI/Private/MetalTexture.cpp
+++ b/Engine/Source/Runtime/Apple/MetalRHI/Private/MetalTexture.cpp
@@ -261,6 +261,13 @@ mtlpp::PixelFormat ToSRGBFormat(mtlpp::PixelFormat MTLFormat)
 	return MTLFormat;
 }
 
+static inline uint32 ComputeLockIndex(uint32 MipIndex, uint32 ArrayIndex)
+{
+	check(MipIndex < MAX_uint16);
+	check(ArrayIndex < MAX_uint16);
+	return (MipIndex & MAX_uint16) | ((ArrayIndex & MAX_uint16) << 16);
+}
+
 void FMetalSurface::PrepareTextureView()
 {
 	// Recreate the texture to enable MTLTextureUsagePixelFormatView which must be off unless we definitely use this feature or we are throwing ~4% performance vs. Windows on the floor.
@@ -578,7 +585,6 @@ FMetalSurface::FMetalSurface(FMetalTextureCreateDesc const& CreateDesc)
 	, Texture           (nil)
 	, MSAATexture       (nil)
 	, MSAAResolveTexture(nil)
-	, BufferLocks       (0)
 	, TotalTextureSize  (0)
 	, Viewport          (nullptr)
 	, ImageSurfaceRef   (nullptr)
@@ -923,7 +929,8 @@ id <MTLBuffer> FMetalSurface::AllocSurface(uint32 MipIndex, uint32 ArrayIndex, E
 	Buffer.label = @"Temporary Surface Backing";
 	
 	// Note: while the lock is active, this map owns the backing store.
-	GRHILockTracker.Lock(this, Buffer, MipIndex, 0, LockMode, false);
+	const uint32 LockIndex = ComputeLockIndex(MipIndex, ArrayIndex);
+	GRHILockTracker.Lock(this, Buffer, LockIndex, MipBytes, LockMode, false);
 	
 #if PLATFORM_MAC
 	// Expand R8_sRGB into RGBA8_sRGB for non Apple Silicon Mac.
@@ -1219,7 +1226,6 @@ void* FMetalSurface::Lock(uint32 MipIndex, uint32 ArrayIndex, EResourceLockMode
 		}
 		case RLM_WriteOnly:
 		{
-			BufferLocks |= 1 << MipIndex;
 			break;
 		}
 		default:
@@ -1233,11 +1239,11 @@ void* FMetalSurface::Lock(uint32 MipIndex, uint32 ArrayIndex, EResourceLockMode
 void FMetalSurface::Unlock(uint32 MipIndex, uint32 ArrayIndex, bool bTryAsync)
 {
 	check(IsInRenderingThread());
-	BufferLocks &= ~(1 << MipIndex);
 	
-	FRHILockTracker::FLockParams Params = GRHILockTracker.Unlock(this, MipIndex);
-	id <MTLBuffer> SourceData = (id <MTLBuffer>) Params.Buffer;
+	const uint32 LockIndex = ComputeLockIndex(MipIndex, ArrayIndex);
+	FRHILockTracker::FLockParams Params = GRHILockTracker.Unlock(this, LockIndex);
 	
+	id <MTLBuffer> SourceData = (id <MTLBuffer>) Params.Buffer;
 	if(bTryAsync)
 	{
 		AsyncUnlock(SourceData, MipIndex, ArrayIndex);
@@ -1267,7 +1273,7 @@ void* FMetalSurface::AsyncLock(class FRHICommandListImmediate& RHICmdList, uint3
 	}
 	else
 	{
-		id <MTLBuffer> Buffer = AllocSurface(MipIndex, 0, LockMode, DestStride);
+		id <MTLBuffer> Buffer = AllocSurface(MipIndex, ArrayIndex, LockMode, DestStride);
 		check(Buffer);
 		
 		BufferData = Buffer.contents;
diff --git a/Engine/Source/Runtime/Apple/MetalRHI/Public/MetalResources.h b/Engine/Source/Runtime/Apple/MetalRHI/Public/MetalResources.h
--- a/Engine/Source/Runtime/Apple/MetalRHI/Public/MetalResources.h
+++ b/Engine/Source/Runtime/Apple/MetalRHI/Public/MetalResources.h
@@ -251,8 +251,6 @@ public:
 		ReadbackRequestedAndComplete 	= ReadbackRequested | ReadbackFenceComplete
 	};
 
-	uint32 BufferLocks;
-
 	// how much memory is allocated for this texture
 	uint64 TotalTextureSize;
 	
