Always use the shadow data for uniform buffer update

#rnx

diff --git a/Engine/Source/Runtime/Apple/AGXRHI/Private/AGXUniformBuffer.cpp b/Engine/Source/Runtime/Apple/AGXRHI/Private/AGXUniformBuffer.cpp
--- a/Engine/Source/Runtime/Apple/AGXRHI/Private/AGXUniformBuffer.cpp
+++ b/Engine/Source/Runtime/Apple/AGXRHI/Private/AGXUniformBuffer.cpp
@@ -11,49 +11,33 @@
 
 #pragma mark Suballocated Uniform Buffer Implementation
 
-FAGXSuballocatedUniformBuffer::FAGXSuballocatedUniformBuffer(const FRHIUniformBufferLayout* Layout, EUniformBufferUsage Usage, EUniformBufferValidation InValidation)
+FAGXSuballocatedUniformBuffer::FAGXSuballocatedUniformBuffer(const void* Contents, const FRHIUniformBufferLayout* Layout, EUniformBufferUsage Usage, EUniformBufferValidation InValidation)
     : FRHIUniformBuffer(Layout)
     , LastFrameUpdated(0)
     , Offset(0)
-    , Backing(nil)
-    , Shadow(nullptr)
+	, Backing(nil)
+	, Shadow(FMemory::Malloc(GetSize()))
     , ResourceTable()
 #if METAL_UNIFORM_BUFFER_VALIDATION
     , Validation(InValidation)
 #endif // METAL_UNIFORM_BUFFER_VALIDATION
 {
-    // Slate can create SingleDraw uniform buffers and use them several frames later. So it must be included.
-    if (Usage == UniformBuffer_SingleDraw || Usage == UniformBuffer_MultiFrame)
-    {
-        Shadow = FMemory::Malloc(GetSize());
-    }
+	FMemory::Memcpy(Shadow, Contents, GetSize());
+	CopyResourceTable(Contents, ResourceTable);
 }
 
 FAGXSuballocatedUniformBuffer::~FAGXSuballocatedUniformBuffer()
 {
-    if (HasShadow())
-    {
-        FMemory::Free(Shadow);
-    }
+    FMemory::Free(Shadow);
 
     // Note: this object does NOT own a reference
     // to the uniform buffer backing store
 }
 
-bool FAGXSuballocatedUniformBuffer::HasShadow()
+void FAGXSuballocatedUniformBuffer::Update(const void* Contents)
 {
-    return Shadow != nullptr;
-}
-
-void FAGXSuballocatedUniformBuffer::Update(const void* Contents, TArray<TRefCountPtr<FRHIResource> > const& InResourceTable)
-{
-    if (HasShadow())
-    {
-        FMemory::Memcpy(Shadow, Contents, GetSize());
-    }
-
-	ResourceTable = InResourceTable;
-
+	FMemory::Memcpy(Shadow, Contents, GetSize());
+	CopyResourceTable(Contents, ResourceTable);
 	PushToGPUBacking(Contents);
 }
 
diff --git a/Engine/Source/Runtime/Apple/AGXRHI/Private/AGXUniformBufferCommands.cpp b/Engine/Source/Runtime/Apple/AGXRHI/Private/AGXUniformBufferCommands.cpp
--- a/Engine/Source/Runtime/Apple/AGXRHI/Private/AGXUniformBufferCommands.cpp
+++ b/Engine/Source/Runtime/Apple/AGXRHI/Private/AGXUniformBufferCommands.cpp
@@ -12,37 +12,28 @@ static void DoUpdateUniformBuffer(FAGXUniformBuffer* UB, const void* Contents)
     // The only way we can be on the RHI thread here is if we're in the process of creating a FLocalUniformBuffer.
     bool bUpdateImmediately = RHICmdList.Bypass() || IsInRHIThread();
 
-	TArray<TRefCountPtr<FRHIResource> > ResourceTable;
-
-	UB->CopyResourceTable(Contents, ResourceTable);
-
     if(bUpdateImmediately)
     {
-        UB->Update(Contents, ResourceTable);
+        UB->Update(Contents);
     }
     else
     {
         const uint32 NumBytes = UB->GetLayout().ConstantBufferSize;
         void* Data = RHICmdList.Alloc(NumBytes, 16);
         FMemory::Memcpy(Data, Contents, NumBytes);
-        
-        RHICmdList.EnqueueLambda([Data, UB, NewResourceTable = MoveTemp(ResourceTable)](FRHICommandListImmediate& RHICmdList)
+
+        RHICmdList.EnqueueLambda([Data, UB](FRHICommandListImmediate& RHICmdList)
         {
-            UB->Update(Data, NewResourceTable);
+            UB->Update(Data);
         });
-        
+
         RHICmdList.RHIThreadFence(true);
     }
 }
 
 FUniformBufferRHIRef FAGXDynamicRHI::RHICreateUniformBuffer(const void* Contents, const FRHIUniformBufferLayout* Layout, EUniformBufferUsage Usage, EUniformBufferValidation Validation)
 {
-	FAGXUniformBuffer* UB = new FAGXUniformBuffer(Layout, Usage, Validation);
-
-	TArray<TRefCountPtr<FRHIResource>> ResourceTable;
-	UB->CopyResourceTable(Contents, ResourceTable);
-	UB->Update(Contents, ResourceTable);
-	return UB;
+	return new FAGXUniformBuffer(Contents, Layout, Usage, Validation);
 }
 
 void FAGXDynamicRHI::RHIUpdateUniformBuffer(FRHIUniformBuffer* UniformBufferRHI, const void* Contents)
diff --git a/Engine/Source/Runtime/Apple/AGXRHI/Public/AGXUniformBuffer.h b/Engine/Source/Runtime/Apple/AGXRHI/Public/AGXUniformBuffer.h
--- a/Engine/Source/Runtime/Apple/AGXRHI/Public/AGXUniformBuffer.h
+++ b/Engine/Source/Runtime/Apple/AGXRHI/Public/AGXUniformBuffer.h
@@ -38,22 +38,21 @@ private:
 public:
     // Creates a uniform buffer.
     // If Usage is SingleDraw or MultiFrame we will keep a copy of the data
-    FAGXSuballocatedUniformBuffer(const FRHIUniformBufferLayout* Layout, EUniformBufferUsage Usage, EUniformBufferValidation Validation);
+    FAGXSuballocatedUniformBuffer(const void* Contents, const FRHIUniformBufferLayout* Layout, EUniformBufferUsage Usage, EUniformBufferValidation Validation);
     ~FAGXSuballocatedUniformBuffer();
-    
-	void Update(const void* Contents, TArray<TRefCountPtr<FRHIResource> > const& InResourceTable);
-    
+
+	void Update(const void* Contents);
+
     // Prepares this uniform buffer for binding.
     // You MUST call this before binding it.
     // If this is a UB that was created before the current frame the data on the GPU will
     // not be correct until this function returns.
     void PrepareToBind();
 
+private:
 	// Copies the RDG resources to a resource table for a deferred update on the RHI thread.
 	void CopyResourceTable(const void* Contents, TArray<TRefCountPtr<FRHIResource> >& OutResourceTable) const;
 
-private:
-    
     // Pushes the data in Contents to the gpu.
     // Updates the frame counter to FrameNumber.
     // (this is to support the case where we create buffer and reuse it many frames later).
@@ -61,8 +60,6 @@ private:
     // and copies Contents into that backing.
     // The amount of data is determined by the Layout
     void PushToGPUBacking(const void* Contents);
-    
-    bool HasShadow();
 };
 
 typedef FAGXSuballocatedUniformBuffer FAGXUniformBuffer;
diff --git a/Engine/Source/Runtime/Apple/MetalRHI/Private/MetalUniformBuffer.cpp b/Engine/Source/Runtime/Apple/MetalRHI/Private/MetalUniformBuffer.cpp
--- a/Engine/Source/Runtime/Apple/MetalRHI/Private/MetalUniformBuffer.cpp
+++ b/Engine/Source/Runtime/Apple/MetalRHI/Private/MetalUniformBuffer.cpp
@@ -11,49 +11,33 @@
 
 #pragma mark Suballocated Uniform Buffer Implementation
 
-FMetalSuballocatedUniformBuffer::FMetalSuballocatedUniformBuffer(const FRHIUniformBufferLayout* Layout, EUniformBufferUsage Usage, EUniformBufferValidation InValidation)
+FMetalSuballocatedUniformBuffer::FMetalSuballocatedUniformBuffer(const void *Contents, const FRHIUniformBufferLayout* Layout, EUniformBufferUsage Usage, EUniformBufferValidation InValidation)
     : FRHIUniformBuffer(Layout)
     , LastFrameUpdated(0)
     , Offset(0)
     , Backing(nil)
-    , Shadow(nullptr)
+    , Shadow(FMemory::Malloc(GetSize()))
     , ResourceTable()
 #if METAL_UNIFORM_BUFFER_VALIDATION
     , Validation(InValidation)
 #endif // METAL_UNIFORM_BUFFER_VALIDATION
 {
-    // Slate can create SingleDraw uniform buffers and use them several frames later. So it must be included.
-    if (Usage == UniformBuffer_SingleDraw || Usage == UniformBuffer_MultiFrame)
-    {
-        Shadow = FMemory::Malloc(GetSize());
-    }
+	FMemory::Memcpy(Shadow, Contents, GetSize());
+	CopyResourceTable(Contents, ResourceTable);
 }
 
 FMetalSuballocatedUniformBuffer::~FMetalSuballocatedUniformBuffer()
 {
-    if (HasShadow())
-    {
-        FMemory::Free(Shadow);
-    }
+	FMemory::Free(Shadow);
 
     // Note: this object does NOT own a reference
     // to the uniform buffer backing store
 }
 
-bool FMetalSuballocatedUniformBuffer::HasShadow()
+void FMetalSuballocatedUniformBuffer::Update(const void* Contents)
 {
-    return Shadow != nullptr;
-}
-
-void FMetalSuballocatedUniformBuffer::Update(const void* Contents, TArray<TRefCountPtr<FRHIResource> > const& InResourceTable)
-{
-    if (HasShadow())
-    {
-        FMemory::Memcpy(Shadow, Contents, GetSize());
-    }
-
-	ResourceTable = InResourceTable;
-
+	FMemory::Memcpy(Shadow, Contents, GetSize());
+	CopyResourceTable(Contents, ResourceTable);
 	PushToGPUBacking(Contents);
 }
 
diff --git a/Engine/Source/Runtime/Apple/MetalRHI/Private/MetalUniformBufferCommands.cpp b/Engine/Source/Runtime/Apple/MetalRHI/Private/MetalUniformBufferCommands.cpp
--- a/Engine/Source/Runtime/Apple/MetalRHI/Private/MetalUniformBufferCommands.cpp
+++ b/Engine/Source/Runtime/Apple/MetalRHI/Private/MetalUniformBufferCommands.cpp
@@ -7,42 +7,33 @@
 static void DoUpdateUniformBuffer(FMetalUniformBuffer* UB, const void* Contents)
 {
     check(IsInRenderingThread() || IsInRHIThread());
-    
+
     FRHICommandListImmediate& RHICmdList = FRHICommandListExecutor::GetImmediateCommandList();
     // The only way we can be on the RHI thread here is if we're in the process of creating a FLocalUniformBuffer.
     bool bUpdateImmediately = RHICmdList.Bypass() || IsInRHIThread();
 
-	TArray<TRefCountPtr<FRHIResource> > ResourceTable;
-
-	UB->CopyResourceTable(Contents, ResourceTable);
-
     if(bUpdateImmediately)
     {
-        UB->Update(Contents, ResourceTable);
+        UB->Update(Contents);
     }
     else
     {
         const uint32 NumBytes = UB->GetLayout().ConstantBufferSize;
         void* Data = RHICmdList.Alloc(NumBytes, 16);
         FMemory::Memcpy(Data, Contents, NumBytes);
-        
-        RHICmdList.EnqueueLambda([Data, UB, NewResourceTable = MoveTemp(ResourceTable)](FRHICommandListImmediate& RHICmdList)
+
+        RHICmdList.EnqueueLambda([Data, UB](FRHICommandListImmediate& RHICmdList)
         {
-            UB->Update(Data, NewResourceTable);
+            UB->Update(Data);
         });
-        
+
         RHICmdList.RHIThreadFence(true);
     }
 }
 
 FUniformBufferRHIRef FMetalDynamicRHI::RHICreateUniformBuffer(const void* Contents, const FRHIUniformBufferLayout* Layout, EUniformBufferUsage Usage, EUniformBufferValidation Validation)
 {
-	FMetalUniformBuffer* UB = new FMetalUniformBuffer(Layout, Usage, Validation);
-
-	TArray<TRefCountPtr<FRHIResource>> ResourceTable;
-	UB->CopyResourceTable(Contents, ResourceTable);
-	UB->Update(Contents, ResourceTable);
-	return UB;
+	return new FMetalUniformBuffer(Contents, Layout, Usage, Validation);
 }
 
 void FMetalDynamicRHI::RHIUpdateUniformBuffer(FRHIUniformBuffer* UniformBufferRHI, const void* Contents)
diff --git a/Engine/Source/Runtime/Apple/MetalRHI/Public/MetalUniformBuffer.h b/Engine/Source/Runtime/Apple/MetalRHI/Public/MetalUniformBuffer.h
--- a/Engine/Source/Runtime/Apple/MetalRHI/Public/MetalUniformBuffer.h
+++ b/Engine/Source/Runtime/Apple/MetalRHI/Public/MetalUniformBuffer.h
@@ -34,26 +34,25 @@ private:
 #if METAL_UNIFORM_BUFFER_VALIDATION
     EUniformBufferValidation Validation;
 #endif
-    
+
 public:
     // Creates a uniform buffer.
     // If Usage is SingleDraw or MultiFrame we will keep a copy of the data
-    FMetalSuballocatedUniformBuffer(const FRHIUniformBufferLayout* Layout, EUniformBufferUsage Usage, EUniformBufferValidation Validation);
+    FMetalSuballocatedUniformBuffer(const void* Contents, const FRHIUniformBufferLayout* Layout, EUniformBufferUsage Usage, EUniformBufferValidation Validation);
     ~FMetalSuballocatedUniformBuffer();
-    
-	void Update(const void* Contents, TArray<TRefCountPtr<FRHIResource> > const& InResourceTable);
-    
+
+	void Update(const void* Contents);
+
     // Prepares this uniform buffer for binding.
     // You MUST call this before binding it.
     // If this is a UB that was created before the current frame the data on the GPU will
     // not be correct until this function returns.
     void PrepareToBind();
 
+private:
 	// Copies the RDG resources to a resource table for a deferred update on the RHI thread.
 	void CopyResourceTable(const void* Contents, TArray<TRefCountPtr<FRHIResource> >& OutResourceTable) const;
 
-private:
-    
     // Pushes the data in Contents to the gpu.
     // Updates the frame counter to FrameNumber.
     // (this is to support the case where we create buffer and reuse it many frames later).
@@ -61,8 +60,6 @@ private:
     // and copies Contents into that backing.
     // The amount of data is determined by the Layout
     void PushToGPUBacking(const void* Contents);
-    
-    bool HasShadow();
 };
 
 typedef FMetalSuballocatedUniformBuffer FMetalUniformBuffer;
