Back out CL 20487525+20493291

- CL 20487525 (Make RHICreateUniformBuffer thread safe) triggered a number of Metal validation errors. Backing out the original change and subsequent hotfix.

#rnx

diff --git a/Engine/Source/Runtime/Apple/MetalRHI/Private/MetalUniformBuffer.cpp b/Engine/Source/Runtime/Apple/MetalRHI/Private/MetalUniformBuffer.cpp
--- a/Engine/Source/Runtime/Apple/MetalRHI/Private/MetalUniformBuffer.cpp
+++ b/Engine/Source/Runtime/Apple/MetalRHI/Private/MetalUniformBuffer.cpp
@@ -11,33 +11,49 @@
 
 #pragma mark Suballocated Uniform Buffer Implementation
 
-FMetalSuballocatedUniformBuffer::FMetalSuballocatedUniformBuffer(const void *Contents, const FRHIUniformBufferLayout* Layout, EUniformBufferUsage Usage, EUniformBufferValidation InValidation)
+FMetalSuballocatedUniformBuffer::FMetalSuballocatedUniformBuffer(const FRHIUniformBufferLayout* Layout, EUniformBufferUsage Usage, EUniformBufferValidation InValidation)
     : FRHIUniformBuffer(Layout)
     , LastFrameUpdated(0)
     , Offset(0)
     , Backing(nil)
-    , Shadow(FMemory::Malloc(GetSize()))
+    , Shadow(nullptr)
     , ResourceTable()
 #if METAL_UNIFORM_BUFFER_VALIDATION
     , Validation(InValidation)
 #endif // METAL_UNIFORM_BUFFER_VALIDATION
 {
-	FMemory::Memcpy(Shadow, Contents, GetSize());
-	CopyResourceTable(Contents, ResourceTable);
+    // Slate can create SingleDraw uniform buffers and use them several frames later. So it must be included.
+    if (Usage == UniformBuffer_SingleDraw || Usage == UniformBuffer_MultiFrame)
+    {
+        Shadow = FMemory::Malloc(GetSize());
+    }
 }
 
 FMetalSuballocatedUniformBuffer::~FMetalSuballocatedUniformBuffer()
 {
-	FMemory::Free(Shadow);
+    if (HasShadow())
+    {
+        FMemory::Free(Shadow);
+    }
 
     // Note: this object does NOT own a reference
     // to the uniform buffer backing store
 }
 
-void FMetalSuballocatedUniformBuffer::Update(const void* Contents)
+bool FMetalSuballocatedUniformBuffer::HasShadow()
+{
+    return Shadow != nullptr;
+}
+
+void FMetalSuballocatedUniformBuffer::Update(const void* Contents, TArray<TRefCountPtr<FRHIResource> > const& InResourceTable)
 {
-	FMemory::Memcpy(Shadow, Contents, GetSize());
-	CopyResourceTable(Contents, ResourceTable);
+    if (HasShadow())
+    {
+        FMemory::Memcpy(Shadow, Contents, GetSize());
+    }
+
+	ResourceTable = InResourceTable;
+
 	PushToGPUBacking(Contents);
 }
 
@@ -46,6 +62,8 @@ void FMetalSuballocatedUniformBuffer::Update(const void* Contents)
 // The amount of data read from Contents is given by the Layout
 void FMetalSuballocatedUniformBuffer::PushToGPUBacking(const void* Contents)
 {
+    check(IsInRenderingThread() ^ IsRunningRHIInSeparateThread());
+    
     FMetalDeviceContext& DeviceContext = GetMetalDeviceContext();
     
     FMetalFrameAllocator* Allocator = DeviceContext.GetUniformAllocator();
@@ -70,7 +88,7 @@ void FMetalSuballocatedUniformBuffer::PrepareToBind()
     }
 }
 
-void FMetalSuballocatedUniformBuffer::CopyResourceTable(const void* Contents, TArray<TRefCountPtr<FRHIResource> >& OutResourceTable) const
+void FMetalSuballocatedUniformBuffer::CopyResourceTable_RenderThread(const void* Contents, TArray<TRefCountPtr<FRHIResource> >& OutResourceTable)
 {
 #if METAL_UNIFORM_BUFFER_VALIDATION
 	if (Validation == EUniformBufferValidation::ValidateResources)
diff --git a/Engine/Source/Runtime/Apple/MetalRHI/Private/MetalUniformBufferCommands.cpp b/Engine/Source/Runtime/Apple/MetalRHI/Private/MetalUniformBufferCommands.cpp
--- a/Engine/Source/Runtime/Apple/MetalRHI/Private/MetalUniformBufferCommands.cpp
+++ b/Engine/Source/Runtime/Apple/MetalRHI/Private/MetalUniformBufferCommands.cpp
@@ -7,33 +7,44 @@
 static void DoUpdateUniformBuffer(FMetalUniformBuffer* UB, const void* Contents)
 {
     check(IsInRenderingThread() || IsInRHIThread());
-
+    
     FRHICommandListImmediate& RHICmdList = FRHICommandListExecutor::GetImmediateCommandList();
     // The only way we can be on the RHI thread here is if we're in the process of creating a FLocalUniformBuffer.
     bool bUpdateImmediately = RHICmdList.Bypass() || IsInRHIThread();
 
+	TArray<TRefCountPtr<FRHIResource> > ResourceTable;
+
+	UB->CopyResourceTable_RenderThread(Contents, ResourceTable);
+
     if(bUpdateImmediately)
     {
-        UB->Update(Contents);
+        UB->Update(Contents, ResourceTable);
     }
     else
     {
         const uint32 NumBytes = UB->GetLayout().ConstantBufferSize;
         void* Data = RHICmdList.Alloc(NumBytes, 16);
         FMemory::Memcpy(Data, Contents, NumBytes);
-
-        RHICmdList.EnqueueLambda([Data, UB](FRHICommandListImmediate& RHICmdList)
+        
+        RHICmdList.EnqueueLambda([Data, UB, NewResourceTable = MoveTemp(ResourceTable)](FRHICommandListImmediate& RHICmdList)
         {
-            UB->Update(Data);
+            UB->Update(Data, NewResourceTable);
         });
-
+        
         RHICmdList.RHIThreadFence(true);
     }
 }
 
 FUniformBufferRHIRef FMetalDynamicRHI::RHICreateUniformBuffer(const void* Contents, const FRHIUniformBufferLayout* Layout, EUniformBufferUsage Usage, EUniformBufferValidation Validation)
 {
-	return new FMetalUniformBuffer(Contents, Layout, Usage, Validation);
+    FMetalDeviceContext& DeviceContext = (FMetalDeviceContext&)GetMetalDeviceContext();
+    FMetalFrameAllocator* UniformAllocator = DeviceContext.GetUniformAllocator();
+    
+    FMetalUniformBuffer* UB = new FMetalUniformBuffer(Layout, Usage, Validation);
+    
+    DoUpdateUniformBuffer(UB, Contents);
+    
+    return UB;
 }
 
 void FMetalDynamicRHI::RHIUpdateUniformBuffer(FRHIUniformBuffer* UniformBufferRHI, const void* Contents)
diff --git a/Engine/Source/Runtime/Apple/MetalRHI/Public/MetalUniformBuffer.h b/Engine/Source/Runtime/Apple/MetalRHI/Public/MetalUniformBuffer.h
--- a/Engine/Source/Runtime/Apple/MetalRHI/Public/MetalUniformBuffer.h
+++ b/Engine/Source/Runtime/Apple/MetalRHI/Public/MetalUniformBuffer.h
@@ -34,25 +34,26 @@ private:
 #if METAL_UNIFORM_BUFFER_VALIDATION
     EUniformBufferValidation Validation;
 #endif
-
+    
 public:
     // Creates a uniform buffer.
     // If Usage is SingleDraw or MultiFrame we will keep a copy of the data
-    FMetalSuballocatedUniformBuffer(const void* Contents, const FRHIUniformBufferLayout* Layout, EUniformBufferUsage Usage, EUniformBufferValidation Validation);
+    FMetalSuballocatedUniformBuffer(const FRHIUniformBufferLayout* Layout, EUniformBufferUsage Usage, EUniformBufferValidation Validation);
     ~FMetalSuballocatedUniformBuffer();
-
-	void Update(const void* Contents);
-
+    
+	void Update(const void* Contents, TArray<TRefCountPtr<FRHIResource> > const& InResourceTable);
+    
     // Prepares this uniform buffer for binding.
     // You MUST call this before binding it.
     // If this is a UB that was created before the current frame the data on the GPU will
     // not be correct until this function returns.
     void PrepareToBind();
 
-private:
 	// Copies the RDG resources to a resource table for a deferred update on the RHI thread.
-	void CopyResourceTable(const void* Contents, TArray<TRefCountPtr<FRHIResource> >& OutResourceTable) const;
+	void CopyResourceTable_RenderThread(const void* Contents, TArray<TRefCountPtr<FRHIResource> >& OutResourceTable);
 
+private:
+    
     // Pushes the data in Contents to the gpu.
     // Updates the frame counter to FrameNumber.
     // (this is to support the case where we create buffer and reuse it many frames later).
@@ -60,6 +61,8 @@ private:
     // and copies Contents into that backing.
     // The amount of data is determined by the Layout
     void PushToGPUBacking(const void* Contents);
+    
+    bool HasShadow();
 };
 
 typedef FMetalSuballocatedUniformBuffer FMetalUniformBuffer;
