AGXRHI: Back out CL 20487525+20493291

- CL 20487525 (Make RHICreateUniformBuffer thread safe) triggered a number of Metal validation errors. Backing out the original change and subsequent hotfix.

#rnx

diff --git a/Engine/Source/Runtime/Apple/AGXRHI/Private/AGXUniformBuffer.cpp b/Engine/Source/Runtime/Apple/AGXRHI/Private/AGXUniformBuffer.cpp
--- a/Engine/Source/Runtime/Apple/AGXRHI/Private/AGXUniformBuffer.cpp
+++ b/Engine/Source/Runtime/Apple/AGXRHI/Private/AGXUniformBuffer.cpp
@@ -11,33 +11,49 @@
 
 #pragma mark Suballocated Uniform Buffer Implementation
 
-FAGXSuballocatedUniformBuffer::FAGXSuballocatedUniformBuffer(const void* Contents, const FRHIUniformBufferLayout* Layout, EUniformBufferUsage Usage, EUniformBufferValidation InValidation)
+FAGXSuballocatedUniformBuffer::FAGXSuballocatedUniformBuffer(const FRHIUniformBufferLayout* Layout, EUniformBufferUsage Usage, EUniformBufferValidation InValidation)
     : FRHIUniformBuffer(Layout)
     , LastFrameUpdated(0)
     , Offset(0)
-	, Backing(nil)
-	, Shadow(FMemory::Malloc(GetSize()))
+    , Backing(nil)
+    , Shadow(nullptr)
     , ResourceTable()
 #if METAL_UNIFORM_BUFFER_VALIDATION
     , Validation(InValidation)
 #endif // METAL_UNIFORM_BUFFER_VALIDATION
 {
-	FMemory::Memcpy(Shadow, Contents, GetSize());
-	CopyResourceTable(Contents, ResourceTable);
+    // Slate can create SingleDraw uniform buffers and use them several frames later. So it must be included.
+    if (Usage == UniformBuffer_SingleDraw || Usage == UniformBuffer_MultiFrame)
+    {
+        Shadow = FMemory::Malloc(GetSize());
+    }
 }
 
 FAGXSuballocatedUniformBuffer::~FAGXSuballocatedUniformBuffer()
 {
-    FMemory::Free(Shadow);
+    if (HasShadow())
+    {
+        FMemory::Free(Shadow);
+    }
 
     // Note: this object does NOT own a reference
     // to the uniform buffer backing store
 }
 
-void FAGXSuballocatedUniformBuffer::Update(const void* Contents)
+bool FAGXSuballocatedUniformBuffer::HasShadow()
+{
+    return Shadow != nullptr;
+}
+
+void FAGXSuballocatedUniformBuffer::Update(const void* Contents, TArray<TRefCountPtr<FRHIResource> > const& InResourceTable)
 {
-	FMemory::Memcpy(Shadow, Contents, GetSize());
-	CopyResourceTable(Contents, ResourceTable);
+    if (HasShadow())
+    {
+        FMemory::Memcpy(Shadow, Contents, GetSize());
+    }
+
+	ResourceTable = InResourceTable;
+
 	PushToGPUBacking(Contents);
 }
 
@@ -46,6 +62,8 @@ void FAGXSuballocatedUniformBuffer::Update(const void* Contents)
 // The amount of data read from Contents is given by the Layout
 void FAGXSuballocatedUniformBuffer::PushToGPUBacking(const void* Contents)
 {
+    check(IsInRenderingThread() ^ IsRunningRHIInSeparateThread());
+    
     FAGXDeviceContext& DeviceContext = GetAGXDeviceContext();
     
     FAGXFrameAllocator* Allocator = DeviceContext.GetUniformAllocator();
@@ -70,7 +88,7 @@ void FAGXSuballocatedUniformBuffer::PrepareToBind()
     }
 }
 
-void FAGXSuballocatedUniformBuffer::CopyResourceTable(const void* Contents, TArray<TRefCountPtr<FRHIResource> >& OutResourceTable) const
+void FAGXSuballocatedUniformBuffer::CopyResourceTable_RenderThread(const void* Contents, TArray<TRefCountPtr<FRHIResource> >& OutResourceTable)
 {
 #if METAL_UNIFORM_BUFFER_VALIDATION
 	if (Validation == EUniformBufferValidation::ValidateResources)
diff --git a/Engine/Source/Runtime/Apple/AGXRHI/Private/AGXUniformBufferCommands.cpp b/Engine/Source/Runtime/Apple/AGXRHI/Private/AGXUniformBufferCommands.cpp
--- a/Engine/Source/Runtime/Apple/AGXRHI/Private/AGXUniformBufferCommands.cpp
+++ b/Engine/Source/Runtime/Apple/AGXRHI/Private/AGXUniformBufferCommands.cpp
@@ -12,28 +12,39 @@ static void DoUpdateUniformBuffer(FAGXUniformBuffer* UB, const void* Contents)
     // The only way we can be on the RHI thread here is if we're in the process of creating a FLocalUniformBuffer.
     bool bUpdateImmediately = RHICmdList.Bypass() || IsInRHIThread();
 
+	TArray<TRefCountPtr<FRHIResource> > ResourceTable;
+
+	UB->CopyResourceTable_RenderThread(Contents, ResourceTable);
+
     if(bUpdateImmediately)
     {
-        UB->Update(Contents);
+        UB->Update(Contents, ResourceTable);
     }
     else
     {
         const uint32 NumBytes = UB->GetLayout().ConstantBufferSize;
         void* Data = RHICmdList.Alloc(NumBytes, 16);
         FMemory::Memcpy(Data, Contents, NumBytes);
-
-        RHICmdList.EnqueueLambda([Data, UB](FRHICommandListImmediate& RHICmdList)
+        
+        RHICmdList.EnqueueLambda([Data, UB, NewResourceTable = MoveTemp(ResourceTable)](FRHICommandListImmediate& RHICmdList)
         {
-            UB->Update(Data);
+            UB->Update(Data, NewResourceTable);
         });
-
+        
         RHICmdList.RHIThreadFence(true);
     }
 }
 
 FUniformBufferRHIRef FAGXDynamicRHI::RHICreateUniformBuffer(const void* Contents, const FRHIUniformBufferLayout* Layout, EUniformBufferUsage Usage, EUniformBufferValidation Validation)
 {
-	return new FAGXUniformBuffer(Contents, Layout, Usage, Validation);
+    FAGXDeviceContext& DeviceContext = (FAGXDeviceContext&)GetAGXDeviceContext();
+    FAGXFrameAllocator* UniformAllocator = DeviceContext.GetUniformAllocator();
+    
+    FAGXUniformBuffer* UB = new FAGXUniformBuffer(Layout, Usage, Validation);
+    
+    DoUpdateUniformBuffer(UB, Contents);
+    
+    return UB;
 }
 
 void FAGXDynamicRHI::RHIUpdateUniformBuffer(FRHIUniformBuffer* UniformBufferRHI, const void* Contents)
diff --git a/Engine/Source/Runtime/Apple/AGXRHI/Public/AGXUniformBuffer.h b/Engine/Source/Runtime/Apple/AGXRHI/Public/AGXUniformBuffer.h
--- a/Engine/Source/Runtime/Apple/AGXRHI/Public/AGXUniformBuffer.h
+++ b/Engine/Source/Runtime/Apple/AGXRHI/Public/AGXUniformBuffer.h
@@ -38,21 +38,22 @@ private:
 public:
     // Creates a uniform buffer.
     // If Usage is SingleDraw or MultiFrame we will keep a copy of the data
-    FAGXSuballocatedUniformBuffer(const void* Contents, const FRHIUniformBufferLayout* Layout, EUniformBufferUsage Usage, EUniformBufferValidation Validation);
+    FAGXSuballocatedUniformBuffer(const FRHIUniformBufferLayout* Layout, EUniformBufferUsage Usage, EUniformBufferValidation Validation);
     ~FAGXSuballocatedUniformBuffer();
-
-	void Update(const void* Contents);
-
+    
+	void Update(const void* Contents, TArray<TRefCountPtr<FRHIResource> > const& InResourceTable);
+    
     // Prepares this uniform buffer for binding.
     // You MUST call this before binding it.
     // If this is a UB that was created before the current frame the data on the GPU will
     // not be correct until this function returns.
     void PrepareToBind();
 
-private:
 	// Copies the RDG resources to a resource table for a deferred update on the RHI thread.
-	void CopyResourceTable(const void* Contents, TArray<TRefCountPtr<FRHIResource> >& OutResourceTable) const;
+	void CopyResourceTable_RenderThread(const void* Contents, TArray<TRefCountPtr<FRHIResource> >& OutResourceTable);
 
+private:
+    
     // Pushes the data in Contents to the gpu.
     // Updates the frame counter to FrameNumber.
     // (this is to support the case where we create buffer and reuse it many frames later).
@@ -60,6 +61,8 @@ private:
     // and copies Contents into that backing.
     // The amount of data is determined by the Layout
     void PushToGPUBacking(const void* Contents);
+    
+    bool HasShadow();
 };
 
 typedef FAGXSuballocatedUniformBuffer FAGXUniformBuffer;
