AGXRHI: Replace mtlpp::ResourceUsage with native type MTLResourceUsage

#rnx

diff --git a/Engine/Source/Runtime/Apple/AGXRHI/Private/AGXCommandEncoder.cpp b/Engine/Source/Runtime/Apple/AGXRHI/Private/AGXCommandEncoder.cpp
--- a/Engine/Source/Runtime/Apple/AGXRHI/Private/AGXCommandEncoder.cpp
+++ b/Engine/Source/Runtime/Apple/AGXRHI/Private/AGXCommandEncoder.cpp
@@ -882,7 +882,7 @@ void FAGXCommandEncoder::SetVisibilityResultMode(mtlpp::VisibilityResultMode con
 	
 #pragma mark - Public Shader Resource Mutators -
 
-void FAGXCommandEncoder::SetShaderBuffer(mtlpp::FunctionType const FunctionType, FAGXBuffer const& Buffer, NSUInteger const Offset, NSUInteger const Length, NSUInteger index, mtlpp::ResourceUsage const Usage, EPixelFormat const Format, NSUInteger const ElementRowPitch)
+void FAGXCommandEncoder::SetShaderBuffer(mtlpp::FunctionType FunctionType, FAGXBuffer const& Buffer, NSUInteger Offset, NSUInteger Length, NSUInteger index, MTLResourceUsage Usage, EPixelFormat Format, NSUInteger ElementRowPitch)
 {
 	check(index < ML_MaxBuffers);
     if(GetAGXDeviceContext().SupportsFeature(EAGXFeaturesSetBufferOffset) && Buffer && (ShaderBuffers[uint32(FunctionType)].Bound & (1 << index)) && ShaderBuffers[uint32(FunctionType)].Buffers[index] == Buffer)
@@ -936,7 +936,7 @@ void FAGXCommandEncoder::SetShaderData(mtlpp::FunctionType const FunctionType, F
 	ShaderBuffers[uint32(FunctionType)].Buffers[Index] = nil;
 	ShaderBuffers[uint32(FunctionType)].Bytes[Index] = Data;
 	ShaderBuffers[uint32(FunctionType)].Offsets[Index] = Offset;
-	ShaderBuffers[uint32(FunctionType)].Usage[Index] = mtlpp::ResourceUsage::Read;
+	ShaderBuffers[uint32(FunctionType)].Usage[Index] = MTLResourceUsageRead;
 	ShaderBuffers[uint32(FunctionType)].SetBufferMetaData(Index, Data ? (Data->Len - Offset) : 0, GAGXBufferFormats[Format].DataFormat, ElementRowPitch);
 	
 	SetShaderBufferInternal(FunctionType, Index);
@@ -981,7 +981,7 @@ void FAGXCommandEncoder::SetShaderBytes(mtlpp::FunctionType const FunctionType,
 		}
 		ShaderBuffers[uint32(FunctionType)].Bytes[Index] = nil;
 		ShaderBuffers[uint32(FunctionType)].Offsets[Index] = 0;
-		ShaderBuffers[uint32(FunctionType)].Usage[Index] = mtlpp::ResourceUsage::Read;
+		ShaderBuffers[uint32(FunctionType)].Usage[Index] = MTLResourceUsageRead;
 		ShaderBuffers[uint32(FunctionType)].SetBufferMetaData(Index, Length, GAGXBufferFormats[PF_Unknown].DataFormat, 0);
 	}
 	else
@@ -991,7 +991,7 @@ void FAGXCommandEncoder::SetShaderBytes(mtlpp::FunctionType const FunctionType,
 		ShaderBuffers[uint32(FunctionType)].Buffers[Index] = nil;
 		ShaderBuffers[uint32(FunctionType)].Bytes[Index] = nil;
 		ShaderBuffers[uint32(FunctionType)].Offsets[Index] = 0;
-		ShaderBuffers[uint32(FunctionType)].Usage[Index] = mtlpp::ResourceUsage(0);
+		ShaderBuffers[uint32(FunctionType)].Usage[Index] = 0;
 		ShaderBuffers[uint32(FunctionType)].SetBufferMetaData(Index, 0, GAGXBufferFormats[PF_Unknown].DataFormat, 0);
 	}
 	
@@ -1025,7 +1025,7 @@ void FAGXCommandEncoder::SetShaderBufferOffset(mtlpp::FunctionType FunctionType,
 	}
 }
 
-void FAGXCommandEncoder::SetShaderTexture(mtlpp::FunctionType FunctionType, FAGXTexture const& Texture, NSUInteger index, mtlpp::ResourceUsage Usage)
+void FAGXCommandEncoder::SetShaderTexture(mtlpp::FunctionType FunctionType, FAGXTexture const& Texture, NSUInteger index, MTLResourceUsage Usage)
 {
 	check(index < ML_MaxTextures);
 	switch (FunctionType)
@@ -1197,7 +1197,7 @@ void FAGXCommandEncoder::SetShaderBufferInternal(mtlpp::FunctionType Function, u
 		
 		if (Buffer.IsSingleUse())
 		{
-			Binding.Usage[Index] = mtlpp::ResourceUsage(0);
+			Binding.Usage[Index] = 0;
 			Binding.Offsets[Index] = 0;
 			Binding.Buffers[Index] = nil;
 			Binding.Bound &= ~(1 << Index);
diff --git a/Engine/Source/Runtime/Apple/AGXRHI/Private/AGXCommandEncoder.h b/Engine/Source/Runtime/Apple/AGXRHI/Private/AGXCommandEncoder.h
--- a/Engine/Source/Runtime/Apple/AGXRHI/Private/AGXCommandEncoder.h
+++ b/Engine/Source/Runtime/Apple/AGXRHI/Private/AGXCommandEncoder.h
@@ -345,7 +345,7 @@ public:
 	 * @param Usage The resource usage mask.
 	 * @param Format The Pixel format to reinterpret the resource as.
 	 */
-	void SetShaderBuffer(mtlpp::FunctionType const FunctionType, FAGXBuffer const& Buffer, NSUInteger const Offset, NSUInteger const Length, NSUInteger const Index, mtlpp::ResourceUsage const Usage, EPixelFormat const Format = PF_Unknown, NSUInteger const ElementRowPitch = 0);
+	void SetShaderBuffer(mtlpp::FunctionType FunctionType, FAGXBuffer const& Buffer, NSUInteger Offset, NSUInteger Length, NSUInteger Index, MTLResourceUsage Usage, EPixelFormat Format = PF_Unknown, NSUInteger ElementRowPitch = 0);
 	
 	/*
 	 * Set an FAGXBufferData to the specified shader frequency at the given bind point index.
@@ -373,7 +373,7 @@ public:
 	 * @param Index The index to modify.
 	 * @param Usage The resource usage mask.
 	 */
-	void SetShaderTexture(mtlpp::FunctionType const FunctionType, FAGXTexture const& Texture, NSUInteger const Index, mtlpp::ResourceUsage const Usage);
+	void SetShaderTexture(mtlpp::FunctionType FunctionType, FAGXTexture const& Texture, NSUInteger Index, MTLResourceUsage Usage);
 	
 	/*
 	 * Set a global sampler for the specified shader frequency at the given bind point index.
@@ -451,7 +451,7 @@ private:
         /** The bound buffer offsets or 0. */
         NSUInteger Offsets[ML_MaxBuffers];
 		/** The usage mask for the bound resource or 0 */
-		mtlpp::ResourceUsage Usage[ML_MaxBuffers];
+		MTLResourceUsage Usage[ML_MaxBuffers];
 		/** The bound buffer constants */
 		struct FSizeConstants
 		{
@@ -517,7 +517,7 @@ public:
 	TSet<ns::AutoReleased<FAGXBuffer>> BufferBindingHistory;
 	TSet<ns::AutoReleased<FAGXTexture>> TextureBindingHistory;
 	
-	TMap<mtlpp::Resource::Type, mtlpp::ResourceUsage> ResourceUsage;
+	TMap<mtlpp::Resource::Type, MTLResourceUsage> ResourceUsage;
 	
 	uint32 EncoderNum;
 	uint32 CmdBufIndex;
diff --git a/Engine/Source/Runtime/Apple/AGXRHI/Private/AGXCommands.cpp b/Engine/Source/Runtime/Apple/AGXRHI/Private/AGXCommands.cpp
--- a/Engine/Source/Runtime/Apple/AGXRHI/Private/AGXCommands.cpp
+++ b/Engine/Source/Runtime/Apple/AGXRHI/Private/AGXCommands.cpp
@@ -309,17 +309,17 @@ void FAGXRHICommandContext::RHISetShaderTexture(FRHIGraphicsShader* ShaderRHI, u
 	{
         if (Surface->Texture || !(Surface->GetDesc().Flags & TexCreate_Presentable))
         {
-            Context->GetCurrentState().SetShaderTexture(Stage, Surface->Texture, TextureIndex, (mtlpp::ResourceUsage)(mtlpp::ResourceUsage::Read|mtlpp::ResourceUsage::Sample));
+            Context->GetCurrentState().SetShaderTexture(Stage, Surface->Texture, TextureIndex, (MTLResourceUsageRead | MTLResourceUsageSample));
         }
         else
         {
             FAGXTexture Tex = Surface->GetCurrentTexture();
-            Context->GetCurrentState().SetShaderTexture(Stage, Tex, TextureIndex, (mtlpp::ResourceUsage)(mtlpp::ResourceUsage::Read|mtlpp::ResourceUsage::Sample));
+            Context->GetCurrentState().SetShaderTexture(Stage, Tex, TextureIndex, (MTLResourceUsageRead | MTLResourceUsageSample));
         }
 	}
 	else
 	{
-		Context->GetCurrentState().SetShaderTexture(Stage, nil, TextureIndex, mtlpp::ResourceUsage(0));
+		Context->GetCurrentState().SetShaderTexture(Stage, nil, TextureIndex, 0);
 	}
 	}
 }
@@ -332,17 +332,17 @@ void FAGXRHICommandContext::RHISetShaderTexture(FRHIComputeShader* ComputeShader
     {
         if (Surface->Texture || !(Surface->GetDesc().Flags & TexCreate_Presentable))
         {
-            Context->GetCurrentState().SetShaderTexture(EAGXShaderStages::Compute, Surface->Texture, TextureIndex, (mtlpp::ResourceUsage)(mtlpp::ResourceUsage::Read|mtlpp::ResourceUsage::Sample));
+            Context->GetCurrentState().SetShaderTexture(EAGXShaderStages::Compute, Surface->Texture, TextureIndex, (MTLResourceUsageRead | MTLResourceUsageSample));
         }
         else
         {
             FAGXTexture Tex = Surface->GetCurrentTexture();
-            Context->GetCurrentState().SetShaderTexture(EAGXShaderStages::Compute, Tex, TextureIndex, (mtlpp::ResourceUsage)(mtlpp::ResourceUsage::Read|mtlpp::ResourceUsage::Sample));
+            Context->GetCurrentState().SetShaderTexture(EAGXShaderStages::Compute, Tex, TextureIndex, (MTLResourceUsageRead | MTLResourceUsageSample));
         }
 	}
 	else
 	{
-		Context->GetCurrentState().SetShaderTexture(EAGXShaderStages::Compute, nil, TextureIndex, mtlpp::ResourceUsage(0));
+		Context->GetCurrentState().SetShaderTexture(EAGXShaderStages::Compute, nil, TextureIndex, 0);
 	}
 	}
 }
diff --git a/Engine/Source/Runtime/Apple/AGXRHI/Private/AGXRenderPass.cpp b/Engine/Source/Runtime/Apple/AGXRHI/Private/AGXRenderPass.cpp
--- a/Engine/Source/Runtime/Apple/AGXRHI/Private/AGXRenderPass.cpp
+++ b/Engine/Source/Runtime/Apple/AGXRHI/Private/AGXRenderPass.cpp
@@ -1143,7 +1143,7 @@ void FAGXRenderPass::CommitRenderResourceTables(void)
 	if (BoundShaderState->VertexShader->SideTableBinding >= 0)
 	{
 		CurrentEncoder.SetShaderSideTable(mtlpp::FunctionType::Vertex, BoundShaderState->VertexShader->SideTableBinding);
-		State.SetShaderBuffer(EAGXShaderStages::Vertex, nil, nil, 0, 0, BoundShaderState->VertexShader->SideTableBinding, mtlpp::ResourceUsage(0));
+		State.SetShaderBuffer(EAGXShaderStages::Vertex, nil, nil, 0, 0, BoundShaderState->VertexShader->SideTableBinding, 0);
 	}
 	
 	if (IsValidRef(BoundShaderState->PixelShader))
@@ -1152,7 +1152,7 @@ void FAGXRenderPass::CommitRenderResourceTables(void)
 		if (BoundShaderState->PixelShader->SideTableBinding >= 0)
 		{
 			CurrentEncoder.SetShaderSideTable(mtlpp::FunctionType::Fragment, BoundShaderState->PixelShader->SideTableBinding);
-			State.SetShaderBuffer(EAGXShaderStages::Pixel, nil, nil, 0, 0, BoundShaderState->PixelShader->SideTableBinding, mtlpp::ResourceUsage(0));
+			State.SetShaderBuffer(EAGXShaderStages::Pixel, nil, nil, 0, 0, BoundShaderState->PixelShader->SideTableBinding, 0);
 		}
 	}
 }
@@ -1167,7 +1167,7 @@ void FAGXRenderPass::CommitDispatchResourceTables(void)
 	if (ComputeShader->SideTableBinding >= 0)
 	{
 		CurrentEncoder.SetShaderSideTable(mtlpp::FunctionType::Kernel, ComputeShader->SideTableBinding);
-		State.SetShaderBuffer(EAGXShaderStages::Compute, nil, nil, 0, 0, ComputeShader->SideTableBinding, mtlpp::ResourceUsage(0));
+		State.SetShaderBuffer(EAGXShaderStages::Compute, nil, nil, 0, 0, ComputeShader->SideTableBinding, 0);
 	}
 }
 
@@ -1181,7 +1181,7 @@ void FAGXRenderPass::CommitAsyncDispatchResourceTables(void)
 	if (ComputeShader->SideTableBinding >= 0)
 	{
 		PrologueEncoder.SetShaderSideTable(mtlpp::FunctionType::Kernel, ComputeShader->SideTableBinding);
-		State.SetShaderBuffer(EAGXShaderStages::Compute, nil, nil, 0, 0, ComputeShader->SideTableBinding, mtlpp::ResourceUsage(0));
+		State.SetShaderBuffer(EAGXShaderStages::Compute, nil, nil, 0, 0, ComputeShader->SideTableBinding, 0);
 	}
 }
 
diff --git a/Engine/Source/Runtime/Apple/AGXRHI/Private/AGXStateCache.cpp b/Engine/Source/Runtime/Apple/AGXRHI/Private/AGXStateCache.cpp
--- a/Engine/Source/Runtime/Apple/AGXRHI/Private/AGXStateCache.cpp
+++ b/Engine/Source/Runtime/Apple/AGXRHI/Private/AGXStateCache.cpp
@@ -231,14 +231,14 @@ FAGXStateCache::~FAGXStateCache()
 			ShaderBuffers[Frequency].Buffers[i].Length = 0;
 			ShaderBuffers[Frequency].Buffers[i].ElementRowPitch = 0;
 			ShaderBuffers[Frequency].Buffers[i].Offset = 0;
-			ShaderBuffers[Frequency].Buffers[i].Usage = mtlpp::ResourceUsage(0);
+			ShaderBuffers[Frequency].Buffers[i].Usage = 0;
 			ShaderBuffers[Frequency].Formats[i] = PF_Unknown;
 		}
 		ShaderBuffers[Frequency].Bound = 0;
 		for (uint32 i = 0; i < ML_MaxTextures; i++)
 		{
 			ShaderTextures[Frequency].Textures[i] = nil;
-			ShaderTextures[Frequency].Usage[i] = mtlpp::ResourceUsage(0);
+			ShaderTextures[Frequency].Usage[i] = 0;
 		}
 		ShaderTextures[Frequency].Bound = 0;
 	}
@@ -293,7 +293,7 @@ void FAGXStateCache::Reset(void)
 		for (uint32 i = 0; i < ML_MaxTextures; i++)
 		{
 			ShaderTextures[Frequency].Textures[i] = nil;
-			ShaderTextures[Frequency].Usage[i] = mtlpp::ResourceUsage(0);
+			ShaderTextures[Frequency].Usage[i] = 0;
 		}
 		ShaderTextures[Frequency].Bound = 0;
 	}
@@ -411,7 +411,7 @@ void FAGXStateCache::SetComputeShader(FAGXComputeShader* InComputeShader)
 		for (uint32 Index = 0; Index < ML_MaxTextures; ++Index)
 		{
 			ShaderTextures[EAGXShaderStages::Compute].Textures[Index] = nil;
-			ShaderTextures[EAGXShaderStages::Compute].Usage[Index] = mtlpp::ResourceUsage(0);
+			ShaderTextures[EAGXShaderStages::Compute].Usage[Index] = 0;
 		}
 		ShaderTextures[EAGXShaderStages::Compute].Bound = 0;
 
@@ -1104,7 +1104,7 @@ void FAGXStateCache::SetVertexStream(uint32 const Index, FAGXBuffer* Buffer, FAG
 	VertexBuffers[Index].Bytes = Bytes;
 	VertexBuffers[Index].Length = Length;
 	
-	SetShaderBuffer(EAGXShaderStages::Vertex, VertexBuffers[Index].Buffer, Bytes, Offset, Length, UNREAL_TO_METAL_BUFFER_INDEX(Index), mtlpp::ResourceUsage::Read);
+	SetShaderBuffer(EAGXShaderStages::Vertex, VertexBuffers[Index].Buffer, Bytes, Offset, Length, UNREAL_TO_METAL_BUFFER_INDEX(Index), MTLResourceUsageRead);
 }
 
 uint32 FAGXStateCache::GetVertexBufferSize(uint32 const Index)
@@ -1336,7 +1336,7 @@ bool FAGXStateCache::NeedsToSetRenderTarget(const FRHIRenderPassInfo& InRenderPa
 	return bAllChecksPassed == false;
 }
 
-void FAGXStateCache::SetShaderBuffer(EAGXShaderStages const Frequency, FAGXBuffer const& Buffer, FAGXBufferData* const Bytes, NSUInteger const Offset, NSUInteger const Length, NSUInteger const Index, mtlpp::ResourceUsage const Usage, EPixelFormat const Format, NSUInteger const ElementRowPitch)
+void FAGXStateCache::SetShaderBuffer(EAGXShaderStages Frequency, FAGXBuffer const& Buffer, FAGXBufferData* const Bytes, NSUInteger Offset, NSUInteger Length, NSUInteger Index, MTLResourceUsage Usage, EPixelFormat Format, NSUInteger ElementRowPitch)
 {
 	check(Frequency < EAGXShaderStages::Num);
 	check(Index < ML_MaxBuffers);
@@ -1369,7 +1369,7 @@ void FAGXStateCache::SetShaderBuffer(EAGXShaderStages const Frequency, FAGXBuffe
 	}
 }
 
-void FAGXStateCache::SetShaderTexture(EAGXShaderStages const Frequency, FAGXTexture const& Texture, NSUInteger const Index, mtlpp::ResourceUsage const Usage)
+void FAGXStateCache::SetShaderTexture(EAGXShaderStages Frequency, FAGXTexture const& Texture, NSUInteger Index, MTLResourceUsage Usage)
 {
 	check(Frequency < EAGXShaderStages::Num);
 	check(Index < ML_MaxTextures);
@@ -1423,12 +1423,12 @@ void FAGXStateCache::SetResource(uint32 ShaderStage, uint32 BindIndex, FRHITextu
 {
 	FAGXSurface* Surface = AGXGetMetalSurfaceFromRHITexture(TextureRHI);
 	ns::AutoReleased<FAGXTexture> Texture;
-	mtlpp::ResourceUsage Usage = (mtlpp::ResourceUsage)0;
+	MTLResourceUsage Usage = 0;
 	if (Surface != nullptr)
 	{
 		TextureRHI->SetLastRenderTime(CurrentTime);
 		Texture = Surface->Texture;
-		Usage = mtlpp::ResourceUsage(mtlpp::ResourceUsage::Read|mtlpp::ResourceUsage::Sample);
+		Usage = MTLResourceUsageRead | MTLResourceUsageSample;
 	}
 	
 	switch (ShaderStage)
@@ -1460,11 +1460,11 @@ void FAGXStateCache::SetShaderResourceView(FAGXContext* Context, EAGXShaderStage
 			FAGXTexture View(SRV->GetTextureView());
 			if (View)
 			{
-				SetShaderTexture(ShaderStage, View, BindIndex, mtlpp::ResourceUsage(mtlpp::ResourceUsage::Read | mtlpp::ResourceUsage::Sample));
+				SetShaderTexture(ShaderStage, View, BindIndex, (MTLResourceUsageRead | MTLResourceUsageSample));
 			}
 			else
 			{
-				SetShaderTexture(ShaderStage, nil, BindIndex, mtlpp::ResourceUsage(0));
+				SetShaderTexture(ShaderStage, nil, BindIndex, 0);
 			}
 		}
 		else
@@ -1476,11 +1476,11 @@ void FAGXStateCache::SetShaderResourceView(FAGXContext* Context, EAGXShaderStage
 				ns::AutoReleased<FAGXTexture> Tex;
 				Tex = SRV->GetLinearTexture();
 
-				SetShaderTexture(ShaderStage, Tex, BindIndex, mtlpp::ResourceUsage(mtlpp::ResourceUsage::Read | mtlpp::ResourceUsage::Sample));
+				SetShaderTexture(ShaderStage, Tex, BindIndex, (MTLResourceUsageRead | MTLResourceUsageSample));
 			}
 			else 
 			{
-				SetShaderBuffer(ShaderStage, Buffer->GetCurrentBufferOrNil(), Buffer->Data, SRV->Offset, Buffer->GetSize(), BindIndex, mtlpp::ResourceUsage::Read, (EPixelFormat)SRV->Format);
+				SetShaderBuffer(ShaderStage, Buffer->GetCurrentBufferOrNil(), Buffer->Data, SRV->Offset, Buffer->GetSize(), BindIndex, MTLResourceUsageRead, (EPixelFormat)SRV->Format);
 			}
 		}
 	}
@@ -1525,7 +1525,7 @@ void FAGXStateCache::SetShaderUnorderedAccessView(EAGXShaderStages ShaderStage,
 			{
 				FPlatformAtomics::InterlockedExchange(&Surface->Written, 1);
 
-				SetShaderTexture(ShaderStage, View, BindIndex, mtlpp::ResourceUsage(mtlpp::ResourceUsage::Read | mtlpp::ResourceUsage::Write));
+				SetShaderTexture(ShaderStage, View, BindIndex, (MTLResourceUsageRead | MTLResourceUsageWrite));
 
 				if (Surface->Texture.GetBuffer() && (EnumHasAllFlags(Surface->GetDesc().Flags, TexCreate_UAV | TexCreate_NoTiling) || EnumHasAllFlags(Surface->GetDesc().Flags, TexCreate_AtomicCompatible)))
 				{
@@ -1535,12 +1535,12 @@ void FAGXStateCache::SetShaderUnorderedAccessView(EAGXShaderStages ShaderStage,
 					FAGXBuffer Buffer(Surface->Texture.GetBuffer(), false);
 					const uint32 BufferOffset = Surface->Texture.GetBufferOffset();
 					const uint32 BufferSize = Surface->Texture.GetBuffer().GetLength();
-					SetShaderBuffer(ShaderStage, Buffer, nullptr, BufferOffset, BufferSize, BindIndex, mtlpp::ResourceUsage(mtlpp::ResourceUsage::Read | mtlpp::ResourceUsage::Write), static_cast<EPixelFormat>(UAV->Format), ElementsPerRow);
+					SetShaderBuffer(ShaderStage, Buffer, nullptr, BufferOffset, BufferSize, BindIndex, (MTLResourceUsageRead | MTLResourceUsageWrite), static_cast<EPixelFormat>(UAV->Format), ElementsPerRow);
 				}
 			}
 			else
 			{
-				SetShaderTexture(ShaderStage, nil, BindIndex, mtlpp::ResourceUsage(0));
+				SetShaderTexture(ShaderStage, nil, BindIndex, 0);
 			}
 		}
 		else
@@ -1552,10 +1552,10 @@ void FAGXStateCache::SetShaderUnorderedAccessView(EAGXShaderStages ShaderStage,
 			{
 				ns::AutoReleased<FAGXTexture> Tex;
 				Tex = UAV->GetLinearTexture();
-				SetShaderTexture(ShaderStage, Tex, BindIndex, mtlpp::ResourceUsage(mtlpp::ResourceUsage::Read | mtlpp::ResourceUsage::Write));
+				SetShaderTexture(ShaderStage, Tex, BindIndex, (MTLResourceUsageRead | MTLResourceUsageWrite));
 			}
 
-			SetShaderBuffer(ShaderStage, Buffer->GetCurrentBufferOrNil(), Buffer->Data, 0, Buffer->GetSize(), BindIndex, mtlpp::ResourceUsage(mtlpp::ResourceUsage::Read | mtlpp::ResourceUsage::Write), (EPixelFormat)UAV->Format);
+			SetShaderBuffer(ShaderStage, Buffer->GetCurrentBufferOrNil(), Buffer->Data, 0, Buffer->GetSize(), BindIndex, (MTLResourceUsageRead | MTLResourceUsageWrite), (EPixelFormat)UAV->Format);
 		}
 	}
 }
diff --git a/Engine/Source/Runtime/Apple/AGXRHI/Private/AGXStateCache.h b/Engine/Source/Runtime/Apple/AGXRHI/Private/AGXStateCache.h
--- a/Engine/Source/Runtime/Apple/AGXRHI/Private/AGXStateCache.h
+++ b/Engine/Source/Runtime/Apple/AGXRHI/Private/AGXStateCache.h
@@ -78,7 +78,7 @@ public:
 	 * @param Usage The resource usage flags.
 	 * @param Format The UAV pixel format.
 	 */
-	void SetShaderBuffer(EAGXShaderStages const Frequency, FAGXBuffer const& Buffer, FAGXBufferData* const Bytes, NSUInteger const Offset, NSUInteger const Length, NSUInteger const Index, mtlpp::ResourceUsage const Usage, EPixelFormat const Format = PF_Unknown, NSUInteger const ElementRowPitch = 0);
+	void SetShaderBuffer(EAGXShaderStages Frequency, FAGXBuffer const& Buffer, FAGXBufferData* const Bytes, NSUInteger Offset, NSUInteger Length, NSUInteger Index, MTLResourceUsage Usage, EPixelFormat Format = PF_Unknown, NSUInteger ElementRowPitch = 0);
 	
 	/*
 	 * Set a global texture for the specified shader frequency at the given bind point index.
@@ -87,7 +87,7 @@ public:
 	 * @param Index The index to modify.
 	 * @param Usage The resource usage flags.
 	 */
-	void SetShaderTexture(EAGXShaderStages const Frequency, FAGXTexture const& Texture, NSUInteger const Index, mtlpp::ResourceUsage const Usage);
+	void SetShaderTexture(EAGXShaderStages Frequency, FAGXTexture const& Texture, NSUInteger Index, MTLResourceUsage Usage);
 	
 	/*
 	 * Set a global sampler for the specified shader frequency at the given bind point index.
@@ -186,7 +186,7 @@ private:
 #pragma mark - Private Type Declarations -
 	struct FAGXBufferBinding
 	{
-		FAGXBufferBinding() : Bytes(nil), Offset(0), Length(0), Usage((mtlpp::ResourceUsage)0) {}
+		FAGXBufferBinding() : Bytes(nil), Offset(0), Length(0), Usage(0) {}
 		/** The bound buffers or nil. */
 		ns::AutoReleased<FAGXBuffer> Buffer;
 		/** Optional bytes buffer used instead of an FAGXBuffer */
@@ -198,7 +198,7 @@ private:
 		/** The bound buffer element row pitch or 0 */
 		NSUInteger ElementRowPitch;
 		/** The bound buffer usage or 0 */
-		mtlpp::ResourceUsage Usage;
+		MTLResourceUsage Usage;
 	};
 	
 	/** A structure of arrays for the current buffer binding settings. */
@@ -220,7 +220,7 @@ private:
 		/** The bound textures or nil. */
 		ns::AutoReleased<FAGXTexture> Textures[ML_MaxTextures];
 		/** The bound texture usage or 0 */
-		mtlpp::ResourceUsage Usage[ML_MaxTextures];
+		MTLResourceUsage Usage[ML_MaxTextures];
 		/** A bitmask for which textures were bound by the application where a bit value of 1 is bound and 0 is unbound. */
 		FAGXTextureMask Bound;
 	};
diff --git a/Engine/Source/Runtime/Apple/AGXRHI/Private/AGXUniformBufferCommands.cpp b/Engine/Source/Runtime/Apple/AGXRHI/Private/AGXUniformBufferCommands.cpp
--- a/Engine/Source/Runtime/Apple/AGXRHI/Private/AGXUniformBufferCommands.cpp
+++ b/Engine/Source/Runtime/Apple/AGXRHI/Private/AGXUniformBufferCommands.cpp
@@ -59,7 +59,7 @@ static void SetUniformBufferInternal(FAGXContext* Context, RHIShaderType* Shader
             UB->PrepareToBind();
             
             FAGXBuffer Buf(UB->Backing, ns::Ownership::AutoRelease);
-            Context->GetCurrentState().SetShaderBuffer(Stage, Buf, nil, UB->Offset, UB->GetSize(), BufferIndex, mtlpp::ResourceUsage::Read);
+            Context->GetCurrentState().SetShaderBuffer(Stage, Buf, nil, UB->Offset, UB->GetSize(), BufferIndex, MTLResourceUsageRead);
         }
     }
 }
diff --git a/Engine/Source/Runtime/Apple/AGXRHI/Private/Shaders/AGXShaderParameterCache.cpp b/Engine/Source/Runtime/Apple/AGXRHI/Private/Shaders/AGXShaderParameterCache.cpp
--- a/Engine/Source/Runtime/Apple/AGXRHI/Private/Shaders/AGXShaderParameterCache.cpp
+++ b/Engine/Source/Runtime/Apple/AGXRHI/Private/Shaders/AGXShaderParameterCache.cpp
@@ -104,13 +104,13 @@ void FAGXShaderParameterCache::CommitPackedGlobals(FAGXStateCache* Cache, FAGXCo
 				uint8 const* Bytes = PackedGlobalUniforms[Index]->Data;
 				ns::AutoReleased<FAGXBuffer> Buffer(Encoder->GetRingBuffer().NewBuffer(Size, 0));
 				FMemory::Memcpy((uint8*)Buffer.GetContents(), Bytes, Size);
-				Cache->SetShaderBuffer((EAGXShaderStages)Frequency, Buffer, nil, 0, Size, UniformBufferIndex, mtlpp::ResourceUsage::Read);
+				Cache->SetShaderBuffer((EAGXShaderStages)Frequency, Buffer, nil, 0, Size, UniformBufferIndex, MTLResourceUsageRead);
 			}
 			else
 			{
 				PackedGlobalUniforms[Index]->Len = Size;
-				Cache->SetShaderBuffer((EAGXShaderStages)Frequency, nil, nil, 0, 0, UniformBufferIndex, mtlpp::ResourceUsage(0));
-				Cache->SetShaderBuffer((EAGXShaderStages)Frequency, nil, PackedGlobalUniforms[Index], 0, Size, UniformBufferIndex, mtlpp::ResourceUsage::Read);
+				Cache->SetShaderBuffer((EAGXShaderStages)Frequency, nil, nil, 0, 0, UniformBufferIndex, 0);
+				Cache->SetShaderBuffer((EAGXShaderStages)Frequency, nil, PackedGlobalUniforms[Index], 0, Size, UniformBufferIndex, MTLResourceUsageRead);
 			}
 
 			// mark as clean
diff --git a/Engine/Source/Runtime/Apple/AGXRHI/Public/AGXRHIContext.h b/Engine/Source/Runtime/Apple/AGXRHI/Public/AGXRHIContext.h
--- a/Engine/Source/Runtime/Apple/AGXRHI/Public/AGXRHIContext.h
+++ b/Engine/Source/Runtime/Apple/AGXRHI/Public/AGXRHIContext.h
@@ -3,7 +3,6 @@
 #pragma once
 
 #include <Metal/Metal.h>
-#include "AGXState.h"
 #include "AGXResources.h"
 #include "AGXViewport.h"
 #include "RHICore.h"
diff --git a/Engine/Source/Runtime/Apple/AGXRHI/Public/AGXState.h b/Engine/Source/Runtime/Apple/AGXRHI/Public/AGXState.h
--- a/Engine/Source/Runtime/Apple/AGXRHI/Public/AGXState.h
+++ b/Engine/Source/Runtime/Apple/AGXRHI/Public/AGXState.h
@@ -6,12 +6,6 @@
 
 #pragma once
 
-#include "CoreMinimal.h"
-#include "RHI.h"
-THIRD_PARTY_INCLUDES_START
-#include "mtlpp.hpp"
-THIRD_PARTY_INCLUDES_END
-
 class FAGXSamplerState : public FRHISamplerState
 {
 public:
