AGXRHI: Replace mtlpp:FunctionType with native MTLFunctionType

#rnx

diff --git a/Engine/Source/Runtime/Apple/AGXRHI/Private/AGXCommandEncoder.cpp b/Engine/Source/Runtime/Apple/AGXRHI/Private/AGXCommandEncoder.cpp
--- a/Engine/Source/Runtime/Apple/AGXRHI/Private/AGXCommandEncoder.cpp
+++ b/Engine/Source/Runtime/Apple/AGXRHI/Private/AGXCommandEncoder.cpp
@@ -37,7 +37,7 @@ FAGXCommandEncoder::FAGXCommandEncoder(FAGXCommandList& CmdList, EAGXCommandEnco
 , CmdBufIndex(0)
 , Type(InType)
 {
-	for (uint32 Frequency = 0; Frequency < uint32(mtlpp::FunctionType::Kernel)+1; Frequency++)
+	for (uint32 Frequency = 0; Frequency < NUM_SHADER_FREQUENCIES; Frequency++)
 	{
 		FMemory::Memzero(ShaderBuffers[Frequency].Bytes);
 		FMemory::Memzero(ShaderBuffers[Frequency].Offsets);
@@ -77,7 +77,7 @@ FAGXCommandEncoder::~FAGXCommandEncoder(void)
 		[DebugGroups release];
 	}
 	
-	for (uint32 Frequency = 0; Frequency < uint32(mtlpp::FunctionType::Kernel)+1; Frequency++)
+	for (uint32 Frequency = 0; Frequency < NUM_SHADER_FREQUENCIES; Frequency++)
 	{
 		for (uint32 i = 0; i < ML_MaxBuffers; i++)
 		{
@@ -115,7 +115,7 @@ void FAGXCommandEncoder::Reset(void)
 		StencilStoreAction = MTLStoreActionUnknown;
 	}
 	
-	for (uint32 Frequency = 0; Frequency < uint32(mtlpp::FunctionType::Kernel)+1; Frequency++)
+	for (uint32 Frequency = 0; Frequency < NUM_SHADER_FREQUENCIES; Frequency++)
 	{
 		for (uint32 i = 0; i < ML_MaxBuffers; i++)
 		{
@@ -133,7 +133,7 @@ void FAGXCommandEncoder::Reset(void)
 
 void FAGXCommandEncoder::ResetLive(void)
 {
-	for (uint32 Frequency = 0; Frequency < uint32(mtlpp::FunctionType::Kernel)+1; Frequency++)
+	for (uint32 Frequency = 0; Frequency < NUM_SHADER_FREQUENCIES; Frequency++)
 	{
 		for (uint32 i = 0; i < ML_MaxBuffers; i++)
 		{
@@ -602,7 +602,7 @@ void FAGXCommandEncoder::EndEncoding(void)
 		}
 	}
 	
-	for (uint32 Frequency = 0; Frequency < uint32(mtlpp::FunctionType::Kernel)+1; Frequency++)
+	for (uint32 Frequency = 0; Frequency < NUM_SHADER_FREQUENCIES; Frequency++)
 	{
 		for (uint32 i = 0; i < ML_MaxBuffers; i++)
 		{
@@ -739,7 +739,7 @@ void FAGXCommandEncoder::SetRenderPassDescriptor(MTLRenderPassDescriptor* InRend
 	}
 	check(RenderPassDesc);
 	
-	for (uint32 Frequency = 0; Frequency < uint32(MTLFunctionTypeKernel)+1; Frequency++)
+	for (uint32 Frequency = 0; Frequency < NUM_SHADER_FREQUENCIES; Frequency++)
 	{
 		for (uint32 i = 0; i < ML_MaxBuffers; i++)
 		{
@@ -882,12 +882,12 @@ void FAGXCommandEncoder::SetVisibilityResultMode(mtlpp::VisibilityResultMode con
 	
 #pragma mark - Public Shader Resource Mutators -
 
-void FAGXCommandEncoder::SetShaderBuffer(mtlpp::FunctionType FunctionType, FAGXBuffer const& Buffer, NSUInteger Offset, NSUInteger Length, NSUInteger index, MTLResourceUsage Usage, EPixelFormat Format, NSUInteger ElementRowPitch)
+void FAGXCommandEncoder::SetShaderBuffer(MTLFunctionType FunctionType, FAGXBuffer const& Buffer, NSUInteger Offset, NSUInteger Length, NSUInteger index, MTLResourceUsage Usage, EPixelFormat Format, NSUInteger ElementRowPitch)
 {
 	check(index < ML_MaxBuffers);
     if(GetAGXDeviceContext().SupportsFeature(EAGXFeaturesSetBufferOffset) && Buffer && (ShaderBuffers[uint32(FunctionType)].Bound & (1 << index)) && ShaderBuffers[uint32(FunctionType)].Buffers[index] == Buffer)
     {
-		if (FunctionType == mtlpp::FunctionType::Vertex || FunctionType == mtlpp::FunctionType::Kernel)
+		if (FunctionType == MTLFunctionTypeVertex || FunctionType == MTLFunctionTypeKernel)
 		{
 			FenceResource(Buffer);
 		}
@@ -920,7 +920,7 @@ void FAGXCommandEncoder::SetShaderBuffer(mtlpp::FunctionType FunctionType, FAGXB
 	}
 }
 
-void FAGXCommandEncoder::SetShaderData(mtlpp::FunctionType const FunctionType, FAGXBufferData* Data, NSUInteger const Offset, NSUInteger const Index, EPixelFormat const Format, NSUInteger const ElementRowPitch)
+void FAGXCommandEncoder::SetShaderData(MTLFunctionType FunctionType, FAGXBufferData* Data, NSUInteger Offset, NSUInteger Index, EPixelFormat Format, NSUInteger ElementRowPitch)
 {
 	check(Index < ML_MaxBuffers);
 	
@@ -942,7 +942,7 @@ void FAGXCommandEncoder::SetShaderData(mtlpp::FunctionType const FunctionType, F
 	SetShaderBufferInternal(FunctionType, Index);
 }
 
-void FAGXCommandEncoder::SetShaderBytes(mtlpp::FunctionType const FunctionType, uint8 const* Bytes, NSUInteger const Length, NSUInteger const Index)
+void FAGXCommandEncoder::SetShaderBytes(MTLFunctionType FunctionType, uint8 const* Bytes, NSUInteger Length, NSUInteger Index)
 {
 	check(Index < ML_MaxBuffers);
 	
@@ -954,15 +954,15 @@ void FAGXCommandEncoder::SetShaderBytes(mtlpp::FunctionType const FunctionType,
 		{
 			switch (FunctionType)
 			{
-				case mtlpp::FunctionType::Vertex:
+				case MTLFunctionTypeVertex:
 					check(RenderCommandEncoder);
 					[RenderCommandEncoder setVertexBytes:static_cast<const void*>(Bytes) length:Length atIndex:Index];
 					break;
-				case mtlpp::FunctionType::Fragment:
+				case MTLFunctionTypeFragment:
 					check(RenderCommandEncoder);
 					[RenderCommandEncoder setFragmentBytes:static_cast<const void*>(Bytes) length:Length atIndex:Index];
 					break;
-				case mtlpp::FunctionType::Kernel:
+				case MTLFunctionTypeKernel:
 					check(ComputeCommandEncoder);
 					[ComputeCommandEncoder setBytes:static_cast<const void*>(Bytes) length:Length atIndex:Index];
 					break;
@@ -998,7 +998,7 @@ void FAGXCommandEncoder::SetShaderBytes(mtlpp::FunctionType const FunctionType,
 	SetShaderBufferInternal(FunctionType, Index);
 }
 
-void FAGXCommandEncoder::SetShaderBufferOffset(mtlpp::FunctionType FunctionType, NSUInteger const Offset, NSUInteger const Length, NSUInteger const index)
+void FAGXCommandEncoder::SetShaderBufferOffset(MTLFunctionType FunctionType, NSUInteger Offset, NSUInteger Length, NSUInteger index)
 {
 	check(index < ML_MaxBuffers);
     checkf(ShaderBuffers[uint32(FunctionType)].Buffers[index] && (ShaderBuffers[uint32(FunctionType)].Bound & (1 << index)), TEXT("Buffer must already be bound"));
@@ -1007,15 +1007,15 @@ void FAGXCommandEncoder::SetShaderBufferOffset(mtlpp::FunctionType FunctionType,
 	ShaderBuffers[uint32(FunctionType)].SetBufferMetaData(index, Length, GAGXBufferFormats[PF_Unknown].DataFormat, 0);
 	switch (FunctionType)
 	{
-		case mtlpp::FunctionType::Vertex:
+		case MTLFunctionTypeVertex:
 			check (RenderCommandEncoder);
 			[RenderCommandEncoder setVertexBufferOffset:(Offset + ShaderBuffers[uint32(FunctionType)].Buffers[index].GetOffset()) atIndex:index];
 			break;
-		case mtlpp::FunctionType::Fragment:
+		case MTLFunctionTypeFragment:
 			check(RenderCommandEncoder);
 			[RenderCommandEncoder setFragmentBufferOffset:(Offset + ShaderBuffers[uint32(FunctionType)].Buffers[index].GetOffset()) atIndex:index];
 			break;
-		case mtlpp::FunctionType::Kernel:
+		case MTLFunctionTypeKernel:
 			check (ComputeCommandEncoder);
 			[ComputeCommandEncoder setBufferOffset:(Offset + ShaderBuffers[uint32(FunctionType)].Buffers[index].GetOffset()) atIndex:index];
 			break;
@@ -1025,21 +1025,21 @@ void FAGXCommandEncoder::SetShaderBufferOffset(mtlpp::FunctionType FunctionType,
 	}
 }
 
-void FAGXCommandEncoder::SetShaderTexture(mtlpp::FunctionType FunctionType, FAGXTexture const& Texture, NSUInteger index, MTLResourceUsage Usage)
+void FAGXCommandEncoder::SetShaderTexture(MTLFunctionType FunctionType, FAGXTexture const& Texture, NSUInteger index, MTLResourceUsage Usage)
 {
 	check(index < ML_MaxTextures);
 	switch (FunctionType)
 	{
-		case mtlpp::FunctionType::Vertex:
+		case MTLFunctionTypeVertex:
 			check (RenderCommandEncoder);
 			FenceResource(Texture);
 			[RenderCommandEncoder setVertexTexture:Texture.GetPtr() atIndex:index];
 			break;
-		case mtlpp::FunctionType::Fragment:
+		case MTLFunctionTypeFragment:
 			check(RenderCommandEncoder);
 			[RenderCommandEncoder setFragmentTexture:Texture.GetPtr() atIndex:index];
 			break;
-		case mtlpp::FunctionType::Kernel:
+		case MTLFunctionTypeKernel:
 			check (ComputeCommandEncoder);
 			FenceResource(Texture);
 			[ComputeCommandEncoder setTexture:Texture.GetPtr() atIndex:index];
@@ -1066,20 +1066,20 @@ void FAGXCommandEncoder::SetShaderTexture(mtlpp::FunctionType FunctionType, FAGX
 	}
 }
 
-void FAGXCommandEncoder::SetShaderSamplerState(mtlpp::FunctionType FunctionType, id<MTLSamplerState> Sampler, NSUInteger index)
+void FAGXCommandEncoder::SetShaderSamplerState(MTLFunctionType FunctionType, id<MTLSamplerState> Sampler, NSUInteger index)
 {
 	check(index < ML_MaxSamplers);
 	switch (FunctionType)
 	{
-		case mtlpp::FunctionType::Vertex:
+		case MTLFunctionTypeVertex:
        		check (RenderCommandEncoder);
 			[RenderCommandEncoder setVertexSamplerState:Sampler atIndex:index];
 			break;
-		case mtlpp::FunctionType::Fragment:
+		case MTLFunctionTypeFragment:
 			check (RenderCommandEncoder);
 			[RenderCommandEncoder setFragmentSamplerState:Sampler atIndex:index];
 			break;
-		case mtlpp::FunctionType::Kernel:
+		case MTLFunctionTypeKernel:
 			check (ComputeCommandEncoder);
 			[ComputeCommandEncoder setSamplerState:Sampler atIndex:index];
 			break;
@@ -1089,11 +1089,11 @@ void FAGXCommandEncoder::SetShaderSamplerState(mtlpp::FunctionType FunctionType,
 	}
 }
 
-void FAGXCommandEncoder::SetShaderSideTable(mtlpp::FunctionType const FunctionType, NSUInteger const Index)
+void FAGXCommandEncoder::SetShaderSideTable(MTLFunctionType const FunctionType, NSUInteger const Index)
 {
 	if (Index < ML_MaxBuffers)
 	{
-		SetShaderData(FunctionType, ShaderBuffers[uint32(FunctionType)].SideTable, 0, Index);
+		SetShaderData(FunctionType, ShaderBuffers[FunctionType].SideTable, 0, Index);
 	}
 }
 
@@ -1140,7 +1140,7 @@ void FAGXCommandEncoder::FenceResource(mtlpp::Buffer const& Resource)
 {
 }
 
-void FAGXCommandEncoder::SetShaderBufferInternal(mtlpp::FunctionType Function, uint32 Index)
+void FAGXCommandEncoder::SetShaderBufferInternal(MTLFunctionType Function, uint32 Index)
 {
 	FAGXBufferBindings& Binding = ShaderBuffers[uint32(Function)];
 
@@ -1170,20 +1170,20 @@ void FAGXCommandEncoder::SetShaderBufferInternal(mtlpp::FunctionType Function, u
         
 		switch (Function)
 		{
-			case mtlpp::FunctionType::Vertex:
+			case MTLFunctionTypeVertex:
 				Binding.Bound |= (1 << Index);
 				check(RenderCommandEncoder);
 				FenceResource(Buffer);
 				[RenderCommandEncoder setVertexBuffer:Buffer.GetPtr() offset:(Offset + Buffer.GetOffset()) atIndex:Index];
 				break;
 
-			case mtlpp::FunctionType::Fragment:
+			case MTLFunctionTypeFragment:
 				Binding.Bound |= (1 << Index);
 				check(RenderCommandEncoder);
 				[RenderCommandEncoder setFragmentBuffer:Buffer.GetPtr() offset:(Offset + Buffer.GetOffset()) atIndex:Index];
 				break;
 
-			case mtlpp::FunctionType::Kernel:
+			case MTLFunctionTypeKernel:
 				Binding.Bound |= (1 << Index);
 				check(ComputeCommandEncoder);
 				FenceResource(Buffer);
@@ -1210,19 +1210,19 @@ void FAGXCommandEncoder::SetShaderBufferInternal(mtlpp::FunctionType Function, u
 		
 		switch (Function)
 		{
-			case mtlpp::FunctionType::Vertex:
+			case MTLFunctionTypeVertex:
 				Binding.Bound |= (1 << Index);
 				check(RenderCommandEncoder);
 				[RenderCommandEncoder setVertexBytes:static_cast<const void*>(Bytes) length:(NSUInteger)Len atIndex:Index];
 				break;
 
-			case mtlpp::FunctionType::Fragment:
+			case MTLFunctionTypeFragment:
 				Binding.Bound |= (1 << Index);
 				check(RenderCommandEncoder);
 				[RenderCommandEncoder setFragmentBytes:static_cast<const void*>(Bytes) length:(NSUInteger)Len atIndex:Index];
 				break;
 
-			case mtlpp::FunctionType::Kernel:
+			case MTLFunctionTypeKernel:
 				Binding.Bound |= (1 << Index);
 				check(ComputeCommandEncoder);
 				[ComputeCommandEncoder setBytes:static_cast<const void*>(Bytes) length:(NSUInteger)Len atIndex:Index];
diff --git a/Engine/Source/Runtime/Apple/AGXRHI/Private/AGXCommandEncoder.h b/Engine/Source/Runtime/Apple/AGXRHI/Private/AGXCommandEncoder.h
--- a/Engine/Source/Runtime/Apple/AGXRHI/Private/AGXCommandEncoder.h
+++ b/Engine/Source/Runtime/Apple/AGXRHI/Private/AGXCommandEncoder.h
@@ -12,48 +12,6 @@ class FAGXCommandQueue;
 class FAGXGraphicsPipelineState;
 struct FAGXCommandBufferFence;
 
-#if 0
-template<typename T>
-class TObjCWrapper : public FThreadSafeRefCountedObject
-{
-	using Pointer = T*;
-
-public:
-	UE_NONCOPYABLE(TObjCWrapper)
-
-	TObjCWrapper(Pointer InObject = nil)
-		: Impl([InObject retain])
-	{
-	}
-
-	~TObjCWrapper()
-	{
-		if (Impl)
-		{
-			[Impl release];
-			Impl = nil;
-		}
-	}
-
-	bool IsValid() const
-	{
-		return (Impl != nil);
-	}
-
-	Pointer AllocInit()
-	{
-		check(!IsValid());
-		Impl = [[T alloc] init];
-		return Impl;
-	}
-
-private:
-	Pointer Impl;
-};
-
-typedef TRefCountPtr<TObjCWrapper<MTLRenderPassDescriptor> > FMTLRenderPassDescriptor;
-#endif // 0
-
 /**
  * Enumeration for submission hints to avoid unclear bool values.
  */
@@ -98,8 +56,6 @@ enum EAGXCommandEncoderType
 class FAGXCommandEncoder
 {
 public:
-#pragma mark - Public C++ Boilerplate -
-
 	/** Default constructor */
 	FAGXCommandEncoder(FAGXCommandList& CmdList, EAGXCommandEncoderType Type);
 	
@@ -112,8 +68,6 @@ public:
 	/** Reset cached state for reuse while in rendering */
 	void ResetLive(void);
 	
-#pragma mark - Public Command Buffer Mutators -
-
 	/**
 	 * Start encoding to CommandBuffer. It is an error to call this with any outstanding command encoders or current command buffer.
 	 * Instead call EndEncoding & CommitCommandBuffer before calling this.
@@ -126,8 +80,6 @@ public:
  	 */
 	void CommitCommandBuffer(uint32 const Flags);
 
-#pragma mark - Public Command Buffer Accessors -
-	
 	/** @returns the current command buffer */
 	mtlpp::CommandBuffer& GetCommandBuffer() { return CommandBuffer; }
 
@@ -137,8 +89,6 @@ public:
 	/** @returns the monotonically incremented command buffer index */
 	uint32 GetCommandBufferIndex() const { return CmdBufIndex; }
 
-#pragma mark - Public Command Encoder Accessors -
-	
 	/** @returns True if and only if there is an active parallel render command encoder, otherwise false. */
 	bool IsParallelRenderCommandEncoderActive() const;
 	
@@ -187,8 +137,6 @@ public:
 	/** @returns The number of encoded passes in the command buffer. */
 	uint32 NumEncodedPasses(void) const { return EncoderNum; }
 	
-#pragma mark - Public Command Encoder Mutators -
-
 	/**
  	 * Begins encoding rendering commands into the current command buffer. No other encoder may be active & the MTLRenderPassDescriptor must previously have been set.
 	 * @param NumChildren The number of child render-encoders to create. 
@@ -215,8 +163,6 @@ public:
 	/** Adds a command-buffer completion handler to the command-buffer */
 	void AddCompletionHandler(mtlpp::CommandBufferHandler Handler);
 	
-#pragma mark - Public Debug Support -
-	
 	/*
 	 * Inserts a debug string into the command buffer.  This does not change any API behavior, but can be useful when debugging.
 	 * @param string The name of the signpost. 
@@ -237,8 +183,6 @@ public:
 	FAGXCommandBufferStats* GetCommandBufferStats(void);
 #endif
 
-#pragma mark - Public Render State Mutators -
-
 	/**
 	 * Set the render pass descriptor - no encoder may be active when this function is called.
 	 * @param RenderPass The render pass descriptor to set. May be nil.
@@ -333,8 +277,6 @@ public:
 	 */
 	void SetVisibilityResultMode(mtlpp::VisibilityResultMode const Mode, NSUInteger const Offset);
 	
-#pragma mark - Public Shader Resource Mutators -
-	
 	/*
 	 * Set a global buffer for the specified shader frequency at the given bind point index.
 	 * @param FunctionType The shader function to modify.
@@ -345,7 +287,7 @@ public:
 	 * @param Usage The resource usage mask.
 	 * @param Format The Pixel format to reinterpret the resource as.
 	 */
-	void SetShaderBuffer(mtlpp::FunctionType FunctionType, FAGXBuffer const& Buffer, NSUInteger Offset, NSUInteger Length, NSUInteger Index, MTLResourceUsage Usage, EPixelFormat Format = PF_Unknown, NSUInteger ElementRowPitch = 0);
+	void SetShaderBuffer(MTLFunctionType FunctionType, FAGXBuffer const& Buffer, NSUInteger Offset, NSUInteger Length, NSUInteger Index, MTLResourceUsage Usage, EPixelFormat Format = PF_Unknown, NSUInteger ElementRowPitch = 0);
 	
 	/*
 	 * Set an FAGXBufferData to the specified shader frequency at the given bind point index.
@@ -355,7 +297,7 @@ public:
 	 * @param Index The index to modify.
 	 * @param Format The pixel format to reinterpret the resource as.
 	 */
-	void SetShaderData(mtlpp::FunctionType const FunctionType, FAGXBufferData* Data, NSUInteger const Offset, NSUInteger const Index, EPixelFormat const Format = PF_Unknown, NSUInteger const ElementRowPitch = 0);
+	void SetShaderData(MTLFunctionType FunctionType, FAGXBufferData* Data, NSUInteger Offset, NSUInteger Index, EPixelFormat Format = PF_Unknown, NSUInteger ElementRowPitch = 0);
 	
 	/*
 	 * Set bytes to the specified shader frequency at the given bind point index.
@@ -364,7 +306,7 @@ public:
 	 * @param Length The length of the buffer or 0 when Bytes is nil.
 	 * @param Index The index to modify.
 	 */
-	void SetShaderBytes(mtlpp::FunctionType const FunctionType, uint8 const* Bytes, NSUInteger const Length, NSUInteger const Index);
+	void SetShaderBytes(MTLFunctionType FunctionType, uint8 const* Bytes, NSUInteger Length, NSUInteger Index);
 	
 	/*
 	 * Set a global texture for the specified shader frequency at the given bind point index.
@@ -373,7 +315,7 @@ public:
 	 * @param Index The index to modify.
 	 * @param Usage The resource usage mask.
 	 */
-	void SetShaderTexture(mtlpp::FunctionType FunctionType, FAGXTexture const& Texture, NSUInteger Index, MTLResourceUsage Usage);
+	void SetShaderTexture(MTLFunctionType FunctionType, FAGXTexture const& Texture, NSUInteger Index, MTLResourceUsage Usage);
 	
 	/*
 	 * Set a global sampler for the specified shader frequency at the given bind point index.
@@ -381,14 +323,14 @@ public:
 	 * @param Sampler The sampler state to bind or nil to clear.
 	 * @param Index The index to modify.
 	 */
-	void SetShaderSamplerState(mtlpp::FunctionType FunctionType, id<MTLSamplerState> Sampler, NSUInteger Index);
+	void SetShaderSamplerState(MTLFunctionType FunctionType, id<MTLSamplerState> Sampler, NSUInteger Index);
 	
 	/*
 	 * Set the shader side-table data for FunctionType at Index.
 	 * @param FunctionType The shader function to modify.
 	 * @param Index The index to bind data to.
 	 */
-	void SetShaderSideTable(mtlpp::FunctionType const FunctionType, NSUInteger const Index);
+	void SetShaderSideTable(MTLFunctionType FunctionType, NSUInteger Index);
 	
 	/*
 	 * Transition resource so that we can barrier fragment->vertex stages.
@@ -396,24 +338,18 @@ public:
 	 */
 	void TransitionResources(mtlpp::Resource const& Resource);
 	
-#pragma mark - Public Compute State Mutators -
-	
 	/*
 	 * Set the compute pipeline state that will be used.
 	 * @param State The state to set - must not be nil.
 	 */
 	void SetComputePipelineState(FAGXShaderPipeline* State);
 
-#pragma mark - Public Ring-Buffer Accessor -
-	
 	/*
 	 * Get the internal ring-buffer used for temporary allocations.
 	 * @returns The temporary allocation buffer for this command-encoder.
 	 */
 	FAGXSubBufferRing& GetRingBuffer(void);
 	
-#pragma mark - Public Resource query Access -
-
 	/*
 	 * Returns True if the Resource has been bound to a command encoder, otherwise false.  History will be cleared after a commit operation
 	 * @returns True if the Resource has been bound to a command encoder, otherwise false.
@@ -422,7 +358,6 @@ public:
 	bool HasBufferBindingHistory(FAGXBuffer const& Buffer) const;
 	
 private:
-#pragma mark - Private Functions -
 	/*
 	 * Set the offset for the buffer bound on the specified shader frequency at the given bind point index.
 	 * @param FunctionType The shader function to modify.
@@ -430,15 +365,16 @@ private:
 	 * @param Length The data length - caller is responsible for accounting for non-zero Offset.
 	 * @param Index The index to modify.
 	 */
-	void SetShaderBufferOffset(mtlpp::FunctionType const FunctionType, NSUInteger const Offset, NSUInteger const Length, NSUInteger const Index);
+	void SetShaderBufferOffset(MTLFunctionType FunctionType, NSUInteger Offset, NSUInteger Length, NSUInteger Index);
 	
-	void SetShaderBufferInternal(mtlpp::FunctionType Function, uint32 Index);
+	void SetShaderBufferInternal(MTLFunctionType Function, uint32 Index);
 	
 	void FenceResource(mtlpp::Texture const& Resource);
 	void FenceResource(mtlpp::Buffer const& Resource);
 	
-#pragma mark - Private Type Declarations -
 private:
+	static const uint32 NUM_SHADER_FREQUENCIES = uint32(MTLFunctionTypeKernel) + 1;
+
     /** A structure of arrays for the current buffer binding settings. */
     struct FAGXBufferBindings
     {
@@ -467,7 +403,6 @@ private:
         /** A bitmask for which buffers were bound by the application where a bit value of 1 is bound and 0 is unbound. */
         uint32 Bound;
         
-public:
         void SetBufferMetaData(NSUInteger Index, NSUInteger Length, NSUInteger Format, NSUInteger ElementRowPitch)
         {
 			Lengths[Index].Length = Length;
@@ -482,13 +417,12 @@ public:
         }
 	};
 	
-#pragma mark - Private Member Variables -
 	FAGXCommandList& CommandList;
 
     // Cache Queue feature
     bool bSupportsMetalFeaturesSetBytes;
     
-	FAGXBufferBindings ShaderBuffers[int(mtlpp::FunctionType::Kernel)+1];
+	FAGXBufferBindings ShaderBuffers[NUM_SHADER_FREQUENCIES];
 	
 	MTLStoreAction ColorStoreActions[MaxSimultaneousRenderTargets];
 	MTLStoreAction DepthStoreAction;
diff --git a/Engine/Source/Runtime/Apple/AGXRHI/Private/AGXRenderPass.cpp b/Engine/Source/Runtime/Apple/AGXRHI/Private/AGXRenderPass.cpp
--- a/Engine/Source/Runtime/Apple/AGXRHI/Private/AGXRenderPass.cpp
+++ b/Engine/Source/Runtime/Apple/AGXRHI/Private/AGXRenderPass.cpp
@@ -1136,22 +1136,22 @@ void FAGXRenderPass::CommitRenderResourceTables(void)
 	
 	State.CommitRenderResources(&CurrentEncoder);
 	
-	State.CommitResourceTable(EAGXShaderStages::Vertex, mtlpp::FunctionType::Vertex, CurrentEncoder);
+	State.CommitResourceTable(EAGXShaderStages::Vertex, MTLFunctionTypeVertex, CurrentEncoder);
 	
 	FAGXGraphicsPipelineState const* BoundShaderState = State.GetGraphicsPSO();
 	
 	if (BoundShaderState->VertexShader->SideTableBinding >= 0)
 	{
-		CurrentEncoder.SetShaderSideTable(mtlpp::FunctionType::Vertex, BoundShaderState->VertexShader->SideTableBinding);
+		CurrentEncoder.SetShaderSideTable(MTLFunctionTypeVertex, BoundShaderState->VertexShader->SideTableBinding);
 		State.SetShaderBuffer(EAGXShaderStages::Vertex, nil, nil, 0, 0, BoundShaderState->VertexShader->SideTableBinding, 0);
 	}
 	
 	if (IsValidRef(BoundShaderState->PixelShader))
 	{
-		State.CommitResourceTable(EAGXShaderStages::Pixel, mtlpp::FunctionType::Fragment, CurrentEncoder);
+		State.CommitResourceTable(EAGXShaderStages::Pixel, MTLFunctionTypeFragment, CurrentEncoder);
 		if (BoundShaderState->PixelShader->SideTableBinding >= 0)
 		{
-			CurrentEncoder.SetShaderSideTable(mtlpp::FunctionType::Fragment, BoundShaderState->PixelShader->SideTableBinding);
+			CurrentEncoder.SetShaderSideTable(MTLFunctionTypeFragment, BoundShaderState->PixelShader->SideTableBinding);
 			State.SetShaderBuffer(EAGXShaderStages::Pixel, nil, nil, 0, 0, BoundShaderState->PixelShader->SideTableBinding, 0);
 		}
 	}
@@ -1161,12 +1161,12 @@ void FAGXRenderPass::CommitDispatchResourceTables(void)
 {
 	State.CommitComputeResources(&CurrentEncoder);
 	
-	State.CommitResourceTable(EAGXShaderStages::Compute, mtlpp::FunctionType::Kernel, CurrentEncoder);
+	State.CommitResourceTable(EAGXShaderStages::Compute, MTLFunctionTypeKernel, CurrentEncoder);
 	
 	FAGXComputeShader const* ComputeShader = State.GetComputeShader();
 	if (ComputeShader->SideTableBinding >= 0)
 	{
-		CurrentEncoder.SetShaderSideTable(mtlpp::FunctionType::Kernel, ComputeShader->SideTableBinding);
+		CurrentEncoder.SetShaderSideTable(MTLFunctionTypeKernel, ComputeShader->SideTableBinding);
 		State.SetShaderBuffer(EAGXShaderStages::Compute, nil, nil, 0, 0, ComputeShader->SideTableBinding, 0);
 	}
 }
@@ -1175,12 +1175,12 @@ void FAGXRenderPass::CommitAsyncDispatchResourceTables(void)
 {
 	State.CommitComputeResources(&PrologueEncoder);
 	
-	State.CommitResourceTable(EAGXShaderStages::Compute, mtlpp::FunctionType::Kernel, PrologueEncoder);
+	State.CommitResourceTable(EAGXShaderStages::Compute, MTLFunctionTypeKernel, PrologueEncoder);
 	
 	FAGXComputeShader const* ComputeShader = State.GetComputeShader();
 	if (ComputeShader->SideTableBinding >= 0)
 	{
-		PrologueEncoder.SetShaderSideTable(mtlpp::FunctionType::Kernel, ComputeShader->SideTableBinding);
+		PrologueEncoder.SetShaderSideTable(MTLFunctionTypeKernel, ComputeShader->SideTableBinding);
 		State.SetShaderBuffer(EAGXShaderStages::Compute, nil, nil, 0, 0, ComputeShader->SideTableBinding, 0);
 	}
 }
diff --git a/Engine/Source/Runtime/Apple/AGXRHI/Private/AGXStateCache.cpp b/Engine/Source/Runtime/Apple/AGXRHI/Private/AGXStateCache.cpp
--- a/Engine/Source/Runtime/Apple/AGXRHI/Private/AGXStateCache.cpp
+++ b/Engine/Source/Runtime/Apple/AGXRHI/Private/AGXStateCache.cpp
@@ -166,7 +166,7 @@ void AGXSafeReleaseMetalRenderPassDescriptor(MTLRenderPassDescriptor* Desc)
 	}
 }
 
-FAGXStateCache::FAGXStateCache(bool const bInImmediate)
+FAGXStateCache::FAGXStateCache(bool bInImmediate)
 : DepthStore(MTLStoreActionUnknown)
 , StencilStore(MTLStoreActionUnknown)
 , VisibilityResults(nullptr)
@@ -246,7 +246,7 @@ FAGXStateCache::~FAGXStateCache()
 	VisibilityResults = nil;
 }
 
-void FAGXStateCache::Reset(void)
+void FAGXStateCache::Reset()
 {
 	SampleCount = 0;
 	
@@ -371,7 +371,7 @@ void FAGXStateCache::SetBlendFactor(FLinearColor const& InBlendFactor)
 	}
 }
 
-void FAGXStateCache::SetStencilRef(uint32 const InStencilRef)
+void FAGXStateCache::SetStencilRef(uint32 InStencilRef)
 {
 	if(StencilRef != InStencilRef)
 	{
@@ -422,7 +422,7 @@ void FAGXStateCache::SetComputeShader(FAGXComputeShader* InComputeShader)
 	}
 }
 
-bool FAGXStateCache::SetRenderPassInfo(FRHIRenderPassInfo const& InRenderTargets, FAGXQueryBuffer* QueryBuffer, bool const bRestart)
+bool FAGXStateCache::SetRenderPassInfo(FRHIRenderPassInfo const& InRenderTargets, FAGXQueryBuffer* QueryBuffer, bool bRestart)
 {
 	bool bNeedsSet = false;
 	
@@ -791,8 +791,8 @@ bool FAGXStateCache::SetRenderPassInfo(FRHIRenderPassInfo const& InRenderTargets
 				DepthClearValue = 1.0f;
 			}
 
-			bool const bCombinedDepthStencilUsingStencil = (DepthTexture && (MTLPixelFormat)DepthTexture.GetPixelFormat() != MTLPixelFormatDepth32Float && RenderPassInfo.DepthStencilRenderTarget.ExclusiveDepthStencil.IsUsingStencil());
-			bool const bUsingDepth = (RenderPassInfo.DepthStencilRenderTarget.ExclusiveDepthStencil.IsUsingDepth() || (bCombinedDepthStencilUsingStencil));
+			const bool bCombinedDepthStencilUsingStencil = (DepthTexture && (MTLPixelFormat)DepthTexture.GetPixelFormat() != MTLPixelFormatDepth32Float && RenderPassInfo.DepthStencilRenderTarget.ExclusiveDepthStencil.IsUsingStencil());
+			const bool bUsingDepth = (RenderPassInfo.DepthStencilRenderTarget.ExclusiveDepthStencil.IsUsingDepth() || (bCombinedDepthStencilUsingStencil));
 			if (DepthTexture && bUsingDepth)
 			{
 				MTLRenderPassDepthAttachmentDescriptor* DepthAttachment = [RenderPass depthAttachment];
@@ -882,8 +882,8 @@ bool FAGXStateCache::SetRenderPassInfo(FRHIRenderPassInfo const& InRenderTargets
             //if we're dealing with a samplecount mismatch we just bail on stencil entirely as stencil
             //doesn't have an autoresolve target to use.
 			
-			bool const bCombinedDepthStencilUsingDepth = (StencilTexture && (MTLPixelFormat)StencilTexture.GetPixelFormat() != MTLPixelFormatStencil8 && RenderPassInfo.DepthStencilRenderTarget.ExclusiveDepthStencil.IsUsingDepth());
-			bool const bUsingStencil = RenderPassInfo.DepthStencilRenderTarget.ExclusiveDepthStencil.IsUsingStencil() || (bCombinedDepthStencilUsingDepth);
+			const bool bCombinedDepthStencilUsingDepth = (StencilTexture && (MTLPixelFormat)StencilTexture.GetPixelFormat() != MTLPixelFormatStencil8 && RenderPassInfo.DepthStencilRenderTarget.ExclusiveDepthStencil.IsUsingDepth());
+			const bool bUsingStencil = RenderPassInfo.DepthStencilRenderTarget.ExclusiveDepthStencil.IsUsingStencil() || (bCombinedDepthStencilUsingDepth);
 			if (StencilTexture && bUsingStencil)
 			{
 				MTLRenderPassStencilAttachmentDescriptor* StencilAttachment = [RenderPass stencilAttachment];
@@ -985,14 +985,14 @@ bool FAGXStateCache::SetRenderPassInfo(FRHIRenderPassInfo const& InRenderTargets
 	return bNeedsSet;
 }
 
-void FAGXStateCache::InvalidateRenderTargets(void)
+void FAGXStateCache::InvalidateRenderTargets()
 {
 	bHasValidRenderTarget = false;
 	bHasValidColorTarget = false;
 	bIsRenderTargetActive = false;
 }
 
-void FAGXStateCache::SetRenderTargetsActive(bool const bActive)
+void FAGXStateCache::SetRenderTargetsActive(bool bActive)
 {
 	bIsRenderTargetActive = bActive;
 }
@@ -1087,7 +1087,7 @@ void FAGXStateCache::SetViewports(MTLViewport const InViewport[], uint32 Count)
 	ActiveViewports = Count;
 }
 
-void FAGXStateCache::SetVertexStream(uint32 const Index, FAGXBuffer* Buffer, FAGXBufferData* Bytes, uint32 const Offset, uint32 const Length)
+void FAGXStateCache::SetVertexStream(uint32 Index, FAGXBuffer* Buffer, FAGXBufferData* Bytes, uint32 Offset, uint32 Length)
 {
 	check(Index < MaxVertexElementCount);
 	check(UNREAL_TO_METAL_BUFFER_INDEX(Index) < MaxMetalStreams);
@@ -1107,7 +1107,7 @@ void FAGXStateCache::SetVertexStream(uint32 const Index, FAGXBuffer* Buffer, FAG
 	SetShaderBuffer(EAGXShaderStages::Vertex, VertexBuffers[Index].Buffer, Bytes, Offset, Length, UNREAL_TO_METAL_BUFFER_INDEX(Index), MTLResourceUsageRead);
 }
 
-uint32 FAGXStateCache::GetVertexBufferSize(uint32 const Index)
+uint32 FAGXStateCache::GetVertexBufferSize(uint32 Index)
 {
 	check(Index < MaxVertexElementCount);
 	check(UNREAL_TO_METAL_BUFFER_INDEX(Index) < MaxMetalStreams);
@@ -1167,7 +1167,7 @@ EPrimitiveType FAGXStateCache::GetPrimitiveType()
 	return GraphicsPSO->GetPrimitiveType();
 }
 
-void FAGXStateCache::BindUniformBuffer(EAGXShaderStages const Freq, uint32 const BufferIndex, FRHIUniformBuffer* BufferRHI)
+void FAGXStateCache::BindUniformBuffer(EAGXShaderStages Freq, uint32 BufferIndex, FRHIUniformBuffer* BufferRHI)
 {
 	check(BufferIndex < ML_MaxBuffers);
 	if (BoundUniformBuffers[Freq][BufferIndex] != BufferRHI)
@@ -1178,12 +1178,12 @@ void FAGXStateCache::BindUniformBuffer(EAGXShaderStages const Freq, uint32 const
 	}
 }
 
-void FAGXStateCache::SetDirtyUniformBuffers(EAGXShaderStages const Freq, uint32 const Dirty)
+void FAGXStateCache::SetDirtyUniformBuffers(EAGXShaderStages Freq, uint32 Dirty)
 {
 	DirtyUniformBuffers[Freq] = Dirty;
 }
 
-void FAGXStateCache::SetVisibilityResultMode(mtlpp::VisibilityResultMode const Mode, NSUInteger const Offset)
+void FAGXStateCache::SetVisibilityResultMode(mtlpp::VisibilityResultMode Mode, NSUInteger Offset)
 {
 	if (VisibilityMode != Mode || VisibilityOffset != Offset)
 	{
@@ -1279,7 +1279,7 @@ bool FAGXStateCache::NeedsToSetRenderTarget(const FRHIRenderPassInfo& InRenderPa
 					FAGXSurface& Surface = *AGXGetMetalSurfaceFromRHITexture(RenderPassInfo.DepthStencilRenderTarget.DepthStencilTarget);
 					
 					const uint32 DepthSampleCount = (Surface.MSAATexture ? Surface.MSAATexture.GetSampleCount() : Surface.Texture.GetSampleCount());
-					bool const bDepthStencilSampleCountMismatchFixup = (SampleCount != DepthSampleCount);
+					const bool bDepthStencilSampleCountMismatchFixup = (SampleCount != DepthSampleCount);
 
 					ERenderTargetStoreAction HighLevelStoreAction = (Surface.MSAATexture && !bDepthStencilSampleCountMismatchFixup) ? ERenderTargetStoreAction::EMultisampleResolve : GetStoreAction(GetDepthActions(RenderPassInfo.DepthStencilRenderTarget.Action));
 					
@@ -1727,7 +1727,7 @@ void FAGXStateCache::CommitComputeResources(FAGXCommandEncoder* Compute)
 	GetShaderParameters(EAGXShaderStages::Compute).CommitPackedGlobals(this, Compute, EAGXShaderStages::Compute, ComputeShader->Bindings);
 }
 
-bool FAGXStateCache::PrepareToRestart(bool const bCurrentApplied)
+bool FAGXStateCache::PrepareToRestart(bool bCurrentApplied)
 {
 	if(CanRestartRenderPass())
 	{
@@ -1806,7 +1806,7 @@ bool FAGXStateCache::PrepareToRestart(bool const bCurrentApplied)
 	}
 }
 
-void FAGXStateCache::SetStateDirty(void)
+void FAGXStateCache::SetStateDirty()
 {	
 	RasterBits = UINT32_MAX;
     PipelineBits = EAGXPipelineFlagMask;
@@ -1818,12 +1818,12 @@ void FAGXStateCache::SetStateDirty(void)
 	}
 }
 
-void FAGXStateCache::SetShaderBufferDirty(EAGXShaderStages const Frequency, NSUInteger const Index)
+void FAGXStateCache::SetShaderBufferDirty(EAGXShaderStages Frequency, NSUInteger Index)
 {
 	ShaderBuffers[Frequency].Bound |= (1 << Index);
 }
 
-void FAGXStateCache::SetRenderStoreActions(FAGXCommandEncoder& CommandEncoder, bool const bConditionalSwitch)
+void FAGXStateCache::SetRenderStoreActions(FAGXCommandEncoder& CommandEncoder, bool bConditionalSwitch)
 {
 	check(CommandEncoder.IsRenderCommandEncoderActive())
 	{
@@ -2149,7 +2149,7 @@ void FAGXStateCache::SetComputePipelineState(FAGXCommandEncoder& CommandEncoder)
 	}
 }
 
-void FAGXStateCache::CommitResourceTable(EAGXShaderStages const Frequency, mtlpp::FunctionType const Type, FAGXCommandEncoder& CommandEncoder)
+void FAGXStateCache::CommitResourceTable(EAGXShaderStages Frequency, MTLFunctionType Type, FAGXCommandEncoder& CommandEncoder)
 {
 	FAGXBufferBindings& BufferBindings = ShaderBuffers[Frequency];
 	while(BufferBindings.Bound)
diff --git a/Engine/Source/Runtime/Apple/AGXRHI/Private/AGXStateCache.h b/Engine/Source/Runtime/Apple/AGXRHI/Private/AGXStateCache.h
--- a/Engine/Source/Runtime/Apple/AGXRHI/Private/AGXStateCache.h
+++ b/Engine/Source/Runtime/Apple/AGXRHI/Private/AGXStateCache.h
@@ -39,32 +39,32 @@ enum EAGXRenderFlags
 class FAGXStateCache
 {
 public:
-	FAGXStateCache(bool const bInImmediate);
+	FAGXStateCache(bool bInImmediate);
 	~FAGXStateCache();
 	
 	/** Reset cached state for reuse */
-	void Reset(void);
+	void Reset();
 
 	void SetScissorRect(bool bEnable, MTLScissorRect const& Rect);
 	void SetBlendFactor(FLinearColor const& InBlendFactor);
-	void SetStencilRef(uint32 const InStencilRef);
+	void SetStencilRef(uint32 InStencilRef);
 	void SetComputeShader(FAGXComputeShader* InComputeShader);
-	bool SetRenderPassInfo(FRHIRenderPassInfo const& InRenderTargets, FAGXQueryBuffer* QueryBuffer, bool const bRestart);
-	void InvalidateRenderTargets(void);
-	void SetRenderTargetsActive(bool const bActive);
+	bool SetRenderPassInfo(FRHIRenderPassInfo const& InRenderTargets, FAGXQueryBuffer* QueryBuffer, bool bRestart);
+	void InvalidateRenderTargets();
+	void SetRenderTargetsActive(bool bActive);
 	void SetViewport(MTLViewport const& InViewport);
 	void SetViewports(MTLViewport const InViewport[], uint32 Count);
-	void SetVertexStream(uint32 const Index, FAGXBuffer* Buffer, FAGXBufferData* Bytes, uint32 const Offset, uint32 const Length);
+	void SetVertexStream(uint32 Index, FAGXBuffer* Buffer, FAGXBufferData* Bytes, uint32 Offset, uint32 Length);
 	void SetGraphicsPipelineState(FAGXGraphicsPipelineState* State);
-	void BindUniformBuffer(EAGXShaderStages const Freq, uint32 const BufferIndex, FRHIUniformBuffer* BufferRHI);
-	void SetDirtyUniformBuffers(EAGXShaderStages const Freq, uint32 const Dirty);
+	void BindUniformBuffer(EAGXShaderStages Freq, uint32 BufferIndex, FRHIUniformBuffer* BufferRHI);
+	void SetDirtyUniformBuffers(EAGXShaderStages Freq, uint32 Dirty);
 	
 	/*
 	 * Monitor if samples pass the depth and stencil tests.
 	 * @param Mode Controls if the counter is disabled or moniters passing samples.
 	 * @param Offset The offset relative to the occlusion query buffer provided when the command encoder was created.  offset must be a multiple of 8.
 	 */
-	void SetVisibilityResultMode(mtlpp::VisibilityResultMode const Mode, NSUInteger const Offset);
+	void SetVisibilityResultMode(mtlpp::VisibilityResultMode Mode, NSUInteger Offset);
 	
 #pragma mark - Public Shader Resource Mutators -
 	/*
@@ -101,11 +101,11 @@ public:
 	
 	void SetShaderUnorderedAccessView(EAGXShaderStages ShaderStage, uint32 BindIndex, FAGXUnorderedAccessView* RESTRICT UAV);
 
-	void SetStateDirty(void);
+	void SetStateDirty();
 	
-	void SetShaderBufferDirty(EAGXShaderStages const Frequency, NSUInteger const Index);
+	void SetShaderBufferDirty(EAGXShaderStages Frequency, NSUInteger Index);
 	
-	void SetRenderStoreActions(FAGXCommandEncoder& CommandEncoder, bool const bConditionalSwitch);
+	void SetRenderStoreActions(FAGXCommandEncoder& CommandEncoder, bool bConditionalSwitch);
 	
 	void SetRenderState(FAGXCommandEncoder& CommandEncoder, FAGXCommandEncoder* PrologueEncoder);
 
@@ -113,11 +113,11 @@ public:
 
 	void CommitComputeResources(FAGXCommandEncoder* Compute);
 	
-	void CommitResourceTable(EAGXShaderStages const Frequency, mtlpp::FunctionType const Type, FAGXCommandEncoder& CommandEncoder);
+	void CommitResourceTable(EAGXShaderStages Frequency, MTLFunctionType Type, FAGXCommandEncoder& CommandEncoder);
 	
-	bool PrepareToRestart(bool const bCurrentApplied);
+	bool PrepareToRestart(bool bCurrentApplied);
 	
-	FAGXShaderParameterCache& GetShaderParameters(EAGXShaderStages const Stage) { return ShaderParameters[Stage]; }
+	FAGXShaderParameterCache& GetShaderParameters(EAGXShaderStages Stage) { return ShaderParameters[Stage]; }
 	FLinearColor const& GetBlendFactor() const { return BlendFactor; }
 	uint32 GetStencilRef() const { return StencilRef; }
 	FAGXDepthStencilState* GetDepthStencilState() const { return DepthStencilState; }
@@ -129,18 +129,18 @@ public:
 	int32 GetNumRenderTargets() { return bHasValidColorTarget ? RenderPassInfo.GetNumColorRenderTargets() : -1; }
 	bool GetHasValidRenderTarget() const { return bHasValidRenderTarget; }
 	bool GetHasValidColorTarget() const { return bHasValidColorTarget; }
-	MTLViewport const& GetViewport(uint32 const Index) const { check(Index < ML_MaxViewports); return Viewport[Index]; }
-	uint32 GetVertexBufferSize(uint32 const Index);
+	MTLViewport const& GetViewport(uint32 Index) const { check(Index < ML_MaxViewports); return Viewport[Index]; }
+	uint32 GetVertexBufferSize(uint32 Index);
 	uint32 GetRenderTargetArraySize() const { return RenderTargetArraySize; }
-	const FRHIUniformBuffer** GetBoundUniformBuffers(EAGXShaderStages const Freq) { return (const FRHIUniformBuffer**)&BoundUniformBuffers[Freq][0]; }
-	uint32 GetDirtyUniformBuffers(EAGXShaderStages const Freq) const { return DirtyUniformBuffers[Freq]; }
+	const FRHIUniformBuffer** GetBoundUniformBuffers(EAGXShaderStages Freq) { return (const FRHIUniformBuffer**)&BoundUniformBuffers[Freq][0]; }
+	uint32 GetDirtyUniformBuffers(EAGXShaderStages Freq) const { return DirtyUniformBuffers[Freq]; }
 	FAGXQueryBuffer* GetVisibilityResultsBuffer() const { return VisibilityResults; }
 	bool GetScissorRectEnabled() const { return bScissorRectEnabled; }
 	bool NeedsToSetRenderTarget(const FRHIRenderPassInfo& RenderPassInfo);
 	bool HasValidDepthStencilSurface() const { return IsValidRef(DepthStencilSurface); }
     bool CanRestartRenderPass() const { return bCanRestartRenderPass; }
-	MTLRenderPassDescriptor* GetRenderPassDescriptor(void) const { return RenderPassDesc; }
-	uint32 GetSampleCount(void) const { return SampleCount; }
+	MTLRenderPassDescriptor* GetRenderPassDescriptor() const { return RenderPassDesc; }
+	uint32 GetSampleCount() const { return SampleCount; }
     bool IsLinearBuffer(EAGXShaderStages ShaderStage, uint32 BindIndex);
 	FAGXShaderPipeline* GetPipelineState() const;
 	EPrimitiveType GetPrimitiveType();
@@ -148,7 +148,7 @@ public:
 	uint32 GetVisibilityResultOffset() { return VisibilityOffset; }
 	
 	FTexture2DRHIRef CreateFallbackDepthStencilSurface(uint32 Width, uint32 Height);
-	bool GetFallbackDepthStencilBound(void) const { return bFallbackDepthStencilBound; }
+	bool GetFallbackDepthStencilBound() const { return bFallbackDepthStencilBound; }
 	
 	void SetRenderPipelineState(FAGXCommandEncoder& CommandEncoder, FAGXCommandEncoder* PrologueEncoder);
     void SetComputePipelineState(FAGXCommandEncoder& CommandEncoder);
