AGXRHI: Replace more mltpp types with native Metal types

- mtlpp::origin
- mltpp::size
- mtlpp::VisibilityMode
- mtlpp::Region
- mtlpp::Index
- mtlpp::TriangleFillMode
- mtlpp::DeptClipMode
- mtlpp::LanguageVersion
- mtlpp::Function::Type

Also, fix-up some use of const and (void)

#rnx

diff --git a/Engine/Source/Runtime/Apple/AGXRHI/Private/AGXCommandEncoder.cpp b/Engine/Source/Runtime/Apple/AGXRHI/Private/AGXCommandEncoder.cpp
--- a/Engine/Source/Runtime/Apple/AGXRHI/Private/AGXCommandEncoder.cpp
+++ b/Engine/Source/Runtime/Apple/AGXRHI/Private/AGXCommandEncoder.cpp
@@ -57,7 +57,7 @@ FAGXCommandEncoder::FAGXCommandEncoder(FAGXCommandList& CmdList, EAGXCommandEnco
 	StencilStoreAction = MTLStoreActionUnknown;
 }
 
-FAGXCommandEncoder::~FAGXCommandEncoder(void)
+FAGXCommandEncoder::~FAGXCommandEncoder()
 {
 	if(CommandBuffer)
 	{
@@ -94,7 +94,7 @@ FAGXCommandEncoder::~FAGXCommandEncoder(void)
 	}
 }
 
-void FAGXCommandEncoder::Reset(void)
+void FAGXCommandEncoder::Reset()
 {
     check(IsRenderCommandEncoderActive() == false
           && IsComputeCommandEncoderActive() == false
@@ -131,7 +131,7 @@ void FAGXCommandEncoder::Reset(void)
 	[DebugGroups removeAllObjects];
 }
 
-void FAGXCommandEncoder::ResetLive(void)
+void FAGXCommandEncoder::ResetLive()
 {
 	for (uint32 Frequency = 0; Frequency < NUM_SHADER_FREQUENCIES; Frequency++)
 	{
@@ -175,7 +175,7 @@ void FAGXCommandEncoder::ResetLive(void)
 
 #pragma mark - Public Command Buffer Mutators -
 
-void FAGXCommandEncoder::StartCommandBuffer(void)
+void FAGXCommandEncoder::StartCommandBuffer()
 {
 	check(!CommandBuffer || EncoderNum == 0);
 	check(IsRenderCommandEncoderActive() == false
@@ -202,15 +202,15 @@ void FAGXCommandEncoder::StartCommandBuffer(void)
 	}
 }
 	
-void FAGXCommandEncoder::CommitCommandBuffer(uint32 const Flags)
+void FAGXCommandEncoder::CommitCommandBuffer(uint32 Flags)
 {
 	check(CommandBuffer);
 	check(IsRenderCommandEncoderActive() == false
           && IsComputeCommandEncoderActive() == false
           && IsBlitCommandEncoderActive() == false);
 
-	bool const bWait = (Flags & EAGXSubmitFlagsWaitOnCommandBuffer);
-	bool const bIsLastCommandBuffer = (Flags & EAGXSubmitFlagsLastCommandBuffer);
+	const bool bWait = (Flags & EAGXSubmitFlagsWaitOnCommandBuffer);
+	const bool bIsLastCommandBuffer = (Flags & EAGXSubmitFlagsLastCommandBuffer);
 	
 	if (EncoderNum == 0 && !bWait && !(Flags & EAGXSubmitFlagsForce))
 	{
@@ -291,12 +291,12 @@ bool FAGXCommandEncoder::IsBlitCommandEncoderActive() const
 	return (BlitCommandEncoder != nil);
 }
 
-bool FAGXCommandEncoder::IsImmediate(void) const
+bool FAGXCommandEncoder::IsImmediate() const
 {
 	return CommandList.IsImmediate();
 }
 
-bool FAGXCommandEncoder::IsParallel(void) const
+bool FAGXCommandEncoder::IsParallel() const
 {
 	return CommandList.IsParallel() && (Type == EAGXCommandEncoderCurrent);
 }
@@ -376,7 +376,7 @@ void FAGXCommandEncoder::BeginParallelRenderCommandEncoding(uint32 NumChildren)
 	}
 }
 
-void FAGXCommandEncoder::BeginRenderCommandEncoding(void)
+void FAGXCommandEncoder::BeginRenderCommandEncoding()
 {
 	check(RenderPassDesc);
 	check(CommandList.IsParallel() || CommandBuffer);
@@ -442,7 +442,7 @@ void FAGXCommandEncoder::BeginComputeCommandEncoding(MTLDispatchType DispatchTyp
 	}
 }
 
-void FAGXCommandEncoder::BeginBlitCommandEncoding(void)
+void FAGXCommandEncoder::BeginBlitCommandEncoding()
 {
 	check(CommandBuffer);
 	check(IsRenderCommandEncoderActive() == false && IsComputeCommandEncoderActive() == false && IsBlitCommandEncoderActive() == false);
@@ -467,7 +467,7 @@ void FAGXCommandEncoder::BeginBlitCommandEncoding(void)
 	}
 }
 
-void FAGXCommandEncoder::EndEncoding(void)
+void FAGXCommandEncoder::EndEncoding()
 {
 	@autoreleasepool
 	{
@@ -686,7 +686,7 @@ void FAGXCommandEncoder::PushDebugGroup(ns::String const& String)
 	}
 }
 
-void FAGXCommandEncoder::PopDebugGroup(void)
+void FAGXCommandEncoder::PopDebugGroup()
 {
 	if (DebugGroups.count > 0)
 	{
@@ -711,7 +711,7 @@ void FAGXCommandEncoder::PopDebugGroup(void)
 }
 
 #if ENABLE_METAL_GPUPROFILE
-FAGXCommandBufferStats* FAGXCommandEncoder::GetCommandBufferStats(void)
+FAGXCommandBufferStats* FAGXCommandEncoder::GetCommandBufferStats()
 {
 	return CommandBufferStats;
 }
@@ -806,7 +806,7 @@ void FAGXCommandEncoder::SetCullMode(MTLCullMode InCullMode)
 	}
 }
 
-void FAGXCommandEncoder::SetDepthBias(float const InDepthBias, float const InSlopeScale, float const InClamp)
+void FAGXCommandEncoder::SetDepthBias(float InDepthBias, float InSlopeScale, float InClamp)
 {
     check (RenderCommandEncoder);
 	{
@@ -831,23 +831,23 @@ void FAGXCommandEncoder::SetScissorRect(MTLScissorRect const Rect[], uint32 NumA
 #endif
 }
 
-void FAGXCommandEncoder::SetTriangleFillMode(mtlpp::TriangleFillMode const InFillMode)
+void FAGXCommandEncoder::SetTriangleFillMode(MTLTriangleFillMode InFillMode)
 {
     check(RenderCommandEncoder);
 	{
-		[RenderCommandEncoder setTriangleFillMode:(MTLTriangleFillMode)InFillMode];
+		[RenderCommandEncoder setTriangleFillMode:InFillMode];
 	}
 }
 
-void FAGXCommandEncoder::SetDepthClipMode(mtlpp::DepthClipMode const InDepthClipMode)
+void FAGXCommandEncoder::SetDepthClipMode(MTLDepthClipMode InDepthClipMode)
 {
 	check(RenderCommandEncoder);
 	{
-		[RenderCommandEncoder setDepthClipMode:(MTLDepthClipMode)InDepthClipMode];
+		[RenderCommandEncoder setDepthClipMode:InDepthClipMode];
 	}
 }
 
-void FAGXCommandEncoder::SetBlendColor(float const Red, float const Green, float const Blue, float const Alpha)
+void FAGXCommandEncoder::SetBlendColor(float Red, float Green, float Blue, float Alpha)
 {
 	check(RenderCommandEncoder);
 	{
@@ -863,7 +863,7 @@ void FAGXCommandEncoder::SetDepthStencilState(id<MTLDepthStencilState> InDepthSt
 	}
 }
 
-void FAGXCommandEncoder::SetStencilReferenceValue(uint32 const ReferenceValue)
+void FAGXCommandEncoder::SetStencilReferenceValue(uint32 ReferenceValue)
 {
     check (RenderCommandEncoder);
 	{
@@ -871,12 +871,12 @@ void FAGXCommandEncoder::SetStencilReferenceValue(uint32 const ReferenceValue)
 	}
 }
 
-void FAGXCommandEncoder::SetVisibilityResultMode(mtlpp::VisibilityResultMode const Mode, NSUInteger const Offset)
+void FAGXCommandEncoder::SetVisibilityResultMode(MTLVisibilityResultMode Mode, NSUInteger Offset)
 {
     check (RenderCommandEncoder);
 	{
-		check(Mode == mtlpp::VisibilityResultMode::Disabled || [RenderPassDesc visibilityResultBuffer]);
-		[RenderCommandEncoder setVisibilityResultMode:(MTLVisibilityResultMode)Mode offset:Offset];
+		check(Mode == MTLVisibilityResultModeDisabled || [RenderPassDesc visibilityResultBuffer]);
+		[RenderCommandEncoder setVisibilityResultMode:Mode offset:Offset];
 	}
 }
 	
@@ -1089,7 +1089,7 @@ void FAGXCommandEncoder::SetShaderSamplerState(MTLFunctionType FunctionType, id<
 	}
 }
 
-void FAGXCommandEncoder::SetShaderSideTable(MTLFunctionType const FunctionType, NSUInteger const Index)
+void FAGXCommandEncoder::SetShaderSideTable(MTLFunctionType FunctionType, NSUInteger Index)
 {
 	if (Index < ML_MaxBuffers)
 	{
@@ -1113,7 +1113,7 @@ void FAGXCommandEncoder::SetComputePipelineState(FAGXShaderPipeline* State)
 
 #pragma mark - Public Ring-Buffer Accessor -
 	
-FAGXSubBufferRing& FAGXCommandEncoder::GetRingBuffer(void)
+FAGXSubBufferRing& FAGXCommandEncoder::GetRingBuffer()
 {
 	return RingBuffer;
 }
diff --git a/Engine/Source/Runtime/Apple/AGXRHI/Private/AGXCommandEncoder.h b/Engine/Source/Runtime/Apple/AGXRHI/Private/AGXCommandEncoder.h
--- a/Engine/Source/Runtime/Apple/AGXRHI/Private/AGXCommandEncoder.h
+++ b/Engine/Source/Runtime/Apple/AGXRHI/Private/AGXCommandEncoder.h
@@ -60,25 +60,25 @@ public:
 	FAGXCommandEncoder(FAGXCommandList& CmdList, EAGXCommandEncoderType Type);
 	
 	/** Destructor */
-	~FAGXCommandEncoder(void);
+	~FAGXCommandEncoder();
 	
 	/** Reset cached state for reuse */
-	void Reset(void);
+	void Reset();
 	
 	/** Reset cached state for reuse while in rendering */
-	void ResetLive(void);
+	void ResetLive();
 	
 	/**
 	 * Start encoding to CommandBuffer. It is an error to call this with any outstanding command encoders or current command buffer.
 	 * Instead call EndEncoding & CommitCommandBuffer before calling this.
 	 */
-	void StartCommandBuffer(void);
+	void StartCommandBuffer();
 	
 	/**
 	 * Commit the existing command buffer if there is one & optionally waiting for completion, if there isn't a current command buffer this is a no-op.
 	 * @param Flags Flags to control commit behaviour.
  	 */
-	void CommitCommandBuffer(uint32 const Flags);
+	void CommitCommandBuffer(uint32 Flags);
 
 	/** @returns the current command buffer */
 	mtlpp::CommandBuffer& GetCommandBuffer() { return CommandBuffer; }
@@ -105,13 +105,13 @@ public:
 	 * True iff the command-encoder submits immediately to the command-queue, false if it performs any buffering.
 	 * @returns True iff the command-list submits immediately to the command-queue, false if it performs any buffering.
 	 */
-	bool IsImmediate(void) const;
+	bool IsImmediate() const;
 	
 	/**
 	 * True iff the command-encoder encodes only to a child of a parallel render command encoder, false if it is standalone.
 	 * @returns True iff the command-encoder encodes only to a child of a parallel render command encoder, false if it is standalone.
 	 */
-	bool IsParallel(void) const;
+	bool IsParallel() const;
 
 	/** @returns True if and only if there is valid render pass descriptor set on the encoder, otherwise false. */
 	bool IsRenderPassDescriptorValid() const;
@@ -135,7 +135,7 @@ public:
 	id<MTLBlitCommandEncoder> GetBlitCommandEncoder() const;
 	
 	/** @returns The number of encoded passes in the command buffer. */
-	uint32 NumEncodedPasses(void) const { return EncoderNum; }
+	uint32 NumEncodedPasses() const { return EncoderNum; }
 	
 	/**
  	 * Begins encoding rendering commands into the current command buffer. No other encoder may be active & the MTLRenderPassDescriptor must previously have been set.
@@ -146,16 +146,16 @@ public:
 	/**
  	 * Begins encoding rendering commands into the current command buffer. No other encoder may be active & the MTLRenderPassDescriptor must previously have been set.
 	 */
-	void BeginRenderCommandEncoding(void);
+	void BeginRenderCommandEncoding();
 	
 	/** Begins encoding compute commands into the current command buffer. No other encoder may be active. */
 	void BeginComputeCommandEncoding(MTLDispatchType Type = MTLDispatchTypeSerial);
 	
 	/** Begins encoding blit commands into the current command buffer. No other encoder may be active. */
-	void BeginBlitCommandEncoding(void);
+	void BeginBlitCommandEncoding();
 	
 	/** Declare that all command generation from this encoder is complete, and detach from the MTLCommandBuffer if there is an encoder active or does nothing if there isn't. */
-	void EndEncoding(void);
+	void EndEncoding();
 	
 	/** Initialises a fence for the current command-buffer optionally adding a command-buffer completion handler to the command-buffer */
 	void InsertCommandBufferFence(FAGXCommandBufferFence& Fence, mtlpp::CommandBufferHandler Handler);
@@ -176,11 +176,11 @@ public:
 	void PushDebugGroup(ns::String const& String);
 	
 	/* Pop the latest named string off of the stack. */
-	void PopDebugGroup(void);
+	void PopDebugGroup();
 	
 #if ENABLE_METAL_GPUPROFILE
 	/* Get the command-buffer stats object. */
-	FAGXCommandBufferStats* GetCommandBufferStats(void);
+	FAGXCommandBufferStats* GetCommandBufferStats();
 #endif
 
 	/**
@@ -228,7 +228,7 @@ public:
 	 * @param SlopeScale The slope-scale to apply.
 	 * @param Clamp The value to clamp to.
 	 */
-	void SetDepthBias(float const DepthBias, float const SlopeScale, float const Clamp);
+	void SetDepthBias(float DepthBias, float SlopeScale, float Clamp);
 	
 	/*
 	 * Specifies a rectangle for a fragment scissor test.  All fragments outside of this rectangle are discarded.
@@ -241,13 +241,13 @@ public:
 	 * Set how to rasterize triangle and triangle strip primitives.
 	 * @param FillMode The fill mode.
 	 */
-	void SetTriangleFillMode(mtlpp::TriangleFillMode const FillMode);
+	void SetTriangleFillMode(MTLTriangleFillMode InFillMode);
 
 	/*
 	 * Set wether to clip or clamp triangles based on depth.
 	 * @param FillMode The fill mode.
 	 */
-	void SetDepthClipMode(mtlpp::DepthClipMode DepthClipMode);
+	void SetDepthClipMode(MTLDepthClipMode InDepthClipMode);
 	
 	/*
 	 * Set the constant blend color used across all blending on all render targets
@@ -256,7 +256,7 @@ public:
 	 * @param Blue The value for the blue channel in 0-1.
 	 * @param Alpha The value for the alpha channel in 0-1.
 	 */
-	void SetBlendColor(float const Red, float const Green, float const Blue, float const Alpha);
+	void SetBlendColor(float Red, float Green, float Blue, float Alpha);
 	
 	/*
 	 * Set the DepthStencil state object.
@@ -268,14 +268,14 @@ public:
 	 * Set the stencil reference value for both the back and front stencil buffers.
 	 * @param ReferenceValue The stencil ref value to use.
 	 */
-	void SetStencilReferenceValue(uint32 const ReferenceValue);
+	void SetStencilReferenceValue(uint32 ReferenceValue);
 	
 	/*
 	 * Monitor if samples pass the depth and stencil tests.
 	 * @param Mode Controls if the counter is disabled or moniters passing samples.
 	 * @param Offset The offset relative to the occlusion query buffer provided when the command encoder was created.  offset must be a multiple of 8.
 	 */
-	void SetVisibilityResultMode(mtlpp::VisibilityResultMode const Mode, NSUInteger const Offset);
+	void SetVisibilityResultMode(MTLVisibilityResultMode Mode, NSUInteger Offset);
 	
 	/*
 	 * Set a global buffer for the specified shader frequency at the given bind point index.
@@ -348,7 +348,7 @@ public:
 	 * Get the internal ring-buffer used for temporary allocations.
 	 * @returns The temporary allocation buffer for this command-encoder.
 	 */
-	FAGXSubBufferRing& GetRingBuffer(void);
+	FAGXSubBufferRing& GetRingBuffer();
 	
 	/*
 	 * Returns True if the Resource has been bound to a command encoder, otherwise false.  History will be cleared after a commit operation
diff --git a/Engine/Source/Runtime/Apple/AGXRHI/Private/AGXContext.cpp b/Engine/Source/Runtime/Apple/AGXRHI/Private/AGXContext.cpp
--- a/Engine/Source/Runtime/Apple/AGXRHI/Private/AGXContext.cpp
+++ b/Engine/Source/Runtime/Apple/AGXRHI/Private/AGXContext.cpp
@@ -1379,7 +1379,7 @@ void FAGXContext::DrawPrimitiveIndirect(uint32 PrimitiveType, FAGXVertexBuffer*
 	RenderPass.DrawPrimitiveIndirect(PrimitiveType, VertexBuffer, ArgumentOffset);
 }
 
-void FAGXContext::DrawIndexedPrimitive(FAGXBuffer const& IndexBuffer, uint32 IndexStride, mtlpp::IndexType IndexType, uint32 PrimitiveType, int32 BaseVertexIndex, uint32 FirstInstance, uint32 NumVertices, uint32 StartIndex, uint32 NumPrimitives, uint32 NumInstances)
+void FAGXContext::DrawIndexedPrimitive(FAGXBuffer const& IndexBuffer, uint32 IndexStride, MTLIndexType IndexType, uint32 PrimitiveType, int32 BaseVertexIndex, uint32 FirstInstance, uint32 NumVertices, uint32 StartIndex, uint32 NumPrimitives, uint32 NumInstances)
 {
 	// finalize any pending state
 	if(!PrepareToDraw(PrimitiveType))
@@ -1412,17 +1412,17 @@ void FAGXContext::DrawIndexedPrimitiveIndirect(uint32 PrimitiveType,FAGXIndexBuf
 	RenderPass.DrawIndexedPrimitiveIndirect(PrimitiveType, IndexBuffer, VertexBuffer, ArgumentOffset);
 }
 
-void FAGXContext::CopyFromTextureToBuffer(FAGXTexture const& Texture, uint32 sourceSlice, uint32 sourceLevel, mtlpp::Origin sourceOrigin, mtlpp::Size sourceSize, FAGXBuffer const& toBuffer, uint32 destinationOffset, uint32 destinationBytesPerRow, uint32 destinationBytesPerImage, mtlpp::BlitOption options)
+void FAGXContext::CopyFromTextureToBuffer(FAGXTexture const& Texture, uint32 sourceSlice, uint32 sourceLevel, MTLOrigin sourceOrigin, MTLSize sourceSize, FAGXBuffer const& toBuffer, uint32 destinationOffset, uint32 destinationBytesPerRow, uint32 destinationBytesPerImage, mtlpp::BlitOption options)
 {
 	RenderPass.CopyFromTextureToBuffer(Texture, sourceSlice, sourceLevel, sourceOrigin, sourceSize, toBuffer, destinationOffset, destinationBytesPerRow, destinationBytesPerImage, options);
 }
 
-void FAGXContext::CopyFromBufferToTexture(FAGXBuffer const& Buffer, uint32 sourceOffset, uint32 sourceBytesPerRow, uint32 sourceBytesPerImage, mtlpp::Size sourceSize, FAGXTexture const& toTexture, uint32 destinationSlice, uint32 destinationLevel, mtlpp::Origin destinationOrigin, mtlpp::BlitOption options)
+void FAGXContext::CopyFromBufferToTexture(FAGXBuffer const& Buffer, uint32 sourceOffset, uint32 sourceBytesPerRow, uint32 sourceBytesPerImage, MTLSize sourceSize, FAGXTexture const& toTexture, uint32 destinationSlice, uint32 destinationLevel, MTLOrigin destinationOrigin, mtlpp::BlitOption options)
 {
 	RenderPass.CopyFromBufferToTexture(Buffer, sourceOffset, sourceBytesPerRow, sourceBytesPerImage, sourceSize, toTexture, destinationSlice, destinationLevel, destinationOrigin, options);
 }
 
-void FAGXContext::CopyFromTextureToTexture(FAGXTexture const& Texture, uint32 sourceSlice, uint32 sourceLevel, mtlpp::Origin sourceOrigin, mtlpp::Size sourceSize, FAGXTexture const& toTexture, uint32 destinationSlice, uint32 destinationLevel, mtlpp::Origin destinationOrigin)
+void FAGXContext::CopyFromTextureToTexture(FAGXTexture const& Texture, uint32 sourceSlice, uint32 sourceLevel, MTLOrigin sourceOrigin, MTLSize sourceSize, FAGXTexture const& toTexture, uint32 destinationSlice, uint32 destinationLevel, MTLOrigin destinationOrigin)
 {
 	RenderPass.CopyFromTextureToTexture(Texture, sourceSlice, sourceLevel, sourceOrigin, sourceSize, toTexture, destinationSlice, destinationLevel, destinationOrigin);
 }
@@ -1432,12 +1432,12 @@ void FAGXContext::CopyFromBufferToBuffer(FAGXBuffer const& SourceBuffer, NSUInte
 	RenderPass.CopyFromBufferToBuffer(SourceBuffer, SourceOffset, DestinationBuffer, DestinationOffset, Size);
 }
 
-bool FAGXContext::AsyncCopyFromBufferToTexture(FAGXBuffer const& Buffer, uint32 sourceOffset, uint32 sourceBytesPerRow, uint32 sourceBytesPerImage, mtlpp::Size sourceSize, FAGXTexture const& toTexture, uint32 destinationSlice, uint32 destinationLevel, mtlpp::Origin destinationOrigin, mtlpp::BlitOption options)
+bool FAGXContext::AsyncCopyFromBufferToTexture(FAGXBuffer const& Buffer, uint32 sourceOffset, uint32 sourceBytesPerRow, uint32 sourceBytesPerImage, MTLSize sourceSize, FAGXTexture const& toTexture, uint32 destinationSlice, uint32 destinationLevel, MTLOrigin destinationOrigin, mtlpp::BlitOption options)
 {
 	return RenderPass.AsyncCopyFromBufferToTexture(Buffer, sourceOffset, sourceBytesPerRow, sourceBytesPerImage, sourceSize, toTexture, destinationSlice, destinationLevel, destinationOrigin, options);
 }
 
-bool FAGXContext::AsyncCopyFromTextureToTexture(FAGXTexture const& Texture, uint32 sourceSlice, uint32 sourceLevel, mtlpp::Origin sourceOrigin, mtlpp::Size sourceSize, FAGXTexture const& toTexture, uint32 destinationSlice, uint32 destinationLevel, mtlpp::Origin destinationOrigin)
+bool FAGXContext::AsyncCopyFromTextureToTexture(FAGXTexture const& Texture, uint32 sourceSlice, uint32 sourceLevel, MTLOrigin sourceOrigin, MTLSize sourceSize, FAGXTexture const& toTexture, uint32 destinationSlice, uint32 destinationLevel, MTLOrigin destinationOrigin)
 {
 	return RenderPass.AsyncCopyFromTextureToTexture(Texture, sourceSlice, sourceLevel, sourceOrigin, sourceSize, toTexture, destinationSlice, destinationLevel, destinationOrigin);
 }
diff --git a/Engine/Source/Runtime/Apple/AGXRHI/Private/AGXContext.h b/Engine/Source/Runtime/Apple/AGXRHI/Private/AGXContext.h
--- a/Engine/Source/Runtime/Apple/AGXRHI/Private/AGXContext.h
+++ b/Engine/Source/Runtime/Apple/AGXRHI/Private/AGXContext.h
@@ -71,7 +71,7 @@ public:
 	
 	void DrawPrimitiveIndirect(uint32 PrimitiveType, FAGXVertexBuffer* VertexBuffer, uint32 ArgumentOffset);
 	
-	void DrawIndexedPrimitive(FAGXBuffer const& IndexBuffer, uint32 IndexStride, mtlpp::IndexType IndexType, uint32 PrimitiveType, int32 BaseVertexIndex, uint32 FirstInstance,
+	void DrawIndexedPrimitive(FAGXBuffer const& IndexBuffer, uint32 IndexStride, MTLIndexType IndexType, uint32 PrimitiveType, int32 BaseVertexIndex, uint32 FirstInstance,
 							  uint32 NumVertices, uint32 StartIndex, uint32 NumPrimitives, uint32 NumInstances);
 	
 	void DrawIndexedIndirect(FAGXIndexBuffer* IndexBufferRHI, uint32 PrimitiveType, FAGXStructuredBuffer* VertexBufferRHI, int32 DrawArgumentsIndex, uint32 NumInstances);
@@ -81,17 +81,17 @@ public:
 	void DrawPatches(uint32 PrimitiveType, FAGXBuffer const& IndexBuffer, uint32 IndexBufferStride, int32 BaseVertexIndex, uint32 FirstInstance, uint32 StartIndex,
 					 uint32 NumPrimitives, uint32 NumInstances);
 	
-	void CopyFromTextureToBuffer(FAGXTexture const& Texture, uint32 sourceSlice, uint32 sourceLevel, mtlpp::Origin sourceOrigin, mtlpp::Size sourceSize, FAGXBuffer const& toBuffer, uint32 destinationOffset, uint32 destinationBytesPerRow, uint32 destinationBytesPerImage, mtlpp::BlitOption options);
+	void CopyFromTextureToBuffer(FAGXTexture const& Texture, uint32 sourceSlice, uint32 sourceLevel, MTLOrigin sourceOrigin, MTLSize sourceSize, FAGXBuffer const& toBuffer, uint32 destinationOffset, uint32 destinationBytesPerRow, uint32 destinationBytesPerImage, mtlpp::BlitOption options);
 	
-	void CopyFromBufferToTexture(FAGXBuffer const& Buffer, uint32 sourceOffset, uint32 sourceBytesPerRow, uint32 sourceBytesPerImage, mtlpp::Size sourceSize, FAGXTexture const& toTexture, uint32 destinationSlice, uint32 destinationLevel, mtlpp::Origin destinationOrigin, mtlpp::BlitOption options);
+	void CopyFromBufferToTexture(FAGXBuffer const& Buffer, uint32 sourceOffset, uint32 sourceBytesPerRow, uint32 sourceBytesPerImage, MTLSize sourceSize, FAGXTexture const& toTexture, uint32 destinationSlice, uint32 destinationLevel, MTLOrigin destinationOrigin, mtlpp::BlitOption options);
 	
-	void CopyFromTextureToTexture(FAGXTexture const& Texture, uint32 sourceSlice, uint32 sourceLevel, mtlpp::Origin sourceOrigin, mtlpp::Size sourceSize, FAGXTexture const& toTexture, uint32 destinationSlice, uint32 destinationLevel, mtlpp::Origin destinationOrigin);
+	void CopyFromTextureToTexture(FAGXTexture const& Texture, uint32 sourceSlice, uint32 sourceLevel, MTLOrigin sourceOrigin, MTLSize sourceSize, FAGXTexture const& toTexture, uint32 destinationSlice, uint32 destinationLevel, MTLOrigin destinationOrigin);
 	
 	void CopyFromBufferToBuffer(FAGXBuffer const& SourceBuffer, NSUInteger SourceOffset, FAGXBuffer const& DestinationBuffer, NSUInteger DestinationOffset, NSUInteger Size);
 	
-    bool AsyncCopyFromBufferToTexture(FAGXBuffer const& Buffer, uint32 sourceOffset, uint32 sourceBytesPerRow, uint32 sourceBytesPerImage, mtlpp::Size sourceSize, FAGXTexture const& toTexture, uint32 destinationSlice, uint32 destinationLevel, mtlpp::Origin destinationOrigin, mtlpp::BlitOption options);
+    bool AsyncCopyFromBufferToTexture(FAGXBuffer const& Buffer, uint32 sourceOffset, uint32 sourceBytesPerRow, uint32 sourceBytesPerImage, MTLSize sourceSize, FAGXTexture const& toTexture, uint32 destinationSlice, uint32 destinationLevel, MTLOrigin destinationOrigin, mtlpp::BlitOption options);
     
-    bool AsyncCopyFromTextureToTexture(FAGXTexture const& Texture, uint32 sourceSlice, uint32 sourceLevel, mtlpp::Origin sourceOrigin, mtlpp::Size sourceSize, FAGXTexture const& toTexture, uint32 destinationSlice, uint32 destinationLevel, mtlpp::Origin destinationOrigin);
+    bool AsyncCopyFromTextureToTexture(FAGXTexture const& Texture, uint32 sourceSlice, uint32 sourceLevel, MTLOrigin sourceOrigin, MTLSize sourceSize, FAGXTexture const& toTexture, uint32 destinationSlice, uint32 destinationLevel, MTLOrigin destinationOrigin);
 	
 	bool CanAsyncCopyToBuffer(FAGXBuffer const& DestinationBuffer);
 	
diff --git a/Engine/Source/Runtime/Apple/AGXRHI/Private/AGXRHIContext.cpp b/Engine/Source/Runtime/Apple/AGXRHI/Private/AGXRHIContext.cpp
--- a/Engine/Source/Runtime/Apple/AGXRHI/Private/AGXRHIContext.cpp
+++ b/Engine/Source/Runtime/Apple/AGXRHI/Private/AGXRHIContext.cpp
@@ -179,8 +179,8 @@ void FAGXRHICommandContext::ResolveTexture(UE::RHICore::FResolveTextureInfo Info
 	// Resolve required - Device must support this - Using Shader for resolve not supported amd NumSamples should be 1
 	check((!bDepthStencil && bSupportsMSAAStoreAndResolve) || (bDepthStencil && bSupportsMSAADepthResolve));
 
-	mtlpp::Origin Origin(0, 0, 0);
-	mtlpp::Size Size(0, 0, 1);
+	MTLOrigin Origin = MTLOriginMake(0, 0, 0);
+	MTLSize Size = MTLSizeMake(0, 0, 1);
 
 	if (Info.ResolveRect.IsValid())
 	{
diff --git a/Engine/Source/Runtime/Apple/AGXRHI/Private/AGXRHIPrivate.h b/Engine/Source/Runtime/Apple/AGXRHI/Private/AGXRHIPrivate.h
--- a/Engine/Source/Runtime/Apple/AGXRHI/Private/AGXRHIPrivate.h
+++ b/Engine/Source/Runtime/Apple/AGXRHI/Private/AGXRHIPrivate.h
@@ -265,7 +265,7 @@ extern uint32 AGXSafeGetRuntimeDebuggingLevel();
 
 extern int32 GAGXBufferZeroFill;
 
-mtlpp::LanguageVersion AGXValidateVersion(uint32 Version);
+MTLLanguageVersion AGXValidateVersion(uint32 Version);
 
 // Needs to be the same as EShaderFrequency when all stages are supported, but unlike EShaderFrequency you can compile out stages.
 enum EAGXShaderStages
diff --git a/Engine/Source/Runtime/Apple/AGXRHI/Private/AGXRHIRenderQuery.cpp b/Engine/Source/Runtime/Apple/AGXRHI/Private/AGXRHIRenderQuery.cpp
--- a/Engine/Source/Runtime/Apple/AGXRHI/Private/AGXRHIRenderQuery.cpp
+++ b/Engine/Source/Runtime/Apple/AGXRHI/Private/AGXRHIRenderQuery.cpp
@@ -205,11 +205,11 @@ void FAGXRHIRenderQuery::Begin(FAGXContext* Context, TSharedPtr<FAGXCommandBuffe
 
 			if ((GMaxRHIFeatureLevel >= ERHIFeatureLevel::SM5) && GetAGXDeviceContext().SupportsFeature(EAGXFeaturesCountingQueries))
 			{
-				Context->GetCurrentState().SetVisibilityResultMode(mtlpp::VisibilityResultMode::Counting, Buffer.Offset);
+				Context->GetCurrentState().SetVisibilityResultMode(MTLVisibilityResultModeCounting, Buffer.Offset);
 			}
 			else
 			{
-				Context->GetCurrentState().SetVisibilityResultMode(mtlpp::VisibilityResultMode::Boolean, Buffer.Offset);
+				Context->GetCurrentState().SetVisibilityResultMode(MTLVisibilityResultModeBoolean, Buffer.Offset);
 			}
 			if (BatchFence.IsValid())
 			{
@@ -242,7 +242,7 @@ void FAGXRHIRenderQuery::End(FAGXContext* Context)
 		{
 			// switch back to non-occlusion rendering
 			check(Buffer.CommandBufferFence.IsValid());
-			Context->GetCurrentState().SetVisibilityResultMode(mtlpp::VisibilityResultMode::Disabled, 0);
+			Context->GetCurrentState().SetVisibilityResultMode(MTLVisibilityResultModeDisabled, 0);
 
 			// For unique, unbatched, queries insert the fence now
 			if (!Buffer.bBatchFence)
diff --git a/Engine/Source/Runtime/Apple/AGXRHI/Private/AGXRenderPass.cpp b/Engine/Source/Runtime/Apple/AGXRHI/Private/AGXRenderPass.cpp
--- a/Engine/Source/Runtime/Apple/AGXRHI/Private/AGXRenderPass.cpp
+++ b/Engine/Source/Runtime/Apple/AGXRHI/Private/AGXRenderPass.cpp
@@ -561,7 +561,7 @@ void FAGXRenderPass::InsertTextureBarrier()
 #endif
 }
 
-void FAGXRenderPass::CopyFromTextureToBuffer(FAGXTexture const& Texture, uint32 sourceSlice, uint32 sourceLevel, mtlpp::Origin sourceOrigin, mtlpp::Size sourceSize, FAGXBuffer const& toBuffer, uint32 destinationOffset, uint32 destinationBytesPerRow, uint32 destinationBytesPerImage, mtlpp::BlitOption options)
+void FAGXRenderPass::CopyFromTextureToBuffer(FAGXTexture const& Texture, uint32 sourceSlice, uint32 sourceLevel, MTLOrigin sourceOrigin, MTLSize sourceSize, FAGXBuffer const& toBuffer, uint32 destinationOffset, uint32 destinationBytesPerRow, uint32 destinationBytesPerImage, mtlpp::BlitOption options)
 {
 	ConditionalSwitchToBlit();
 	
@@ -573,8 +573,8 @@ void FAGXRenderPass::CopyFromTextureToBuffer(FAGXTexture const& Texture, uint32
 		[Encoder			 copyFromTexture:Texture.GetPtr()
 								 sourceSlice:(NSUInteger)sourceSlice
 								 sourceLevel:(NSUInteger)sourceLevel
-								sourceOrigin:*reinterpret_cast<MTLOrigin*>(&sourceOrigin)
-								  sourceSize:*reinterpret_cast<MTLSize*>(&sourceSize)
+								sourceOrigin:sourceOrigin
+								  sourceSize:sourceSize
 									toBuffer:toBuffer.GetPtr()
 						   destinationOffset:(NSUInteger)destinationOffset + toBuffer.GetOffset()
 					  destinationBytesPerRow:(NSUInteger)destinationBytesPerRow
@@ -585,7 +585,7 @@ void FAGXRenderPass::CopyFromTextureToBuffer(FAGXTexture const& Texture, uint32
 	ConditionalSubmit();
 }
 
-void FAGXRenderPass::CopyFromBufferToTexture(FAGXBuffer const& Buffer, uint32 sourceOffset, uint32 sourceBytesPerRow, uint32 sourceBytesPerImage, mtlpp::Size sourceSize, FAGXTexture const& toTexture, uint32 destinationSlice, uint32 destinationLevel, mtlpp::Origin destinationOrigin, mtlpp::BlitOption options)
+void FAGXRenderPass::CopyFromBufferToTexture(FAGXBuffer const& Buffer, uint32 sourceOffset, uint32 sourceBytesPerRow, uint32 sourceBytesPerImage, MTLSize sourceSize, FAGXTexture const& toTexture, uint32 destinationSlice, uint32 destinationLevel, MTLOrigin destinationOrigin, mtlpp::BlitOption options)
 {
 	ConditionalSwitchToBlit();
 	
@@ -600,11 +600,11 @@ void FAGXRenderPass::CopyFromBufferToTexture(FAGXBuffer const& Buffer, uint32 so
 						   sourceOffset:(NSUInteger)sourceOffset + Buffer.GetOffset()
 					  sourceBytesPerRow:(NSUInteger)sourceBytesPerRow
 					sourceBytesPerImage:(NSUInteger)sourceBytesPerImage
-							 sourceSize:*reinterpret_cast<MTLSize*>(&sourceSize)
+							 sourceSize:sourceSize
 							  toTexture:toTexture.GetPtr()
 					   destinationSlice:(NSUInteger)destinationSlice
 					   destinationLevel:(NSUInteger)destinationLevel
-					  destinationOrigin:*reinterpret_cast<MTLOrigin*>(&destinationOrigin)];
+					  destinationOrigin:destinationOrigin];
 	}
 	else
 	{
@@ -612,18 +612,18 @@ void FAGXRenderPass::CopyFromBufferToTexture(FAGXBuffer const& Buffer, uint32 so
 						   sourceOffset:(NSUInteger)sourceOffset + Buffer.GetOffset()
 					  sourceBytesPerRow:(NSUInteger)sourceBytesPerRow
 					sourceBytesPerImage:(NSUInteger)sourceBytesPerImage
-							 sourceSize:*reinterpret_cast<MTLSize*>(&sourceSize)
+							 sourceSize:sourceSize
 							  toTexture:toTexture.GetPtr()
 					   destinationSlice:(NSUInteger)destinationSlice
 					   destinationLevel:(NSUInteger)destinationLevel
-					  destinationOrigin:*reinterpret_cast<MTLOrigin*>(&destinationOrigin)
+					  destinationOrigin:destinationOrigin
 								options:*reinterpret_cast<MTLBlitOption*>(&options)];
 	}
 	
 	ConditionalSubmit();
 }
 
-void FAGXRenderPass::CopyFromTextureToTexture(FAGXTexture const& Texture, uint32 sourceSlice, uint32 sourceLevel, mtlpp::Origin sourceOrigin, mtlpp::Size sourceSize, FAGXTexture const& toTexture, uint32 destinationSlice, uint32 destinationLevel, mtlpp::Origin destinationOrigin)
+void FAGXRenderPass::CopyFromTextureToTexture(FAGXTexture const& Texture, uint32 sourceSlice, uint32 sourceLevel, MTLOrigin sourceOrigin, MTLSize sourceSize, FAGXTexture const& toTexture, uint32 destinationSlice, uint32 destinationLevel, MTLOrigin destinationOrigin)
 {
 	ConditionalSwitchToBlit();
 	
@@ -635,12 +635,12 @@ void FAGXRenderPass::CopyFromTextureToTexture(FAGXTexture const& Texture, uint32
 	[Encoder	  copyFromTexture:Texture.GetPtr()
 					  sourceSlice:(NSUInteger)sourceSlice
 					  sourceLevel:(NSUInteger)sourceLevel
-					 sourceOrigin:*reinterpret_cast<MTLOrigin*>(&sourceOrigin)
-					   sourceSize:*reinterpret_cast<MTLSize*>(&sourceSize)
+					 sourceOrigin:sourceOrigin
+					   sourceSize:sourceSize
 						toTexture:toTexture.GetPtr()
 				 destinationSlice:(NSUInteger)destinationSlice
 				 destinationLevel:(NSUInteger)destinationLevel
-				destinationOrigin:*reinterpret_cast<MTLOrigin*>(&destinationOrigin)];
+				destinationOrigin:destinationOrigin];
 	
 	ConditionalSubmit();
 }
@@ -663,7 +663,7 @@ void FAGXRenderPass::CopyFromBufferToBuffer(FAGXBuffer const& SourceBuffer, NSUI
 	ConditionalSubmit();
 }
 
-void FAGXRenderPass::PresentTexture(FAGXTexture const& Texture, uint32 sourceSlice, uint32 sourceLevel, mtlpp::Origin sourceOrigin, mtlpp::Size sourceSize, FAGXTexture const& toTexture, uint32 destinationSlice, uint32 destinationLevel, mtlpp::Origin destinationOrigin)
+void FAGXRenderPass::PresentTexture(FAGXTexture const& Texture, uint32 sourceSlice, uint32 sourceLevel, MTLOrigin sourceOrigin, MTLSize sourceSize, FAGXTexture const& toTexture, uint32 destinationSlice, uint32 destinationLevel, MTLOrigin destinationOrigin)
 {
 	ConditionalSwitchToBlit();
 	
@@ -675,12 +675,12 @@ void FAGXRenderPass::PresentTexture(FAGXTexture const& Texture, uint32 sourceSli
 	[Encoder	  copyFromTexture:Texture.GetPtr()
 					  sourceSlice:(NSUInteger)sourceSlice
 					  sourceLevel:(NSUInteger)sourceLevel
-					 sourceOrigin:*reinterpret_cast<MTLOrigin*>(&sourceOrigin)
-					   sourceSize:*reinterpret_cast<MTLSize*>(&sourceSize)
+					 sourceOrigin:sourceOrigin
+					   sourceSize:sourceSize
 						toTexture:toTexture.GetPtr()
 				 destinationSlice:(NSUInteger)destinationSlice
 				 destinationLevel:(NSUInteger)destinationLevel
-				destinationOrigin:*reinterpret_cast<MTLOrigin*>(&destinationOrigin)];
+				destinationOrigin:destinationOrigin];
 }
 
 void FAGXRenderPass::SynchronizeTexture(FAGXTexture const& Texture, uint32 Slice, uint32 Level)
@@ -744,7 +744,7 @@ void FAGXRenderPass::FillBuffer(FAGXBuffer const& Buffer, ns::Range Range, uint8
 	}
 }
 
-bool FAGXRenderPass::AsyncCopyFromBufferToTexture(FAGXBuffer const& Buffer, uint32 sourceOffset, uint32 sourceBytesPerRow, uint32 sourceBytesPerImage, mtlpp::Size sourceSize, FAGXTexture const& toTexture, uint32 destinationSlice, uint32 destinationLevel, mtlpp::Origin destinationOrigin, mtlpp::BlitOption options)
+bool FAGXRenderPass::AsyncCopyFromBufferToTexture(FAGXBuffer const& Buffer, uint32 sourceOffset, uint32 sourceBytesPerRow, uint32 sourceBytesPerImage, MTLSize sourceSize, FAGXTexture const& toTexture, uint32 destinationSlice, uint32 destinationLevel, MTLOrigin destinationOrigin, mtlpp::BlitOption options)
 {
 	id<MTLBlitCommandEncoder> TargetEncoder = nil;
 	bool bAsync = !CurrentEncoder.HasTextureBindingHistory(toTexture);
@@ -769,11 +769,11 @@ bool FAGXRenderPass::AsyncCopyFromBufferToTexture(FAGXBuffer const& Buffer, uint
 							   sourceOffset:(NSUInteger)sourceOffset + Buffer.GetOffset()
 						  sourceBytesPerRow:(NSUInteger)sourceBytesPerRow
 						sourceBytesPerImage:(NSUInteger)sourceBytesPerImage
-								 sourceSize:*reinterpret_cast<MTLSize*>(&sourceSize)
+								 sourceSize:sourceSize
 								  toTexture:toTexture.GetPtr()
 						   destinationSlice:(NSUInteger)destinationSlice
 						   destinationLevel:(NSUInteger)destinationLevel
-						  destinationOrigin:*reinterpret_cast<MTLOrigin*>(&destinationOrigin)];
+						  destinationOrigin:destinationOrigin];
 	}
 	else
 	{
@@ -781,18 +781,18 @@ bool FAGXRenderPass::AsyncCopyFromBufferToTexture(FAGXBuffer const& Buffer, uint
 							   sourceOffset:(NSUInteger)sourceOffset + Buffer.GetOffset()
 						  sourceBytesPerRow:(NSUInteger)sourceBytesPerRow
 						sourceBytesPerImage:(NSUInteger)sourceBytesPerImage
-								 sourceSize:*reinterpret_cast<MTLSize*>(&sourceSize)
+								 sourceSize:sourceSize
 								  toTexture:toTexture.GetPtr()
 						   destinationSlice:(NSUInteger)destinationSlice
 						   destinationLevel:(NSUInteger)destinationLevel
-						  destinationOrigin:*reinterpret_cast<MTLOrigin*>(&destinationOrigin)
+						  destinationOrigin:destinationOrigin
 									options:*reinterpret_cast<MTLBlitOption*>(&options)];
 	}
 	
 	return bAsync;
 }
 
-bool FAGXRenderPass::AsyncCopyFromTextureToTexture(FAGXTexture const& Texture, uint32 sourceSlice, uint32 sourceLevel, mtlpp::Origin sourceOrigin, mtlpp::Size sourceSize, FAGXTexture const& toTexture, uint32 destinationSlice, uint32 destinationLevel, mtlpp::Origin destinationOrigin)
+bool FAGXRenderPass::AsyncCopyFromTextureToTexture(FAGXTexture const& Texture, uint32 sourceSlice, uint32 sourceLevel, MTLOrigin sourceOrigin, MTLSize sourceSize, FAGXTexture const& toTexture, uint32 destinationSlice, uint32 destinationLevel, MTLOrigin destinationOrigin)
 {
 	id<MTLBlitCommandEncoder> TargetEncoder = nil;
 	bool bAsync = !CurrentEncoder.HasTextureBindingHistory(toTexture);
@@ -814,12 +814,12 @@ bool FAGXRenderPass::AsyncCopyFromTextureToTexture(FAGXTexture const& Texture, u
 	[TargetEncoder	  copyFromTexture:Texture.GetPtr()
 						  sourceSlice:(NSUInteger)sourceSlice
 						  sourceLevel:(NSUInteger)sourceLevel
-						 sourceOrigin:*reinterpret_cast<MTLOrigin*>(&sourceOrigin)
-						   sourceSize:*reinterpret_cast<MTLSize*>(&sourceSize)
+						 sourceOrigin:sourceOrigin
+						   sourceSize:sourceSize
 							toTexture:toTexture.GetPtr()
 					 destinationSlice:(NSUInteger)destinationSlice
 					 destinationLevel:(NSUInteger)destinationLevel
-					destinationOrigin:*reinterpret_cast<MTLOrigin*>(&destinationOrigin)];
+					destinationOrigin:destinationOrigin];
 
 	return bAsync;
 }
diff --git a/Engine/Source/Runtime/Apple/AGXRHI/Private/AGXRenderPass.h b/Engine/Source/Runtime/Apple/AGXRHI/Private/AGXRenderPass.h
--- a/Engine/Source/Runtime/Apple/AGXRHI/Private/AGXRenderPass.h
+++ b/Engine/Source/Runtime/Apple/AGXRHI/Private/AGXRenderPass.h
@@ -48,15 +48,15 @@ public:
     
     void EndRenderPass(void);
     
-    void CopyFromTextureToBuffer(FAGXTexture const& Texture, uint32 sourceSlice, uint32 sourceLevel, mtlpp::Origin sourceOrigin, mtlpp::Size sourceSize, FAGXBuffer const& toBuffer, uint32 destinationOffset, uint32 destinationBytesPerRow, uint32 destinationBytesPerImage, mtlpp::BlitOption options);
+    void CopyFromTextureToBuffer(FAGXTexture const& Texture, uint32 sourceSlice, uint32 sourceLevel, MTLOrigin sourceOrigin, MTLSize sourceSize, FAGXBuffer const& toBuffer, uint32 destinationOffset, uint32 destinationBytesPerRow, uint32 destinationBytesPerImage, mtlpp::BlitOption options);
     
-    void CopyFromBufferToTexture(FAGXBuffer const& Buffer, uint32 sourceOffset, uint32 sourceBytesPerRow, uint32 sourceBytesPerImage, mtlpp::Size sourceSize, FAGXTexture const& toTexture, uint32 destinationSlice, uint32 destinationLevel, mtlpp::Origin destinationOrigin, mtlpp::BlitOption options);
+    void CopyFromBufferToTexture(FAGXBuffer const& Buffer, uint32 sourceOffset, uint32 sourceBytesPerRow, uint32 sourceBytesPerImage, MTLSize sourceSize, FAGXTexture const& toTexture, uint32 destinationSlice, uint32 destinationLevel, MTLOrigin destinationOrigin, mtlpp::BlitOption options);
     
-    void CopyFromTextureToTexture(FAGXTexture const& Texture, uint32 sourceSlice, uint32 sourceLevel, mtlpp::Origin sourceOrigin, mtlpp::Size sourceSize, FAGXTexture const& toTexture, uint32 destinationSlice, uint32 destinationLevel, mtlpp::Origin destinationOrigin);
+    void CopyFromTextureToTexture(FAGXTexture const& Texture, uint32 sourceSlice, uint32 sourceLevel, MTLOrigin sourceOrigin, MTLSize sourceSize, FAGXTexture const& toTexture, uint32 destinationSlice, uint32 destinationLevel, MTLOrigin destinationOrigin);
 	
 	void CopyFromBufferToBuffer(FAGXBuffer const& SourceBuffer, NSUInteger SourceOffset, FAGXBuffer const& DestinationBuffer, NSUInteger DestinationOffset, NSUInteger Size);
 	
-	void PresentTexture(FAGXTexture const& Texture, uint32 sourceSlice, uint32 sourceLevel, mtlpp::Origin sourceOrigin, mtlpp::Size sourceSize, FAGXTexture const& toTexture, uint32 destinationSlice, uint32 destinationLevel, mtlpp::Origin destinationOrigin);
+	void PresentTexture(FAGXTexture const& Texture, uint32 sourceSlice, uint32 sourceLevel, MTLOrigin sourceOrigin, MTLSize sourceSize, FAGXTexture const& toTexture, uint32 destinationSlice, uint32 destinationLevel, MTLOrigin destinationOrigin);
     
     void SynchronizeTexture(FAGXTexture const& Texture, uint32 Slice, uint32 Level);
     
@@ -64,9 +64,9 @@ public:
     
 	void FillBuffer(FAGXBuffer const& Buffer, ns::Range Range, uint8 Value);
 	
-	bool AsyncCopyFromBufferToTexture(FAGXBuffer const& Buffer, uint32 sourceOffset, uint32 sourceBytesPerRow, uint32 sourceBytesPerImage, mtlpp::Size sourceSize, FAGXTexture const& toTexture, uint32 destinationSlice, uint32 destinationLevel, mtlpp::Origin destinationOrigin, mtlpp::BlitOption options);
+	bool AsyncCopyFromBufferToTexture(FAGXBuffer const& Buffer, uint32 sourceOffset, uint32 sourceBytesPerRow, uint32 sourceBytesPerImage, MTLSize sourceSize, FAGXTexture const& toTexture, uint32 destinationSlice, uint32 destinationLevel, MTLOrigin destinationOrigin, mtlpp::BlitOption options);
 	
-	bool AsyncCopyFromTextureToTexture(FAGXTexture const& Texture, uint32 sourceSlice, uint32 sourceLevel, mtlpp::Origin sourceOrigin, mtlpp::Size sourceSize, FAGXTexture const& toTexture, uint32 destinationSlice, uint32 destinationLevel, mtlpp::Origin destinationOrigin);
+	bool AsyncCopyFromTextureToTexture(FAGXTexture const& Texture, uint32 sourceSlice, uint32 sourceLevel, MTLOrigin sourceOrigin, MTLSize sourceSize, FAGXTexture const& toTexture, uint32 destinationSlice, uint32 destinationLevel, MTLOrigin destinationOrigin);
 	
 	bool CanAsyncCopyToBuffer(FAGXBuffer const& DestinationBuffer);
 	
diff --git a/Engine/Source/Runtime/Apple/AGXRHI/Private/AGXRenderTarget.cpp b/Engine/Source/Runtime/Apple/AGXRHI/Private/AGXRenderTarget.cpp
--- a/Engine/Source/Runtime/Apple/AGXRHI/Private/AGXRenderTarget.cpp
+++ b/Engine/Source/Runtime/Apple/AGXRHI/Private/AGXRenderTarget.cpp
@@ -73,8 +73,8 @@ void FAGXRHICommandContext::RHICopyToResolveTarget(FRHITexture* SourceTextureRHI
 		checkf( DestinationDesc.IsTexture2D()   || DestinationDesc.IsTextureCube(), TEXT("Only Tex2D & Cube are tested to work so far!"));
 		checkf(!DestinationDesc.IsTextureCube() || DestinationDesc.ArraySize == 1,  TEXT("Cube arrays might not work yet."));
 
-		mtlpp::Origin Origin(0, 0, 0);
-		mtlpp::Size Size(0, 0, 1);
+		MTLOrigin Origin = MTLOriginMake(0, 0, 0);
+		MTLSize Size = MTLSizeMake(0, 0, 1);
 		if (ResolveParams.Rect.IsValid())
 		{
 			// Partial copy
@@ -265,7 +265,7 @@ void FAGXDynamicRHI::RHIReadSurfaceData(FRHITexture* TextureRHI, FIntRect Rect,
 	OutData.AddUninitialized(SizeX * SizeY);
 	
 	FColor* OutDataPtr = OutData.GetData();
-	mtlpp::Region Region(Rect.Min.X, Rect.Min.Y, SizeX, SizeY);
+	MTLRegion Region = MTLRegionMake2D(Rect.Min.X, Rect.Min.Y, SizeX, SizeY);
     
 	FAGXTexture Texture = Surface->Texture;
     if(!Texture && EnumHasAnyFlags(Surface->GetDesc().Flags, TexCreate_Presentable))
@@ -307,10 +307,10 @@ void FAGXDynamicRHI::RHIReadSurfaceData(FRHITexture* TextureRHI, FIntRect Rect,
 			
 			TempTexture = GMtlppDevice.NewTexture(Desc);
 			
-			ImmediateContext.Context->CopyFromTextureToTexture(Texture, 0, InFlags.GetMip(), mtlpp::Origin(Region.origin), mtlpp::Size(Region.size), TempTexture, 0, 0, mtlpp::Origin(0, 0, 0));
+			ImmediateContext.Context->CopyFromTextureToTexture(Texture, 0, InFlags.GetMip(), Region.origin, Region.size, TempTexture, 0, 0, MTLOriginMake(0, 0, 0));
 			
 			Texture = TempTexture;
-			Region = mtlpp::Region(0, 0, SizeX, SizeY);
+			Region = MTLRegionMake2D(0, 0, SizeX, SizeY);
 		}
 #if PLATFORM_MAC
 		if ([Texture.GetPtr() storageMode] == MTLStorageModeManaged)
@@ -329,7 +329,12 @@ void FAGXDynamicRHI::RHIReadSurfaceData(FRHITexture* TextureRHI, FIntRect Rect,
 		TArray<uint8> Data;
 		Data.AddUninitialized(BytesPerImage);
 		
-		Texture.GetBytes(Data.GetData(), Stride, BytesPerImage, Region, 0, 0);
+		[Texture.GetPtr() getBytes:(void*)Data.GetData()
+					   bytesPerRow:Stride
+					 bytesPerImage:BytesPerImage
+						fromRegion:Region
+					   mipmapLevel:0
+							 slice:0];
 		
 		ConvertSurfaceDataToFColor(Surface->GetDesc().Format, SizeX, SizeY, (uint8*)Data.GetData(), Stride, OutDataPtr, InFlags);
 		
@@ -441,7 +446,7 @@ void FAGXDynamicRHI::RHIReadSurfaceFloatData(FRHITexture* TextureRHI, FIntRect R
 	OutData.Empty();
 	OutData.AddUninitialized(SizeX * SizeY);
 	
-	mtlpp::Region Region = mtlpp::Region(Rect.Min.X, Rect.Min.Y, SizeX, SizeY);
+	MTLRegion Region = MTLRegionMake2D(Rect.Min.X, Rect.Min.Y, SizeX, SizeY);
 	
 	// function wants details about the destination, not the source
 	const uint32 Stride = GPixelFormats[Surface->GetDesc().Format].BlockBytes * SizeX;
@@ -507,7 +512,7 @@ void FAGXDynamicRHI::RHIRead3DSurfaceFloatData(FRHITexture* TextureRHI,FIntRect
 	OutData.Empty();
 	OutData.AddUninitialized(SizeX * SizeY * SizeZ);
 	
-	mtlpp::Region Region = mtlpp::Region(InRect.Min.X, InRect.Min.Y, ZMinMax.X, SizeX, SizeY, SizeZ);
+	MTLRegion Region = MTLRegionMake3D(InRect.Min.X, InRect.Min.Y, ZMinMax.X, SizeX, SizeY, SizeZ);
 	
 	// function wants details about the destination, not the source
 	const uint32 Stride = GPixelFormats[Surface->GetDesc().Format].BlockBytes * SizeX;
diff --git a/Engine/Source/Runtime/Apple/AGXRHI/Private/AGXShaders.cpp b/Engine/Source/Runtime/Apple/AGXRHI/Private/AGXShaders.cpp
--- a/Engine/Source/Runtime/Apple/AGXRHI/Private/AGXShaders.cpp
+++ b/Engine/Source/Runtime/Apple/AGXRHI/Private/AGXShaders.cpp
@@ -44,52 +44,52 @@ NSString* AGXDecodeMetalSourceCode(uint32 CodeSize, TArray<uint8> const& Compres
 	return GlslCodeNSString;
 }
 
-mtlpp::LanguageVersion AGXValidateVersion(uint32 Version)
+MTLLanguageVersion AGXValidateVersion(uint32 Version)
 {
-    mtlpp::LanguageVersion Result = mtlpp::LanguageVersion::Version2_2;
+    MTLLanguageVersion Result = MTLLanguageVersion2_2;
 #if PLATFORM_MAC
-    Result = mtlpp::LanguageVersion::Version2_2;
-    switch(Version)
+    Result = MTLLanguageVersion2_2;
+    switch (Version)
     {
         case 7:
-            Result = mtlpp::LanguageVersion::Version2_4;
+            Result = MTLLanguageVersion2_4;
             break;
         case 6:
-            Result = mtlpp::LanguageVersion::Version2_3;
+            Result = MTLLanguageVersion2_3;
             break;
         case 5:
-            Result = mtlpp::LanguageVersion::Version2_2;
+            Result = MTLLanguageVersion2_2;
             break;
         case 0:
             Version = 5;
-            Result = mtlpp::LanguageVersion::Version2_2; // minimum version as of UE5.0
+            Result = MTLLanguageVersion2_2; // minimum version as of UE5.0
             break;
         default:
             //EMacMetalShaderStandard::MacMetalSLStandard_Minimum is currently 2.2
             UE_LOG(LogTemp, Warning, TEXT("The Metal version currently set is not supported anymore. Set it in the Project Settings. Defaulting to the minimum version."));
             Version = 5;
-            Result = mtlpp::LanguageVersion::Version2_2;
+            Result = MTLLanguageVersion2_2;
             break;
     }
 #else
-    Result = mtlpp::LanguageVersion::Version2_3;
-    switch(Version)
+    Result = MTLLanguageVersion2_3;
+    switch (Version)
     {
         case 7:
-            Result = mtlpp::LanguageVersion::Version2_4;
+            Result = MTLLanguageVersion2_4;
             break;
         case 6:
-            Result = mtlpp::LanguageVersion::Version2_3;
+            Result = MTLLanguageVersion2_3;
             break;
         case 0:
             Version = 6;
-            Result = mtlpp::LanguageVersion::Version2_3; // minimum version as of UE5.0
+            Result = MTLLanguageVersion2_3; // minimum version as of UE5.0
             break;
         default:
             //EMacMetalShaderStandard::MacMetalSLStandard_Minimum and EIOSMetalShaderStandard::IOSMetalSLStandard_Minimum is currently 2.3
             UE_LOG(LogTemp, Warning, TEXT("The Metal version currently set is not supported anymore. Set it in the Project Settings. Defaulting to the minimum version."));
             Version = 6;
-            Result = mtlpp::LanguageVersion::Version2_3;
+            Result = MTLLanguageVersion2_3;
             break;
     }
 #endif
diff --git a/Engine/Source/Runtime/Apple/AGXRHI/Private/AGXStateCache.cpp b/Engine/Source/Runtime/Apple/AGXRHI/Private/AGXStateCache.cpp
--- a/Engine/Source/Runtime/Apple/AGXRHI/Private/AGXStateCache.cpp
+++ b/Engine/Source/Runtime/Apple/AGXRHI/Private/AGXStateCache.cpp
@@ -17,13 +17,13 @@
 	#define FMETALTEXTUREMASK_MAX UINT32_MAX
 #endif
 
-static mtlpp::TriangleFillMode TranslateFillMode(ERasterizerFillMode FillMode)
+static MTLTriangleFillMode TranslateFillMode(ERasterizerFillMode FillMode)
 {
 	switch (FillMode)
 	{
-		case FM_Wireframe:	return mtlpp::TriangleFillMode::Lines;
-		case FM_Point:		return mtlpp::TriangleFillMode::Fill;
-		default:			return mtlpp::TriangleFillMode::Fill;
+		case FM_Wireframe:	return MTLTriangleFillModeLines;
+		case FM_Point:		return MTLTriangleFillModeFill;
+		default:			return MTLTriangleFillModeFill;
 	};
 }
 
@@ -37,13 +37,13 @@ static MTLCullMode TranslateCullMode(ERasterizerCullMode CullMode)
 	}
 }
 
-static mtlpp::DepthClipMode TranslateDepthClipMode(ERasterizerDepthClipMode DepthClipMode)
+static MTLDepthClipMode TranslateDepthClipMode(ERasterizerDepthClipMode DepthClipMode)
 {
 	switch (DepthClipMode)
 	{
-	case ERasterizerDepthClipMode::DepthClip:	return mtlpp::DepthClipMode::Clip;
-	case ERasterizerDepthClipMode::DepthClamp:	return mtlpp::DepthClipMode::Clamp;
-	default:									return mtlpp::DepthClipMode::Clip;
+		case ERasterizerDepthClipMode::DepthClip:	return MTLDepthClipModeClip;
+		case ERasterizerDepthClipMode::DepthClamp:	return MTLDepthClipModeClamp;
+		default:									return MTLDepthClipModeClip;
 	}
 }
 
@@ -170,7 +170,7 @@ FAGXStateCache::FAGXStateCache(bool bInImmediate)
 : DepthStore(MTLStoreActionUnknown)
 , StencilStore(MTLStoreActionUnknown)
 , VisibilityResults(nullptr)
-, VisibilityMode(mtlpp::VisibilityResultMode::Disabled)
+, VisibilityMode(MTLVisibilityResultModeDisabled)
 , VisibilityOffset(0)
 , VisibilityWritten(0)
 , DepthStencilState(nullptr)
@@ -299,7 +299,7 @@ void FAGXStateCache::Reset()
 	}
 	
 	VisibilityResults = nil;
-	VisibilityMode = mtlpp::VisibilityResultMode::Disabled;
+	VisibilityMode = MTLVisibilityResultModeDisabled;
 	VisibilityOffset = 0;
 	VisibilityWritten = 0;
 	
@@ -1107,7 +1107,7 @@ void FAGXStateCache::SetVertexStream(uint32 Index, FAGXBuffer* Buffer, FAGXBuffe
 	SetShaderBuffer(EAGXShaderStages::Vertex, VertexBuffers[Index].Buffer, Bytes, Offset, Length, UNREAL_TO_METAL_BUFFER_INDEX(Index), MTLResourceUsageRead);
 }
 
-uint32 FAGXStateCache::GetVertexBufferSize(uint32 Index)
+uint32 FAGXStateCache::GetVertexBufferSize(uint32 Index) const
 {
 	check(Index < MaxVertexElementCount);
 	check(UNREAL_TO_METAL_BUFFER_INDEX(Index) < MaxMetalStreams);
@@ -1161,7 +1161,7 @@ FAGXShaderPipeline* FAGXStateCache::GetPipelineState() const
 	return GraphicsPSO->GetPipeline();
 }
 
-EPrimitiveType FAGXStateCache::GetPrimitiveType()
+EPrimitiveType FAGXStateCache::GetPrimitiveType() const
 {
 	check(IsValidRef(GraphicsPSO));
 	return GraphicsPSO->GetPrimitiveType();
@@ -1183,7 +1183,7 @@ void FAGXStateCache::SetDirtyUniformBuffers(EAGXShaderStages Freq, uint32 Dirty)
 	DirtyUniformBuffers[Freq] = Dirty;
 }
 
-void FAGXStateCache::SetVisibilityResultMode(mtlpp::VisibilityResultMode Mode, NSUInteger Offset)
+void FAGXStateCache::SetVisibilityResultMode(MTLVisibilityResultMode Mode, NSUInteger Offset)
 {
 	if (VisibilityMode != Mode || VisibilityOffset != Offset)
 	{
@@ -1486,7 +1486,7 @@ void FAGXStateCache::SetShaderResourceView(FAGXContext* Context, EAGXShaderStage
 	}
 }
 
-bool FAGXStateCache::IsLinearBuffer(EAGXShaderStages ShaderStage, uint32 BindIndex)
+bool FAGXStateCache::IsLinearBuffer(EAGXShaderStages ShaderStage, uint32 BindIndex) const
 {
     switch (ShaderStage)
     {
@@ -1927,7 +1927,7 @@ void FAGXStateCache::SetRenderState(FAGXCommandEncoder& CommandEncoder, FAGXComm
 		if (RasterBits & EAGXRenderFlagVisibilityResultMode)
 		{
 			CommandEncoder.SetVisibilityResultMode(VisibilityMode, VisibilityOffset);
-			if (VisibilityMode != mtlpp::VisibilityResultMode::Disabled)
+			if (VisibilityMode != MTLVisibilityResultModeDisabled)
 			{
             	VisibilityWritten = VisibilityOffset + FAGXQueryBufferPool::EQueryResultMaxSize;
 			}
diff --git a/Engine/Source/Runtime/Apple/AGXRHI/Private/AGXStateCache.h b/Engine/Source/Runtime/Apple/AGXRHI/Private/AGXStateCache.h
--- a/Engine/Source/Runtime/Apple/AGXRHI/Private/AGXStateCache.h
+++ b/Engine/Source/Runtime/Apple/AGXRHI/Private/AGXStateCache.h
@@ -64,7 +64,7 @@ public:
 	 * @param Mode Controls if the counter is disabled or moniters passing samples.
 	 * @param Offset The offset relative to the occlusion query buffer provided when the command encoder was created.  offset must be a multiple of 8.
 	 */
-	void SetVisibilityResultMode(mtlpp::VisibilityResultMode Mode, NSUInteger Offset);
+	void SetVisibilityResultMode(MTLVisibilityResultMode Mode, NSUInteger Offset);
 	
 #pragma mark - Public Shader Resource Mutators -
 	/*
@@ -126,13 +126,13 @@ public:
 	FAGXComputeShader* GetComputeShader() const { return ComputeShader; }
 	CGSize GetFrameBufferSize() const { return FrameBufferSize; }
 	FRHIRenderPassInfo const& GetRenderPassInfo() const { return RenderPassInfo; }
-	int32 GetNumRenderTargets() { return bHasValidColorTarget ? RenderPassInfo.GetNumColorRenderTargets() : -1; }
+	int32 GetNumRenderTargets() const { return bHasValidColorTarget ? RenderPassInfo.GetNumColorRenderTargets() : -1; }
 	bool GetHasValidRenderTarget() const { return bHasValidRenderTarget; }
 	bool GetHasValidColorTarget() const { return bHasValidColorTarget; }
-	MTLViewport const& GetViewport(uint32 Index) const { check(Index < ML_MaxViewports); return Viewport[Index]; }
-	uint32 GetVertexBufferSize(uint32 Index);
+	const MTLViewport& GetViewport(uint32 Index) const { check(Index < ML_MaxViewports); return Viewport[Index]; }
+	uint32 GetVertexBufferSize(uint32 Index) const;
 	uint32 GetRenderTargetArraySize() const { return RenderTargetArraySize; }
-	const FRHIUniformBuffer** GetBoundUniformBuffers(EAGXShaderStages Freq) { return (const FRHIUniformBuffer**)&BoundUniformBuffers[Freq][0]; }
+	const FRHIUniformBuffer** GetBoundUniformBuffers(EAGXShaderStages Freq) const { return (const FRHIUniformBuffer**)&BoundUniformBuffers[Freq][0]; }
 	uint32 GetDirtyUniformBuffers(EAGXShaderStages Freq) const { return DirtyUniformBuffers[Freq]; }
 	FAGXQueryBuffer* GetVisibilityResultsBuffer() const { return VisibilityResults; }
 	bool GetScissorRectEnabled() const { return bScissorRectEnabled; }
@@ -141,11 +141,11 @@ public:
     bool CanRestartRenderPass() const { return bCanRestartRenderPass; }
 	MTLRenderPassDescriptor* GetRenderPassDescriptor() const { return RenderPassDesc; }
 	uint32 GetSampleCount() const { return SampleCount; }
-    bool IsLinearBuffer(EAGXShaderStages ShaderStage, uint32 BindIndex);
+    bool IsLinearBuffer(EAGXShaderStages ShaderStage, uint32 BindIndex) const;
 	FAGXShaderPipeline* GetPipelineState() const;
-	EPrimitiveType GetPrimitiveType();
-	mtlpp::VisibilityResultMode GetVisibilityResultMode() { return VisibilityMode; }
-	uint32 GetVisibilityResultOffset() { return VisibilityOffset; }
+	EPrimitiveType GetPrimitiveType() const;
+	MTLVisibilityResultMode GetVisibilityResultMode() const { return VisibilityMode; }
+	uint32 GetVisibilityResultOffset() const { return VisibilityOffset; }
 	
 	FTexture2DRHIRef CreateFallbackDepthStencilSurface(uint32 Width, uint32 Height);
 	bool GetFallbackDepthStencilBound() const { return bFallbackDepthStencilBound; }
@@ -269,7 +269,7 @@ private:
 	MTLStoreAction StencilStore;
 
 	FAGXQueryBuffer* VisibilityResults;
-	mtlpp::VisibilityResultMode VisibilityMode;
+	MTLVisibilityResultMode VisibilityMode;
 	NSUInteger VisibilityOffset;
 	NSUInteger VisibilityWritten;
 
diff --git a/Engine/Source/Runtime/Apple/AGXRHI/Private/AGXTexture.cpp b/Engine/Source/Runtime/Apple/AGXRHI/Private/AGXTexture.cpp
--- a/Engine/Source/Runtime/Apple/AGXRHI/Private/AGXTexture.cpp
+++ b/Engine/Source/Runtime/Apple/AGXRHI/Private/AGXTexture.cpp
@@ -904,11 +904,11 @@ void FAGXSurface::UpdateSurfaceAndDestroySourceBuffer(id <MTLBuffer> SourceBuffe
 	uint32 Stride;
 	uint32 BytesPerImage = GetMipSize(MipIndex, &Stride, true);
 	
-	mtlpp::Region Region;
+	MTLRegion Region;
 	if (GetDesc().IsTexture3D())
 	{
 		// upload the texture to the texture slice
-		Region = mtlpp::Region(
+		Region = MTLRegionMake3D(
 			0, 0, 0,
 			FMath::Max<uint32>(GetDesc().Extent.X >> MipIndex, 1),
 			FMath::Max<uint32>(GetDesc().Extent.Y >> MipIndex, 1),
@@ -918,7 +918,7 @@ void FAGXSurface::UpdateSurfaceAndDestroySourceBuffer(id <MTLBuffer> SourceBuffe
 	else
 	{
 		// upload the texture to the texture slice
-		Region = mtlpp::Region(
+		Region = MTLRegionMake2D(
 			0, 0,
 			FMath::Max<uint32>(GetDesc().Extent.X >> MipIndex, 1),
 			FMath::Max<uint32>(GetDesc().Extent.Y >> MipIndex, 1)
@@ -1025,7 +1025,13 @@ void FAGXSurface::UpdateSurfaceAndDestroySourceBuffer(id <MTLBuffer> SourceBuffe
 		}
 #endif
 		
-		MTLPP_VALIDATE(mtlpp::Texture, Texture, AGXSafeGetRuntimeDebuggingLevel() >= EAGXDebugLevelValidation, Replace(Region, MipIndex, ArrayIndex, SourceBuffer.contents, Stride, BytesPerImage));
+		[Texture.GetPtr() replaceRegion:Region
+							mipmapLevel:MipIndex
+								  slice:ArrayIndex
+							  withBytes:(const void*)[SourceBuffer contents]
+							bytesPerRow:Stride
+						  bytesPerImage:BytesPerImage];
+
 		[SourceBuffer release];
 		
 		INC_DWORD_STAT_BY(STAT_AGXTextureMemUpdate, BytesPerImage);
@@ -1056,11 +1062,11 @@ void* FAGXSurface::Lock(uint32 MipIndex, uint32 ArrayIndex, EResourceLockMode Lo
 			FRHICommandListImmediate& RHICmdList = FRHICommandListExecutor::GetImmediateCommandList();
 			const bool bIssueImmediateCommands = RHICmdList.Bypass() || IsInRHIThread();
 			
-			mtlpp::Region Region;
+			MTLRegion Region;
 			if (GetDesc().IsTexture3D())
 			{
 				// upload the texture to the texture slice
-				Region = mtlpp::Region(
+				Region = MTLRegionMake3D(
 					0, 0, 0,
 					FMath::Max<uint32>(GetDesc().Extent.X >> MipIndex, 1),
 					FMath::Max<uint32>(GetDesc().Extent.Y >> MipIndex, 1),
@@ -1069,7 +1075,7 @@ void* FAGXSurface::Lock(uint32 MipIndex, uint32 ArrayIndex, EResourceLockMode Lo
 			else
 			{
 				// upload the texture to the texture slice
-				Region = mtlpp::Region(
+				Region = MTLRegionMake2D(
 					0, 0,
 					FMath::Max<uint32>(GetDesc().Extent.X >> MipIndex, 1),
 					FMath::Max<uint32>(GetDesc().Extent.Y >> MipIndex, 1)
@@ -1140,7 +1146,13 @@ void* FAGXSurface::Lock(uint32 MipIndex, uint32 ArrayIndex, EResourceLockMode Lo
 					BytesPerRow = 0;
 					MipBytes = 0;
 				}
-				MTLPP_VALIDATE(mtlpp::Texture, Texture, AGXSafeGetRuntimeDebuggingLevel() >= EAGXDebugLevelValidation, GetBytes(MTLPP_VALIDATE(mtlpp::Buffer, SourceData, AGXSafeGetRuntimeDebuggingLevel() >= EAGXDebugLevelValidation, GetContents()), BytesPerRow, MipBytes, Region, MipIndex, ArrayIndex));
+
+				[Texture.GetPtr() getBytes:(void*)((uint8*)[SourceData.GetPtr() contents] + SourceData.GetOffset())
+							   bytesPerRow:BytesPerRow
+							 bytesPerImage:MipBytes
+								fromRegion:Region
+							   mipmapLevel:MipIndex
+									 slice:ArrayIndex];
 			}
 			
 #if PLATFORM_MAC
@@ -1500,7 +1512,7 @@ FTexture2DRHIRef FAGXDynamicRHI::RHIAsyncReallocateTexture2D(FRHITexture2D* OldT
 
 			// only handling straight 2D textures here
 			uint32 SliceIndex = 0;
-			mtlpp::Origin Origin(0, 0, 0);
+			MTLOrigin Origin = MTLOriginMake(0, 0, 0);
 
 			FAGXTexture Tex = OldTexture->Texture;
 
@@ -1515,7 +1527,7 @@ FTexture2DRHIRef FAGXDynamicRHI::RHIAsyncReallocateTexture2D(FRHITexture2D* OldT
 				const uint32 MipSizeX = FMath::Max<uint32>(1, NewSizeX >> (MipIndex + DestMipOffset));
 				const uint32 MipSizeY = FMath::Max<uint32>(1, NewSizeY >> (MipIndex + DestMipOffset));
 
-				bAsync &= Context.AsyncCopyFromTextureToTexture(OldTexture->Texture, SliceIndex, MipIndex + SourceMipOffset, Origin, mtlpp::Size(MipSizeX, MipSizeY, 1), NewTexture->Texture, SliceIndex, MipIndex + DestMipOffset, Origin);
+				bAsync &= Context.AsyncCopyFromTextureToTexture(OldTexture->Texture, SliceIndex, MipIndex + SourceMipOffset, Origin, MTLSizeMake(MipSizeX, MipSizeY, 1), NewTexture->Texture, SliceIndex, MipIndex + DestMipOffset, Origin);
 			}
 
 			// when done, decrement the counter to indicate it's safe
@@ -1677,7 +1689,7 @@ static void InternalUpdateTexture2D(FAGXContext& Context, FRHITexture2D* Texture
 	FAGXSurface* Texture = ResourceCast(TextureRHI);
 	FAGXTexture Tex = Texture->Texture;
 	
-	mtlpp::Region Region = mtlpp::Region(UpdateRegion.DestX, UpdateRegion.DestY, UpdateRegion.Width, UpdateRegion.Height);
+	MTLRegion Region = MTLRegionMake2D(UpdateRegion.DestX, UpdateRegion.DestY, UpdateRegion.Width, UpdateRegion.Height);
 	
 	if ([Tex.GetPtr() storageMode] == MTLStorageModePrivate)
 	{
@@ -1701,7 +1713,10 @@ static void InternalUpdateTexture2D(FAGXContext& Context, FRHITexture2D* Texture
 	}
 	else
 	{
-		MTLPP_VALIDATE(mtlpp::Texture, Tex, AGXSafeGetRuntimeDebuggingLevel() >= EAGXDebugLevelValidation, Replace(Region, MipIndex, 0, MTLPP_VALIDATE(mtlpp::Buffer, Buffer, AGXSafeGetRuntimeDebuggingLevel() >= EAGXDebugLevelValidation, GetContents()), SourcePitch, 0));
+		[Tex.GetPtr() replaceRegion:Region
+						mipmapLevel:MipIndex
+						  withBytes:(const void*)((uint8*)[Buffer.GetPtr() contents] + Buffer.GetOffset())
+						bytesPerRow:SourcePitch];
 	}
 
 	FPlatformAtomics::InterlockedExchange(&Texture->Written, 1);
@@ -1772,8 +1787,12 @@ void FAGXDynamicRHI::RHIUpdateTexture2D(FRHITexture2D* TextureRHI, uint32 MipInd
 				SourceData = (uint8*)ExpandedData.GetData();
 			}
 #endif
-			mtlpp::Region Region = mtlpp::Region(UpdateRegion.DestX, UpdateRegion.DestY, UpdateRegion.Width, UpdateRegion.Height);
-			MTLPP_VALIDATE(mtlpp::Texture, Tex, AGXSafeGetRuntimeDebuggingLevel() >= EAGXDebugLevelValidation, Replace(Region, MipIndex, 0, SourceData, SourcePitch, 0));
+			MTLRegion Region = MTLRegionMake2D(UpdateRegion.DestX, UpdateRegion.DestY, UpdateRegion.Width, UpdateRegion.Height);
+
+			[Tex.GetPtr() replaceRegion:Region
+							mipmapLevel:MipIndex
+							  withBytes:(const void*)SourceData
+							bytesPerRow:SourcePitch];
 
 			FPlatformAtomics::InterlockedExchange(&Texture->Written, 1);
 		}
@@ -1821,7 +1840,7 @@ static void InternalUpdateTexture3D(FAGXContext& Context, FRHITexture3D* Texture
 	FAGXSurface* Texture = ResourceCast(TextureRHI);
 	FAGXTexture Tex = Texture->Texture;
 	
-	mtlpp::Region Region = mtlpp::Region(UpdateRegion.DestX, UpdateRegion.DestY, UpdateRegion.DestZ, UpdateRegion.Width, UpdateRegion.Height, UpdateRegion.Depth);
+	MTLRegion Region = MTLRegionMake3D(UpdateRegion.DestX, UpdateRegion.DestY, UpdateRegion.DestZ, UpdateRegion.Width, UpdateRegion.Height, UpdateRegion.Depth);
 	
 	if ([Tex.GetPtr() storageMode] == MTLStorageModePrivate)
 	{
@@ -1843,7 +1862,12 @@ static void InternalUpdateTexture3D(FAGXContext& Context, FRHITexture3D* Texture
 	}
 	else
 	{
-		MTLPP_VALIDATE(mtlpp::Texture, Tex, AGXSafeGetRuntimeDebuggingLevel() >= EAGXDebugLevelValidation, Replace(Region, MipIndex, 0, (uint8*)Buffer.GetContents(), SourceRowPitch, SourceDepthPitch));
+		[Tex.GetPtr() replaceRegion:Region
+						mipmapLevel:MipIndex
+							  slice:0
+						  withBytes:(const void*)((uint8*)[Buffer.GetPtr() contents] + Buffer.GetOffset())
+						bytesPerRow:SourceRowPitch
+					  bytesPerImage:SourceDepthPitch];
 	}
 
 	FPlatformAtomics::InterlockedExchange(&Texture->Written, 1);
@@ -1941,8 +1965,15 @@ void FAGXDynamicRHI::RHIUpdateTexture3D(FRHITexture3D* TextureRHI,uint32 MipInde
 		}
 		else
 		{
-			mtlpp::Region Region = mtlpp::Region(UpdateRegion.DestX, UpdateRegion.DestY, UpdateRegion.DestZ, UpdateRegion.Width, UpdateRegion.Height, UpdateRegion.Depth);
-			MTLPP_VALIDATE(mtlpp::Texture, Tex, AGXSafeGetRuntimeDebuggingLevel() >= EAGXDebugLevelValidation, Replace(Region, MipIndex, 0, SourceData, SourceRowPitch, SourceDepthPitch));
+			MTLRegion Region = MTLRegionMake3D(UpdateRegion.DestX, UpdateRegion.DestY, UpdateRegion.DestZ, UpdateRegion.Width, UpdateRegion.Height, UpdateRegion.Depth);
+
+			[Tex.GetPtr() replaceRegion:Region
+							mipmapLevel:MipIndex
+								  slice:0
+							  withBytes:(const void*)SourceData
+							bytesPerRow:SourceRowPitch
+						  bytesPerImage:SourceDepthPitch];
+
 			FPlatformAtomics::InterlockedExchange(&Texture->Written, 1);
 		}
 		
@@ -2053,11 +2084,11 @@ void FAGXRHICommandContext::RHICopyTexture(FRHITexture* SourceTextureRHI, FRHITe
 				{
 					uint32 SourceMipIndex = CopyInfo.SourceMipIndex + MipIndex;
 					uint32 DestMipIndex = CopyInfo.DestMipIndex + MipIndex;
-					mtlpp::Size SourceSize(FMath::Max(Size.X >> MipIndex, 1), FMath::Max(Size.Y >> MipIndex, 1), FMath::Max(Size.Z >> MipIndex, 1));
-					mtlpp::Size DestSize = SourceSize;
+					MTLSize SourceSize = MTLSizeMake(FMath::Max(Size.X >> MipIndex, 1), FMath::Max(Size.Y >> MipIndex, 1), FMath::Max(Size.Z >> MipIndex, 1));
+					MTLSize DestSize = SourceSize;
 
-					mtlpp::Origin SourceOrigin(CopyInfo.SourcePosition.X >> MipIndex, CopyInfo.SourcePosition.Y >> MipIndex, CopyInfo.SourcePosition.Z >> MipIndex);
-					mtlpp::Origin DestinationOrigin(CopyInfo.DestPosition.X >> MipIndex, CopyInfo.DestPosition.Y >> MipIndex, CopyInfo.DestPosition.Z >> MipIndex);
+					MTLOrigin SourceOrigin = MTLOriginMake(CopyInfo.SourcePosition.X >> MipIndex, CopyInfo.SourcePosition.Y >> MipIndex, CopyInfo.SourcePosition.Z >> MipIndex);
+					MTLOrigin DestinationOrigin = MTLOriginMake(CopyInfo.DestPosition.X >> MipIndex, CopyInfo.DestPosition.Y >> MipIndex, CopyInfo.DestPosition.Z >> MipIndex);
 
 					if (TextureFormatCompatible)
 					{
@@ -2068,7 +2099,7 @@ void FAGXRHICommandContext::RHICopyTexture(FRHITexture* SourceTextureRHI, FRHITe
 					// Account for create with TexCreate_SRGB flag which could make these different
 					if (TextureFormatExactMatch && (SrcTexture.GetPixelFormat() == MetalDestTexture->Texture.GetPixelFormat()))
 					{
-						GetInternalContext().CopyFromTextureToTexture(SrcTexture, SourceSliceIndex, SourceMipIndex, SourceOrigin,SourceSize,MetalDestTexture->Texture, DestSliceIndex, DestMipIndex, DestinationOrigin);
+						GetInternalContext().CopyFromTextureToTexture(SrcTexture, SourceSliceIndex, SourceMipIndex, SourceOrigin, SourceSize, MetalDestTexture->Texture, DestSliceIndex, DestMipIndex, DestinationOrigin);
 					}
 					else
 					{
diff --git a/Engine/Source/Runtime/Apple/AGXRHI/Private/AGXUAV.cpp b/Engine/Source/Runtime/Apple/AGXRHI/Private/AGXUAV.cpp
--- a/Engine/Source/Runtime/Apple/AGXRHI/Private/AGXUAV.cpp
+++ b/Engine/Source/Runtime/Apple/AGXRHI/Private/AGXUAV.cpp
@@ -51,7 +51,7 @@ FAGXResourceViewBase::FAGXResourceViewBase(FRHIBuffer* InBuffer, uint32 InStartO
 		}
 		else
 		{
-			Format = (SourceBuffer->IndexType == mtlpp::IndexType::UInt16) ? PF_R16_UINT : PF_R32_UINT;
+			Format = (SourceBuffer->IndexType == MTLIndexTypeUInt16) ? PF_R16_UINT : PF_R32_UINT;
 			Stride = SourceBuffer->GetStride();
 
 			check(Stride == ((Format == PF_R16_UINT) ? 2 : 4));
@@ -460,7 +460,7 @@ void FAGXDynamicRHI::RHIUpdateShaderResourceView(FRHIShaderResourceView* SRVRHI,
 	SRV->SourceBuffer = ResourceCast(BufferRHI);
 	SRV->Stride = 0;
 
-	SRV->Format = SRV->SourceBuffer && SRV->SourceBuffer->IndexType != mtlpp::IndexType::UInt16
+	SRV->Format = SRV->SourceBuffer && SRV->SourceBuffer->IndexType != MTLIndexTypeUInt16
 		? PF_R32_UINT
 		: PF_R16_UINT;
 
diff --git a/Engine/Source/Runtime/Apple/AGXRHI/Private/AGXVertexBuffer.cpp b/Engine/Source/Runtime/Apple/AGXRHI/Private/AGXVertexBuffer.cpp
--- a/Engine/Source/Runtime/Apple/AGXRHI/Private/AGXVertexBuffer.cpp
+++ b/Engine/Source/Runtime/Apple/AGXRHI/Private/AGXVertexBuffer.cpp
@@ -602,7 +602,7 @@ void FAGXRHIBuffer::Unlock()
 FAGXResourceMultiBuffer::FAGXResourceMultiBuffer(uint32 InSize, EBufferUsageFlags InUsage, EAGXBufferUsage InAgxUsage, uint32 InStride, FResourceArrayInterface* ResourceArray, ERHIResourceType Type)
 	: FRHIBuffer(InSize, InUsage, InStride)
 	, FAGXRHIBuffer(InSize, InUsage, InAgxUsage, Type)
-	, IndexType((InStride == 2) ? mtlpp::IndexType::UInt16 : mtlpp::IndexType::UInt32)
+	, IndexType((InStride == 2) ? MTLIndexTypeUInt16 : MTLIndexTypeUInt32)
 {
 	if (EnumHasAnyFlags(InUsage, BUF_StructuredBuffer))
 	{
diff --git a/Engine/Source/Runtime/Apple/AGXRHI/Private/Shaders/AGXCompiledShaderCache.h b/Engine/Source/Runtime/Apple/AGXRHI/Private/Shaders/AGXCompiledShaderCache.h
--- a/Engine/Source/Runtime/Apple/AGXRHI/Private/Shaders/AGXCompiledShaderCache.h
+++ b/Engine/Source/Runtime/Apple/AGXRHI/Private/Shaders/AGXCompiledShaderCache.h
@@ -48,7 +48,7 @@ public:
 private:
 	FRWLock Lock;
 	TMap<FAGXCompiledShaderKey, mtlpp::Function> Cache;
-	TMap<mtlpp::Function::Type, mtlpp::Library> LibCache;
+	TMap<id<MTLFunction>, mtlpp::Library> LibCache;
 };
 
 extern FAGXCompiledShaderCache& GetAGXCompiledShaderCache();
diff --git a/Engine/Source/Runtime/Apple/AGXRHI/Private/Shaders/Types/Templates/AGXBaseShader.h b/Engine/Source/Runtime/Apple/AGXRHI/Private/Shaders/Types/Templates/AGXBaseShader.h
--- a/Engine/Source/Runtime/Apple/AGXRHI/Private/Shaders/Types/Templates/AGXBaseShader.h
+++ b/Engine/Source/Runtime/Apple/AGXRHI/Private/Shaders/Types/Templates/AGXBaseShader.h
@@ -20,7 +20,7 @@
 #pragma mark - AGX RHI Base Shader Class Support Routines
 
 
-extern mtlpp::LanguageVersion AGXValidateVersion(uint32 Version);
+extern MTLLanguageVersion AGXValidateVersion(uint32 Version);
 
 
 //------------------------------------------------------------------------------
@@ -263,23 +263,23 @@ void TAGXBaseShader<BaseResourceType, ShaderType>::Init(TArrayView<const uint8>
 
 			NSString* NewShaderString = [ShaderString stringByReplacingOccurrencesOfString:@"#pragma once" withString:@""];
 
-			mtlpp::CompileOptions CompileOptions;
+			MTLCompileOptions* CompileOptions = [[MTLCompileOptions alloc] init];
 
 #if DEBUG_METAL_SHADERS
 			static bool bForceFastMath = FParse::Param(FCommandLine::Get(), TEXT("metalfastmath"));
 			static bool bForceNoFastMath = FParse::Param(FCommandLine::Get(), TEXT("metalnofastmath"));
 			if (bForceNoFastMath)
 			{
-				CompileOptions.SetFastMathEnabled(NO);
+				[CompileOptions setFastMathEnabled:NO];
 			}
 			else if (bForceFastMath)
 			{
-				CompileOptions.SetFastMathEnabled(YES);
+				[CompileOptions setFastMathEnabled:YES];
 			}
 			else
 #endif
 			{
-				CompileOptions.SetFastMathEnabled((BOOL)(!(Header.CompileFlags & (1 << CFLAG_NoFastMath))));
+				[CompileOptions setFastMathEnabled:(BOOL)(!(Header.CompileFlags & (1 << CFLAG_NoFastMath)))];
 			}
 
 #if !PLATFORM_MAC || DEBUG_METAL_SHADERS
@@ -290,42 +290,50 @@ void TAGXBaseShader<BaseResourceType, ShaderType>::Init(TArrayView<const uint8>
 #if DEBUG_METAL_SHADERS
 			[PreprocessorMacros addEntriesFromDictionary: @{ @"MTLSL_ENABLE_DEBUG_INFO" : @(1)}];
 #endif
-			CompileOptions.SetPreprocessorMacros(PreprocessorMacros);
+			[CompileOptions setPreprocessorMacros:PreprocessorMacros];
 #endif
 
-			mtlpp::LanguageVersion MetalVersion;
-			switch(Header.Version)
+			MTLLanguageVersion MetalVersion;
+			switch (Header.Version)
 			{
 				case 7:
-					MetalVersion = mtlpp::LanguageVersion::Version2_4;
+					MetalVersion = MTLLanguageVersion2_4;
 					break;
+
 				case 6:
-					MetalVersion = mtlpp::LanguageVersion::Version2_3;
+					MetalVersion = MTLLanguageVersion2_3;
 					break;
+
 				case 5:
-					MetalVersion = mtlpp::LanguageVersion::Version2_2;
+					MetalVersion = MTLLanguageVersion2_2;
 					break;
+
 				default:
 					UE_LOG(LogRHI, Fatal, TEXT("Failed to create shader with unknown version %d: %s"), Header.Version, *FString(NewShaderString));
-					MetalVersion = mtlpp::LanguageVersion::Version2_2;
+					MetalVersion = MTLLanguageVersion2_2;
 					break;
-				}
-			CompileOptions.SetLanguageVersion(MetalVersion);
+			}
+
+			[CompileOptions setLanguageVersion:MetalVersion];
 
-			ns::AutoReleasedError Error;
-			Library = GMtlppDevice.NewLibrary(NewShaderString, CompileOptions, &Error);
-			if (Library == nil)
+			NSError* Error = nil;
+			id<MTLLibrary> NewLibrary = [GMtlDevice newLibraryWithSource:NewShaderString options:CompileOptions error:&Error];
+			if (NewLibrary == nil)
 			{
 				UE_LOG(LogRHI, Error, TEXT("*********** Error\n%s"), *FString(NewShaderString));
-				UE_LOG(LogRHI, Fatal, TEXT("Failed to create shader: %s"), *FString([Error.GetPtr() description]));
+				UE_LOG(LogRHI, Fatal, TEXT("Failed to create shader: %s"), *FString([Error description]));
 			}
 			else if (Error != nil)
 			{
 				// Warning...
 				UE_LOG(LogRHI, Warning, TEXT("*********** Warning\n%s"), *FString(NewShaderString));
-				UE_LOG(LogRHI, Warning, TEXT("Created shader with warnings: %s"), *FString([Error.GetPtr() description]));
+				UE_LOG(LogRHI, Warning, TEXT("Created shader with warnings: %s"), *FString([Error description]));
 			}
 
+			Library = mtlpp::Library(NewLibrary, nullptr, ns::Ownership::Assign);
+
+			[CompileOptions release];
+
 			GlslCodeNSString = NewShaderString;
 			[GlslCodeNSString retain];
 		}
diff --git a/Engine/Source/Runtime/Apple/AGXRHI/Public/AGXResources.h b/Engine/Source/Runtime/Apple/AGXRHI/Public/AGXResources.h
--- a/Engine/Source/Runtime/Apple/AGXRHI/Public/AGXResources.h
+++ b/Engine/Source/Runtime/Apple/AGXRHI/Public/AGXResources.h
@@ -477,7 +477,7 @@ public:
 	void Swap(FAGXResourceMultiBuffer& Other);
 
 	// 16- or 32-bit; used for index buffers only.
-	mtlpp::IndexType IndexType;
+	MTLIndexType IndexType;
 };
 
 typedef FAGXResourceMultiBuffer FAGXIndexBuffer;
