AGXRHI: Remove yet more mtlpp types in favor of native Metal types

- mtlpp::VertexFormat
- mtlpp::BlitOption

#rnx

diff --git a/Engine/Source/Runtime/Apple/AGXRHI/Private/AGXCommandQueue.cpp b/Engine/Source/Runtime/Apple/AGXRHI/Private/AGXCommandQueue.cpp
--- a/Engine/Source/Runtime/Apple/AGXRHI/Private/AGXCommandQueue.cpp
+++ b/Engine/Source/Runtime/Apple/AGXRHI/Private/AGXCommandQueue.cpp
@@ -13,7 +13,7 @@
 #include "command_buffer.hpp"
 
 uint64 FAGXCommandQueue::Features = 0;
-extern mtlpp::VertexFormat GAGXFColorVertexFormat;
+extern MTLVertexFormat GAGXFColorVertexFormat;
 
 FAGXCommandQueue::FAGXCommandQueue(uint32 const MaxNumCommandBuffers /* = 0 */)
 : ParallelCommandLists(0)
@@ -51,7 +51,7 @@ FAGXCommandQueue::FAGXCommandQueue(uint32 const MaxNumCommandBuffers /* = 0 */)
 	}
 
 	Features |= EAGXFeaturesPrivateBufferSubAllocation | EAGXFeaturesGPUCaptureManager | EAGXFeaturesBufferSubAllocation | EAGXFeaturesParallelRenderEncoders | EAGXFeaturesPipelineBufferMutability | EAGXFeaturesMaxThreadsPerThreadgroup;
-	GAGXFColorVertexFormat = mtlpp::VertexFormat::UChar4Normalized_BGRA;
+	GAGXFColorVertexFormat = MTLVertexFormatUChar4Normalized_BGRA;
 
 #else
 	if ([GMtlDevice supportsFeatureSet:MTLFeatureSet_iOS_GPUFamily3_v1])
@@ -70,7 +70,7 @@ FAGXCommandQueue::FAGXCommandQueue(uint32 const MaxNumCommandBuffers /* = 0 */)
 	// These options are fine however as thye just change how we allocate small buffers
 	Features |= EAGXFeaturesBufferSubAllocation;
 	Features |= EAGXFeaturesPrivateBufferSubAllocation;
-	GAGXFColorVertexFormat = mtlpp::VertexFormat::UChar4Normalized_BGRA;
+	GAGXFColorVertexFormat = MTLVertexFormatUChar4Normalized_BGRA;
 
 	Features |= EAGXFeaturesPresentMinDuration
 		| EAGXFeaturesGPUCaptureManager
@@ -108,7 +108,7 @@ FAGXCommandQueue::FAGXCommandQueue(uint32 const MaxNumCommandBuffers /* = 0 */)
 			| EAGXFeaturesPrivateBufferSubAllocation
 					  | EAGXFeaturesMaxThreadsPerThreadgroup );
 		
-		GAGXFColorVertexFormat = mtlpp::VertexFormat::UChar4Normalized_BGRA;
+		GAGXFColorVertexFormat = MTLVertexFormatUChar4Normalized_BGRA;
 		
 		if ([[GMtlDevice name] rangeOfString:@"Nvidia" options:NSCaseInsensitiveSearch].location == NSNotFound && !FParse::Param(FCommandLine::Get(),TEXT("nometalparallelencoder")))
 		{
diff --git a/Engine/Source/Runtime/Apple/AGXRHI/Private/AGXContext.cpp b/Engine/Source/Runtime/Apple/AGXRHI/Private/AGXContext.cpp
--- a/Engine/Source/Runtime/Apple/AGXRHI/Private/AGXContext.cpp
+++ b/Engine/Source/Runtime/Apple/AGXRHI/Private/AGXContext.cpp
@@ -1412,12 +1412,12 @@ void FAGXContext::DrawIndexedPrimitiveIndirect(uint32 PrimitiveType,FAGXIndexBuf
 	RenderPass.DrawIndexedPrimitiveIndirect(PrimitiveType, IndexBuffer, VertexBuffer, ArgumentOffset);
 }
 
-void FAGXContext::CopyFromTextureToBuffer(FAGXTexture const& Texture, uint32 sourceSlice, uint32 sourceLevel, MTLOrigin sourceOrigin, MTLSize sourceSize, FAGXBuffer const& toBuffer, uint32 destinationOffset, uint32 destinationBytesPerRow, uint32 destinationBytesPerImage, mtlpp::BlitOption options)
+void FAGXContext::CopyFromTextureToBuffer(FAGXTexture const& Texture, uint32 sourceSlice, uint32 sourceLevel, MTLOrigin sourceOrigin, MTLSize sourceSize, FAGXBuffer const& toBuffer, uint32 destinationOffset, uint32 destinationBytesPerRow, uint32 destinationBytesPerImage, MTLBlitOption options)
 {
 	RenderPass.CopyFromTextureToBuffer(Texture, sourceSlice, sourceLevel, sourceOrigin, sourceSize, toBuffer, destinationOffset, destinationBytesPerRow, destinationBytesPerImage, options);
 }
 
-void FAGXContext::CopyFromBufferToTexture(FAGXBuffer const& Buffer, uint32 sourceOffset, uint32 sourceBytesPerRow, uint32 sourceBytesPerImage, MTLSize sourceSize, FAGXTexture const& toTexture, uint32 destinationSlice, uint32 destinationLevel, MTLOrigin destinationOrigin, mtlpp::BlitOption options)
+void FAGXContext::CopyFromBufferToTexture(FAGXBuffer const& Buffer, uint32 sourceOffset, uint32 sourceBytesPerRow, uint32 sourceBytesPerImage, MTLSize sourceSize, FAGXTexture const& toTexture, uint32 destinationSlice, uint32 destinationLevel, MTLOrigin destinationOrigin, MTLBlitOption options)
 {
 	RenderPass.CopyFromBufferToTexture(Buffer, sourceOffset, sourceBytesPerRow, sourceBytesPerImage, sourceSize, toTexture, destinationSlice, destinationLevel, destinationOrigin, options);
 }
@@ -1432,7 +1432,7 @@ void FAGXContext::CopyFromBufferToBuffer(FAGXBuffer const& SourceBuffer, NSUInte
 	RenderPass.CopyFromBufferToBuffer(SourceBuffer, SourceOffset, DestinationBuffer, DestinationOffset, Size);
 }
 
-bool FAGXContext::AsyncCopyFromBufferToTexture(FAGXBuffer const& Buffer, uint32 sourceOffset, uint32 sourceBytesPerRow, uint32 sourceBytesPerImage, MTLSize sourceSize, FAGXTexture const& toTexture, uint32 destinationSlice, uint32 destinationLevel, MTLOrigin destinationOrigin, mtlpp::BlitOption options)
+bool FAGXContext::AsyncCopyFromBufferToTexture(FAGXBuffer const& Buffer, uint32 sourceOffset, uint32 sourceBytesPerRow, uint32 sourceBytesPerImage, MTLSize sourceSize, FAGXTexture const& toTexture, uint32 destinationSlice, uint32 destinationLevel, MTLOrigin destinationOrigin, MTLBlitOption options)
 {
 	return RenderPass.AsyncCopyFromBufferToTexture(Buffer, sourceOffset, sourceBytesPerRow, sourceBytesPerImage, sourceSize, toTexture, destinationSlice, destinationLevel, destinationOrigin, options);
 }
diff --git a/Engine/Source/Runtime/Apple/AGXRHI/Private/AGXContext.h b/Engine/Source/Runtime/Apple/AGXRHI/Private/AGXContext.h
--- a/Engine/Source/Runtime/Apple/AGXRHI/Private/AGXContext.h
+++ b/Engine/Source/Runtime/Apple/AGXRHI/Private/AGXContext.h
@@ -81,15 +81,15 @@ public:
 	void DrawPatches(uint32 PrimitiveType, FAGXBuffer const& IndexBuffer, uint32 IndexBufferStride, int32 BaseVertexIndex, uint32 FirstInstance, uint32 StartIndex,
 					 uint32 NumPrimitives, uint32 NumInstances);
 	
-	void CopyFromTextureToBuffer(FAGXTexture const& Texture, uint32 sourceSlice, uint32 sourceLevel, MTLOrigin sourceOrigin, MTLSize sourceSize, FAGXBuffer const& toBuffer, uint32 destinationOffset, uint32 destinationBytesPerRow, uint32 destinationBytesPerImage, mtlpp::BlitOption options);
+	void CopyFromTextureToBuffer(FAGXTexture const& Texture, uint32 sourceSlice, uint32 sourceLevel, MTLOrigin sourceOrigin, MTLSize sourceSize, FAGXBuffer const& toBuffer, uint32 destinationOffset, uint32 destinationBytesPerRow, uint32 destinationBytesPerImage, MTLBlitOption options);
 	
-	void CopyFromBufferToTexture(FAGXBuffer const& Buffer, uint32 sourceOffset, uint32 sourceBytesPerRow, uint32 sourceBytesPerImage, MTLSize sourceSize, FAGXTexture const& toTexture, uint32 destinationSlice, uint32 destinationLevel, MTLOrigin destinationOrigin, mtlpp::BlitOption options);
+	void CopyFromBufferToTexture(FAGXBuffer const& Buffer, uint32 sourceOffset, uint32 sourceBytesPerRow, uint32 sourceBytesPerImage, MTLSize sourceSize, FAGXTexture const& toTexture, uint32 destinationSlice, uint32 destinationLevel, MTLOrigin destinationOrigin, MTLBlitOption options);
 	
 	void CopyFromTextureToTexture(FAGXTexture const& Texture, uint32 sourceSlice, uint32 sourceLevel, MTLOrigin sourceOrigin, MTLSize sourceSize, FAGXTexture const& toTexture, uint32 destinationSlice, uint32 destinationLevel, MTLOrigin destinationOrigin);
 	
 	void CopyFromBufferToBuffer(FAGXBuffer const& SourceBuffer, NSUInteger SourceOffset, FAGXBuffer const& DestinationBuffer, NSUInteger DestinationOffset, NSUInteger Size);
 	
-    bool AsyncCopyFromBufferToTexture(FAGXBuffer const& Buffer, uint32 sourceOffset, uint32 sourceBytesPerRow, uint32 sourceBytesPerImage, MTLSize sourceSize, FAGXTexture const& toTexture, uint32 destinationSlice, uint32 destinationLevel, MTLOrigin destinationOrigin, mtlpp::BlitOption options);
+    bool AsyncCopyFromBufferToTexture(FAGXBuffer const& Buffer, uint32 sourceOffset, uint32 sourceBytesPerRow, uint32 sourceBytesPerImage, MTLSize sourceSize, FAGXTexture const& toTexture, uint32 destinationSlice, uint32 destinationLevel, MTLOrigin destinationOrigin, MTLBlitOption options);
     
     bool AsyncCopyFromTextureToTexture(FAGXTexture const& Texture, uint32 sourceSlice, uint32 sourceLevel, MTLOrigin sourceOrigin, MTLSize sourceSize, FAGXTexture const& toTexture, uint32 destinationSlice, uint32 destinationLevel, MTLOrigin destinationOrigin);
 	
diff --git a/Engine/Source/Runtime/Apple/AGXRHI/Private/AGXRenderPass.cpp b/Engine/Source/Runtime/Apple/AGXRHI/Private/AGXRenderPass.cpp
--- a/Engine/Source/Runtime/Apple/AGXRHI/Private/AGXRenderPass.cpp
+++ b/Engine/Source/Runtime/Apple/AGXRHI/Private/AGXRenderPass.cpp
@@ -561,7 +561,7 @@ void FAGXRenderPass::InsertTextureBarrier()
 #endif
 }
 
-void FAGXRenderPass::CopyFromTextureToBuffer(FAGXTexture const& Texture, uint32 sourceSlice, uint32 sourceLevel, MTLOrigin sourceOrigin, MTLSize sourceSize, FAGXBuffer const& toBuffer, uint32 destinationOffset, uint32 destinationBytesPerRow, uint32 destinationBytesPerImage, mtlpp::BlitOption options)
+void FAGXRenderPass::CopyFromTextureToBuffer(FAGXTexture const& Texture, uint32 sourceSlice, uint32 sourceLevel, MTLOrigin sourceOrigin, MTLSize sourceSize, FAGXBuffer const& toBuffer, uint32 destinationOffset, uint32 destinationBytesPerRow, uint32 destinationBytesPerImage, MTLBlitOption options)
 {
 	ConditionalSwitchToBlit();
 	
@@ -579,13 +579,13 @@ void FAGXRenderPass::CopyFromTextureToBuffer(FAGXTexture const& Texture, uint32
 						   destinationOffset:(NSUInteger)destinationOffset + toBuffer.GetOffset()
 					  destinationBytesPerRow:(NSUInteger)destinationBytesPerRow
 					destinationBytesPerImage:(NSUInteger)destinationBytesPerImage
-									 options:*reinterpret_cast<MTLBlitOption*>(&options)];
+									 options:options];
 	}
 	
 	ConditionalSubmit();
 }
 
-void FAGXRenderPass::CopyFromBufferToTexture(FAGXBuffer const& Buffer, uint32 sourceOffset, uint32 sourceBytesPerRow, uint32 sourceBytesPerImage, MTLSize sourceSize, FAGXTexture const& toTexture, uint32 destinationSlice, uint32 destinationLevel, MTLOrigin destinationOrigin, mtlpp::BlitOption options)
+void FAGXRenderPass::CopyFromBufferToTexture(FAGXBuffer const& Buffer, uint32 sourceOffset, uint32 sourceBytesPerRow, uint32 sourceBytesPerImage, MTLSize sourceSize, FAGXTexture const& toTexture, uint32 destinationSlice, uint32 destinationLevel, MTLOrigin destinationOrigin, MTLBlitOption options)
 {
 	ConditionalSwitchToBlit();
 	
@@ -594,7 +594,7 @@ void FAGXRenderPass::CopyFromBufferToTexture(FAGXBuffer const& Buffer, uint32 so
 	
 	METAL_GPUPROFILE(FAGXProfiler::GetProfiler()->EncodeBlit(CurrentEncoder.GetCommandBufferStats(), __FUNCTION__));
 	
-	if (options == mtlpp::BlitOption::None)
+	if (options == MTLBlitOptionNone)
 	{
 		[Encoder		 copyFromBuffer:Buffer.GetPtr()
 						   sourceOffset:(NSUInteger)sourceOffset + Buffer.GetOffset()
@@ -617,7 +617,7 @@ void FAGXRenderPass::CopyFromBufferToTexture(FAGXBuffer const& Buffer, uint32 so
 					   destinationSlice:(NSUInteger)destinationSlice
 					   destinationLevel:(NSUInteger)destinationLevel
 					  destinationOrigin:destinationOrigin
-								options:*reinterpret_cast<MTLBlitOption*>(&options)];
+								options:options];
 	}
 	
 	ConditionalSubmit();
@@ -744,7 +744,7 @@ void FAGXRenderPass::FillBuffer(FAGXBuffer const& Buffer, ns::Range Range, uint8
 	}
 }
 
-bool FAGXRenderPass::AsyncCopyFromBufferToTexture(FAGXBuffer const& Buffer, uint32 sourceOffset, uint32 sourceBytesPerRow, uint32 sourceBytesPerImage, MTLSize sourceSize, FAGXTexture const& toTexture, uint32 destinationSlice, uint32 destinationLevel, MTLOrigin destinationOrigin, mtlpp::BlitOption options)
+bool FAGXRenderPass::AsyncCopyFromBufferToTexture(FAGXBuffer const& Buffer, uint32 sourceOffset, uint32 sourceBytesPerRow, uint32 sourceBytesPerImage, MTLSize sourceSize, FAGXTexture const& toTexture, uint32 destinationSlice, uint32 destinationLevel, MTLOrigin destinationOrigin, MTLBlitOption options)
 {
 	id<MTLBlitCommandEncoder> TargetEncoder = nil;
 	bool bAsync = !CurrentEncoder.HasTextureBindingHistory(toTexture);
@@ -763,7 +763,7 @@ bool FAGXRenderPass::AsyncCopyFromBufferToTexture(FAGXBuffer const& Buffer, uint
 	
 	check(TargetEncoder);
 	
-	if (options == mtlpp::BlitOption::None)
+	if (options == MTLBlitOptionNone)
 	{
 		[TargetEncoder		 copyFromBuffer:Buffer.GetPtr()
 							   sourceOffset:(NSUInteger)sourceOffset + Buffer.GetOffset()
@@ -786,7 +786,7 @@ bool FAGXRenderPass::AsyncCopyFromBufferToTexture(FAGXBuffer const& Buffer, uint
 						   destinationSlice:(NSUInteger)destinationSlice
 						   destinationLevel:(NSUInteger)destinationLevel
 						  destinationOrigin:destinationOrigin
-									options:*reinterpret_cast<MTLBlitOption*>(&options)];
+									options:options];
 	}
 	
 	return bAsync;
diff --git a/Engine/Source/Runtime/Apple/AGXRHI/Private/AGXRenderPass.h b/Engine/Source/Runtime/Apple/AGXRHI/Private/AGXRenderPass.h
--- a/Engine/Source/Runtime/Apple/AGXRHI/Private/AGXRenderPass.h
+++ b/Engine/Source/Runtime/Apple/AGXRHI/Private/AGXRenderPass.h
@@ -48,9 +48,9 @@ public:
     
     void EndRenderPass(void);
     
-    void CopyFromTextureToBuffer(FAGXTexture const& Texture, uint32 sourceSlice, uint32 sourceLevel, MTLOrigin sourceOrigin, MTLSize sourceSize, FAGXBuffer const& toBuffer, uint32 destinationOffset, uint32 destinationBytesPerRow, uint32 destinationBytesPerImage, mtlpp::BlitOption options);
+    void CopyFromTextureToBuffer(FAGXTexture const& Texture, uint32 sourceSlice, uint32 sourceLevel, MTLOrigin sourceOrigin, MTLSize sourceSize, FAGXBuffer const& toBuffer, uint32 destinationOffset, uint32 destinationBytesPerRow, uint32 destinationBytesPerImage, MTLBlitOption options);
     
-    void CopyFromBufferToTexture(FAGXBuffer const& Buffer, uint32 sourceOffset, uint32 sourceBytesPerRow, uint32 sourceBytesPerImage, MTLSize sourceSize, FAGXTexture const& toTexture, uint32 destinationSlice, uint32 destinationLevel, MTLOrigin destinationOrigin, mtlpp::BlitOption options);
+    void CopyFromBufferToTexture(FAGXBuffer const& Buffer, uint32 sourceOffset, uint32 sourceBytesPerRow, uint32 sourceBytesPerImage, MTLSize sourceSize, FAGXTexture const& toTexture, uint32 destinationSlice, uint32 destinationLevel, MTLOrigin destinationOrigin, MTLBlitOption options);
     
     void CopyFromTextureToTexture(FAGXTexture const& Texture, uint32 sourceSlice, uint32 sourceLevel, MTLOrigin sourceOrigin, MTLSize sourceSize, FAGXTexture const& toTexture, uint32 destinationSlice, uint32 destinationLevel, MTLOrigin destinationOrigin);
 	
@@ -64,7 +64,7 @@ public:
     
 	void FillBuffer(FAGXBuffer const& Buffer, ns::Range Range, uint8 Value);
 	
-	bool AsyncCopyFromBufferToTexture(FAGXBuffer const& Buffer, uint32 sourceOffset, uint32 sourceBytesPerRow, uint32 sourceBytesPerImage, MTLSize sourceSize, FAGXTexture const& toTexture, uint32 destinationSlice, uint32 destinationLevel, MTLOrigin destinationOrigin, mtlpp::BlitOption options);
+	bool AsyncCopyFromBufferToTexture(FAGXBuffer const& Buffer, uint32 sourceOffset, uint32 sourceBytesPerRow, uint32 sourceBytesPerImage, MTLSize sourceSize, FAGXTexture const& toTexture, uint32 destinationSlice, uint32 destinationLevel, MTLOrigin destinationOrigin, MTLBlitOption options);
 	
 	bool AsyncCopyFromTextureToTexture(FAGXTexture const& Texture, uint32 sourceSlice, uint32 sourceLevel, MTLOrigin sourceOrigin, MTLSize sourceSize, FAGXTexture const& toTexture, uint32 destinationSlice, uint32 destinationLevel, MTLOrigin destinationOrigin);
 	
diff --git a/Engine/Source/Runtime/Apple/AGXRHI/Private/AGXRenderTarget.cpp b/Engine/Source/Runtime/Apple/AGXRHI/Private/AGXRenderTarget.cpp
--- a/Engine/Source/Runtime/Apple/AGXRHI/Private/AGXRenderTarget.cpp
+++ b/Engine/Source/Runtime/Apple/AGXRHI/Private/AGXRenderTarget.cpp
@@ -357,17 +357,17 @@ void FAGXDynamicRHI::RHIReadSurfaceData(FRHITexture* TextureRHI, FIntRect Rect,
 			
 			if (Surface->GetDesc().Format != PF_DepthStencil)
 			{
-				ImmediateContext.Context->CopyFromTextureToBuffer(Texture, 0, InFlags.GetMip(), Region.origin, Region.size, Buffer, 0, AlignedStride, BytesPerImage, mtlpp::BlitOption::None);
+				ImmediateContext.Context->CopyFromTextureToBuffer(Texture, 0, InFlags.GetMip(), Region.origin, Region.size, Buffer, 0, AlignedStride, BytesPerImage, MTLBlitOptionNone);
 			}
 			else
 			{
 				if (!InFlags.GetOutputStencil())
 				{
-					ImmediateContext.Context->CopyFromTextureToBuffer(Texture, 0, InFlags.GetMip(), Region.origin, Region.size, Buffer, 0, AlignedStride, BytesPerImage, mtlpp::BlitOption::DepthFromDepthStencil);
+					ImmediateContext.Context->CopyFromTextureToBuffer(Texture, 0, InFlags.GetMip(), Region.origin, Region.size, Buffer, 0, AlignedStride, BytesPerImage, MTLBlitOptionDepthFromDepthStencil);
 				}
 				else
 				{
-					ImmediateContext.Context->CopyFromTextureToBuffer(Texture, 0, InFlags.GetMip(), Region.origin, Region.size, Buffer, 0, AlignedStride, BytesPerImage, mtlpp::BlitOption::StencilFromDepthStencil);
+					ImmediateContext.Context->CopyFromTextureToBuffer(Texture, 0, InFlags.GetMip(), Region.origin, Region.size, Buffer, 0, AlignedStride, BytesPerImage, MTLBlitOptionStencilFromDepthStencil);
 				}
 			}
 			
@@ -459,7 +459,7 @@ void FAGXDynamicRHI::RHIReadSurfaceFloatData(FRHITexture* TextureRHI, FIntRect R
 		// Synchronise the texture with the CPU
 		SCOPE_CYCLE_COUNTER(STAT_AGXTexturePageOffTime);
 		
-		ImmediateContext.Context->CopyFromTextureToBuffer(Texture, ArrayIndex, MipIndex, Region.origin, Region.size, Buffer, 0, AlignedStride, BytesPerImage, mtlpp::BlitOption::None);
+		ImmediateContext.Context->CopyFromTextureToBuffer(Texture, ArrayIndex, MipIndex, Region.origin, Region.size, Buffer, 0, AlignedStride, BytesPerImage, MTLBlitOptionNone);
 		
 		//kick the current command buffer.
 		ImmediateContext.Context->SubmitCommandBufferAndWait();
@@ -525,7 +525,7 @@ void FAGXDynamicRHI::RHIRead3DSurfaceFloatData(FRHITexture* TextureRHI,FIntRect
 		// Synchronise the texture with the CPU
 		SCOPE_CYCLE_COUNTER(STAT_AGXTexturePageOffTime);
 		
-		ImmediateContext.Context->CopyFromTextureToBuffer(Texture, 0, 0, Region.origin, Region.size, Buffer, 0, AlignedStride, BytesPerImage, mtlpp::BlitOption::None);
+		ImmediateContext.Context->CopyFromTextureToBuffer(Texture, 0, 0, Region.origin, Region.size, Buffer, 0, AlignedStride, BytesPerImage, MTLBlitOptionNone);
 		
 		//kick the current command buffer.
 		ImmediateContext.Context->SubmitCommandBufferAndWait();
diff --git a/Engine/Source/Runtime/Apple/AGXRHI/Private/AGXTexture.cpp b/Engine/Source/Runtime/Apple/AGXRHI/Private/AGXTexture.cpp
--- a/Engine/Source/Runtime/Apple/AGXRHI/Private/AGXTexture.cpp
+++ b/Engine/Source/Runtime/Apple/AGXRHI/Private/AGXTexture.cpp
@@ -964,11 +964,11 @@ void FAGXSurface::UpdateSurfaceAndDestroySourceBuffer(id <MTLBuffer> SourceBuffe
 		
 		bool const bWait = ((GetAGXDeviceContext().GetNumActiveContexts() == 1) && (GAGXMaxOutstandingAsyncTexUploads > 0) && (Count >= GAGXMaxOutstandingAsyncTexUploads));
 		
-		mtlpp::BlitOption Options = mtlpp::BlitOption::None;
+		MTLBlitOption Options = MTLBlitOptionNone;
 #if !PLATFORM_MAC
 		if ((MTLPixelFormat)Texture.GetPixelFormat() >= MTLPixelFormatPVRTC_RGB_2BPP && (MTLPixelFormat)Texture.GetPixelFormat() <= MTLPixelFormatPVRTC_RGBA_4BPP_sRGB)
 		{
-			Options = mtlpp::BlitOption::RowLinearPVRTC;
+			Options = MTLBlitOptionRowLinearPVRTC;
 		}
 #endif
 		
@@ -1088,7 +1088,7 @@ void* FAGXSurface::Lock(uint32 MipIndex, uint32 ArrayIndex, EResourceLockMode Lo
 				auto CopyTexToBuf =
 				[this, &ArrayIndex, &MipIndex, &Region, &SourceData, &DestStride, &MipBytes](FRHICommandListImmediate& RHICmdList)
 				{
-					GetAGXDeviceContext().CopyFromTextureToBuffer(this->Texture, ArrayIndex, MipIndex, Region.origin, Region.size, SourceData, 0, DestStride, MipBytes, mtlpp::BlitOption::None);
+					GetAGXDeviceContext().CopyFromTextureToBuffer(this->Texture, ArrayIndex, MipIndex, Region.origin, Region.size, SourceData, 0, DestStride, MipBytes, MTLBlitOptionNone);
 					//kick the current command buffer.
 					GetAGXDeviceContext().SubmitCommandBufferAndWait();
 				};
@@ -1699,11 +1699,11 @@ static void InternalUpdateTexture2D(FAGXContext& Context, FRHITexture2D* Texture
 		const uint32 NumRows = UpdateRegion.Height / (uint32)FormatInfo.BlockSizeY;
 		uint32 BytesPerImage = SourcePitch * NumRows;
 		
-		mtlpp::BlitOption Options = mtlpp::BlitOption::None;
+		MTLBlitOption Options = MTLBlitOptionNone;
 #if !PLATFORM_MAC
 		if ((MTLPixelFormat)Tex.GetPixelFormat() >= MTLPixelFormatPVRTC_RGB_2BPP && (MTLPixelFormat)Tex.GetPixelFormat() <= MTLPixelFormatPVRTC_RGBA_4BPP_sRGB)
 		{
-			Options = mtlpp::BlitOption::RowLinearPVRTC;
+			Options = MTLBlitOptionRowLinearPVRTC;
 		}
 #endif
 		if(Context.AsyncCopyFromBufferToTexture(Buffer, 0, SourcePitch, BytesPerImage, Region.size, Tex, 0, MipIndex, Region.origin, Options))
@@ -1848,11 +1848,11 @@ static void InternalUpdateTexture3D(FAGXContext& Context, FRHITexture3D* Texture
 		const uint32 NumRows = UpdateRegion.Height / (uint32)FormatInfo.BlockSizeY;
 		const uint32 BytesPerImage = SourceRowPitch * NumRows;
 
-		mtlpp::BlitOption Options = mtlpp::BlitOption::None;
+		MTLBlitOption Options = MTLBlitOptionNone;
 #if !PLATFORM_MAC
 		if ((MTLPixelFormat)Tex.GetPixelFormat() >= MTLPixelFormatPVRTC_RGB_2BPP && (MTLPixelFormat)Tex.GetPixelFormat() <= MTLPixelFormatPVRTC_RGBA_4BPP_sRGB)
 		{
-			Options = mtlpp::BlitOption::RowLinearPVRTC;
+			Options = MTLBlitOptionRowLinearPVRTC;
 		}
 #endif
 		if(Context.AsyncCopyFromBufferToTexture(Buffer, 0, SourceRowPitch, BytesPerImage, Region.size, Tex, 0, MipIndex, Region.origin, Options))
@@ -2125,11 +2125,11 @@ void FAGXRHICommandContext::RHICopyTexture(FRHITexture* SourceTextureRHI, FRHITe
 						
 						check(Buffer);
 						
-						mtlpp::BlitOption Options = mtlpp::BlitOption::None;
+						MTLBlitOption Options = MTLBlitOptionNone;
 #if !PLATFORM_MAC
 						if ((MTLPixelFormat)MetalSrcTexture->Texture.GetPixelFormat() >= MTLPixelFormatPVRTC_RGB_2BPP && (MTLPixelFormat)MetalSrcTexture->Texture.GetPixelFormat() <= MTLPixelFormatPVRTC_RGBA_4BPP_sRGB)
 						{
-							Options = mtlpp::BlitOption::RowLinearPVRTC;
+							Options = MTLBlitOptionRowLinearPVRTC;
 						}
 #endif
 						GetInternalContext().CopyFromTextureToBuffer(MetalSrcTexture->Texture, SourceSliceIndex, SourceMipIndex, SourceOrigin, SourceSize, Buffer, 0, AlignedStride, BytesPerImage, Options);
diff --git a/Engine/Source/Runtime/Apple/AGXRHI/Private/AGXVertexDeclaration.cpp b/Engine/Source/Runtime/Apple/AGXRHI/Private/AGXVertexDeclaration.cpp
--- a/Engine/Source/Runtime/Apple/AGXRHI/Private/AGXVertexDeclaration.cpp
+++ b/Engine/Source/Runtime/Apple/AGXRHI/Private/AGXVertexDeclaration.cpp
@@ -16,7 +16,7 @@
 #pragma mark - AGX RHI Vertex Declaration Globals
 
 
-mtlpp::VertexFormat GAGXFColorVertexFormat = mtlpp::VertexFormat::UChar4Normalized;
+MTLVertexFormat GAGXFColorVertexFormat = MTLVertexFormatUChar4Normalized;
 
 
 //------------------------------------------------------------------------------
@@ -24,31 +24,31 @@ mtlpp::VertexFormat GAGXFColorVertexFormat = mtlpp::VertexFormat::UChar4Normaliz
 #pragma mark - AGX RHI Vertex Declaration Support Routines
 
 
-static mtlpp::VertexFormat TranslateElementTypeToMTLType(EVertexElementType Type)
+static MTLVertexFormat TranslateElementTypeToMTLType(EVertexElementType Type)
 {
 	switch (Type)
 	{
-		case VET_Float1:		return mtlpp::VertexFormat::Float;
-		case VET_Float2:		return mtlpp::VertexFormat::Float2;
-		case VET_Float3:		return mtlpp::VertexFormat::Float3;
-		case VET_Float4:		return mtlpp::VertexFormat::Float4;
-		case VET_PackedNormal:	return mtlpp::VertexFormat::Char4Normalized;
-		case VET_UByte4:		return mtlpp::VertexFormat::UChar4;
-		case VET_UByte4N:		return mtlpp::VertexFormat::UChar4Normalized;
+		case VET_Float1:		return MTLVertexFormatFloat;
+		case VET_Float2:		return MTLVertexFormatFloat2;
+		case VET_Float3:		return MTLVertexFormatFloat3;
+		case VET_Float4:		return MTLVertexFormatFloat4;
+		case VET_PackedNormal:	return MTLVertexFormatChar4Normalized;
+		case VET_UByte4:		return MTLVertexFormatUChar4;
+		case VET_UByte4N:		return MTLVertexFormatUChar4Normalized;
 		case VET_Color:			return GAGXFColorVertexFormat;
-		case VET_Short2:		return mtlpp::VertexFormat::Short2;
-		case VET_Short4:		return mtlpp::VertexFormat::Short4;
-		case VET_Short2N:		return mtlpp::VertexFormat::Short2Normalized;
-		case VET_Half2:			return mtlpp::VertexFormat::Half2;
-		case VET_Half4:			return mtlpp::VertexFormat::Half4;
-		case VET_Short4N:		return mtlpp::VertexFormat::Short4Normalized;
-		case VET_UShort2:		return mtlpp::VertexFormat::UShort2;
-		case VET_UShort4:		return mtlpp::VertexFormat::UShort4;
-		case VET_UShort2N:		return mtlpp::VertexFormat::UShort2Normalized;
-		case VET_UShort4N:		return mtlpp::VertexFormat::UShort4Normalized;
-		case VET_URGB10A2N:		return mtlpp::VertexFormat::UInt1010102Normalized;
-		case VET_UInt:			return mtlpp::VertexFormat::UInt;
-        default:				METAL_FATAL_ERROR(TEXT("Unknown vertex element type %d!"), (uint32)Type); return mtlpp::VertexFormat::Float;
+		case VET_Short2:		return MTLVertexFormatShort2;
+		case VET_Short4:		return MTLVertexFormatShort4;
+		case VET_Short2N:		return MTLVertexFormatShort2Normalized;
+		case VET_Half2:			return MTLVertexFormatHalf2;
+		case VET_Half4:			return MTLVertexFormatHalf4;
+		case VET_Short4N:		return MTLVertexFormatShort4Normalized;
+		case VET_UShort2:		return MTLVertexFormatUShort2;
+		case VET_UShort4:		return MTLVertexFormatUShort4;
+		case VET_UShort2N:		return MTLVertexFormatUShort2Normalized;
+		case VET_UShort4N:		return MTLVertexFormatUShort4Normalized;
+		case VET_URGB10A2N:		return MTLVertexFormatUInt1010102Normalized;
+		case VET_UInt:			return MTLVertexFormatUInt;
+        default:				METAL_FATAL_ERROR(TEXT("Unknown vertex element type %u!"), (uint32)Type); return MTLVertexFormatFloat;
     };
 }
 
@@ -171,7 +171,7 @@ void FAGXVertexDeclaration::GenerateLayout(const FVertexDeclarationElementList&
 
 		// set the format for each element
 		mtlpp::VertexAttributeDescriptor Attrib = Attributes[Element.AttributeIndex];
-		Attrib.SetFormat(TranslateElementTypeToMTLType(Element.Type));
+		Attrib.SetFormat(mtlpp::VertexFormat(TranslateElementTypeToMTLType(Element.Type)));
 		Attrib.SetOffset(Element.Offset);
 		Attrib.SetBufferIndex(ShaderBufferIndex);
 	}
