AGXRHI: Remove mtlpp types from LLM and BlendState namespaces

#rnx

diff --git a/Engine/Source/Runtime/Apple/AGXRHI/Private/AGXBuffer.cpp b/Engine/Source/Runtime/Apple/AGXRHI/Private/AGXBuffer.cpp
--- a/Engine/Source/Runtime/Apple/AGXRHI/Private/AGXBuffer.cpp
+++ b/Engine/Source/Runtime/Apple/AGXRHI/Private/AGXBuffer.cpp
@@ -182,7 +182,7 @@ FAGXSubBufferHeap::FAGXSubBufferHeap(NSUInteger Size, NSUInteger Alignment, MTLR
 	check(ParentBuffer.GetPtr());
 	check(ParentBuffer.GetLength() >= FullSize);
 #if STATS || ENABLE_LOW_LEVEL_MEM_TRACKER
-	AGXLLM::LogAllocBuffer(ParentBuffer);
+	AGXLLM::LogAllocBuffer(ParentBuffer.GetPtr());
 #endif
 	FreeRanges.Add(ns::Range(0, FullSize));
 
@@ -414,7 +414,7 @@ FAGXSubBufferLinear::FAGXSubBufferLinear(NSUInteger Size, NSUInteger Alignment,
 	check(ParentBuffer.GetPtr());
 	check(ParentBuffer.GetLength() >= FullSize);
 #if STATS || ENABLE_LOW_LEVEL_MEM_TRACKER
-	AGXLLM::LogAllocBuffer(ParentBuffer);
+	AGXLLM::LogAllocBuffer(ParentBuffer.GetPtr());
 #endif
 	INC_MEMORY_STAT_BY(STAT_AGXBufferUnusedMemory, FullSize);
 	INC_MEMORY_STAT_BY(STAT_AGXLinearBufferUnusedMemory, FullSize);
@@ -534,7 +534,7 @@ FAGXSubBufferMagazine::FAGXSubBufferMagazine(NSUInteger Size, NSUInteger ChunkSi
 	check(ParentBuffer.GetLength() >= FullSize);
 	METAL_FATAL_ASSERT(ParentBuffer, TEXT("Failed to create heap of size %u and resource options %u"), Size, (uint32)Options);
 #if STATS || ENABLE_LOW_LEVEL_MEM_TRACKER
-	AGXLLM::LogAllocBuffer(ParentBuffer);
+	AGXLLM::LogAllocBuffer(ParentBuffer.GetPtr());
 #endif
 	
 	INC_MEMORY_STAT_BY(STAT_AGXBufferUnusedMemory, FullSize);
@@ -906,7 +906,7 @@ FAGXBuffer FAGXBufferPoolPolicyData::CreateResource(CreationArguments Args)
 	FAGXBuffer NewBuf(MTLPP_VALIDATE(mtlpp::Device, GMtlppDevice, AGXSafeGetRuntimeDebuggingLevel() >= EAGXDebugLevelValidation, NewBuffer(BufferSize, mtlpp::ResourceOptions(Args.Options)), true));
 
 #if STATS || ENABLE_LOW_LEVEL_MEM_TRACKER
-	AGXLLM::LogAllocBuffer(NewBuf);
+	AGXLLM::LogAllocBuffer(NewBuf.GetPtr());
 #endif
 	INC_MEMORY_STAT_BY(STAT_AGXBufferUnusedMemory, NewBuf.GetLength());
 	INC_MEMORY_STAT_BY(STAT_AGXPooledBufferUnusedMemory, NewBuf.GetLength());
@@ -970,7 +970,7 @@ FAGXTexture FAGXTexturePool::CreateTexture(mtlpp::TextureDescriptor Desc)
 		METAL_GPUPROFILE(FAGXScopedCPUStats CPUStat(FString::Printf(TEXT("AllocTexture: %s"), TEXT("")/**FString([Desc.GetPtr() description])*/)));
 		Texture = MTLPP_VALIDATE(mtlpp::Device, GMtlppDevice, AGXSafeGetRuntimeDebuggingLevel() >= EAGXDebugLevelValidation, NewTexture(Desc));
 #if STATS || ENABLE_LOW_LEVEL_MEM_TRACKER
-		AGXLLM::LogAllocTexture(Desc, Texture);
+		AGXLLM::LogAllocTexture(Desc.GetPtr(), Texture.GetPtr());
 #endif
 	}
 	return Texture;
@@ -1273,7 +1273,7 @@ FAGXBuffer FAGXResourceHeap::CreateBuffer(uint32 Size, uint32 Alignment, EBuffer
 		mtlpp::Buffer MtlppBuffer([GMtlDevice newBufferWithLength:BlockSize options:Options], nullptr, ns::Ownership::Assign);
 		Buffer = FAGXBuffer(MoveTemp(MtlppBuffer), false);
 #if STATS || ENABLE_LOW_LEVEL_MEM_TRACKER
-		AGXLLM::LogAllocBuffer(Buffer);
+		AGXLLM::LogAllocBuffer(Buffer.GetPtr());
 #endif
 		INC_MEMORY_STAT_BY(STAT_AGXDeviceBufferMemory, Buffer.GetLength());
 	}
diff --git a/Engine/Source/Runtime/Apple/AGXRHI/Private/AGXLLM.cpp b/Engine/Source/Runtime/Apple/AGXRHI/Private/AGXLLM.cpp
--- a/Engine/Source/Runtime/Apple/AGXRHI/Private/AGXLLM.cpp
+++ b/Engine/Source/Runtime/Apple/AGXRHI/Private/AGXLLM.cpp
@@ -143,19 +143,19 @@ static EPixelFormat MetalToRHIPixelFormat(MTLPixelFormat Format)
 	return PF_MAX;
 }
 
-void AGXLLM::LogAllocTexture(mtlpp::TextureDescriptor const& Desc, mtlpp::Texture const& Texture)
+void AGXLLM::LogAllocTexture(MTLTextureDescriptor* Desc, id<MTLTexture> Texture)
 {
 	MTLSizeAndAlign SizeAlign;
 	if (FAGXCommandQueue::SupportsFeature(EAGXFeaturesGPUCaptureManager))
 	{
-		SizeAlign = [GMtlDevice heapTextureSizeAndAlignWithDescriptor:Desc.GetPtr()];
+		SizeAlign = [GMtlDevice heapTextureSizeAndAlignWithDescriptor:Desc];
 	}
 	
-	void* Ptr = (void*)Texture.GetPtr();
+	void* Ptr = (void*)Texture;
 	uint64 Size = static_cast<uint64>(SizeAlign.size);
 	
 #if PLATFORM_IOS
-	bool bMemoryless = ([Texture.GetPtr() storageMode] == MTLStorageModeMemoryless);
+	bool bMemoryless = ([Texture storageMode] == MTLStorageModeMemoryless);
 	if (!bMemoryless)
 #endif
 	{
@@ -168,9 +168,9 @@ void AGXLLM::LogAllocTexture(mtlpp::TextureDescriptor const& Desc, mtlpp::Textur
 	{
 		LLM_SCOPED_PAUSE_TRACKING(ELLMAllocType::System);
 		
-		if ([Desc.GetPtr() usage] & MTLTextureUsageRenderTarget)
+		if ([Desc usage] & MTLTextureUsageRenderTarget)
 		{
-			objc_setAssociatedObject(Texture.GetPtr(), (void*)&AGXLLM::LogAllocTexture,
+			objc_setAssociatedObject(Texture, (void*)&AGXLLM::LogAllocTexture,
 			[[[FAGXDeallocHandler alloc] initWithBlock:^{
 				LLM_PLATFORM_SCOPE_METAL(ELLMTagAGX::RenderTargets);
 				
@@ -188,7 +188,7 @@ void AGXLLM::LogAllocTexture(mtlpp::TextureDescriptor const& Desc, mtlpp::Textur
 		}
 		else
 		{
-			objc_setAssociatedObject(Texture.GetPtr(), (void*)&AGXLLM::LogAllocTexture,
+			objc_setAssociatedObject(Texture, (void*)&AGXLLM::LogAllocTexture,
 			[[[FAGXDeallocHandler alloc] initWithBlock:^{
 				LLM_PLATFORM_SCOPE_METAL(ELLMTagAGX::Textures);
 			
@@ -207,10 +207,10 @@ void AGXLLM::LogAllocTexture(mtlpp::TextureDescriptor const& Desc, mtlpp::Textur
 	}
 }
 
-void AGXLLM::LogAllocBuffer(mtlpp::Buffer const& Buffer)
+void AGXLLM::LogAllocBuffer(id<MTLBuffer> Buffer)
 {
-	void* Ptr = (void*)Buffer.GetPtr();
-	uint64 Size = Buffer.GetLength();
+	void* Ptr = (void*)Buffer;
+	uint64 Size = [Buffer length];
 	
 	INC_MEMORY_STAT_BY(STAT_AGXBufferMemory, Size);
 	INC_DWORD_STAT(STAT_AGXBufferCount);
@@ -220,7 +220,7 @@ void AGXLLM::LogAllocBuffer(mtlpp::Buffer const& Buffer)
 	{
 		LLM_SCOPED_PAUSE_TRACKING(ELLMAllocType::System);
 		
-		objc_setAssociatedObject(Buffer.GetPtr(), (void*)&AGXLLM::LogAllocBuffer,
+		objc_setAssociatedObject(Buffer, (void*)&AGXLLM::LogAllocBuffer,
 		[[[FAGXDeallocHandler alloc] initWithBlock:^{
 			LLM_PLATFORM_SCOPE_METAL(ELLMTagAGX::Buffers);
 			
@@ -232,15 +232,3 @@ void AGXLLM::LogAllocBuffer(mtlpp::Buffer const& Buffer)
 		OBJC_ASSOCIATION_RETAIN);
 	}
 }
-
-void AGXLLM::LogAliasTexture(mtlpp::Texture const& Texture)
-{
-	objc_setAssociatedObject(Texture.GetPtr(), (void*)&AGXLLM::LogAllocTexture, nil, OBJC_ASSOCIATION_RETAIN);
-}
-
-void AGXLLM::LogAliasBuffer(mtlpp::Buffer const& Buffer)
-{
-	objc_setAssociatedObject(Buffer.GetPtr(), (void*)&AGXLLM::LogAllocBuffer, nil, OBJC_ASSOCIATION_RETAIN);
-}
-
-
diff --git a/Engine/Source/Runtime/Apple/AGXRHI/Private/AGXLLM.h b/Engine/Source/Runtime/Apple/AGXRHI/Private/AGXLLM.h
--- a/Engine/Source/Runtime/Apple/AGXRHI/Private/AGXLLM.h
+++ b/Engine/Source/Runtime/Apple/AGXRHI/Private/AGXLLM.h
@@ -46,11 +46,6 @@ namespace AGXLLM
 // These work without the LLM module
 namespace AGXLLM
 {
-	void LogAllocTexture(mtlpp::TextureDescriptor const& Desc, mtlpp::Texture const& Texture);
-	void LogAllocBuffer(mtlpp::Buffer const& Buffer);
-	
-	void LogAliasTexture(mtlpp::Texture const& Texture);
-	void LogAliasBuffer(mtlpp::Buffer const& Buffer);
+	void LogAllocTexture(MTLTextureDescriptor* Desc, id<MTLTexture> Texture);
+	void LogAllocBuffer(id<MTLBuffer> Buffer);
 }
-
-
diff --git a/Engine/Source/Runtime/Apple/AGXRHI/Private/AGXRHIRenderQuery.cpp b/Engine/Source/Runtime/Apple/AGXRHI/Private/AGXRHIRenderQuery.cpp
--- a/Engine/Source/Runtime/Apple/AGXRHI/Private/AGXRHIRenderQuery.cpp
+++ b/Engine/Source/Runtime/Apple/AGXRHI/Private/AGXRHIRenderQuery.cpp
@@ -109,7 +109,7 @@ FAGXQueryBuffer* FAGXQueryBufferPool::GetCurrentQueryBuffer()
 			FMemory::Memzero((((uint8*)Buffer.GetContents())), EQueryBufferMaxSize);
 
 #if STATS || ENABLE_LOW_LEVEL_MEM_TRACKER
-			AGXLLM::LogAllocBuffer(Buffer);
+			AGXLLM::LogAllocBuffer(Buffer.GetPtr());
 #endif // STATS || ENABLE_LOW_LEVEL_MEM_TRACKER
 		}
 
diff --git a/Engine/Source/Runtime/Apple/AGXRHI/Private/AGXState.cpp b/Engine/Source/Runtime/Apple/AGXRHI/Private/AGXState.cpp
--- a/Engine/Source/Runtime/Apple/AGXRHI/Private/AGXState.cpp
+++ b/Engine/Source/Runtime/Apple/AGXRHI/Private/AGXState.cpp
@@ -446,7 +446,7 @@ FAGXBlendState::FAGXBlendState(FBlendStateInitializerRHI const& Initializer)
 			[BlendState release];
 		}
 
-		RenderTargetStates[RenderTargetIndex].BlendState = mtlpp::RenderPipelineColorAttachmentDescriptor(BlendState);
+		RenderTargetStates[RenderTargetIndex].BlendState = BlendState;
 
 		// get the unique key
 		uint32 BlendBitMask =
diff --git a/Engine/Source/Runtime/Apple/AGXRHI/Private/AGXVertexBuffer.cpp b/Engine/Source/Runtime/Apple/AGXRHI/Private/AGXVertexBuffer.cpp
--- a/Engine/Source/Runtime/Apple/AGXRHI/Private/AGXVertexBuffer.cpp
+++ b/Engine/Source/Runtime/Apple/AGXRHI/Private/AGXVertexBuffer.cpp
@@ -200,7 +200,7 @@ FAGXRHIBuffer::FAGXRHIBuffer(uint32 InSize, EBufferUsageFlags InUsage, EAGXBuffe
 				Backing.Buffer = FAGXBuffer(MTLPP_VALIDATE(mtlpp::Device, GMtlppDevice, AGXSafeGetRuntimeDebuggingLevel() >= EAGXDebugLevelValidation, NewBuffer(AllocSize, mtlpp::ResourceOptions(Options))), false);
 				
 #if STATS || ENABLE_LOW_LEVEL_MEM_TRACKER
-				AGXLLM::LogAllocBuffer(Backing.Buffer);
+				AGXLLM::LogAllocBuffer(Backing.Buffer.GetPtr());
 #endif
 				INC_MEMORY_STAT_BY(STAT_AGXDeviceBufferMemory, Backing.Buffer.GetLength());
 				
diff --git a/Engine/Source/Runtime/Apple/AGXRHI/Public/AGXState.h b/Engine/Source/Runtime/Apple/AGXRHI/Public/AGXState.h
--- a/Engine/Source/Runtime/Apple/AGXRHI/Public/AGXState.h
+++ b/Engine/Source/Runtime/Apple/AGXRHI/Public/AGXState.h
@@ -69,8 +69,10 @@ public:
 
 	struct FBlendPerMRT
 	{
-		mtlpp::RenderPipelineColorAttachmentDescriptor BlendState;
-		uint8 BlendStateKey;
+		MTLRenderPipelineColorAttachmentDescriptor* BlendState = nil;
+		uint8 BlendStateKey = 0;
+
+		FBlendPerMRT() = default;
 	};
 	FBlendPerMRT RenderTargetStates[MaxSimultaneousRenderTargets];
 	bool bUseIndependentRenderTargetBlendStates;
