AGXRHI: Replace mtlpp::VertexDescriptor with native Metal type MTLVertexDescriptor

#rnx

diff --git a/Engine/Source/Runtime/Apple/AGXRHI/Private/AGXHashedVertexDescriptor.cpp b/Engine/Source/Runtime/Apple/AGXRHI/Private/AGXHashedVertexDescriptor.cpp
--- a/Engine/Source/Runtime/Apple/AGXRHI/Private/AGXHashedVertexDescriptor.cpp
+++ b/Engine/Source/Runtime/Apple/AGXRHI/Private/AGXHashedVertexDescriptor.cpp
@@ -14,16 +14,9 @@
 #pragma mark - AGX RHI Hashed Vertex Descriptor
 
 
-FAGXHashedVertexDescriptor::FAGXHashedVertexDescriptor()
-	: VertexDescHash(0)
-	, VertexDesc(nil)
-{
-	// void
-}
-
-FAGXHashedVertexDescriptor::FAGXHashedVertexDescriptor(mtlpp::VertexDescriptor Desc, uint32 Hash)
+FAGXHashedVertexDescriptor::FAGXHashedVertexDescriptor(MTLVertexDescriptor* Desc, uint32 Hash)
 	: VertexDescHash(Hash)
-	, VertexDesc(Desc)
+	, VertexDesc([Desc retain])
 {
 	// void
 }
@@ -37,7 +30,7 @@ FAGXHashedVertexDescriptor::FAGXHashedVertexDescriptor(FAGXHashedVertexDescripto
 
 FAGXHashedVertexDescriptor::~FAGXHashedVertexDescriptor()
 {
-	// void
+	[VertexDesc release];
 }
 
 FAGXHashedVertexDescriptor& FAGXHashedVertexDescriptor::operator=(FAGXHashedVertexDescriptor const& Other)
@@ -45,7 +38,8 @@ FAGXHashedVertexDescriptor& FAGXHashedVertexDescriptor::operator=(FAGXHashedVert
 	if (this != &Other)
 	{
 		VertexDescHash = Other.VertexDescHash;
-		VertexDesc = Other.VertexDesc;
+		[VertexDesc release];
+		VertexDesc = [Other.VertexDesc retain];
 	}
 	return *this;
 }
@@ -59,39 +53,26 @@ bool FAGXHashedVertexDescriptor::operator==(FAGXHashedVertexDescriptor const& Ot
 		if (VertexDescHash == Other.VertexDescHash)
 		{
 			bEqual = true;
-			if (VertexDesc.GetPtr() != Other.VertexDesc.GetPtr())
+			if (VertexDesc != Other.VertexDesc)
 			{
-				ns::Array<mtlpp::VertexBufferLayoutDescriptor> Layouts = VertexDesc.GetLayouts();
-				ns::Array<mtlpp::VertexAttributeDescriptor> Attributes = VertexDesc.GetAttributes();
-
-				ns::Array<mtlpp::VertexBufferLayoutDescriptor> OtherLayouts = Other.VertexDesc.GetLayouts();
-				ns::Array<mtlpp::VertexAttributeDescriptor> OtherAttributes = Other.VertexDesc.GetAttributes();
-				check(Layouts && Attributes && OtherLayouts && OtherAttributes);
-
-				for (uint32 i = 0; bEqual && i < MaxVertexElementCount; i++)
+				for (uint32 Index = 0; bEqual && (Index < MaxVertexElementCount); ++Index)
 				{
-					mtlpp::VertexBufferLayoutDescriptor LayoutDesc = Layouts[(NSUInteger)i];
-					mtlpp::VertexBufferLayoutDescriptor OtherLayoutDesc = OtherLayouts[(NSUInteger)i];
+					bEqual &= ((VertexDesc.layouts[Index] != nil) == (Other.VertexDesc.layouts[Index] != nil));
 
-					bEqual &= ((LayoutDesc != nil) == (OtherLayoutDesc != nil));
-
-					if (LayoutDesc && OtherLayoutDesc)
+					if (bEqual && (VertexDesc.layouts[Index] != nil))
 					{
-						bEqual &= (LayoutDesc.GetStride() == OtherLayoutDesc.GetStride());
-						bEqual &= (LayoutDesc.GetStepFunction() == OtherLayoutDesc.GetStepFunction());
-						bEqual &= (LayoutDesc.GetStepRate() == OtherLayoutDesc.GetStepRate());
+						bEqual &= (VertexDesc.layouts[Index].stride         == Other.VertexDesc.layouts[Index].stride);
+						bEqual &= (VertexDesc.layouts[Index].stepFunction   == Other.VertexDesc.layouts[Index].stepFunction);
+						bEqual &= (VertexDesc.layouts[Index].stepRate       == Other.VertexDesc.layouts[Index].stepRate);
 					}
 
-					mtlpp::VertexAttributeDescriptor AttrDesc = Attributes[(NSUInteger)i];
-					mtlpp::VertexAttributeDescriptor OtherAttrDesc = OtherAttributes[(NSUInteger)i];
-
-					bEqual &= ((AttrDesc != nil) == (OtherAttrDesc != nil));
+					bEqual &= ((VertexDesc.attributes[Index] != nil) == (Other.VertexDesc.attributes[Index] != nil));
 
-					if (AttrDesc && OtherAttrDesc)
+					if (bEqual && (VertexDesc.attributes[Index] != nil))
 					{
-						bEqual &= (AttrDesc.GetFormat() == OtherAttrDesc.GetFormat());
-						bEqual &= (AttrDesc.GetOffset() == OtherAttrDesc.GetOffset());
-						bEqual &= (AttrDesc.GetBufferIndex() == OtherAttrDesc.GetBufferIndex());
+						bEqual &= (VertexDesc.attributes[Index].format      == Other.VertexDesc.attributes[Index].format);
+						bEqual &= (VertexDesc.attributes[Index].offset      == Other.VertexDesc.attributes[Index].offset);
+						bEqual &= (VertexDesc.attributes[Index].bufferIndex == Other.VertexDesc.attributes[Index].bufferIndex);
 					}
 				}
 			}
diff --git a/Engine/Source/Runtime/Apple/AGXRHI/Private/AGXHashedVertexDescriptor.h b/Engine/Source/Runtime/Apple/AGXRHI/Private/AGXHashedVertexDescriptor.h
--- a/Engine/Source/Runtime/Apple/AGXRHI/Private/AGXHashedVertexDescriptor.h
+++ b/Engine/Source/Runtime/Apple/AGXRHI/Private/AGXHashedVertexDescriptor.h
@@ -19,11 +19,11 @@
  */
 struct FAGXHashedVertexDescriptor
 {
-	NSUInteger VertexDescHash;
-	mtlpp::VertexDescriptor VertexDesc;
+	uint32 VertexDescHash = 0;
+	MTLVertexDescriptor* VertexDesc = nil;
 
-	FAGXHashedVertexDescriptor();
-	FAGXHashedVertexDescriptor(mtlpp::VertexDescriptor Desc, uint32 Hash);
+	FAGXHashedVertexDescriptor() = default;
+	FAGXHashedVertexDescriptor(MTLVertexDescriptor* Desc, uint32 Hash);
 	FAGXHashedVertexDescriptor(FAGXHashedVertexDescriptor const& Other);
 	~FAGXHashedVertexDescriptor();
 
diff --git a/Engine/Source/Runtime/Apple/AGXRHI/Private/AGXPipeline.cpp b/Engine/Source/Runtime/Apple/AGXRHI/Private/AGXPipeline.cpp
--- a/Engine/Source/Runtime/Apple/AGXRHI/Private/AGXPipeline.cpp
+++ b/Engine/Source/Runtime/Apple/AGXRHI/Private/AGXPipeline.cpp
@@ -495,7 +495,7 @@ static MTLVertexDescriptor* GetMaskedVertexDescriptor(MTLVertexDescriptor* Input
 {
 	for (uint32 Attr = 0; Attr < MaxMetalStreams; Attr++)
 	{
-		if (!InOutMask.IsFieldEnabled((int32)Attr) && [InputDesc.attributes objectAtIndexedSubscript:Attr] != nil)
+		if (!InOutMask.IsFieldEnabled((int32)Attr) && (InputDesc.attributes[Attr] != nil))
 		{
 			MTLVertexDescriptor* Desc = [[InputDesc copy] autorelease];
 			CrossCompiler::FShaderBindingInOutMask BuffersUsed;
@@ -503,18 +503,18 @@ static MTLVertexDescriptor* GetMaskedVertexDescriptor(MTLVertexDescriptor* Input
 			{
 				if (!InOutMask.IsFieldEnabled(MetalStreamIndex))
 				{
-					[Desc.attributes setObject:nil atIndexedSubscript:MetalStreamIndex];
+					Desc.attributes[MetalStreamIndex] = nil;
 				}
 				else
 				{
-					BuffersUsed.EnableField([Desc.attributes objectAtIndexedSubscript : MetalStreamIndex].bufferIndex);
+					BuffersUsed.EnableField(Desc.attributes[MetalStreamIndex].bufferIndex);
 				}
 			}
 			for (int32 BufferIndex = 0; BufferIndex < ML_MaxBuffers; ++BufferIndex)
 			{
 				if (!BuffersUsed.IsFieldEnabled(BufferIndex))
 				{
-					[Desc.layouts setObject:nil atIndexedSubscript:BufferIndex];
+					Desc.layouts[BufferIndex] = nil;
 				}
 			}
 			return Desc;
@@ -570,7 +570,7 @@ static bool ConfigureRenderPipelineDescriptor(MTLRenderPipelineDescriptor* Rende
 			MetalFormat = (MTLPixelFormat)AGXToSRGBFormat((MTLPixelFormat)MetalFormat);
 		}
 		
-		MTLRenderPipelineColorAttachmentDescriptor* Attachment = [ColorAttachments objectAtIndexedSubscript:ActiveTargetIndex];
+		MTLRenderPipelineColorAttachmentDescriptor* Attachment = ColorAttachments[ActiveTargetIndex];
 		[Attachment setPixelFormat:MetalFormat];
 		
 		MTLRenderPipelineColorAttachmentDescriptor* Blend = BlendState->RenderTargetStates[ActiveTargetIndex].BlendState;
@@ -668,12 +668,12 @@ static bool ConfigureRenderPipelineDescriptor(MTLRenderPipelineDescriptor* Rende
 				
 				if (Index < ML_MaxBuffers)
 				{
-					[[VertexPipelineBuffers objectAtIndexedSubscript:Index] setMutability:MTLMutabilityImmutable];
+					VertexPipelineBuffers[Index].mutability = MTLMutabilityImmutable;
 				}
 			}
 			if (VertexSideTable > 0)
 			{
-				[[VertexPipelineBuffers objectAtIndexedSubscript:VertexSideTable] setMutability:MTLMutabilityImmutable];
+				VertexPipelineBuffers[VertexSideTable].mutability = MTLMutabilityImmutable;
 			}
 		}
 		
@@ -688,12 +688,12 @@ static bool ConfigureRenderPipelineDescriptor(MTLRenderPipelineDescriptor* Rende
 				
 				if (Index < ML_MaxBuffers)
 				{
-					[[FragmentPipelineBuffers objectAtIndexedSubscript:Index] setMutability:MTLMutabilityImmutable];
+					FragmentPipelineBuffers[Index].mutability = MTLMutabilityImmutable;
 				}
 			}
 			if (PixelShader->SideTableBinding > 0)
 			{
-				[[FragmentPipelineBuffers objectAtIndexedSubscript:PixelShader->SideTableBinding] setMutability:MTLMutabilityImmutable];
+				FragmentPipelineBuffers[PixelShader->SideTableBinding].mutability = MTLMutabilityImmutable;
 			}
 		}
 	}
diff --git a/Engine/Source/Runtime/Apple/AGXRHI/Private/AGXVertexDeclaration.cpp b/Engine/Source/Runtime/Apple/AGXRHI/Private/AGXVertexDeclaration.cpp
--- a/Engine/Source/Runtime/Apple/AGXRHI/Private/AGXVertexDeclaration.cpp
+++ b/Engine/Source/Runtime/Apple/AGXRHI/Private/AGXVertexDeclaration.cpp
@@ -106,10 +106,7 @@ bool FAGXVertexDeclaration::GetInitializer(FVertexDeclarationElementList& Init)
 
 void FAGXVertexDeclaration::GenerateLayout(const FVertexDeclarationElementList& InElements)
 {
-	mtlpp::VertexDescriptor NewLayout;
-
-	ns::Array<mtlpp::VertexBufferLayoutDescriptor> Layouts = NewLayout.GetLayouts();
-	ns::Array<mtlpp::VertexAttributeDescriptor> Attributes = NewLayout.GetAttributes();
+	MTLVertexDescriptor* NewLayout = [[MTLVertexDescriptor alloc] init];;
 
 	BaseHash = 0;
 	uint32 StrideHash = BaseHash;
@@ -138,8 +135,8 @@ void FAGXVertexDeclaration::GenerateLayout(const FVertexDeclarationElementList&
 		if (ExistingStride == NULL)
 		{
 			// handle 0 stride buffers
-			mtlpp::VertexStepFunction Function = (Element.Stride == 0 ? mtlpp::VertexStepFunction::Constant : (Element.bUseInstanceIndex ? mtlpp::VertexStepFunction::PerInstance : mtlpp::VertexStepFunction::PerVertex));
-			uint32 StepRate = (Element.Stride == 0 ? 0 : 1);
+			const MTLVertexStepFunction StepFunction = (Element.Stride == 0 ? MTLVertexStepFunctionConstant : (Element.bUseInstanceIndex ? MTLVertexStepFunctionPerInstance : MTLVertexStepFunctionPerVertex));
+			const uint32 StepRate = (Element.Stride == 0 ? 0 : 1);
 
 			// even with MTLVertexStepFunctionConstant, it needs a non-zero stride (not sure why)
 			if (Element.Stride == 0)
@@ -155,10 +152,9 @@ void FAGXVertexDeclaration::GenerateLayout(const FVertexDeclarationElementList&
 			}
 
 			// set the stride once per buffer
-			mtlpp::VertexBufferLayoutDescriptor VBLayout = Layouts[ShaderBufferIndex];
-			VBLayout.SetStride(Stride);
-			VBLayout.SetStepFunction(Function);
-			VBLayout.SetStepRate(StepRate);
+			NewLayout.layouts[ShaderBufferIndex].stride       = Stride;
+			NewLayout.layouts[ShaderBufferIndex].stepFunction = StepFunction;
+			NewLayout.layouts[ShaderBufferIndex].stepRate     = StepRate;
 
 			// track this buffer and stride
 			BufferStrides.Add(ShaderBufferIndex, Element.Stride);
@@ -170,11 +166,12 @@ void FAGXVertexDeclaration::GenerateLayout(const FVertexDeclarationElementList&
 		}
 
 		// set the format for each element
-		mtlpp::VertexAttributeDescriptor Attrib = Attributes[Element.AttributeIndex];
-		Attrib.SetFormat(mtlpp::VertexFormat(TranslateElementTypeToMTLType(Element.Type)));
-		Attrib.SetOffset(Element.Offset);
-		Attrib.SetBufferIndex(ShaderBufferIndex);
+		NewLayout.attributes[Element.AttributeIndex].format      = TranslateElementTypeToMTLType(Element.Type);
+		NewLayout.attributes[Element.AttributeIndex].offset      = Element.Offset;
+		NewLayout.attributes[Element.AttributeIndex].bufferIndex = ShaderBufferIndex;
 	}
 
 	Layout = FAGXHashedVertexDescriptor(NewLayout, HashCombine(BaseHash, StrideHash));
+
+	[NewLayout release];
 }
