AGXRHI: Replace mtlpp::Resource with the Metal native type MTLResource

#rnx

diff --git a/Engine/Source/Runtime/Apple/AGXRHI/Private/AGXCommandEncoder.cpp b/Engine/Source/Runtime/Apple/AGXRHI/Private/AGXCommandEncoder.cpp
--- a/Engine/Source/Runtime/Apple/AGXRHI/Private/AGXCommandEncoder.cpp
+++ b/Engine/Source/Runtime/Apple/AGXRHI/Private/AGXCommandEncoder.cpp
@@ -887,10 +887,6 @@ void FAGXCommandEncoder::SetShaderBuffer(MTLFunctionType FunctionType, FAGXBuffe
 	check(index < ML_MaxBuffers);
     if(GetAGXDeviceContext().SupportsFeature(EAGXFeaturesSetBufferOffset) && Buffer && (ShaderBuffers[uint32(FunctionType)].Bound & (1 << index)) && ShaderBuffers[uint32(FunctionType)].Buffers[index] == Buffer)
     {
-		if (FunctionType == MTLFunctionTypeVertex || FunctionType == MTLFunctionTypeKernel)
-		{
-			FenceResource(Buffer);
-		}
 		SetShaderBufferOffset(FunctionType, Offset, Length, index);
 		ShaderBuffers[uint32(FunctionType)].Usage[index] = Usage;
 		ShaderBuffers[uint32(FunctionType)].SetBufferMetaData(index, Length, GAGXBufferFormats[Format].DataFormat, ElementRowPitch);
@@ -1032,7 +1028,6 @@ void FAGXCommandEncoder::SetShaderTexture(MTLFunctionType FunctionType, FAGXText
 	{
 		case MTLFunctionTypeVertex:
 			check (RenderCommandEncoder);
-			FenceResource(Texture);
 			[RenderCommandEncoder setVertexTexture:Texture.GetPtr() atIndex:index];
 			break;
 		case MTLFunctionTypeFragment:
@@ -1041,7 +1036,6 @@ void FAGXCommandEncoder::SetShaderTexture(MTLFunctionType FunctionType, FAGXText
 			break;
 		case MTLFunctionTypeKernel:
 			check (ComputeCommandEncoder);
-			FenceResource(Texture);
 			[ComputeCommandEncoder setTexture:Texture.GetPtr() atIndex:index];
 			break;
 		default:
@@ -1097,7 +1091,7 @@ void FAGXCommandEncoder::SetShaderSideTable(MTLFunctionType FunctionType, NSUInt
 	}
 }
 
-void FAGXCommandEncoder::TransitionResources(mtlpp::Resource const& Resource)
+void FAGXCommandEncoder::TransitionResources(id<MTLResource> Resource)
 {
 }
 
@@ -1132,14 +1126,6 @@ bool FAGXCommandEncoder::HasBufferBindingHistory(FAGXBuffer const& Buffer) const
 
 #pragma mark - Private Functions -
 
-void FAGXCommandEncoder::FenceResource(mtlpp::Texture const& Resource)
-{
-}
-
-void FAGXCommandEncoder::FenceResource(mtlpp::Buffer const& Resource)
-{
-}
-
 void FAGXCommandEncoder::SetShaderBufferInternal(MTLFunctionType Function, uint32 Index)
 {
 	FAGXBufferBindings& Binding = ShaderBuffers[uint32(Function)];
@@ -1173,7 +1159,6 @@ void FAGXCommandEncoder::SetShaderBufferInternal(MTLFunctionType Function, uint3
 			case MTLFunctionTypeVertex:
 				Binding.Bound |= (1 << Index);
 				check(RenderCommandEncoder);
-				FenceResource(Buffer);
 				[RenderCommandEncoder setVertexBuffer:Buffer.GetPtr() offset:(Offset + Buffer.GetOffset()) atIndex:Index];
 				break;
 
@@ -1186,7 +1171,6 @@ void FAGXCommandEncoder::SetShaderBufferInternal(MTLFunctionType Function, uint3
 			case MTLFunctionTypeKernel:
 				Binding.Bound |= (1 << Index);
 				check(ComputeCommandEncoder);
-				FenceResource(Buffer);
 				[ComputeCommandEncoder setBuffer:Buffer.GetPtr() offset:(Offset + Buffer.GetOffset()) atIndex:Index];
 				break;
 
diff --git a/Engine/Source/Runtime/Apple/AGXRHI/Private/AGXCommandEncoder.h b/Engine/Source/Runtime/Apple/AGXRHI/Private/AGXCommandEncoder.h
--- a/Engine/Source/Runtime/Apple/AGXRHI/Private/AGXCommandEncoder.h
+++ b/Engine/Source/Runtime/Apple/AGXRHI/Private/AGXCommandEncoder.h
@@ -336,7 +336,7 @@ public:
 	 * Transition resource so that we can barrier fragment->vertex stages.
 	 * @param Resource The resource we are going to make readable having been written.
 	 */
-	void TransitionResources(mtlpp::Resource const& Resource);
+	void TransitionResources(id<MTLResource> Resource);
 	
 	/*
 	 * Set the compute pipeline state that will be used.
@@ -369,9 +369,6 @@ private:
 	
 	void SetShaderBufferInternal(MTLFunctionType Function, uint32 Index);
 	
-	void FenceResource(mtlpp::Texture const& Resource);
-	void FenceResource(mtlpp::Buffer const& Resource);
-	
 private:
 	static const uint32 NUM_SHADER_FREQUENCIES = uint32(MTLFunctionTypeKernel) + 1;
 
@@ -451,7 +448,7 @@ private:
 	TSet<ns::AutoReleased<FAGXBuffer>> BufferBindingHistory;
 	TSet<ns::AutoReleased<FAGXTexture>> TextureBindingHistory;
 	
-	TMap<mtlpp::Resource::Type, MTLResourceUsage> ResourceUsage;
+	TMap<id<MTLResource>, MTLResourceUsage> ResourceUsage;
 	
 	uint32 EncoderNum;
 	uint32 CmdBufIndex;
diff --git a/Engine/Source/Runtime/Apple/AGXRHI/Private/AGXContext.cpp b/Engine/Source/Runtime/Apple/AGXRHI/Private/AGXContext.cpp
--- a/Engine/Source/Runtime/Apple/AGXRHI/Private/AGXContext.cpp
+++ b/Engine/Source/Runtime/Apple/AGXRHI/Private/AGXContext.cpp
@@ -1036,17 +1036,17 @@ void FAGXContext::TransitionResource(FRHIUnorderedAccessView* InResource)
 		FAGXSurface* Surface = UAV->GetSourceTexture();
 		if (Surface->Texture)
 		{
-			RenderPass.TransitionResources(Surface->Texture);
+			RenderPass.TransitionResources(Surface->Texture.GetPtr());
 			if (Surface->MSAATexture)
 			{
-				RenderPass.TransitionResources(Surface->MSAATexture);
+				RenderPass.TransitionResources(Surface->MSAATexture.GetPtr());
 			}
 		}
 	}
 	else
 	{
 		FAGXResourceMultiBuffer* Buffer = UAV->GetSourceBuffer();
-		RenderPass.TransitionResources(Buffer->GetCurrentBuffer());
+		RenderPass.TransitionResources(Buffer->GetCurrentBuffer().GetPtr());
 	}
 }
 
@@ -1056,10 +1056,10 @@ void FAGXContext::TransitionResource(FRHITexture* InResource)
 
 	if ((Surface != nullptr) && Surface->Texture)
 	{
-		RenderPass.TransitionResources(Surface->Texture);
+		RenderPass.TransitionResources(Surface->Texture.GetPtr());
 		if (Surface->MSAATexture)
 		{
-			RenderPass.TransitionResources(Surface->MSAATexture);
+			RenderPass.TransitionResources(Surface->MSAATexture.GetPtr());
 		}
 	}
 }
@@ -1471,7 +1471,7 @@ void FAGXContext::SynchronizeTexture(FAGXTexture const& Texture, uint32 Slice, u
 	RenderPass.SynchronizeTexture(Texture, Slice, Level);
 }
 
-void FAGXContext::SynchronizeResource(mtlpp::Resource const& Resource)
+void FAGXContext::SynchronizeResource(id<MTLResource> Resource)
 {
 	RenderPass.SynchronizeResource(Resource);
 }
diff --git a/Engine/Source/Runtime/Apple/AGXRHI/Private/AGXContext.h b/Engine/Source/Runtime/Apple/AGXRHI/Private/AGXContext.h
--- a/Engine/Source/Runtime/Apple/AGXRHI/Private/AGXContext.h
+++ b/Engine/Source/Runtime/Apple/AGXRHI/Private/AGXContext.h
@@ -103,7 +103,7 @@ public:
 	
 	void SynchronizeTexture(FAGXTexture const& Texture, uint32 Slice, uint32 Level);
 	
-	void SynchronizeResource(mtlpp::Resource const& Resource);
+	void SynchronizeResource(id<MTLResource> Resource);
 	
 	void FillBuffer(FAGXBuffer const& Buffer, ns::Range Range, uint8 Value);
 
@@ -165,7 +165,7 @@ void FAGXContext::TransitionRHIResource(T* InResource)
 	auto Resource = ResourceCast(InResource);
 	if (Resource->GetCurrentBufferOrNil())
 	{
-		RenderPass.TransitionResources(Resource->GetCurrentBuffer());
+		RenderPass.TransitionResources(Resource->GetCurrentBuffer().GetPtr());
 	}
 }
 
diff --git a/Engine/Source/Runtime/Apple/AGXRHI/Private/AGXRenderPass.cpp b/Engine/Source/Runtime/Apple/AGXRHI/Private/AGXRenderPass.cpp
--- a/Engine/Source/Runtime/Apple/AGXRHI/Private/AGXRenderPass.cpp
+++ b/Engine/Source/Runtime/Apple/AGXRHI/Private/AGXRenderPass.cpp
@@ -699,7 +699,7 @@ void FAGXRenderPass::SynchronizeTexture(FAGXTexture const& Texture, uint32 Slice
 #endif
 }
 
-void FAGXRenderPass::SynchronizeResource(mtlpp::Resource const& Resource)
+void FAGXRenderPass::SynchronizeResource(id<MTLResource> Resource)
 {
 	check(Resource);
 #if PLATFORM_MAC
@@ -709,7 +709,7 @@ void FAGXRenderPass::SynchronizeResource(mtlpp::Resource const& Resource)
 	check(Encoder);
 	
 	// METAL_GPUPROFILE(FAGXProfiler::GetProfiler()->EncodeBlit(CurrentEncoder.GetCommandBufferStats(), __FUNCTION__));
-	[Encoder synchronizeResource:Resource.GetPtr()];
+	[Encoder synchronizeResource:Resource];
 
 	ConditionalSubmit();
 #endif
@@ -935,7 +935,7 @@ void FAGXRenderPass::AddAsyncCommandBufferHandlers(mtlpp::CommandBufferHandler S
 	}
 }
 
-void FAGXRenderPass::TransitionResources(mtlpp::Resource const& Resource)
+void FAGXRenderPass::TransitionResources(id<MTLResource> Resource)
 {
 	PrologueEncoder.TransitionResources(Resource);
 	CurrentEncoder.TransitionResources(Resource);
diff --git a/Engine/Source/Runtime/Apple/AGXRHI/Private/AGXRenderPass.h b/Engine/Source/Runtime/Apple/AGXRHI/Private/AGXRenderPass.h
--- a/Engine/Source/Runtime/Apple/AGXRHI/Private/AGXRenderPass.h
+++ b/Engine/Source/Runtime/Apple/AGXRHI/Private/AGXRenderPass.h
@@ -60,7 +60,7 @@ public:
     
     void SynchronizeTexture(FAGXTexture const& Texture, uint32 Slice, uint32 Level);
     
-	void SynchronizeResource(mtlpp::Resource const& Resource);
+	void SynchronizeResource(id<MTLResource> Resource);
     
 	void FillBuffer(FAGXBuffer const& Buffer, ns::Range Range, uint8 Value);
 	
@@ -86,7 +86,7 @@ public:
 	
 	void AddAsyncCommandBufferHandlers(mtlpp::CommandBufferHandler Scheduled, mtlpp::CommandBufferHandler Completion);
 	
-	void TransitionResources(mtlpp::Resource const& Resource);
+	void TransitionResources(id<MTLResource> Resource);
 
 #pragma mark - Public Debug Support -
 	
diff --git a/Engine/Source/Runtime/Apple/AGXRHI/Private/AGXVertexBuffer.cpp b/Engine/Source/Runtime/Apple/AGXRHI/Private/AGXVertexBuffer.cpp
--- a/Engine/Source/Runtime/Apple/AGXRHI/Private/AGXVertexBuffer.cpp
+++ b/Engine/Source/Runtime/Apple/AGXRHI/Private/AGXVertexBuffer.cpp
@@ -502,7 +502,7 @@ void* FAGXRHIBuffer::Lock(bool bIsOnRHIThread, EResourceLockMode InLockMode, uin
 			SCOPE_CYCLE_COUNTER(STAT_AGXBufferPageOffTime);
 			
 			// Synchronise the buffer with the CPU
-			GetAGXDeviceContext().SynchronizeResource(GetCurrentBacking().Buffer);
+			GetAGXDeviceContext().SynchronizeResource(GetCurrentBacking().Buffer.GetPtr());
 			
 			//kick the current command buffer.
 			GetAGXDeviceContext().SubmitCommandBufferAndWait();
