AGXRHI: Replace mtlpp::CaptureManager et al. with native Metal APIs.

- Moved FObjCWrapperRetained to AGXRHIPrivate.h. This is a useful utility struct for stuffing objective-c objects into (for example) a TUniquePtr for:
  1. Easier and more automatic object lifetime management
  2. Conveniently placing objective-c objects into UE containers.
- Remove EAGXFeaturesGPUCaptureManager. This feature is now supported by all Apple platforms UE5 supports and no longer requiers a feature flag.
- Remove mtlpp from FAGXCaptureManager

#rnx

diff --git a/Engine/Source/Runtime/Apple/AGXRHI/Private/AGXCaptureManager.cpp b/Engine/Source/Runtime/Apple/AGXRHI/Private/AGXCaptureManager.cpp
--- a/Engine/Source/Runtime/Apple/AGXRHI/Private/AGXCaptureManager.cpp
+++ b/Engine/Source/Runtime/Apple/AGXRHI/Private/AGXCaptureManager.cpp
@@ -3,41 +3,46 @@
 #include "AGXRHIPrivate.h"
 #include "AGXCaptureManager.h"
 #include "AGXCommandQueue.h"
-#include "capture_manager.hpp"
-
-bool GAGXSupportsCaptureManager = false;
 
 FAGXCaptureManager::FAGXCaptureManager(FAGXCommandQueue& InQueue)
 : Queue(InQueue)
 {
-	MTLPP_IF_AVAILABLE(10.13, 11.0, 11.0)
+	MTLCaptureManager* CaptureManager = [MTLCaptureManager sharedCaptureManager];
+	
+	id<MTLCaptureScope> DefaultCaptureScope = [CaptureManager newCaptureScopeWithDevice:GMtlDevice];
+	DefaultCaptureScope.label = @"1 Frame";
+	
+	CaptureManager.defaultCaptureScope = DefaultCaptureScope;
+	
+	FAGXCaptureScope DefaultScope;
+	DefaultScope.CaptureScope = MakeUnique<RetainedCaptureScopeType>(DefaultCaptureScope);
+	DefaultScope.Type = EAGXCaptureTypePresent;
+	DefaultScope.StepCount = 1;
+	DefaultScope.LastTrigger = 0;
+	ActiveScopes.Add(DefaultScope);
+	
+	[DefaultCaptureScope beginScope];
+	
+	// The local reference is no longer required (object lifetime is managed by DefaultScope).
+	[DefaultCaptureScope release];
+	
+	uint32 PresentStepCounts[] = {2, 5, 10, 15, 30, 60, 90, 120};
+	for (uint32 i = 0; i < (sizeof(PresentStepCounts) / sizeof(uint32)); i++)
 	{
-		GAGXSupportsCaptureManager = true;
-		
-		mtlpp::CaptureManager Manager = mtlpp::CaptureManager::SharedCaptureManager();
-		Manager.SetDefaultCaptureScope(Manager.NewCaptureScopeWithDevice(GMtlppDevice));
-		Manager.GetDefaultCaptureScope().SetLabel(@"1 Frame");
-		
-		FAGXCaptureScope DefaultScope;
-		DefaultScope.MTLScope = Manager.GetDefaultCaptureScope();
-		DefaultScope.Type = EAGXCaptureTypePresent;
-		DefaultScope.StepCount = 1;
-		DefaultScope.LastTrigger = 0;
-		ActiveScopes.Add(DefaultScope);
-		DefaultScope.MTLScope.BeginScope();
-		
-		uint32 PresentStepCounts[] = {2, 5, 10, 15, 30, 60, 90, 120};
-		for (uint32 i = 0; i < (sizeof(PresentStepCounts) / sizeof(uint32)); i++)
-		{
-			FAGXCaptureScope Scope;
-			Scope.MTLScope = Manager.NewCaptureScopeWithDevice(GMtlppDevice);
-			Scope.MTLScope.SetLabel(FString::Printf(TEXT("%u Frames"), PresentStepCounts[i]).GetNSString());
-			Scope.Type = EAGXCaptureTypePresent;
-			Scope.StepCount = PresentStepCounts[i];
-			Scope.LastTrigger = 0;
-			ActiveScopes.Add(Scope);
-			Scope.MTLScope.BeginScope();
-		}
+		id<MTLCaptureScope> CaptureScope = [CaptureManager newCaptureScopeWithDevice:GMtlDevice];
+		CaptureScope.label = [NSString stringWithFormat:@"%u Frames", PresentStepCounts[i]];
+	
+		FAGXCaptureScope Scope;
+		Scope.CaptureScope = MakeUnique<RetainedCaptureScopeType>(CaptureScope);
+		Scope.Type = EAGXCaptureTypePresent;
+		Scope.StepCount = PresentStepCounts[i];
+		Scope.LastTrigger = 0;
+		ActiveScopes.Add(Scope);
+	
+		[CaptureScope beginScope];
+	
+		// The local reference is no longer required (object lifetime is managed by Scope).
+		[CaptureScope release];
 	}
 }
 
@@ -47,47 +52,37 @@ FAGXCaptureManager::~FAGXCaptureManager()
 
 void FAGXCaptureManager::PresentFrame(uint32 FrameNumber)
 {
-	if (GAGXSupportsCaptureManager)
+	for (FAGXCaptureScope& Scope : ActiveScopes)
 	{
-		for (FAGXCaptureScope& Scope : ActiveScopes)
+		uint32 Diff = 0;
+		if (FrameNumber > Scope.LastTrigger)
 		{
-			uint32 Diff = 0;
-			if (FrameNumber > Scope.LastTrigger)
-			{
-				Diff = FrameNumber - Scope.LastTrigger;
-			}
-			else
-			{
-				Diff = (UINT32_MAX - Scope.LastTrigger) + FrameNumber;
-			}
-			
-			if (Diff >= Scope.StepCount)
-			{
-				Scope.MTLScope.EndScope();
-				Scope.MTLScope.BeginScope();
-				Scope.LastTrigger = FrameNumber;
-			}
+			Diff = FrameNumber - Scope.LastTrigger;
+		}
+		else
+		{
+			Diff = (UINT32_MAX - Scope.LastTrigger) + FrameNumber;
+		}
+		
+		if (Diff >= Scope.StepCount)
+		{
+			[Scope.CaptureScope.Get()->Object endScope];
+			[Scope.CaptureScope.Get()->Object beginScope];
+			Scope.LastTrigger = FrameNumber;
 		}
-	}
-	else
-	{
-		Queue.InsertDebugCaptureBoundary();
 	}
 }
 
-void FAGXCaptureManager::BeginCapture(void)
+void FAGXCaptureManager::BeginCapture()
 {
-	if (GAGXSupportsCaptureManager)
-	{
-		mtlpp::CaptureManager::SharedCaptureManager().StartCaptureWithDevice(GMtlppDevice);
-	}
+PRAGMA_DISABLE_DEPRECATION_WARNINGS
+
+	[[MTLCaptureManager sharedCaptureManager] startCaptureWithDevice:GMtlDevice];
+
+PRAGMA_ENABLE_DEPRECATION_WARNINGS
 }
 
-void FAGXCaptureManager::EndCapture(void)
+void FAGXCaptureManager::EndCapture()
 {
-	if (GAGXSupportsCaptureManager)
-	{
-		mtlpp::CaptureManager::SharedCaptureManager().StopCapture();
-	}
+	[[MTLCaptureManager sharedCaptureManager] stopCapture];
 }
-
diff --git a/Engine/Source/Runtime/Apple/AGXRHI/Private/AGXCaptureManager.h b/Engine/Source/Runtime/Apple/AGXRHI/Private/AGXCaptureManager.h
--- a/Engine/Source/Runtime/Apple/AGXRHI/Private/AGXCaptureManager.h
+++ b/Engine/Source/Runtime/Apple/AGXRHI/Private/AGXCaptureManager.h
@@ -2,12 +2,6 @@
 
 #pragma once
 
-#include <Metal/Metal.h>
-#include "capture_scope.hpp"
-#include "device.hpp"
-
-class FAGXCommandQueue;
-
 class FAGXCaptureManager
 {
 public:
@@ -19,14 +13,15 @@ public:
 	
 	// Programmatic captures without an Xcode capture scope.
 	// Use them to instrument the code manually to debug issues.
-	void BeginCapture(void);
-	void EndCapture(void);
+	void BeginCapture();
+	void EndCapture();
 	
 private:
 	FAGXCommandQueue& Queue;
-	bool bSupportsCaptureManager;
 	
 private:
+	using RetainedCaptureScopeType = FObjCWrapperRetained< id<MTLCaptureScope> >;
+
 	enum EAGXCaptureType
 	{
 		EAGXCaptureTypeUnknown,
@@ -40,7 +35,31 @@ private:
 		EAGXCaptureType Type;
 		uint32 StepCount;
 		uint32 LastTrigger;
-		mtlpp::CaptureScope MTLScope;
+		TUniquePtr<RetainedCaptureScopeType> CaptureScope;
+
+		FAGXCaptureScope()
+		: Type(EAGXCaptureTypeUnknown)
+		, StepCount(0)
+		, LastTrigger(0)
+		, CaptureScope()
+		{
+		}
+
+		FAGXCaptureScope(const FAGXCaptureScope& Other)
+		: Type(Other.Type)
+		, StepCount(Other.StepCount)
+		, LastTrigger(Other.LastTrigger)
+		, CaptureScope(MakeUnique<RetainedCaptureScopeType>(Other.CaptureScope.Get()->Object))
+		{
+		}
+
+		FAGXCaptureScope(FAGXCaptureScope&& Other)
+		: Type(Other.Type)
+		, StepCount(Other.StepCount)
+		, LastTrigger(Other.LastTrigger)
+		, CaptureScope(MoveTemp(Other.CaptureScope))
+		{
+		}
 	};
 	
 	TArray<FAGXCaptureScope> ActiveScopes;
diff --git a/Engine/Source/Runtime/Apple/AGXRHI/Private/AGXCommandQueue.cpp b/Engine/Source/Runtime/Apple/AGXRHI/Private/AGXCommandQueue.cpp
--- a/Engine/Source/Runtime/Apple/AGXRHI/Private/AGXCommandQueue.cpp
+++ b/Engine/Source/Runtime/Apple/AGXRHI/Private/AGXCommandQueue.cpp
@@ -50,7 +50,7 @@ FAGXCommandQueue::FAGXCommandQueue(uint32 const MaxNumCommandBuffers /* = 0 */)
 		Features |= EAGXFeaturesCountingQueries | EAGXFeaturesBaseVertexInstance | EAGXFeaturesIndirectBuffer | EAGXFeaturesMSAADepthResolve | EAGXFeaturesMSAAStoreAndResolve;
 	}
 
-	Features |= EAGXFeaturesPrivateBufferSubAllocation | EAGXFeaturesGPUCaptureManager | EAGXFeaturesBufferSubAllocation | EAGXFeaturesParallelRenderEncoders | EAGXFeaturesPipelineBufferMutability | EAGXFeaturesMaxThreadsPerThreadgroup;
+	Features |= EAGXFeaturesPrivateBufferSubAllocation | EAGXFeaturesBufferSubAllocation | EAGXFeaturesParallelRenderEncoders | EAGXFeaturesPipelineBufferMutability | EAGXFeaturesMaxThreadsPerThreadgroup;
 	GAGXFColorVertexFormat = MTLVertexFormatUChar4Normalized_BGRA;
 
 #else
@@ -73,7 +73,6 @@ FAGXCommandQueue::FAGXCommandQueue(uint32 const MaxNumCommandBuffers /* = 0 */)
 	GAGXFColorVertexFormat = MTLVertexFormatUChar4Normalized_BGRA;
 
 	Features |= EAGXFeaturesPresentMinDuration
-		| EAGXFeaturesGPUCaptureManager
 		| EAGXFeaturesBufferSubAllocation
 		| EAGXFeaturesParallelRenderEncoders
 		| EAGXFeaturesPipelineBufferMutability;
@@ -123,7 +122,7 @@ FAGXCommandQueue::FAGXCommandQueue(uint32 const MaxNumCommandBuffers /* = 0 */)
     
     if ([GMtlDevice supportsFeatureSet:MTLFeatureSet_macOS_GPUFamily1_v3])
 	{
-		Features |= EAGXFeaturesMultipleViewports | EAGXFeaturesPipelineBufferMutability | EAGXFeaturesGPUCaptureManager;
+		Features |= EAGXFeaturesMultipleViewports | EAGXFeaturesPipelineBufferMutability;
 	}
 #endif
 	
diff --git a/Engine/Source/Runtime/Apple/AGXRHI/Private/AGXCommandQueue.h b/Engine/Source/Runtime/Apple/AGXRHI/Private/AGXCommandQueue.h
--- a/Engine/Source/Runtime/Apple/AGXRHI/Private/AGXCommandQueue.h
+++ b/Engine/Source/Runtime/Apple/AGXRHI/Private/AGXCommandQueue.h
@@ -52,7 +52,7 @@ typedef NS_OPTIONS(uint64, EAGXFeatures)
     /** Supports minimum on-glass duration for drawables */
     EAGXFeaturesPresentMinDuration = 1llu << 16llu,
     /** Supports programmatic frame capture API */
-    EAGXFeaturesGPUCaptureManager = 1llu << 17llu,
+    EAGXFeaturesGPUCaptureManager_REMOVED = 1llu << 17llu,
 	/** Supports efficient buffer-blits */
 	EAGXFeaturesEfficientBufferBlits = 1llu << 18llu,
 	/** Supports any kind of buffer sub-allocation */
diff --git a/Engine/Source/Runtime/Apple/AGXRHI/Private/AGXLLM.cpp b/Engine/Source/Runtime/Apple/AGXRHI/Private/AGXLLM.cpp
--- a/Engine/Source/Runtime/Apple/AGXRHI/Private/AGXLLM.cpp
+++ b/Engine/Source/Runtime/Apple/AGXRHI/Private/AGXLLM.cpp
@@ -145,11 +145,7 @@ static EPixelFormat MetalToRHIPixelFormat(MTLPixelFormat Format)
 
 void AGXLLM::LogAllocTexture(MTLTextureDescriptor* Desc, id<MTLTexture> Texture)
 {
-	MTLSizeAndAlign SizeAlign;
-	if (FAGXCommandQueue::SupportsFeature(EAGXFeaturesGPUCaptureManager))
-	{
-		SizeAlign = [GMtlDevice heapTextureSizeAndAlignWithDescriptor:Desc];
-	}
+	MTLSizeAndAlign SizeAlign= [GMtlDevice heapTextureSizeAndAlignWithDescriptor:Desc];
 	
 	void* Ptr = (void*)Texture;
 	uint64 Size = static_cast<uint64>(SizeAlign.size);
diff --git a/Engine/Source/Runtime/Apple/AGXRHI/Private/AGXRHIPrivate.h b/Engine/Source/Runtime/Apple/AGXRHI/Private/AGXRHIPrivate.h
--- a/Engine/Source/Runtime/Apple/AGXRHI/Private/AGXRHIPrivate.h
+++ b/Engine/Source/Runtime/Apple/AGXRHI/Private/AGXRHIPrivate.h
@@ -127,8 +127,28 @@ enum EMTLTextureType
 #define METAL_SUPPORTS_INDIRECT_ARGUMENT_BUFFERS 1
 #define METAL_SUPPORTS_CAPTURE_MANAGER 1
 #define METAL_SUPPORTS_TILE_SHADERS 1
-// In addition to compile-time SDK checks we also need a way to check if these are available on runtime
-extern bool GAGXSupportsCaptureManager;
+
+template<typename ReferenceType>
+struct FObjCWrapperRetained
+{
+	using Type = ReferenceType;
+
+	FObjCWrapperRetained() = delete;
+	FObjCWrapperRetained(const FObjCWrapperRetained&) = delete;
+	FObjCWrapperRetained(FObjCWrapperRetained&&) = delete;
+
+	explicit FObjCWrapperRetained(Type InObject = nil)
+		: Object([InObject retain])
+	{
+	}
+
+	~FObjCWrapperRetained()
+	{
+		[Object release];
+	}
+
+	Type Object;
+};
 
 struct FAGXBufferFormat
 {
diff --git a/Engine/Source/Runtime/Apple/AGXRHI/Private/AGXState.cpp b/Engine/Source/Runtime/Apple/AGXRHI/Private/AGXState.cpp
--- a/Engine/Source/Runtime/Apple/AGXRHI/Private/AGXState.cpp
+++ b/Engine/Source/Runtime/Apple/AGXRHI/Private/AGXState.cpp
@@ -166,24 +166,11 @@ static EColorWriteMask TranslateWriteMask(MTLColorWriteMask WriteMask)
 	return (EColorWriteMask)Result;
 }
 
-struct FObjCWrapperRetained
-{
-	FObjCWrapperRetained(id InObject = nil)
-		: Object([InObject retain])
-	{
-	}
-
-	~FObjCWrapperRetained()
-	{
-		[Object release];
-	}
-
-	id Object;
-};
-
 template <typename InitializerType>
 class FAGXStateObjectCache
 {
+	using RetainedObjectType = FObjCWrapperRetained< id >;
+
 public:
 	FAGXStateObjectCache()
 	{
@@ -200,7 +187,7 @@ public:
 			Mutex.ReadLock();
 		}
 		
-		TUniquePtr<FObjCWrapperRetained>* State = Cache.Find(Init);
+		TUniquePtr<RetainedObjectType>* State = Cache.Find(Init);
 
 		if (IsRunningRHIInSeparateThread())
 		{
@@ -217,7 +204,7 @@ public:
 			Mutex.WriteLock();
 		}
 		
-		Cache.Add(Init, MakeUnique<FObjCWrapperRetained>(State));
+		Cache.Add(Init, MakeUnique<RetainedObjectType>(State));
 		
 		if (IsRunningRHIInSeparateThread())
 		{
@@ -226,7 +213,7 @@ public:
 	}
 	
 private:
-	TMap<InitializerType, TUniquePtr<FObjCWrapperRetained> > Cache;
+	TMap<InitializerType, TUniquePtr<RetainedObjectType>> Cache;
 	FRWLock Mutex;
 };
 
