AGXRHI: Replace mtlpp::TextureDescriptor with native Metal type MTLTextureDescriptor.

- Plus, relocate FAGXTextureDesc and FAGXTextureCreateDesc out of public header AGXResources.h into private header AGXRHIPrivate.h

#rnx

diff --git a/Engine/Source/Runtime/Apple/AGXRHI/Private/AGXBuffer.cpp b/Engine/Source/Runtime/Apple/AGXRHI/Private/AGXBuffer.cpp
--- a/Engine/Source/Runtime/Apple/AGXRHI/Private/AGXBuffer.cpp
+++ b/Engine/Source/Runtime/Apple/AGXRHI/Private/AGXBuffer.cpp
@@ -934,19 +934,19 @@ FAGXTexturePool::~FAGXTexturePool()
 {
 }
 
-FAGXTexture FAGXTexturePool::CreateTexture(mtlpp::TextureDescriptor Desc)
+FAGXTexture FAGXTexturePool::CreateTexture(MTLTextureDescriptor* Desc)
 {
 	FAGXTexturePool::Descriptor Descriptor;
-	Descriptor.textureType = (NSUInteger)Desc.GetTextureType();
-	Descriptor.pixelFormat = (NSUInteger)Desc.GetPixelFormat();
-	Descriptor.width = Desc.GetWidth();
-	Descriptor.height = Desc.GetHeight();
-	Descriptor.depth = Desc.GetDepth();
-	Descriptor.mipmapLevelCount = Desc.GetMipmapLevelCount();
-	Descriptor.sampleCount = Desc.GetSampleCount();
-	Descriptor.arrayLength = Desc.GetArrayLength();
-	Descriptor.resourceOptions = Desc.GetResourceOptions();
-	Descriptor.usage = Desc.GetUsage();
+	Descriptor.textureType = Desc.textureType;
+	Descriptor.pixelFormat = Desc.pixelFormat;
+	Descriptor.width = Desc.width;
+	Descriptor.height = Desc.height;
+	Descriptor.depth = Desc.depth;
+	Descriptor.mipmapLevelCount = Desc.mipmapLevelCount;
+	Descriptor.sampleCount = Desc.sampleCount;
+	Descriptor.arrayLength = Desc.arrayLength;
+	Descriptor.resourceOptions = Desc.resourceOptions;
+	Descriptor.usage = Desc.usage;
 	if (Descriptor.usage == MTLTextureUsageUnknown)
 	{
 		Descriptor.usage = (MTLTextureUsageShaderRead | MTLTextureUsageShaderWrite | MTLTextureUsageRenderTarget | MTLTextureUsagePixelFormatView);
@@ -970,7 +970,7 @@ FAGXTexture FAGXTexturePool::CreateTexture(mtlpp::TextureDescriptor Desc)
 		METAL_GPUPROFILE(FAGXScopedCPUStats CPUStat(FString::Printf(TEXT("AllocTexture: %s"), TEXT("")/**FString([Desc.GetPtr() description])*/)));
 		Texture = MTLPP_VALIDATE(mtlpp::Device, GMtlppDevice, AGXSafeGetRuntimeDebuggingLevel() >= EAGXDebugLevelValidation, NewTexture(Desc));
 #if STATS || ENABLE_LOW_LEVEL_MEM_TRACKER
-		AGXLLM::LogAllocTexture(Desc.GetPtr(), Texture.GetPtr());
+		AGXLLM::LogAllocTexture(Desc, Texture.GetPtr());
 #endif
 	}
 	return Texture;
@@ -1333,12 +1333,12 @@ void FAGXResourceHeap::ReleaseBuffer(FAGXBuffer& Buffer)
 	}
 }
 
-FAGXTexture FAGXResourceHeap::CreateTexture(mtlpp::TextureDescriptor Desc, FAGXSurface* Surface)
+FAGXTexture FAGXResourceHeap::CreateTexture(MTLTextureDescriptor* Desc, FAGXSurface* Surface)
 {
 	LLM_SCOPE_METAL(ELLMTagAGX::Textures);
 	LLM_PLATFORM_SCOPE_METAL(ELLMTagAGX::Textures);
 	
-	if ([Desc.GetPtr() usage] & MTLTextureUsageRenderTarget)
+	if (Desc.usage & MTLTextureUsageRenderTarget)
 	{
 		LLM_PLATFORM_SCOPE_METAL(ELLMTagAGX::RenderTargets);
 		return TargetPool.CreateTexture(Desc);
diff --git a/Engine/Source/Runtime/Apple/AGXRHI/Private/AGXBuffer.h b/Engine/Source/Runtime/Apple/AGXRHI/Private/AGXBuffer.h
--- a/Engine/Source/Runtime/Apple/AGXRHI/Private/AGXBuffer.h
+++ b/Engine/Source/Runtime/Apple/AGXRHI/Private/AGXBuffer.h
@@ -332,7 +332,7 @@ public:
 	FAGXTexturePool(FCriticalSection& PoolMutex);
 	~FAGXTexturePool();
 	
-	FAGXTexture CreateTexture(mtlpp::TextureDescriptor Desc);
+	FAGXTexture CreateTexture(MTLTextureDescriptor* Desc);
 	void ReleaseTexture(FAGXTexture& Texture);
 	
 	void Drain(bool const bForce);
@@ -411,7 +411,7 @@ public:
 	void Init(FAGXCommandQueue& Queue);
 	
     FAGXBuffer CreateBuffer(uint32 InSize, uint32 InAlignment, EBufferUsageFlags InFlags, MTLResourceOptions InOptions, bool bForceUnique = false);
-	FAGXTexture CreateTexture(mtlpp::TextureDescriptor Desc, FAGXSurface* Surface);
+	FAGXTexture CreateTexture(MTLTextureDescriptor* Desc, FAGXSurface* Surface);
 	
 	void ReleaseBuffer(FAGXBuffer& Buffer);
 	void ReleaseTexture(FAGXSurface* Surface, FAGXTexture& Texture);
diff --git a/Engine/Source/Runtime/Apple/AGXRHI/Private/AGXContext.cpp b/Engine/Source/Runtime/Apple/AGXRHI/Private/AGXContext.cpp
--- a/Engine/Source/Runtime/Apple/AGXRHI/Private/AGXContext.cpp
+++ b/Engine/Source/Runtime/Apple/AGXRHI/Private/AGXContext.cpp
@@ -707,7 +707,7 @@ void FAGXDeviceContext::ReleaseTexture(FAGXTexture& Texture)
 	}
 }
 
-FAGXTexture FAGXDeviceContext::CreateTexture(FAGXSurface* Surface, mtlpp::TextureDescriptor Descriptor)
+FAGXTexture FAGXDeviceContext::CreateTexture(FAGXSurface* Surface, MTLTextureDescriptor* Descriptor)
 {
 	FAGXTexture Tex = Heap.CreateTexture(Descriptor, Surface);
 #if METAL_DEBUG_OPTIONS
diff --git a/Engine/Source/Runtime/Apple/AGXRHI/Private/AGXContext.h b/Engine/Source/Runtime/Apple/AGXRHI/Private/AGXContext.h
--- a/Engine/Source/Runtime/Apple/AGXRHI/Private/AGXContext.h
+++ b/Engine/Source/Runtime/Apple/AGXRHI/Private/AGXContext.h
@@ -181,7 +181,7 @@ public:
 	
 	inline FAGXResourceHeap& GetResourceHeap(void) { return Heap; }
 	
-	FAGXTexture CreateTexture(FAGXSurface* Surface, mtlpp::TextureDescriptor Descriptor);
+	FAGXTexture CreateTexture(FAGXSurface* Surface, MTLTextureDescriptor* Descriptor);
 	FAGXBuffer CreatePooledBuffer(FAGXPooledBufferArgs const& Args);
 	void ReleaseBuffer(FAGXBuffer& Buf);
 	void ReleaseObject(id Object);
diff --git a/Engine/Source/Runtime/Apple/AGXRHI/Private/AGXRHIPrivate.h b/Engine/Source/Runtime/Apple/AGXRHI/Private/AGXRHIPrivate.h
--- a/Engine/Source/Runtime/Apple/AGXRHI/Private/AGXRHIPrivate.h
+++ b/Engine/Source/Runtime/Apple/AGXRHI/Private/AGXRHIPrivate.h
@@ -150,6 +150,31 @@ struct FObjCWrapperRetained
 	Type Object;
 };
 
+struct FAGXTextureDesc
+{
+	using ReferencedTextureDescriptorType = FObjCWrapperRetained<MTLTextureDescriptor*>;
+
+	FAGXTextureDesc(const FRHITextureDesc& InDesc);
+	
+	TUniquePtr<ReferencedTextureDescriptorType> Desc;
+	MTLPixelFormat PixelFormat = MTLPixelFormatInvalid;
+	bool bMemoryless = false;
+	bool bIsRenderTarget = false;
+	uint8 FormatKey = 0;
+};
+
+struct FAGXTextureCreateDesc : public FRHITextureCreateDesc, public FAGXTextureDesc
+{
+	FAGXTextureCreateDesc(FRHITextureCreateDesc const& CreateDesc)
+		: FRHITextureCreateDesc(CreateDesc)
+		, FAGXTextureDesc(CreateDesc)
+	{
+		// @todo: texture type unification - Metal can override NumSamples based on command line options.
+		// We should instead require the renderer to do this.
+		NumSamples = [Desc.Get()->Object sampleCount];
+	}
+};
+
 struct FAGXBufferFormat
 {
 	// Valid linear texture pixel formats - potentially different than the actual texture formats
diff --git a/Engine/Source/Runtime/Apple/AGXRHI/Private/AGXRenderTarget.cpp b/Engine/Source/Runtime/Apple/AGXRHI/Private/AGXRenderTarget.cpp
--- a/Engine/Source/Runtime/Apple/AGXRHI/Private/AGXRenderTarget.cpp
+++ b/Engine/Source/Runtime/Apple/AGXRHI/Private/AGXRenderTarget.cpp
@@ -282,38 +282,39 @@ void FAGXDynamicRHI::RHIReadSurfaceData(FRHITexture* TextureRHI, FIntRect Rect,
 	{
 		SCOPE_CYCLE_COUNTER(STAT_AGXTexturePageOffTime);
 		
+		id<MTLTexture> MetalTexture = Texture.GetPtr();
 		FAGXTexture TempTexture = nil;
-		if ([Texture.GetPtr() storageMode] == MTLStorageModePrivate)
+		if (MetalTexture.storageMode == MTLStorageModePrivate)
 		{
 #if PLATFORM_MAC
 			MTLResourceOptions ResourceStorageMode = MTLResourceStorageModeManaged;
 #else
 			MTLResourceOptions ResourceStorageMode = MTLResourceStorageModeShared;
 #endif
-			MTLPixelFormat MetalFormat = (MTLPixelFormat)GPixelFormats[Surface->GetDesc().Format].PlatformFormat;
-			mtlpp::TextureDescriptor Desc;
-			Desc.SetTextureType(Texture.GetTextureType());
-			Desc.SetPixelFormat(Texture.GetPixelFormat());
-			Desc.SetWidth(SizeX);
-			Desc.SetHeight(SizeY);
-			Desc.SetDepth(1);
-			Desc.SetMipmapLevelCount(1); // Only consider a single subresource and not the whole texture (like in the other RHIs)
-			Desc.SetSampleCount(Texture.GetSampleCount());
-			Desc.SetArrayLength(Texture.GetArrayLength());
+			MTLTextureDescriptor* TextureDescriptor = [[MTLTextureDescriptor alloc] init];
+			TextureDescriptor.textureType      = MetalTexture.textureType;
+			TextureDescriptor.pixelFormat      = MetalTexture.pixelFormat;
+			TextureDescriptor.width            = SizeX;
+			TextureDescriptor.height           = SizeY;
+			TextureDescriptor.depth            = 1;
+			TextureDescriptor.mipmapLevelCount = 1; // Only consider a single subresource and not the whole texture (like in the other RHIs)
+			TextureDescriptor.sampleCount      = MetalTexture.sampleCount;
+			TextureDescriptor.arrayLength      = MetalTexture.arrayLength;
+			TextureDescriptor.resourceOptions  = ((MetalTexture.resourceOptions & ~MTLResourceStorageModeMask) | ResourceStorageMode);
+			TextureDescriptor.usage            = MetalTexture.usage;
 			
-			Desc.SetResourceOptions(mtlpp::ResourceOptions(([Texture.GetPtr() resourceOptions] & ~MTLResourceStorageModeMask) | ResourceStorageMode));
-			
-			Desc.SetUsage(Texture.GetUsage());
-			
-			TempTexture = GMtlppDevice.NewTexture(Desc);
+			TempTexture = FAGXTexture([GMtlDevice newTextureWithDescriptor:TextureDescriptor], ns::Ownership::Assign);
 			
 			ImmediateContext.Context->CopyFromTextureToTexture(Texture, 0, InFlags.GetMip(), Region.origin, Region.size, TempTexture, 0, 0, MTLOriginMake(0, 0, 0));
 			
 			Texture = TempTexture;
+			MetalTexture = Texture.GetPtr();
 			Region = MTLRegionMake2D(0, 0, SizeX, SizeY);
+
+			[TextureDescriptor release];
 		}
 #if PLATFORM_MAC
-		if ([Texture.GetPtr() storageMode] == MTLStorageModeManaged)
+		if (MetalTexture.storageMode == MTLStorageModeManaged)
 		{
 			// Synchronise the texture with the CPU
 			ImmediateContext.Context->SynchronizeTexture(Texture, 0, InFlags.GetMip());
@@ -329,12 +330,12 @@ void FAGXDynamicRHI::RHIReadSurfaceData(FRHITexture* TextureRHI, FIntRect Rect,
 		TArray<uint8> Data;
 		Data.AddUninitialized(BytesPerImage);
 		
-		[Texture.GetPtr() getBytes:(void*)Data.GetData()
-					   bytesPerRow:Stride
-					 bytesPerImage:BytesPerImage
-						fromRegion:Region
-					   mipmapLevel:0
-							 slice:0];
+		[MetalTexture getBytes:(void*)Data.GetData()
+				   bytesPerRow:Stride
+				 bytesPerImage:BytesPerImage
+					fromRegion:Region
+				   mipmapLevel:0
+						 slice:0];
 		
 		ConvertSurfaceDataToFColor(Surface->GetDesc().Format, SizeX, SizeY, (uint8*)Data.GetData(), Stride, OutDataPtr, InFlags);
 		
diff --git a/Engine/Source/Runtime/Apple/AGXRHI/Private/AGXTexture.cpp b/Engine/Source/Runtime/Apple/AGXRHI/Private/AGXTexture.cpp
--- a/Engine/Source/Runtime/Apple/AGXRHI/Private/AGXTexture.cpp
+++ b/Engine/Source/Runtime/Apple/AGXRHI/Private/AGXTexture.cpp
@@ -304,20 +304,25 @@ void FAGXSurface::PrepareTextureView()
 
 FAGXTexture FAGXSurface::Reallocate(FAGXTexture InTexture, MTLTextureUsage UsageModifier)
 {
-	mtlpp::TextureDescriptor Desc;
-	Desc.SetTextureType(InTexture.GetTextureType());
-	Desc.SetPixelFormat(InTexture.GetPixelFormat());
-	Desc.SetWidth(InTexture.GetWidth());
-	Desc.SetHeight(InTexture.GetHeight());
-	Desc.SetDepth(InTexture.GetDepth());
-	Desc.SetMipmapLevelCount(InTexture.GetMipmapLevelCount());
-	Desc.SetSampleCount(InTexture.GetSampleCount());
-	Desc.SetArrayLength(InTexture.GetArrayLength());
-	Desc.SetResourceOptions(mtlpp::ResourceOptions([InTexture.GetPtr() resourceOptions]));
-	Desc.SetUsage(mtlpp::TextureUsage([InTexture.GetPtr() usage] | UsageModifier));
-	
-	FAGXTexture NewTex = GetAGXDeviceContext().CreateTexture(this, Desc);
+	id<MTLTexture> InMetalTexture = InTexture.GetPtr();
+
+	MTLTextureDescriptor* TextureDescriptor = [[MTLTextureDescriptor alloc] init];
+	TextureDescriptor.textureType      = InMetalTexture.textureType;
+	TextureDescriptor.pixelFormat      = InMetalTexture.pixelFormat;
+	TextureDescriptor.width            = InMetalTexture.width;
+	TextureDescriptor.height           = InMetalTexture.height;
+	TextureDescriptor.depth            = InMetalTexture.depth;
+	TextureDescriptor.mipmapLevelCount = InMetalTexture.mipmapLevelCount;
+	TextureDescriptor.sampleCount      = InMetalTexture.sampleCount;
+	TextureDescriptor.arrayLength      = InMetalTexture.arrayLength;
+	TextureDescriptor.resourceOptions  = InMetalTexture.resourceOptions;
+	TextureDescriptor.usage            = (InMetalTexture.usage | UsageModifier);
+	
+	FAGXTexture NewTex = GetAGXDeviceContext().CreateTexture(this, TextureDescriptor);
 	check(NewTex);
+	
+	[TextureDescriptor release];
+	
 	return NewTex;
 }
 
@@ -368,61 +373,78 @@ uint8 AGXGetMetalPixelFormatKey(MTLPixelFormat Format)
 	return PixelFormatKeyMap.Get(Format);
 }
 
-FAGXTextureDesc::FAGXTextureDesc(FRHITextureDesc const& InDesc)
+FAGXTextureDesc::FAGXTextureDesc(const FRHITextureDesc& InDesc)
 	: bIsRenderTarget(IsRenderTarget(InDesc.Flags))
 {
-	MTLFormat = (MTLPixelFormat)GPixelFormats[InDesc.Format].PlatformFormat;
+	PixelFormat = (MTLPixelFormat)GPixelFormats[InDesc.Format].PlatformFormat;
+
 	if (EnumHasAnyFlags(InDesc.Flags, TexCreate_SRGB))
 	{
-		MTLFormat = AGXToSRGBFormat(MTLFormat);
+		PixelFormat = AGXToSRGBFormat(PixelFormat);
 	}
 
 	// get a unique key for this surface's format
-	FormatKey = AGXGetMetalPixelFormatKey(MTLFormat);
+	FormatKey = AGXGetMetalPixelFormatKey(PixelFormat);
 
 	if (InDesc.IsTextureCube())
 	{
-		Desc = mtlpp::TextureDescriptor::TextureCubeDescriptor((mtlpp::PixelFormat)MTLFormat, InDesc.Extent.X, (InDesc.NumMips > 1));
+		Desc = MakeUnique<ReferencedTextureDescriptorType>(
+				[MTLTextureDescriptor textureCubeDescriptorWithPixelFormat:PixelFormat
+																	  size:InDesc.Extent.X
+																 mipmapped:(InDesc.NumMips > 1)]);
 	}
 	else if (InDesc.IsTexture3D())
 	{
-		Desc.SetTextureType(mtlpp::TextureType::Texture3D);
-		Desc.SetWidth(InDesc.Extent.X);
-		Desc.SetHeight(InDesc.Extent.Y);
-		Desc.SetDepth(InDesc.Depth);
-		Desc.SetPixelFormat((mtlpp::PixelFormat)MTLFormat);
-		Desc.SetArrayLength(1);
-		Desc.SetMipmapLevelCount(1);
-		Desc.SetSampleCount(1);
+		MTLTextureDescriptor* TextureDescriptor = [[MTLTextureDescriptor alloc] init];
+		TextureDescriptor.textureType      = MTLTextureType3D;
+		TextureDescriptor.width            = InDesc.Extent.X;
+		TextureDescriptor.height           = InDesc.Extent.Y;
+		TextureDescriptor.depth            = InDesc.Depth;
+		TextureDescriptor.pixelFormat      = PixelFormat;;
+		TextureDescriptor.arrayLength      = 1;
+		TextureDescriptor.mipmapLevelCount = 1;
+		TextureDescriptor.sampleCount      = 1;
+
+		Desc = MakeUnique<ReferencedTextureDescriptorType>(TextureDescriptor);
+
+		[TextureDescriptor release];
 	}
 	else
 	{
-		Desc = mtlpp::TextureDescriptor::Texture2DDescriptor((mtlpp::PixelFormat)MTLFormat, InDesc.Extent.X, InDesc.Extent.Y, (InDesc.NumMips > 1));
-		Desc.SetArrayLength(InDesc.ArraySize);
+		Desc = MakeUnique<ReferencedTextureDescriptorType>(
+				[MTLTextureDescriptor texture2DDescriptorWithPixelFormat:PixelFormat
+																   width:InDesc.Extent.X
+																  height:InDesc.Extent.Y
+															   mipmapped:(InDesc.NumMips > 1)]);
+		[Desc.Get()->Object setArrayLength:InDesc.ArraySize];
 	}
+	check(Desc.IsValid());
+
+	MTLTextureDescriptor* TextureDescriptor = Desc.Get()->Object;
 
 	// flesh out the descriptor
 	if (InDesc.IsTextureArray())
 	{
-		Desc.SetArrayLength(InDesc.ArraySize);
+		TextureDescriptor.arrayLength = InDesc.ArraySize;
 		if (InDesc.IsTextureCube())
 		{
 			if (FAGXCommandQueue::SupportsFeature(EAGXFeaturesCubemapArrays))
 			{
-				Desc.SetTextureType(mtlpp::TextureType::TextureCubeArray);
+				TextureDescriptor.textureType = MTLTextureTypeCubeArray;
 			}
 			else
 			{
-				Desc.SetTextureType(mtlpp::TextureType::Texture2DArray);
-				Desc.SetArrayLength(InDesc.ArraySize * 6);
+				TextureDescriptor.textureType = MTLTextureType2DArray;
+				TextureDescriptor.arrayLength = (InDesc.ArraySize * 6);
 			}
 		}
 		else
 		{
-			Desc.SetTextureType(mtlpp::TextureType::Texture2DArray);
+			TextureDescriptor.textureType = MTLTextureType2DArray;
 		}
 	}
-	Desc.SetMipmapLevelCount(InDesc.NumMips);
+
+	TextureDescriptor.mipmapLevelCount = InDesc.NumMips;
 
 	{
 		MTLResourceOptions ResourceStorageMode;
@@ -503,8 +525,8 @@ FAGXTextureDesc::FAGXTextureDesc(FRHITextureDesc const& InDesc)
 		}
 #endif
 
-		Desc.SetResourceOptions(mtlpp::ResourceOptions(MTLResourceCPUCacheModeDefaultCache | ResourceStorageMode | MTLResourceHazardTrackingModeDefault));
-		Desc.SetUsage(mtlpp::TextureUsage(ConvertFlagsToUsage(InDesc.Flags)));
+		TextureDescriptor.resourceOptions = (MTLResourceCPUCacheModeDefaultCache | ResourceStorageMode | MTLResourceHazardTrackingModeDefault);
+		TextureDescriptor.usage = ConvertFlagsToUsage(InDesc.Flags);
 	}
 	
 	if (!FParse::Param(FCommandLine::Get(), TEXT("nomsaa")))
@@ -512,17 +534,17 @@ FAGXTextureDesc::FAGXTextureDesc(FRHITextureDesc const& InDesc)
 		if (InDesc.NumSamples > 1)
 		{
 			check(bIsRenderTarget);
-			Desc.SetTextureType(mtlpp::TextureType::Texture2DMultisample);
+			TextureDescriptor.textureType = MTLTextureType2DMultisample;
 
 			// allow commandline to override
 			uint32 NewNumSamples;
 			if (FParse::Value(FCommandLine::Get(), TEXT("msaa="), NewNumSamples))
 			{
-				Desc.SetSampleCount(NewNumSamples);
+				TextureDescriptor.sampleCount = NewNumSamples;
 			}
 			else
 			{
-				Desc.SetSampleCount(InDesc.NumSamples);
+				TextureDescriptor.sampleCount = InDesc.NumSamples;
 			}
 		}
 	}
@@ -548,6 +570,8 @@ FAGXSurface::FAGXSurface(FAGXTextureCreateDesc const& CreateDesc)
 		return;
 	}
 
+	MTLTextureDescriptor* TextureDescriptor = CreateDesc.Desc.Get()->Object;
+
 	FResourceBulkDataInterface* BulkData = CreateDesc.BulkData;
 
 	// The bulk data interface can be used to create external textures for VR and media player.
@@ -567,7 +591,7 @@ FAGXSurface::FAGXSurface(FAGXTextureCreateDesc const& CreateDesc)
 #else
 				Texture = CVMetalTextureGetTexture((CVMetalTextureRef)ImageSurfaceRef);
 #endif
-				METAL_FATAL_ASSERT(Texture, TEXT("Failed to create texture, desc %s"), *FString([CreateDesc.Desc.GetPtr() description]));
+				METAL_FATAL_ASSERT(Texture, TEXT("Failed to create texture, desc %s"), *FString([TextureDescriptor description]));
 
 				BulkData->Discard();
 				BulkData = nullptr;
@@ -580,16 +604,19 @@ FAGXSurface::FAGXSurface(FAGXTextureCreateDesc const& CreateDesc)
 				ImageSurfaceRef = (CFTypeRef)BulkData->GetResourceBulkData();
 				CFRetain(ImageSurfaceRef);
 
-				mtlpp::TextureDescriptor DescCopy([CreateDesc.Desc.GetPtr() copy], ns::Ownership::Assign);
-				[DescCopy.GetPtr() setResourceOptions:(([DescCopy.GetPtr() resourceOptions] & ~MTLResourceStorageModeMask) | MTLResourceStorageModeManaged)];
+				MTLTextureDescriptor* DescCopy = [TextureDescriptor copy];
+				DescCopy.resourceOptions = ((DescCopy.resourceOptions & ~MTLResourceStorageModeMask) | MTLResourceStorageModeManaged);
 
-				ns::IOSurface MtlppIOSurface((IOSurfaceRef)ImageSurfaceRef, ns::Ownership::AutoRelease);
-				Texture = MTLPP_VALIDATE(mtlpp::Device, GMtlppDevice, AGXSafeGetRuntimeDebuggingLevel() >= EAGXDebugLevelValidation, NewTextureWithDescriptor(DescCopy, MtlppIOSurface, 0));
+				Texture = mtlpp::Texture([GMtlDevice newTextureWithDescriptor:DescCopy
+																   iosurface:(IOSurfaceRef)ImageSurfaceRef
+																	   plane:0], nullptr, ns::Ownership::Assign);
 				
-				METAL_FATAL_ASSERT(Texture, TEXT("Failed to create texture, desc %s"), *FString([DescCopy.GetPtr() description]));
+				METAL_FATAL_ASSERT(Texture, TEXT("Failed to create texture, desc %s"), *FString([DescCopy description]));
 
 				BulkData->Discard();
 				BulkData = nullptr;
+
+				[DescCopy release];
 			}
 			break;
 #endif
@@ -601,29 +628,31 @@ FAGXSurface::FAGXSurface(FAGXTextureCreateDesc const& CreateDesc)
 		// Non VR/media texture case (i.e. a regular texture).
 		// Create the actual texture resource.
 
-		const bool bBufferCompatibleOption = 	(CreateDesc.Desc.GetTextureType() == mtlpp::TextureType::Texture2D || CreateDesc.Desc.GetTextureType() == mtlpp::TextureType::TextureBuffer) &&
-												CreateDesc.NumMips == 1 && CreateDesc.ArraySize == 1 && CreateDesc.NumSamples == 1 && CreateDesc.Desc.GetDepth() == 1;
+		const bool bBufferCompatibleOption = (TextureDescriptor.textureType == MTLTextureType2D || TextureDescriptor.textureType == MTLTextureTypeTextureBuffer) &&
+											  CreateDesc.NumMips == 1 && CreateDesc.ArraySize == 1 && CreateDesc.NumSamples == 1 && TextureDescriptor.depth == 1;
 
-		if(bBufferCompatibleOption && (EnumHasAllFlags(CreateDesc.Flags, TexCreate_UAV | TexCreate_NoTiling) || EnumHasAllFlags(CreateDesc.Flags, TexCreate_AtomicCompatible)))
+		if (bBufferCompatibleOption && (EnumHasAllFlags(CreateDesc.Flags, TexCreate_UAV | TexCreate_NoTiling) || EnumHasAllFlags(CreateDesc.Flags, TexCreate_AtomicCompatible)))
 		{
-			const uint32 MinimumByteAlignment = GMtlppDevice.GetMinimumLinearTextureAlignmentForPixelFormat((mtlpp::PixelFormat)CreateDesc.MTLFormat);
-			const NSUInteger BytesPerRow = Align(CreateDesc.Desc.GetWidth() * GPixelFormats[CreateDesc.Format].BlockBytes, MinimumByteAlignment);
+			const uint32 MinimumByteAlignment = [GMtlDevice minimumLinearTextureAlignmentForPixelFormat:CreateDesc.PixelFormat];
+			const NSUInteger BytesPerRow = Align(TextureDescriptor.width * GPixelFormats[CreateDesc.Format].BlockBytes, MinimumByteAlignment);
 
 			// Backing buffer resource options must match the texture we are going to create from it
-			FAGXPooledBufferArgs Args(BytesPerRow * CreateDesc.Desc.GetHeight(), BUF_Dynamic, FAGXPooledBufferArgs::PrivateStorageResourceOptions);
+			FAGXPooledBufferArgs Args(BytesPerRow * TextureDescriptor.height, BUF_Dynamic, FAGXPooledBufferArgs::PrivateStorageResourceOptions);
 			FAGXBuffer Buffer = GetAGXDeviceContext().CreatePooledBuffer(Args);
 
-			Texture = Buffer.NewTexture(CreateDesc.Desc, 0, BytesPerRow);
+			Texture = mtlpp::Texture([Buffer.GetPtr() newTextureWithDescriptor:TextureDescriptor
+																		offset:Buffer.GetOffset()
+																   bytesPerRow:BytesPerRow], nullptr, ns::Ownership::Assign);
 		}
 		else
 		{
 			// If we are in here then either the texture description is not buffer compatable or these flags were not set
 			// assert that these flag combinations are not set as they require a buffer backed texture and the texture description is not compatible with that
-			checkf(!EnumHasAllFlags(CreateDesc.Flags, TexCreate_AtomicCompatible), TEXT("Requested buffer backed texture that breaks Metal linear texture limitations: %s"), *FString([CreateDesc.Desc description]));
-			Texture = GetAGXDeviceContext().CreateTexture(this, CreateDesc.Desc);
+			checkf(!EnumHasAllFlags(CreateDesc.Flags, TexCreate_AtomicCompatible), TEXT("Requested buffer backed texture that breaks Metal linear texture limitations: %s"), *FString([TextureDescriptor description]));
+			Texture = GetAGXDeviceContext().CreateTexture(this, TextureDescriptor);
 		}
 		
-		METAL_FATAL_ASSERT(Texture, TEXT("Failed to create texture, desc %s"), *FString([CreateDesc.Desc description]));
+		METAL_FATAL_ASSERT(Texture, TEXT("Failed to create texture, desc %s"), *FString([TextureDescriptor description]));
 	}
 
 	if (BulkData)
@@ -652,7 +681,7 @@ FAGXSurface::FAGXSurface(FAGXTextureCreateDesc const& CreateDesc)
 
 	if (CreateDesc.NumSamples > 1 && !FParse::Param(FCommandLine::Get(), TEXT("nomsaa")))
 	{
-		MSAATexture = GetAGXDeviceContext().CreateTexture(this, CreateDesc.Desc);
+		MSAATexture = GetAGXDeviceContext().CreateTexture(this, TextureDescriptor);
 
 		//device doesn't support HW depth resolve.  This case only valid on mobile renderer or
 		//on Mac where RHISupportsSeparateMSAAAndResolveTextures is true.
@@ -693,7 +722,7 @@ FAGXSurface::FAGXSurface(FAGXTextureCreateDesc const& CreateDesc)
 
 		if (MSAATexture.GetPtr() == nil)
 		{
-			UE_LOG(LogAGX, Warning, TEXT("Failed to create MSAA texture with descriptor: %s"), *FString([CreateDesc.Desc.GetPtr() description]));
+			UE_LOG(LogAGX, Warning, TEXT("Failed to create MSAA texture with descriptor: %s"), *FString([TextureDescriptor description]));
 		}
 	}
 
diff --git a/Engine/Source/Runtime/Apple/AGXRHI/Private/AGXVertexBuffer.cpp b/Engine/Source/Runtime/Apple/AGXRHI/Private/AGXVertexBuffer.cpp
--- a/Engine/Source/Runtime/Apple/AGXRHI/Private/AGXVertexBuffer.cpp
+++ b/Engine/Source/Runtime/Apple/AGXRHI/Private/AGXVertexBuffer.cpp
@@ -293,7 +293,7 @@ void FAGXRHIBuffer::AllocLinearTextures(const LinearTextureMapKey& InLinearTextu
 		
 		MTLPixelFormat MTLFormat = (MTLPixelFormat)GAGXBufferFormats[InFormat].LinearTextureFormat;
 		
-		mtlpp::TextureDescriptor Desc;
+		MTLTextureDescriptor* Desc = nil;
 		MTLTextureUsage TexUsage = MTLTextureUsageUnknown;
 		if (EnumHasAnyFlags(Usage, BUF_ShaderResource))
 		{
@@ -323,15 +323,20 @@ void FAGXRHIBuffer::AllocLinearTextures(const LinearTextureMapKey& InLinearTextu
 
 		{
 			MTLResourceOptions Options = MTLResourceCPUCacheModeDefaultCache | (StorageMode << MTLResourceStorageModeShift) | MTLResourceHazardTrackingModeDefault;
-			Desc = mtlpp::TextureDescriptor::TextureBufferDescriptor((mtlpp::PixelFormat)MTLFormat, NumElements, mtlpp::ResourceOptions(Options), mtlpp::TextureUsage(TexUsage));
-			Desc.SetAllowGPUOptimisedContents(false);
+			Desc = [MTLTextureDescriptor textureBufferDescriptorWithPixelFormat:MTLFormat
+																		  width:NumElements
+																resourceOptions:Options
+																		  usage:TexUsage];
+			Desc.allowGPUOptimizedContents = NO;
 		}
 
 		for(FAGXBufferAndViews& Backing : BufferPool)
 		{
 			FAGXBuffer& Buffer = Backing.Buffer;
-			FAGXTexture NewTexture = MTLPP_VALIDATE(mtlpp::Buffer, Buffer, AGXSafeGetRuntimeDebuggingLevel() >= EAGXDebugLevelValidation, NewTexture(Desc, Offset, RowBytes));
-			METAL_FATAL_ASSERT(NewTexture, TEXT("Failed to create linear texture, desc %s from buffer %s"), *FString([Desc description]), *FString([Buffer description]));
+			FAGXTexture NewTexture([Buffer.GetPtr() newTextureWithDescriptor:Desc
+																	  offset:(Offset + Buffer.GetOffset())
+																 bytesPerRow:RowBytes], ns::Ownership::Assign);
+			METAL_FATAL_ASSERT(NewTexture, TEXT("Failed to create linear texture, desc %s from buffer %s"), *FString([Desc description]), *FString([Buffer.GetPtr() description]));
 			
 			check(GAGXBufferFormats[InFormat].LinearTextureFormat == MTLPixelFormatRG11B10Float || GAGXBufferFormats[InFormat].LinearTextureFormat == (MTLPixelFormat)NewTexture.GetPixelFormat());
 			Backing.Views.Add(InLinearTextureMapKey, NewTexture);
diff --git a/Engine/Source/Runtime/Apple/AGXRHI/Public/AGXResources.h b/Engine/Source/Runtime/Apple/AGXRHI/Public/AGXResources.h
--- a/Engine/Source/Runtime/Apple/AGXRHI/Public/AGXResources.h
+++ b/Engine/Source/Runtime/Apple/AGXRHI/Public/AGXResources.h
@@ -133,29 +133,6 @@ public:
 	}
 };
 
-struct FAGXTextureDesc
-{
-	FAGXTextureDesc(FRHITextureDesc const& InDesc);
-	
-	mtlpp::TextureDescriptor Desc;
-	MTLPixelFormat MTLFormat;
-	bool bMemoryless = false;
-	bool bIsRenderTarget = false;
-	uint8 FormatKey = 0;
-};
-
-struct FAGXTextureCreateDesc : public FRHITextureCreateDesc, public FAGXTextureDesc
-{
-	FAGXTextureCreateDesc(FRHITextureCreateDesc const& CreateDesc)
-		: FRHITextureCreateDesc(CreateDesc)
-		, FAGXTextureDesc(CreateDesc)
-	{
-		// @todo: texture type unification - Metal can override NumSamples based on command line options.
-		// We should instead require the renderer to do this.
-		NumSamples = Desc.GetSampleCount();
-	}
-};
-
 // AGX RHI texture resource
 class AGXRHI_API FAGXSurface : public FRHITexture
 {
@@ -164,7 +141,7 @@ public:
 	/** 
 	 * Constructor that will create Texture and Color/DepthBuffers as needed
 	 */
-	FAGXSurface(FAGXTextureCreateDesc const& CreateDesc);
+	FAGXSurface(const struct FAGXTextureCreateDesc& CreateDesc);
 	
 	/**
 	 * Destructor
