Fix for Apple A8 device support
  - Devices crashing when using Base_Vertex metal shader semantics
  - Runtime error when using Virtual Texturing on A8 devices which is unsupported.

diff --git a/Engine/Source/Developer/Apple/MetalShaderFormat/Private/MetalDerivedData.cpp b/Engine/Source/Developer/Apple/MetalShaderFormat/Private/MetalDerivedData.cpp
--- a/Engine/Source/Developer/Apple/MetalShaderFormat/Private/MetalDerivedData.cpp
+++ b/Engine/Source/Developer/Apple/MetalShaderFormat/Private/MetalDerivedData.cpp
@@ -226,6 +226,13 @@ bool DoCompileMetalShader(
 	bool const bZeroInitialise = Input.Environment.CompilerFlags.Contains(CFLAG_ZeroInitialise);
 	bool const bBoundsChecks = Input.Environment.CompilerFlags.Contains(CFLAG_BoundsChecking);
 
+    bool bSupportAppleA8 = false;
+    FString const* SupportAppleA8 = Input.Environment.GetDefinitions().Find(TEXT("SUPPORT_APPLE_A8"));
+    if (SupportAppleA8)
+    {
+        LexFromString(bSupportAppleA8, *(*SupportAppleA8));
+    }
+    
 	bool bSwizzleSample = false;
 	FString const* Swizzle = Input.Environment.GetDefinitions().Find(TEXT("METAL_SWIZZLE_SAMPLES"));
 	if (Swizzle)
@@ -975,16 +982,16 @@ bool DoCompileMetalShader(
 				break;
 			case EMetalGPUSemanticsTBDRDesktop:
 				TargetDesc.Language = CrossCompiler::EShaderConductorLanguage::Metal_iOS;
-				TargetDesc.CompileFlags.SetDefine(TEXT("ios_support_base_vertex_instance"), 1);
+                TargetDesc.CompileFlags.SetDefine(TEXT("ios_support_base_vertex_instance"), !bSupportAppleA8);
 				TargetDesc.CompileFlags.SetDefine(TEXT("use_framebuffer_fetch_subpasses"), 1);
 				TargetDesc.CompileFlags.SetDefine(TEXT("emulate_cube_array"), 1);
 				break;
 			case EMetalGPUSemanticsMobile:
 			default:
 				TargetDesc.Language = CrossCompiler::EShaderConductorLanguage::Metal_iOS;
-				TargetDesc.CompileFlags.SetDefine(TEXT("use_framebuffer_fetch_subpasses"), 1);
+                TargetDesc.CompileFlags.SetDefine(TEXT("ios_support_base_vertex_instance"), !bSupportAppleA8);
+                TargetDesc.CompileFlags.SetDefine(TEXT("use_framebuffer_fetch_subpasses"), 1);
 				TargetDesc.CompileFlags.SetDefine(TEXT("emulate_cube_array"), 1);
-				TargetDesc.CompileFlags.SetDefine(TEXT("ios_support_base_vertex_instance"), 1);
 				break;
 			}
 
diff --git a/Engine/Source/Runtime/Apple/MetalRHI/Private/MetalRHI.cpp b/Engine/Source/Runtime/Apple/MetalRHI/Private/MetalRHI.cpp
--- a/Engine/Source/Runtime/Apple/MetalRHI/Private/MetalRHI.cpp
+++ b/Engine/Source/Runtime/Apple/MetalRHI/Private/MetalRHI.cpp
@@ -185,19 +185,26 @@ FMetalDynamicRHI::FMetalDynamicRHI(ERHIFeatureLevel::Type RequestedFeatureLevel)
 	mtlpp::Device Device = ImmediateContext.Context->GetDevice();
 		
 #if PLATFORM_IOS
-	// A8 can use 256 bits of MRTs
+    bool bSupportAppleA8 = false;
+    GConfig->GetBool(TEXT("/Script/IOSRuntimeSettings.IOSRuntimeSettings"), TEXT("bSupportAppleA8"), bSupportAppleA8, GEngineIni);
+        
+    bool bIsA8FeatureSet = false;
+        
 #if PLATFORM_TVOS
 	GRHISupportsDrawIndirect = Device.SupportsFeatureSet(mtlpp::FeatureSet::tvOS_GPUFamily2_v1);
 	GRHISupportsPixelShaderUAVs = Device.SupportsFeatureSet(mtlpp::FeatureSet::tvOS_GPUFamily2_v1);
+        
+    if (!Device.SupportsFeatureSet(mtlpp::FeatureSet::tvOS_GPUFamily2_v1))
+    {
+        bIsA8FeatureSet = true;
+    }
+        
 #else
+	if (!Device.SupportsFeatureSet(mtlpp::FeatureSet::iOS_GPUFamily3_v1))
 	{
-		bool bSupportAppleA8 = false;
-		GConfig->GetBool(TEXT("/Script/IOSRuntimeSettings.IOSRuntimeSettings"), TEXT("bSupportAppleA8"), bSupportAppleA8, GEngineIni);
-		if (!Device.SupportsFeatureSet(mtlpp::FeatureSet::iOS_GPUFamily3_v1) && !bSupportAppleA8)
-		{
-			UE_LOG(LogMetal, Fatal, TEXT("This device does not supports the Apple A9 feature set which is the minimum for this build. Please check the Support Apple A8 checkbox in the IOS Project Settings."));
-		}
-	}
+        bIsA8FeatureSet = true;
+    }
+    
 	GRHISupportsRWTextureBuffers = Device.SupportsFeatureSet(mtlpp::FeatureSet::iOS_GPUFamily4_v1);
 	GRHISupportsDrawIndirect = Device.SupportsFeatureSet(mtlpp::FeatureSet::iOS_GPUFamily3_v1);
 	GRHISupportsPixelShaderUAVs = Device.SupportsFeatureSet(mtlpp::FeatureSet::iOS_GPUFamily3_v1);
@@ -229,6 +236,20 @@ FMetalDynamicRHI::FMetalDynamicRHI(ERHIFeatureLevel::Type RequestedFeatureLevel)
 	bSupportsPointLights = GSupportsVolumeTextureRendering;
 #endif
 
+    if(bIsA8FeatureSet)
+    {
+        if(!bSupportAppleA8)
+        {
+            UE_LOG(LogMetal, Fatal, TEXT("This device does not supports the Apple A8x or above feature set which is the minimum for this build. Please check the Support Apple A8 checkbox in the IOS Project Settings."));
+        }
+        
+        static auto* CVarMobileVirtualTextures = IConsoleManager::Get().FindTConsoleVariableDataInt(TEXT("r.Mobile.VirtualTextures"));
+        if(CVarMobileVirtualTextures->GetValueOnAnyThread() != 0)
+        {
+            UE_LOG(LogMetal, Fatal, TEXT("Mobile Virtual Textures require a minimum of the Apple A9 feature set."));
+        }
+    }
+        
     bool bProjectSupportsMRTs = false;
     GConfig->GetBool(TEXT("/Script/IOSRuntimeSettings.IOSRuntimeSettings"), TEXT("bSupportsMetalMRT"), bProjectSupportsMRTs, GEngineIni);
 
diff --git a/Engine/Source/Runtime/Engine/Private/ShaderCompiler/ShaderCompiler.cpp b/Engine/Source/Runtime/Engine/Private/ShaderCompiler/ShaderCompiler.cpp
--- a/Engine/Source/Runtime/Engine/Private/ShaderCompiler/ShaderCompiler.cpp
+++ b/Engine/Source/Runtime/Engine/Private/ShaderCompiler/ShaderCompiler.cpp
@@ -5722,6 +5722,8 @@ void GlobalBeginCompileShader(
 			bool bForceFloats = false;
 			FString IndirectArgumentTier;
 			bool bEnableMathOptimisations = true;
+            bool bSupportAppleA8 = false;
+            
 			if (IsPCPlatform(EShaderPlatform(Target.Platform)))
 			{
 				GConfig->GetBool(TEXT("/Script/MacTargetPlatform.MacTargetSettings"), TEXT("UseFastIntrinsics"), bAllowFastIntrinsics, GEngineIni);
@@ -5737,6 +5739,7 @@ void GlobalBeginCompileShader(
 				GConfig->GetBool(TEXT("/Script/IOSRuntimeSettings.IOSRuntimeSettings"), TEXT("EnableMathOptimisations"), bEnableMathOptimisations, GEngineIni);
 				GConfig->GetBool(TEXT("/Script/IOSRuntimeSettings.IOSRuntimeSettings"), TEXT("ForceFloats"), bForceFloats, GEngineIni);
 				GConfig->GetString(TEXT("/Script/IOSRuntimeSettings.IOSRuntimeSettings"), TEXT("IndirectArgumentTier"), IndirectArgumentTier, GEngineIni);
+                GConfig->GetBool(TEXT("/Script/IOSRuntimeSettings.IOSRuntimeSettings"), TEXT("bSupportAppleA8"), bSupportAppleA8, GEngineIni);
                 
 				// Force no development shaders on iOS
 				bAllowDevelopmentShaderCompile = false;
@@ -5746,6 +5749,7 @@ void GlobalBeginCompileShader(
             
 			Input.Environment.SetDefine(TEXT("METAL_USE_FAST_INTRINSICS"), bAllowFastIntrinsics);
 			Input.Environment.SetDefine(TEXT("METAL_INDIRECT_ARGUMENT_BUFFERS"), IndirectArgumentTier);
+            Input.Environment.SetDefine(TEXT("SUPPORT_APPLE_A8"), bSupportAppleA8);
 			
 			// Same as console-variable above, but that's global and this is per-platform, per-project
 			if (!bEnableMathOptimisations)
diff --git a/Engine/Source/Runtime/IOS/IOSRuntimeSettings/Classes/IOSRuntimeSettings.h b/Engine/Source/Runtime/IOS/IOSRuntimeSettings/Classes/IOSRuntimeSettings.h
--- a/Engine/Source/Runtime/IOS/IOSRuntimeSettings/Classes/IOSRuntimeSettings.h
+++ b/Engine/Source/Runtime/IOS/IOSRuntimeSettings/Classes/IOSRuntimeSettings.h
@@ -462,7 +462,11 @@ public:
 	UPROPERTY(config, EditAnywhere, Category = Rendering, Meta = (DisplayName = "Tier of Indirect Argument Buffers to use when compiling shaders", ConfigRestartRequired = true))
 	int32 IndirectArgumentTier;
 	
-    /** Supports Apple A8 devices, disables 3d texture compression */
+    /** Supports Apple A8 devices.
+     * Disables 3d texture compression
+     * Virtual Textures are not supported on A8 devices due to missing hardware features
+     * This will disable also Base_Vertex semantics in the msl which may have negative consequences if shaders rely on it
+     */
     UPROPERTY(config, EditAnywhere, Category = Rendering, Meta = (DisplayName = "Support Apple A8", ConfigRestartRequired = true))
     bool bSupportAppleA8;
     
diff --git a/Engine/Source/Runtime/RenderCore/Private/Shader.cpp b/Engine/Source/Runtime/RenderCore/Private/Shader.cpp
--- a/Engine/Source/Runtime/RenderCore/Private/Shader.cpp
+++ b/Engine/Source/Runtime/RenderCore/Private/Shader.cpp
@@ -1621,7 +1621,9 @@ void ShaderMapAppendKeyString(EShaderPlatform Platform, FString& KeyString)
 		bool bAllowFastIntrinsics = false;
 		bool bEnableMathOptimisations = true;
 		bool bForceFloats = false;
+        bool bSupportAppleA8 = false;
 		int32 IndirectArgumentTier = 0;
+        
 		if (IsPCPlatform(Platform))
 		{
 			GConfig->GetBool(TEXT("/Script/MacTargetPlatform.MacTargetSettings"), TEXT("UseFastIntrinsics"), bAllowFastIntrinsics, GEngineIni);
@@ -1634,6 +1636,7 @@ void ShaderMapAppendKeyString(EShaderPlatform Platform, FString& KeyString)
 			GConfig->GetBool(TEXT("/Script/IOSRuntimeSettings.IOSRuntimeSettings"), TEXT("EnableMathOptimisations"), bEnableMathOptimisations, GEngineIni);
 			GConfig->GetBool(TEXT("/Script/IOSRuntimeSettings.IOSRuntimeSettings"), TEXT("ForceFloats"), bForceFloats, GEngineIni);
 			GConfig->GetInt(TEXT("/Script/IOSRuntimeSettings.IOSRuntimeSettings"), TEXT("IndirectArgumentTier"), IndirectArgumentTier, GEngineIni);
+            GConfig->GetBool(TEXT("/Script/IOSRuntimeSettings.IOSRuntimeSettings"), TEXT("bSupportAppleA8"), bSupportAppleA8, GEngineIni);
 		}
 		
 		if (bAllowFastIntrinsics)
@@ -1652,6 +1655,11 @@ void ShaderMapAppendKeyString(EShaderPlatform Platform, FString& KeyString)
 			KeyString += TEXT("_FP32");
 		}
 		
+        if(bSupportAppleA8)
+        {
+            KeyString += TEXT("_A8GPU");
+        }
+        
 		KeyString += FString::Printf(TEXT("_IAB%d"), IndirectArgumentTier);
 		
 		// Shaders built for archiving - for Metal that requires compiling the code in a different way so that we can strip it later
