AGXRHI: Replace mtlpp::library and mtlpp::function with native Metal types.

- Also, rename FObjCWrapperRetained to FObjCWrapper and improve it slightly to make it more conveneint to work within UE containers.

#rnx

diff --git a/Engine/Source/Runtime/Apple/AGXRHI/Private/AGXCaptureManager.cpp b/Engine/Source/Runtime/Apple/AGXRHI/Private/AGXCaptureManager.cpp
--- a/Engine/Source/Runtime/Apple/AGXRHI/Private/AGXCaptureManager.cpp
+++ b/Engine/Source/Runtime/Apple/AGXRHI/Private/AGXCaptureManager.cpp
@@ -15,7 +15,7 @@ FAGXCaptureManager::FAGXCaptureManager(FAGXCommandQueue& InQueue)
 	CaptureManager.defaultCaptureScope = DefaultCaptureScope;
 	
 	FAGXCaptureScope DefaultScope;
-	DefaultScope.CaptureScope = MakeUnique<RetainedCaptureScopeType>(DefaultCaptureScope);
+	DefaultScope.CaptureScope = new FMTLCaptureScope(DefaultCaptureScope, /* bRetain = */ false);
 	DefaultScope.Type = EAGXCaptureTypePresent;
 	DefaultScope.StepCount = 1;
 	DefaultScope.LastTrigger = 0;
@@ -23,9 +23,6 @@ FAGXCaptureManager::FAGXCaptureManager(FAGXCommandQueue& InQueue)
 	
 	[DefaultCaptureScope beginScope];
 	
-	// The local reference is no longer required (object lifetime is managed by DefaultScope).
-	[DefaultCaptureScope release];
-	
 	uint32 PresentStepCounts[] = {2, 5, 10, 15, 30, 60, 90, 120};
 	for (uint32 i = 0; i < (sizeof(PresentStepCounts) / sizeof(uint32)); i++)
 	{
@@ -33,16 +30,13 @@ FAGXCaptureManager::FAGXCaptureManager(FAGXCommandQueue& InQueue)
 		CaptureScope.label = [NSString stringWithFormat:@"%u Frames", PresentStepCounts[i]];
 	
 		FAGXCaptureScope Scope;
-		Scope.CaptureScope = MakeUnique<RetainedCaptureScopeType>(CaptureScope);
+		Scope.CaptureScope = new FMTLCaptureScope(CaptureScope, /* bRetain = */ false);
 		Scope.Type = EAGXCaptureTypePresent;
 		Scope.StepCount = PresentStepCounts[i];
 		Scope.LastTrigger = 0;
 		ActiveScopes.Add(Scope);
 	
 		[CaptureScope beginScope];
-	
-		// The local reference is no longer required (object lifetime is managed by Scope).
-		[CaptureScope release];
 	}
 }
 
@@ -66,8 +60,8 @@ void FAGXCaptureManager::PresentFrame(uint32 FrameNumber)
 		
 		if (Diff >= Scope.StepCount)
 		{
-			[Scope.CaptureScope.Get()->Object endScope];
-			[Scope.CaptureScope.Get()->Object beginScope];
+			[Scope.CaptureScope->Get() endScope];
+			[Scope.CaptureScope->Get() beginScope];
 			Scope.LastTrigger = FrameNumber;
 		}
 	}
diff --git a/Engine/Source/Runtime/Apple/AGXRHI/Private/AGXCaptureManager.h b/Engine/Source/Runtime/Apple/AGXRHI/Private/AGXCaptureManager.h
--- a/Engine/Source/Runtime/Apple/AGXRHI/Private/AGXCaptureManager.h
+++ b/Engine/Source/Runtime/Apple/AGXRHI/Private/AGXCaptureManager.h
@@ -20,8 +20,6 @@ private:
 	FAGXCommandQueue& Queue;
 	
 private:
-	using RetainedCaptureScopeType = FObjCWrapperRetained< id<MTLCaptureScope> >;
-
 	enum EAGXCaptureType
 	{
 		EAGXCaptureTypeUnknown,
@@ -35,7 +33,7 @@ private:
 		EAGXCaptureType Type;
 		uint32 StepCount;
 		uint32 LastTrigger;
-		TUniquePtr<RetainedCaptureScopeType> CaptureScope;
+		TRefCountPtr<FMTLCaptureScope> CaptureScope;
 
 		FAGXCaptureScope()
 		: Type(EAGXCaptureTypeUnknown)
@@ -49,7 +47,7 @@ private:
 		: Type(Other.Type)
 		, StepCount(Other.StepCount)
 		, LastTrigger(Other.LastTrigger)
-		, CaptureScope(MakeUnique<RetainedCaptureScopeType>(Other.CaptureScope.Get()->Object))
+		, CaptureScope(Other.CaptureScope)
 		{
 		}
 
diff --git a/Engine/Source/Runtime/Apple/AGXRHI/Private/AGXDynamicRHI_Shaders.cpp b/Engine/Source/Runtime/Apple/AGXRHI/Private/AGXDynamicRHI_Shaders.cpp
--- a/Engine/Source/Runtime/Apple/AGXRHI/Private/AGXDynamicRHI_Shaders.cpp
+++ b/Engine/Source/Runtime/Apple/AGXRHI/Private/AGXDynamicRHI_Shaders.cpp
@@ -36,7 +36,8 @@ FGeometryShaderRHIRef FAGXDynamicRHI::RHICreateGeometryShader(TArrayView<const u
 	@autoreleasepool {
 		FAGXGeometryShader* Shader = new FAGXGeometryShader;
 		FMetalCodeHeader Header;
-		Shader->Init(Code, Header);
+		TRefCountPtr<FMTLLibrary> Library;
+		Shader->Init(Code, Header, Library);
 		return Shader;
 	}
 }
@@ -44,7 +45,8 @@ FGeometryShaderRHIRef FAGXDynamicRHI::RHICreateGeometryShader(TArrayView<const u
 FComputeShaderRHIRef FAGXDynamicRHI::RHICreateComputeShader(TArrayView<const uint8> Code, const FSHAHash& Hash)
 {
 	@autoreleasepool {
-		return new FAGXComputeShader(Code);
+		TRefCountPtr<FMTLLibrary> Library;
+		return new FAGXComputeShader(Code, Library);
 	}
 }
 
@@ -97,7 +99,7 @@ FRHIShaderLibraryRef FAGXDynamicRHI::RHICreateShaderLibrary(EShaderPlatform Plat
 			{
 				check(((SerializedShaders.GetNumShaders() + Header.NumShadersPerLibrary - 1) / Header.NumShadersPerLibrary) == Header.NumLibraries);
 
-				TArray<mtlpp::Library> Libraries;
+				TArray<TRefCountPtr<FMTLLibrary>> Libraries;
 				Libraries.Empty(Header.NumLibraries);
 
 				for (uint32 i = 0; i < Header.NumLibraries; i++)
@@ -107,10 +109,10 @@ FRHIShaderLibraryRef FAGXDynamicRHI::RHICreateShaderLibrary(EShaderPlatform Plat
 
 					METAL_GPUPROFILE(FAGXScopedCPUStats CPUStat(FString::Printf(TEXT("NewLibraryFile: %s"), *MetalLibraryFilePath)));
 					NSError* Error;
-					mtlpp::Library Library([GMtlDevice newLibraryWithFile:MetalLibraryFilePath.GetNSString() error:&Error], nullptr, ns::Ownership::Assign);
+					id<MTLLibrary> Library = [GMtlDevice newLibraryWithFile:MetalLibraryFilePath.GetNSString() error:&Error];
 					if (Library != nil)
 					{
-						Libraries.Add(Library);
+						Libraries.Add(new FMTLLibrary(Library, /* bRetain = */ false));
 					}
 					else
 					{
diff --git a/Engine/Source/Runtime/Apple/AGXRHI/Private/AGXPipeline.cpp b/Engine/Source/Runtime/Apple/AGXRHI/Private/AGXPipeline.cpp
--- a/Engine/Source/Runtime/Apple/AGXRHI/Private/AGXPipeline.cpp
+++ b/Engine/Source/Runtime/Apple/AGXRHI/Private/AGXPipeline.cpp
@@ -869,8 +869,8 @@ static FAGXShaderPipeline* CreateMTLRenderPipeline(bool const bSync, FAGXGraphic
     FAGXVertexShader* VertexShader = (FAGXVertexShader*)Init.BoundShaderState.VertexShaderRHI;
     FAGXPixelShader* PixelShader = (FAGXPixelShader*)Init.BoundShaderState.PixelShaderRHI;
     
- 	id<MTLFunction> vertexFunction = VertexShader->GetFunction().GetPtr();
-    id<MTLFunction> fragmentFunction = PixelShader ? PixelShader->GetFunction().GetPtr() : nil;
+ 	id<MTLFunction> vertexFunction = VertexShader->GetFunction();
+    id<MTLFunction> fragmentFunction = PixelShader ? PixelShader->GetFunction() : nil;
 
     FAGXShaderPipeline* Pipeline = nil;
     if (vertexFunction && ((PixelShader != nullptr) == (fragmentFunction != nil)))
diff --git a/Engine/Source/Runtime/Apple/AGXRHI/Private/AGXRHIPrivate.h b/Engine/Source/Runtime/Apple/AGXRHI/Private/AGXRHIPrivate.h
--- a/Engine/Source/Runtime/Apple/AGXRHI/Private/AGXRHIPrivate.h
+++ b/Engine/Source/Runtime/Apple/AGXRHI/Private/AGXRHIPrivate.h
@@ -129,34 +129,43 @@ enum EMTLTextureType
 #define METAL_SUPPORTS_TILE_SHADERS 1
 
 template<typename ReferenceType>
-struct FObjCWrapperRetained
+class FObjCWrapper : public FThreadSafeRefCountedObject
 {
-	using Type = ReferenceType;
-
-	FObjCWrapperRetained() = delete;
-	FObjCWrapperRetained(const FObjCWrapperRetained&) = delete;
-	FObjCWrapperRetained(FObjCWrapperRetained&&) = delete;
-
-	explicit FObjCWrapperRetained(Type InObject = nil)
-		: Object([InObject retain])
+public:
+	FObjCWrapper() = delete;
+	FObjCWrapper(const FObjCWrapper&) = delete;
+	FObjCWrapper(FObjCWrapper&&) = delete;
+
+	explicit FObjCWrapper(ReferenceType InObject = nil, bool bRetain = true)
+		: FThreadSafeRefCountedObject()
+		, Object(bRetain ? [InObject retain] : InObject)
 	{
 	}
 
-	~FObjCWrapperRetained()
+	~FObjCWrapper()
 	{
 		[Object release];
 	}
 
-	Type Object;
+	ReferenceType Get() const
+	{
+		return Object;
+	}
+
+private:
+	ReferenceType Object;
 };
 
+typedef FObjCWrapper<id>                    FNSObject;
+typedef FObjCWrapper<id<MTLCaptureScope>>   FMTLCaptureScope;
+typedef FObjCWrapper<id<MTLLibrary>>        FMTLLibrary;
+typedef FObjCWrapper<MTLTextureDescriptor*> FMTLTextureDescriptor;
+
 struct FAGXTextureDesc
 {
-	using ReferencedTextureDescriptorType = FObjCWrapperRetained<MTLTextureDescriptor*>;
-
 	FAGXTextureDesc(const FRHITextureDesc& InDesc);
 	
-	TUniquePtr<ReferencedTextureDescriptorType> Desc;
+	TRefCountPtr<FMTLTextureDescriptor> Desc;
 	MTLPixelFormat PixelFormat = MTLPixelFormatInvalid;
 	bool bMemoryless = false;
 	bool bIsRenderTarget = false;
@@ -171,7 +180,7 @@ struct FAGXTextureCreateDesc : public FRHITextureCreateDesc, public FAGXTextureD
 	{
 		// @todo: texture type unification - Metal can override NumSamples based on command line options.
 		// We should instead require the renderer to do this.
-		NumSamples = [Desc.Get()->Object sampleCount];
+		NumSamples = [Desc->Get() sampleCount];
 	}
 };
 
diff --git a/Engine/Source/Runtime/Apple/AGXRHI/Private/AGXState.cpp b/Engine/Source/Runtime/Apple/AGXRHI/Private/AGXState.cpp
--- a/Engine/Source/Runtime/Apple/AGXRHI/Private/AGXState.cpp
+++ b/Engine/Source/Runtime/Apple/AGXRHI/Private/AGXState.cpp
@@ -169,8 +169,6 @@ static EColorWriteMask TranslateWriteMask(MTLColorWriteMask WriteMask)
 template <typename InitializerType>
 class FAGXStateObjectCache
 {
-	using RetainedObjectType = FObjCWrapperRetained< id >;
-
 public:
 	FAGXStateObjectCache()
 	{
@@ -187,14 +185,14 @@ public:
 			Mutex.ReadLock();
 		}
 		
-		TUniquePtr<RetainedObjectType>* State = Cache.Find(Init);
+		const TRefCountPtr<FNSObject>* State = Cache.Find(Init);
 
 		if (IsRunningRHIInSeparateThread())
 		{
 			Mutex.ReadUnlock();
 		}
 		
-		return State ? State->Get()->Object : nil;
+		return State ? State->GetReference()->Get() : nil;
 	}
 	
 	void Add(InitializerType Init, id State)
@@ -204,7 +202,7 @@ public:
 			Mutex.WriteLock();
 		}
 		
-		Cache.Add(Init, MakeUnique<RetainedObjectType>(State));
+		Cache.Add(Init, new FNSObject(State, /* bRetain = */ true));
 		
 		if (IsRunningRHIInSeparateThread())
 		{
@@ -213,7 +211,7 @@ public:
 	}
 	
 private:
-	TMap<InitializerType, TUniquePtr<RetainedObjectType>> Cache;
+	TMap<InitializerType, TRefCountPtr<FNSObject>> Cache;
 	FRWLock Mutex;
 };
 
diff --git a/Engine/Source/Runtime/Apple/AGXRHI/Private/AGXTexture.cpp b/Engine/Source/Runtime/Apple/AGXRHI/Private/AGXTexture.cpp
--- a/Engine/Source/Runtime/Apple/AGXRHI/Private/AGXTexture.cpp
+++ b/Engine/Source/Runtime/Apple/AGXRHI/Private/AGXTexture.cpp
@@ -388,10 +388,11 @@ FAGXTextureDesc::FAGXTextureDesc(const FRHITextureDesc& InDesc)
 
 	if (InDesc.IsTextureCube())
 	{
-		Desc = MakeUnique<ReferencedTextureDescriptorType>(
+		Desc = new FMTLTextureDescriptor(
 				[MTLTextureDescriptor textureCubeDescriptorWithPixelFormat:PixelFormat
 																	  size:InDesc.Extent.X
-																 mipmapped:(InDesc.NumMips > 1)]);
+																 mipmapped:(InDesc.NumMips > 1)],
+				/* bRetain = */ true);
 	}
 	else if (InDesc.IsTexture3D())
 	{
@@ -405,22 +406,21 @@ FAGXTextureDesc::FAGXTextureDesc(const FRHITextureDesc& InDesc)
 		TextureDescriptor.mipmapLevelCount = 1;
 		TextureDescriptor.sampleCount      = 1;
 
-		Desc = MakeUnique<ReferencedTextureDescriptorType>(TextureDescriptor);
-
-		[TextureDescriptor release];
+		Desc = new FMTLTextureDescriptor(TextureDescriptor, /* bRetain = */ false);
 	}
 	else
 	{
-		Desc = MakeUnique<ReferencedTextureDescriptorType>(
+		Desc = new FMTLTextureDescriptor(
 				[MTLTextureDescriptor texture2DDescriptorWithPixelFormat:PixelFormat
 																   width:InDesc.Extent.X
 																  height:InDesc.Extent.Y
-															   mipmapped:(InDesc.NumMips > 1)]);
-		[Desc.Get()->Object setArrayLength:InDesc.ArraySize];
+															   mipmapped:(InDesc.NumMips > 1)],
+				/* bRetain = */ true);
+		[Desc->Get() setArrayLength:InDesc.ArraySize];
 	}
 	check(Desc.IsValid());
 
-	MTLTextureDescriptor* TextureDescriptor = Desc.Get()->Object;
+	MTLTextureDescriptor* TextureDescriptor = Desc->Get();
 
 	// flesh out the descriptor
 	if (InDesc.IsTextureArray())
@@ -570,7 +570,7 @@ FAGXSurface::FAGXSurface(FAGXTextureCreateDesc const& CreateDesc)
 		return;
 	}
 
-	MTLTextureDescriptor* TextureDescriptor = CreateDesc.Desc.Get()->Object;
+	MTLTextureDescriptor* TextureDescriptor = CreateDesc.Desc->Get();
 
 	FResourceBulkDataInterface* BulkData = CreateDesc.BulkData;
 
@@ -587,7 +587,7 @@ FAGXSurface::FAGXSurface(FAGXTextureCreateDesc const& CreateDesc)
 				CFRetain(ImageSurfaceRef);
 				
 #if !COREVIDEO_SUPPORTS_METAL
-				Texture = MTLPP_VALIDATE(mtlpp::Device, GMtlppDevice, AGXSafeGetRuntimeDebuggingLevel() >= EAGXDebugLevelValidation, NewTexture(CreateDesc.Desc, CVPixelBufferGetIOSurface((CVPixelBufferRef)ImageSurfaceRef), 0));
+				Texture = FAGXTexture([GMtlDevice newTextureWithDescriptor:CreateDesc.Desc->Get() iosurface:CVPixelBufferGetIOSurface((CVPixelBufferRef)ImageSurfaceRef) plane:0], ns::Ownership::Assign);
 #else
 				Texture = CVMetalTextureGetTexture((CVMetalTextureRef)ImageSurfaceRef);
 #endif
@@ -607,9 +607,9 @@ FAGXSurface::FAGXSurface(FAGXTextureCreateDesc const& CreateDesc)
 				MTLTextureDescriptor* DescCopy = [TextureDescriptor copy];
 				DescCopy.resourceOptions = ((DescCopy.resourceOptions & ~MTLResourceStorageModeMask) | MTLResourceStorageModeManaged);
 
-				Texture = mtlpp::Texture([GMtlDevice newTextureWithDescriptor:DescCopy
-																   iosurface:(IOSurfaceRef)ImageSurfaceRef
-																	   plane:0], nullptr, ns::Ownership::Assign);
+				Texture = FAGXTexture([GMtlDevice newTextureWithDescriptor:DescCopy
+																 iosurface:(IOSurfaceRef)ImageSurfaceRef
+																	 plane:0], ns::Ownership::Assign);
 				
 				METAL_FATAL_ASSERT(Texture, TEXT("Failed to create texture, desc %s"), *FString([DescCopy description]));
 
diff --git a/Engine/Source/Runtime/Apple/AGXRHI/Private/Shaders/AGXCompiledShaderCache.cpp b/Engine/Source/Runtime/Apple/AGXRHI/Private/Shaders/AGXCompiledShaderCache.cpp
--- a/Engine/Source/Runtime/Apple/AGXRHI/Private/Shaders/AGXCompiledShaderCache.cpp
+++ b/Engine/Source/Runtime/Apple/AGXRHI/Private/Shaders/AGXCompiledShaderCache.cpp
@@ -5,11 +5,7 @@
 =============================================================================*/
 
 #include "CoreMinimal.h"
-
-THIRD_PARTY_INCLUDES_START
-#include "mtlpp.hpp"
-THIRD_PARTY_INCLUDES_END
-
+#include "AGXRHIPrivate.h"
 #include "AGXCompiledShaderKey.h"
 #include "AGXCompiledShaderCache.h"
 
diff --git a/Engine/Source/Runtime/Apple/AGXRHI/Private/Shaders/AGXCompiledShaderCache.h b/Engine/Source/Runtime/Apple/AGXRHI/Private/Shaders/AGXCompiledShaderCache.h
--- a/Engine/Source/Runtime/Apple/AGXRHI/Private/Shaders/AGXCompiledShaderCache.h
+++ b/Engine/Source/Runtime/Apple/AGXRHI/Private/Shaders/AGXCompiledShaderCache.h
@@ -10,7 +10,6 @@
 
 struct FAGXCompiledShaderCache
 {
-public:
 	FAGXCompiledShaderCache()
 	{
 		// VOID
@@ -18,37 +17,38 @@ public:
 
 	~FAGXCompiledShaderCache()
 	{
-		// VOID
+		for (auto& Entry : Cache)
+		{
+			[Entry.Value release];
+		}
 	}
 
-	mtlpp::Function FindRef(FAGXCompiledShaderKey const& Key)
+	id<MTLFunction> FindRef(const FAGXCompiledShaderKey& Key)
 	{
 		FRWScopeLock ScopedLock(Lock, SLT_ReadOnly);
-		mtlpp::Function Func = Cache.FindRef(Key);
-		return Func;
+		return Cache.FindRef(Key);
 	}
 
-	mtlpp::Library FindLibrary(mtlpp::Function const& Function)
+	TRefCountPtr<FMTLLibrary> FindLibrary(id<MTLFunction> Function)
 	{
 		FRWScopeLock ScopedLock(Lock, SLT_ReadOnly);
-		mtlpp::Library Lib = LibCache.FindRef(Function.GetPtr());
-		return Lib;
+		return LibCache.FindRef(Function);
 	}
 
-	void Add(FAGXCompiledShaderKey Key, mtlpp::Library const& Lib, mtlpp::Function const& Function)
+	void Add(FAGXCompiledShaderKey Key, const TRefCountPtr<FMTLLibrary>& Library, id<MTLFunction> Function)
 	{
 		FRWScopeLock ScopedLock(Lock, SLT_Write);
 		if (Cache.FindRef(Key) == nil)
 		{
 			Cache.Add(Key, Function);
-			LibCache.Add(Function.GetPtr(), Lib);
+			LibCache.Add(Function, Library);
 		}
 	}
 
 private:
 	FRWLock Lock;
-	TMap<FAGXCompiledShaderKey, mtlpp::Function> Cache;
-	TMap<id<MTLFunction>, mtlpp::Library> LibCache;
+	TMap<FAGXCompiledShaderKey, id<MTLFunction>> Cache;
+	TMap<id<MTLFunction>, TRefCountPtr<FMTLLibrary>> LibCache;
 };
 
 extern FAGXCompiledShaderCache& GetAGXCompiledShaderCache();
diff --git a/Engine/Source/Runtime/Apple/AGXRHI/Private/Shaders/AGXShaderLibrary.cpp b/Engine/Source/Runtime/Apple/AGXRHI/Private/Shaders/AGXShaderLibrary.cpp
--- a/Engine/Source/Runtime/Apple/AGXRHI/Private/Shaders/AGXShaderLibrary.cpp
+++ b/Engine/Source/Runtime/Apple/AGXRHI/Private/Shaders/AGXShaderLibrary.cpp
@@ -20,7 +20,7 @@
 
 
 template<typename ShaderType>
-static TRefCountPtr<FRHIShader> AGXCreateMetalShader(TArrayView<const uint8> InCode, mtlpp::Library InLibrary)
+static TRefCountPtr<FRHIShader> AGXCreateMetalShader(TArrayView<const uint8> InCode, const TRefCountPtr<FMTLLibrary>& InLibrary)
 {
 	ShaderType* Shader = new ShaderType(InCode, InLibrary);
 	if (!Shader->GetFunction())
@@ -53,7 +53,7 @@ FAGXShaderLibrary::FAGXShaderLibrary(EShaderPlatform Platform,
 										 const FMetalShaderLibraryHeader& InHeader,
 										 const FSerializedShaderArchive& InSerializedShaders,
 										 const TArray<uint8>& InShaderCode,
-										 const TArray<mtlpp::Library>& InLibrary)
+										 const TArray<TRefCountPtr<FMTLLibrary>>& InLibrary)
 	: FRHIShaderLibrary(Platform, Name)
 	, ShaderLibraryFilename(InShaderLibraryFilename)
 	, Library(InLibrary)
diff --git a/Engine/Source/Runtime/Apple/AGXRHI/Private/Shaders/AGXShaderLibrary.h b/Engine/Source/Runtime/Apple/AGXRHI/Private/Shaders/AGXShaderLibrary.h
--- a/Engine/Source/Runtime/Apple/AGXRHI/Private/Shaders/AGXShaderLibrary.h
+++ b/Engine/Source/Runtime/Apple/AGXRHI/Private/Shaders/AGXShaderLibrary.h
@@ -18,7 +18,7 @@ public:
 						const FMetalShaderLibraryHeader& InHeader,
 						const FSerializedShaderArchive& InSerializedShaders,
 						const TArray<uint8>& InShaderCode,
-						const TArray<mtlpp::Library>& InLibrary);
+						const TArray<TRefCountPtr<FMTLLibrary>>& InLibrary);
 
 	virtual ~FAGXShaderLibrary();
 
@@ -39,7 +39,7 @@ public:
 
 private:
 	FString ShaderLibraryFilename;
-	TArray<mtlpp::Library> Library;
+	TArray<TRefCountPtr<FMTLLibrary>> Library;
 	FMetalShaderLibraryHeader Header;
 	FSerializedShaderArchive SerializedShaders;
 	TArray<uint8> ShaderCode;
diff --git a/Engine/Source/Runtime/Apple/AGXRHI/Private/Shaders/Debugging/AGXShaderDebugCache.cpp b/Engine/Source/Runtime/Apple/AGXRHI/Private/Shaders/Debugging/AGXShaderDebugCache.cpp
--- a/Engine/Source/Runtime/Apple/AGXRHI/Private/Shaders/Debugging/AGXShaderDebugCache.cpp
+++ b/Engine/Source/Runtime/Apple/AGXRHI/Private/Shaders/Debugging/AGXShaderDebugCache.cpp
@@ -6,10 +6,6 @@
 
 #include "CoreMinimal.h"
 
-THIRD_PARTY_INCLUDES_START
-#include "mtlpp.hpp"
-THIRD_PARTY_INCLUDES_END
-
 #include "HAL/FileManager.h"
 #include "HAL/PlatformFileManager.h"
 #include "Misc/Paths.h"
@@ -32,19 +28,18 @@ FAGXShaderDebugZipFile* FAGXShaderDebugCache::GetDebugFile(FString Path)
 	return Ref;
 }
 
-ns::String FAGXShaderDebugCache::GetShaderCode(uint32 ShaderSrcLen, uint32 ShaderSrcCRC)
+NSString* FAGXShaderDebugCache::GetShaderCode(uint32 ShaderSrcLen, uint32 ShaderSrcCRC)
 {
-	ns::String Code;
 	FScopeLock Lock(&Mutex);
 	for (auto const& Ref : DebugFiles)
 	{
-		Code = Ref.Value->GetShaderCode(ShaderSrcLen, ShaderSrcCRC);
-		if (Code)
+		NSString* Code = Ref.Value->GetShaderCode(ShaderSrcLen, ShaderSrcCRC);
+		if (Code != nil)
 		{
-			break;
+			return Code;
 		}
 	}
-	return Code;
+	return nil;
 }
 
 #endif // !UE_BUILD_SHIPPING
diff --git a/Engine/Source/Runtime/Apple/AGXRHI/Private/Shaders/Debugging/AGXShaderDebugCache.h b/Engine/Source/Runtime/Apple/AGXRHI/Private/Shaders/Debugging/AGXShaderDebugCache.h
--- a/Engine/Source/Runtime/Apple/AGXRHI/Private/Shaders/Debugging/AGXShaderDebugCache.h
+++ b/Engine/Source/Runtime/Apple/AGXRHI/Private/Shaders/Debugging/AGXShaderDebugCache.h
@@ -17,7 +17,7 @@ struct FAGXShaderDebugCache
 	}
 	
 	class FAGXShaderDebugZipFile* GetDebugFile(FString Path);
-	ns::String GetShaderCode(uint32 ShaderSrcLen, uint32 ShaderSrcCRC);
+	NSString* GetShaderCode(uint32 ShaderSrcLen, uint32 ShaderSrcCRC);
 	
 	FCriticalSection Mutex;
 	TMap<FString, class FAGXShaderDebugZipFile*> DebugFiles;
diff --git a/Engine/Source/Runtime/Apple/AGXRHI/Private/Shaders/Debugging/AGXShaderDebugZipFile.cpp b/Engine/Source/Runtime/Apple/AGXRHI/Private/Shaders/Debugging/AGXShaderDebugZipFile.cpp
--- a/Engine/Source/Runtime/Apple/AGXRHI/Private/Shaders/Debugging/AGXShaderDebugZipFile.cpp
+++ b/Engine/Source/Runtime/Apple/AGXRHI/Private/Shaders/Debugging/AGXShaderDebugZipFile.cpp
@@ -6,11 +6,6 @@
 
 #include "CoreMinimal.h"
 
-THIRD_PARTY_INCLUDES_START
-#include "mtlpp.hpp"
-THIRD_PARTY_INCLUDES_END
-
-
 #include "HAL/FileManager.h"
 #include "HAL/PlatformFileManager.h"
 #include "Misc/Paths.h"
@@ -222,9 +217,9 @@ FAGXShaderDebugZipFile::~FAGXShaderDebugZipFile()
 	}
 }
 
-ns::String FAGXShaderDebugZipFile::GetShaderCode(uint32 ShaderSrcLen, uint32 ShaderSrcCRC)
+NSString* FAGXShaderDebugZipFile::GetShaderCode(uint32 ShaderSrcLen, uint32 ShaderSrcCRC)
 {
-	ns::String Source;
+	NSString* Source = nil;
 	FScopeLock Lock(&Mutex);
 	FString Name = FString::Printf(TEXT("%u_%u.metal"), ShaderSrcLen, ShaderSrcCRC);
 	for (auto const& Entry : Files)
diff --git a/Engine/Source/Runtime/Apple/AGXRHI/Private/Shaders/Debugging/AGXShaderDebugZipFile.h b/Engine/Source/Runtime/Apple/AGXRHI/Private/Shaders/Debugging/AGXShaderDebugZipFile.h
--- a/Engine/Source/Runtime/Apple/AGXRHI/Private/Shaders/Debugging/AGXShaderDebugZipFile.h
+++ b/Engine/Source/Runtime/Apple/AGXRHI/Private/Shaders/Debugging/AGXShaderDebugZipFile.h
@@ -33,7 +33,7 @@ public:
 	FAGXShaderDebugZipFile(FString LibPath);
 	~FAGXShaderDebugZipFile();
 
-	ns::String GetShaderCode(uint32 ShaderSrcLen, uint32 ShaderSrcCRC);
+	NSString* GetShaderCode(uint32 ShaderSrcLen, uint32 ShaderSrcCRC);
 
 private:
 	FCriticalSection Mutex;
diff --git a/Engine/Source/Runtime/Apple/AGXRHI/Private/Shaders/Types/AGXComputeShader.cpp b/Engine/Source/Runtime/Apple/AGXRHI/Private/Shaders/Types/AGXComputeShader.cpp
--- a/Engine/Source/Runtime/Apple/AGXRHI/Private/Shaders/Types/AGXComputeShader.cpp
+++ b/Engine/Source/Runtime/Apple/AGXRHI/Private/Shaders/Types/AGXComputeShader.cpp
@@ -15,7 +15,7 @@
 #pragma mark - AGX RHI Compute Shader Class
 
 
-FAGXComputeShader::FAGXComputeShader(TArrayView<const uint8> InCode, mtlpp::Library InLibrary)
+FAGXComputeShader::FAGXComputeShader(TArrayView<const uint8> InCode, const TRefCountPtr<FMTLLibrary>& InLibrary)
 	: NumThreadsX(0)
 	, NumThreadsY(0)
 	, NumThreadsZ(0)
@@ -39,7 +39,7 @@ FAGXShaderPipeline* FAGXComputeShader::GetPipeline()
 {
 	if (!Pipeline)
 	{
-		id<MTLFunction> Func = GetCompiledFunction().GetPtr();
+		id<MTLFunction> Func = GetCompiledFunction();
 		check(Func != nil);
 
 		NSError* Error = nil;
@@ -113,7 +113,7 @@ FAGXShaderPipeline* FAGXComputeShader::GetPipeline()
 	return Pipeline;
 }
 
-mtlpp::Function FAGXComputeShader::GetFunction()
+id<MTLFunction> FAGXComputeShader::GetFunction()
 {
 	return GetCompiledFunction();
 }
diff --git a/Engine/Source/Runtime/Apple/AGXRHI/Private/Shaders/Types/AGXComputeShader.h b/Engine/Source/Runtime/Apple/AGXRHI/Private/Shaders/Types/AGXComputeShader.h
--- a/Engine/Source/Runtime/Apple/AGXRHI/Private/Shaders/Types/AGXComputeShader.h
+++ b/Engine/Source/Runtime/Apple/AGXRHI/Private/Shaders/Types/AGXComputeShader.h
@@ -15,11 +15,11 @@
 class FAGXComputeShader : public TAGXBaseShader<FRHIComputeShader, SF_Compute>
 {
 public:
-	FAGXComputeShader(TArrayView<const uint8> InCode, mtlpp::Library InLibrary = nil);
+	FAGXComputeShader(TArrayView<const uint8> InCode, const TRefCountPtr<FMTLLibrary>& InLibrary);
 	virtual ~FAGXComputeShader();
 
 	FAGXShaderPipeline* GetPipeline();
-	mtlpp::Function GetFunction();
+	id<MTLFunction> GetFunction();
 
 	// thread group counts
 	int32 NumThreadsX;
diff --git a/Engine/Source/Runtime/Apple/AGXRHI/Private/Shaders/Types/AGXPixelShader.cpp b/Engine/Source/Runtime/Apple/AGXRHI/Private/Shaders/Types/AGXPixelShader.cpp
--- a/Engine/Source/Runtime/Apple/AGXRHI/Private/Shaders/Types/AGXPixelShader.cpp
+++ b/Engine/Source/Runtime/Apple/AGXRHI/Private/Shaders/Types/AGXPixelShader.cpp
@@ -18,16 +18,17 @@
 FAGXPixelShader::FAGXPixelShader(TArrayView<const uint8> InCode)
 {
 	FMetalCodeHeader Header;
-	Init(InCode, Header);
+	TRefCountPtr<FMTLLibrary> Library;
+	Init(InCode, Header, Library);
 }
 
-FAGXPixelShader::FAGXPixelShader(TArrayView<const uint8> InCode, mtlpp::Library InLibrary)
+FAGXPixelShader::FAGXPixelShader(TArrayView<const uint8> InCode, const TRefCountPtr<FMTLLibrary>& InLibrary)
 {
 	FMetalCodeHeader Header;
 	Init(InCode, Header, InLibrary);
 }
 
-mtlpp::Function FAGXPixelShader::GetFunction()
+id<MTLFunction> FAGXPixelShader::GetFunction()
 {
 	return GetCompiledFunction();
 }
diff --git a/Engine/Source/Runtime/Apple/AGXRHI/Private/Shaders/Types/AGXPixelShader.h b/Engine/Source/Runtime/Apple/AGXRHI/Private/Shaders/Types/AGXPixelShader.h
--- a/Engine/Source/Runtime/Apple/AGXRHI/Private/Shaders/Types/AGXPixelShader.h
+++ b/Engine/Source/Runtime/Apple/AGXRHI/Private/Shaders/Types/AGXPixelShader.h
@@ -16,7 +16,7 @@ class FAGXPixelShader : public TAGXBaseShader<FRHIPixelShader, SF_Pixel>
 {
 public:
 	FAGXPixelShader(TArrayView<const uint8> InCode);
-	FAGXPixelShader(TArrayView<const uint8> InCode, mtlpp::Library InLibrary);
+	FAGXPixelShader(TArrayView<const uint8> InCode, const TRefCountPtr<FMTLLibrary>& InLibrary);
 
-	mtlpp::Function GetFunction();
+	id<MTLFunction> GetFunction();
 };
diff --git a/Engine/Source/Runtime/Apple/AGXRHI/Private/Shaders/Types/AGXVertexShader.cpp b/Engine/Source/Runtime/Apple/AGXRHI/Private/Shaders/Types/AGXVertexShader.cpp
--- a/Engine/Source/Runtime/Apple/AGXRHI/Private/Shaders/Types/AGXVertexShader.cpp
+++ b/Engine/Source/Runtime/Apple/AGXRHI/Private/Shaders/Types/AGXVertexShader.cpp
@@ -18,16 +18,17 @@
 FAGXVertexShader::FAGXVertexShader(TArrayView<const uint8> InCode)
 {
 	FMetalCodeHeader Header;
-	Init(InCode, Header);
+	TRefCountPtr<FMTLLibrary> Library;
+	Init(InCode, Header, Library);
 }
 
-FAGXVertexShader::FAGXVertexShader(TArrayView<const uint8> InCode, mtlpp::Library InLibrary)
+FAGXVertexShader::FAGXVertexShader(TArrayView<const uint8> InCode, const TRefCountPtr<FMTLLibrary>& InLibrary)
 {
 	FMetalCodeHeader Header;
 	Init(InCode, Header, InLibrary);
 }
 
-mtlpp::Function FAGXVertexShader::GetFunction()
+id<MTLFunction> FAGXVertexShader::GetFunction()
 {
 	return GetCompiledFunction();
 }
diff --git a/Engine/Source/Runtime/Apple/AGXRHI/Private/Shaders/Types/AGXVertexShader.h b/Engine/Source/Runtime/Apple/AGXRHI/Private/Shaders/Types/AGXVertexShader.h
--- a/Engine/Source/Runtime/Apple/AGXRHI/Private/Shaders/Types/AGXVertexShader.h
+++ b/Engine/Source/Runtime/Apple/AGXRHI/Private/Shaders/Types/AGXVertexShader.h
@@ -16,7 +16,7 @@ class FAGXVertexShader : public TAGXBaseShader<FRHIVertexShader, SF_Vertex>
 {
 public:
 	FAGXVertexShader(TArrayView<const uint8> InCode);
-	FAGXVertexShader(TArrayView<const uint8> InCode, mtlpp::Library InLibrary);
+	FAGXVertexShader(TArrayView<const uint8> InCode, const TRefCountPtr<FMTLLibrary>& InLibrary);
 
-	mtlpp::Function GetFunction();
+	id<MTLFunction> GetFunction();
 };
diff --git a/Engine/Source/Runtime/Apple/AGXRHI/Private/Shaders/Types/Templates/AGXBaseShader.h b/Engine/Source/Runtime/Apple/AGXRHI/Private/Shaders/Types/Templates/AGXBaseShader.h
--- a/Engine/Source/Runtime/Apple/AGXRHI/Private/Shaders/Types/Templates/AGXBaseShader.h
+++ b/Engine/Source/Runtime/Apple/AGXRHI/Private/Shaders/Types/Templates/AGXBaseShader.h
@@ -55,7 +55,7 @@ public:
 		this->Destroy();
 	}
 
-	void Init(TArrayView<const uint8> InCode, FMetalCodeHeader& Header, mtlpp::Library InLibrary = nil);
+	void Init(TArrayView<const uint8> InCode, FMetalCodeHeader& Header, const TRefCountPtr<FMTLLibrary>& InLibrary);
 	void Destroy();
 
 	/**
@@ -90,14 +90,14 @@ public:
 	uint32 ConstantValueHash = 0;
 
 protected:
-	mtlpp::Function GetCompiledFunction(bool const bAsync = false);
+	id<MTLFunction> GetCompiledFunction(bool bAsync = false);
 
 	// this is the compiler shader
-	mtlpp::Function Function = nil;
+	id<MTLFunction> Function = nil;
 
 private:
 	// This is the MTLLibrary for the shader so we can dynamically refine the MTLFunction
-	mtlpp::Library Library = nil;
+	TRefCountPtr<FMTLLibrary> Library;
 
 	/** The debuggable text source */
 	NSString* GlslCodeNSString = nil;
@@ -120,7 +120,7 @@ private:
 
 
 template<typename BaseResourceType, int32 ShaderType>
-void TAGXBaseShader<BaseResourceType, ShaderType>::Init(TArrayView<const uint8> InShaderCode, FMetalCodeHeader& Header, mtlpp::Library InLibrary)
+void TAGXBaseShader<BaseResourceType, ShaderType>::Init(TArrayView<const uint8> InShaderCode, FMetalCodeHeader& Header, const TRefCountPtr<FMTLLibrary>& InLibrary)
 {
 	FShaderCodeReader ShaderCode(InShaderCode);
 
@@ -152,7 +152,7 @@ void TAGXBaseShader<BaseResourceType, ShaderType>::Init(TArrayView<const uint8>
 	const ANSICHAR* SourceCode = (ANSICHAR*)InShaderCode.GetData() + CodeOffset;
 
 	// Only archived shaders should be in here.
-	UE_CLOG(InLibrary && !(Header.CompileFlags & (1 << CFLAG_Archive)), LogAGX, Warning, TEXT("Shader being loaded wasn't marked for archiving but a MTLLibrary was provided - this is unsupported."));
+	UE_CLOG(InLibrary.IsValid() && InLibrary->Get() != nil && !(Header.CompileFlags & (1 << CFLAG_Archive)), LogAGX, Warning, TEXT("Shader being loaded wasn't marked for archiving but a MTLLibrary was provided - this is unsupported."));
 
 	if (!OfflineCompiledFlag)
 	{
@@ -180,7 +180,7 @@ void TAGXBaseShader<BaseResourceType, ShaderType>::Init(TArrayView<const uint8>
 #if !UE_BUILD_SHIPPING
 		else if(bForceTextShaders)
 		{
-			GlslCodeNSString = [FAGXShaderDebugCache::Get().GetShaderCode(SourceLen, SourceCRC).GetPtr() retain];
+			GlslCodeNSString = [FAGXShaderDebugCache::Get().GetShaderCode(SourceLen, SourceCRC) retain];
 		}
 #endif
 		if (bForceTextShaders && CodeSize && CompressedSource.Num())
@@ -208,7 +208,7 @@ void TAGXBaseShader<BaseResourceType, ShaderType>::Init(TArrayView<const uint8>
 	FAGXCompiledShaderKey Key(Header.SourceLen, Header.SourceCRC, FunctionConstantHash);
 
 	Function = GetAGXCompiledShaderCache().FindRef(Key);
-	if (!Library && Function)
+	if (!Library.IsValid() && Function)
 	{
 		Library = GetAGXCompiledShaderCache().FindLibrary(Function);
 	}
@@ -218,11 +218,11 @@ void TAGXBaseShader<BaseResourceType, ShaderType>::Init(TArrayView<const uint8>
 	}
 
 	Bindings = Header.Bindings;
-	if (bNeedsCompiling || !Library)
+	if (bNeedsCompiling || !Library.IsValid())
 	{
 		if (bOfflineCompile METAL_DEBUG_OPTION(&& !(bHasShaderSource && bForceTextShaders)))
 		{
-			if (InLibrary)
+			if (InLibrary.IsValid())
 			{
 				Library = InLibrary;
 			}
@@ -235,18 +235,21 @@ void TAGXBaseShader<BaseResourceType, ShaderType>::Init(TArrayView<const uint8>
 
 				// allow GCD to copy the data into its own buffer
 				//		dispatch_data_t GCDBuffer = dispatch_data_create(InShaderCode.GetTypedData() + CodeOffset, ShaderCode.GetActualShaderCodeSize() - CodeOffset, nil, DISPATCH_DATA_DESTRUCTOR_DEFAULT);
-				ns::AutoReleasedError AError;
+				NSError* Error = nil;
 				void* Buffer = FMemory::Malloc( BufferSize );
 				FMemory::Memcpy( Buffer, InShaderCode.GetData() + CodeOffset, BufferSize );
 				dispatch_data_t GCDBuffer = dispatch_data_create(Buffer, BufferSize, dispatch_get_main_queue(), ^(void) { FMemory::Free(Buffer); } );
 
 				// load up the already compiled shader
-				Library = GMtlppDevice.NewLibrary(GCDBuffer, &AError);
+				id<MTLLibrary> NewLibrary = [GMtlDevice newLibraryWithData:GCDBuffer error:&Error];
 				dispatch_release(GCDBuffer);
-
-				if (Library == nil)
+				if (NewLibrary == nil)
 				{
-					NSLog(@"Failed to create library: %@", ns::Error(AError).GetPtr());
+					UE_LOG(LogAGX, Error, TEXT("Failed to create library: %s"), *FString(Error.description));
+				}
+				else
+				{
+					Library = new FMTLLibrary(NewLibrary, /* bRetain = */ false);
 				}
 			}
 		}
@@ -309,7 +312,7 @@ void TAGXBaseShader<BaseResourceType, ShaderType>::Init(TArrayView<const uint8>
 					break;
 
 				default:
-					UE_LOG(LogRHI, Fatal, TEXT("Failed to create shader with unknown version %d: %s"), Header.Version, *FString(NewShaderString));
+					UE_LOG(LogAGX, Fatal, TEXT("Failed to create shader with unknown version %d: %s"), Header.Version, *FString(NewShaderString));
 					MetalVersion = MTLLanguageVersion2_2;
 					break;
 			}
@@ -320,17 +323,17 @@ void TAGXBaseShader<BaseResourceType, ShaderType>::Init(TArrayView<const uint8>
 			id<MTLLibrary> NewLibrary = [GMtlDevice newLibraryWithSource:NewShaderString options:CompileOptions error:&Error];
 			if (NewLibrary == nil)
 			{
-				UE_LOG(LogRHI, Error, TEXT("*********** Error\n%s"), *FString(NewShaderString));
-				UE_LOG(LogRHI, Fatal, TEXT("Failed to create shader: %s"), *FString([Error description]));
+				UE_LOG(LogAGX, Error, TEXT("*********** Error\n%s"), *FString(NewShaderString));
+				UE_LOG(LogAGX, Fatal, TEXT("Failed to create shader: %s"), *FString(Error.description));
 			}
 			else if (Error != nil)
 			{
 				// Warning...
-				UE_LOG(LogRHI, Warning, TEXT("*********** Warning\n%s"), *FString(NewShaderString));
-				UE_LOG(LogRHI, Warning, TEXT("Created shader with warnings: %s"), *FString([Error description]));
+				UE_LOG(LogAGX, Warning, TEXT("*********** Warning\n%s"), *FString(NewShaderString));
+				UE_LOG(LogAGX, Warning, TEXT("Created shader with warnings: %s"), *FString(Error.description));
 			}
 
-			Library = mtlpp::Library(NewLibrary, nullptr, ns::Ownership::Assign);
+			Library = new FMTLLibrary(NewLibrary, /* bRetain = */ false);
 
 			[CompileOptions release];
 
@@ -385,67 +388,63 @@ uint32 TAGXBaseShader<BaseResourceType, ShaderType>::GetRefCount() const
 }
 
 template<typename BaseResourceType, int32 ShaderType>
-mtlpp::Function TAGXBaseShader<BaseResourceType, ShaderType>::GetCompiledFunction(bool const bAsync)
+id<MTLFunction> TAGXBaseShader<BaseResourceType, ShaderType>::GetCompiledFunction(bool bAsync)
 {
-	mtlpp::Function Func = Function;
-
-	if (!Func)
+	if (Function == nil)
 	{
 		// Find the existing compiled shader in the cache.
-		uint32 FunctionConstantHash = ConstantValueHash;
-		FAGXCompiledShaderKey Key(SourceLen, SourceCRC, FunctionConstantHash);
-		Func = Function = GetAGXCompiledShaderCache().FindRef(Key);
+		FAGXCompiledShaderKey Key(SourceLen, SourceCRC, ConstantValueHash);
+		Function = GetAGXCompiledShaderCache().FindRef(Key);
 
-		if (!Func)
+		if (Function == nil)
 		{
 			// Get the function from the library - the function name is "Main" followed by the CRC32 of the source MTLSL as 0-padded hex.
 			// This ensures that even if we move to a unified library that the function names will be unique - duplicates will only have one entry in the library.
 			NSString* Name = [NSString stringWithFormat:@"Main_%0.8x_%0.8x", SourceLen, SourceCRC];
-			mtlpp::FunctionConstantValues ConstantValues(nil);
+			METAL_GPUPROFILE(FString NameString(Name));
+
+			MTLFunctionConstantValues* ConstantValues = nil;
 			if (bHasFunctionConstants)
 			{
-				ConstantValues = mtlpp::FunctionConstantValues();
+				ConstantValues = [[MTLFunctionConstantValues alloc] init];
 
 				if (bDeviceFunctionConstants)
 				{
 					// Index 33 is the device vendor id constant
-					ConstantValues.SetConstantValue(&GRHIVendorId, mtlpp::DataType::UInt, @"GAGXDeviceManufacturer");
+					[ConstantValues setConstantValue:&GRHIVendorId type:MTLDataTypeUInt withName:@"GAGXDeviceManufacturer"];
 				}
 			}
 
 			if (!bHasFunctionConstants || !bAsync)
 			{
-				METAL_GPUPROFILE(FAGXScopedCPUStats CPUStat(FString::Printf(TEXT("NewFunction: %s"), *FString(Name))));
+				METAL_GPUPROFILE(FAGXScopedCPUStats CPUStat(FString::Printf(TEXT("NewFunction: %s"), *NameString)));
 				if (!bHasFunctionConstants)
 				{
-					Function = Library.NewFunction(Name);
+					Function = [Library->Get() newFunctionWithName:Name];
 				}
 				else
 				{
-					ns::AutoReleasedError AError;
-					Function = Library.NewFunction(Name, ConstantValues, &AError);
-					ns::Error Error = AError;
-					UE_CLOG(Function == nil, LogAGX, Error, TEXT("Failed to create function: %s"), *FString(Error.GetPtr().description));
+					NSError* Error = nil;
+					Function = [Library->Get() newFunctionWithName:Name constantValues:ConstantValues error:&Error];
+					UE_CLOG(Function == nil, LogAGX, Error, TEXT("Failed to create function: %s"), *FString(Error.description));
 					UE_CLOG(Function == nil, LogAGX, Fatal, TEXT("*********** Error\n%s"), *FString(GetSourceCode()));
 				}
 
 				check(Function);
 				GetAGXCompiledShaderCache().Add(Key, Library, Function);
-
-				Func = Function;
 			}
 			else
 			{
-				METAL_GPUPROFILE(FAGXScopedCPUStats CPUStat(FString::Printf(TEXT("NewFunctionAsync: %s"), *FString(Name))));
+				METAL_GPUPROFILE(FAGXScopedCPUStats CPUStat(FString::Printf(TEXT("NewFunctionAsync: %s"), *NameString)));
 				METAL_GPUPROFILE(uint64 CPUStart = CPUStat.Stats ? CPUStat.Stats->CPUStartTime : 0);
 #if ENABLE_METAL_GPUPROFILE
-				ns::String nsName(Name);
-				Library.NewFunction(Name, ConstantValues, [Key, this, CPUStart, nsName](mtlpp::Function const& NewFunction, ns::Error const& Error){
+				TFunction<void (id<MTLFunction>, NSError*)> FunctionCompletionHandler = [Key, this, CPUStart, NameString](id<MTLFunction> NewFunction, NSError* OptionalError)
 #else
-				Library.NewFunction(Name, ConstantValues, [Key, this](mtlpp::Function const& NewFunction, ns::Error const& Error){
+				TFunction<void (id<MTLFunction>, NSError*)> FunctionCompletionHandler = [Key, this](id<MTLFunction> NewFunction, NSError* OptionalError)
 #endif
-					METAL_GPUPROFILE(FAGXScopedCPUStats CompletionStat(FString::Printf(TEXT("NewFunctionCompletion: %s"), *FString(nsName.GetPtr()))));
-					UE_CLOG(NewFunction == nil, LogAGX, Error, TEXT("Failed to create function: %s"), *FString(Error.GetPtr().description));
+				{
+					METAL_GPUPROFILE(FAGXScopedCPUStats CompletionStat(FString::Printf(TEXT("NewFunctionCompletion: %s"), *NameString)));
+					UE_CLOG(NewFunction == nil, LogAGX, Error, TEXT("Failed to create function: %s"), *FString(OptionalError.description));
 					UE_CLOG(NewFunction == nil, LogAGX, Fatal, TEXT("*********** Error\n%s"), *FString(GetSourceCode()));
 
 					GetAGXCompiledShaderCache().Add(Key, Library, NewFunction);
@@ -455,13 +454,18 @@ mtlpp::Function TAGXBaseShader<BaseResourceType, ShaderType>::GetCompiledFunctio
 						CompletionStat.Stats->CPUStartTime = CPUStart;
 					}
 #endif
-				});
+				};
+
+				[Library->Get() newFunctionWithName:Name constantValues:ConstantValues completionHandler:^(id<MTLFunction> NewFunction, NSError* OptionalError)
+				{
+					FunctionCompletionHandler(NewFunction, OptionalError);
+				}];
 
 				return nil;
 			}
 		}
 	}
 
-	check(Func);
-	return Func;
+	check(Function);
+	return Function;
 }
