Fix a bug of resolving MSAA texture on IOS.

diff --git a/Engine/Source/Runtime/Apple/MetalRHI/Private/MetalTexture.cpp b/Engine/Source/Runtime/Apple/MetalRHI/Private/MetalTexture.cpp
--- a/Engine/Source/Runtime/Apple/MetalRHI/Private/MetalTexture.cpp
+++ b/Engine/Source/Runtime/Apple/MetalRHI/Private/MetalTexture.cpp
@@ -401,8 +401,9 @@ uint8 GetMetalPixelFormatKey(mtlpp::PixelFormat Format)
 	return PixelFormatKeyMap.Get(Format);
 }
 
-FMetalTextureDesc::FMetalTextureDesc(FRHITextureDesc const& InDesc)
-	: bIsRenderTarget(IsRenderTarget(InDesc.Flags))
+FMetalTextureCreateDesc::FMetalTextureCreateDesc(FRHITextureCreateDesc const& InDesc)
+	: FRHITextureCreateDesc(InDesc)
+	, bIsRenderTarget(IsRenderTarget(InDesc.Flags))
 {
 	MTLFormat = (mtlpp::PixelFormat)GPixelFormats[InDesc.Format].PlatformFormat;
 	if (EnumHasAnyFlags(InDesc.Flags, TexCreate_SRGB))
@@ -547,36 +548,6 @@ FMetalTextureDesc::FMetalTextureDesc(FRHITextureDesc const& InDesc)
 		static mtlpp::ResourceOptions GeneralResourceOption = FMetalCommandQueue::GetCompatibleResourceOptions(mtlpp::ResourceOptions::HazardTrackingModeUntracked);
 		Desc.SetResourceOptions((mtlpp::ResourceOptions)(Desc.GetResourceOptions() | GeneralResourceOption));
 	}
-
-	
-	if (!FParse::Param(FCommandLine::Get(), TEXT("nomsaa")))
-	{
-		if (InDesc.NumSamples > 1)
-		{
-			check(bIsRenderTarget);
-			Desc.SetTextureType(mtlpp::TextureType::Texture2DMultisample);
-
-			// allow commandline to override
-			uint32 NewNumSamples;
-			if (FParse::Value(FCommandLine::Get(), TEXT("msaa="), NewNumSamples))
-			{
-				Desc.SetSampleCount(NewNumSamples);
-			}
-			else
-			{
-				Desc.SetSampleCount(InDesc.NumSamples);
-			}
-
-#if PLATFORM_IOS
-			if (GMaxRHIFeatureLevel < ERHIFeatureLevel::SM5)
-			{
-				bMemoryless = true;
-				Desc.SetStorageMode(mtlpp::StorageMode::Memoryless);
-				Desc.SetResourceOptions(mtlpp::ResourceOptions::StorageModeMemoryless);
-			}
-#endif
-		}
-	}
 }
 
 FMetalSurface::FMetalSurface(FMetalTextureCreateDesc const& CreateDesc)
@@ -704,7 +675,32 @@ FMetalSurface::FMetalSurface(FMetalTextureCreateDesc const& CreateDesc)
 
 	if (CreateDesc.NumSamples > 1 && !FParse::Param(FCommandLine::Get(), TEXT("nomsaa")))
 	{
-		MSAATexture = GetMetalDeviceContext().CreateTexture(this, CreateDesc.Desc);
+		mtlpp::TextureDescriptor Desc = CreateDesc.Desc;
+		check(CreateDesc.bIsRenderTarget);
+		Desc.SetTextureType(mtlpp::TextureType::Texture2DMultisample);
+
+		// allow commandline to override
+		uint32 NewNumSamples;
+		if (FParse::Value(FCommandLine::Get(), TEXT("msaa="), NewNumSamples))
+		{
+			Desc.SetSampleCount(NewNumSamples);
+		}
+		else
+		{
+			Desc.SetSampleCount(CreateDesc.NumSamples);
+		}
+
+		bool bMemoryless = false;
+#if PLATFORM_IOS
+		if (GMaxRHIFeatureLevel < ERHIFeatureLevel::SM5)
+		{
+			bMemoryless = true;
+			Desc.SetStorageMode(mtlpp::StorageMode::Memoryless);
+			Desc.SetResourceOptions(mtlpp::ResourceOptions::StorageModeMemoryless);
+		}
+#endif
+
+		MSAATexture = GetMetalDeviceContext().CreateTexture(this, Desc);
 			
 		//device doesn't support HW depth resolve.  This case only valid on mobile renderer or
 		//on Mac where RHISupportsSeparateMSAAAndResolveTextures is true.
@@ -717,7 +713,7 @@ FMetalSurface::FMetalSurface(FMetalTextureCreateDesc const& CreateDesc)
 			// we don't have the resolve texture, so we just update the memory size with the MSAA size
 			TotalTextureSize = TotalTextureSize * CreateDesc.NumSamples;
 		}
-		else if (!CreateDesc.bMemoryless)
+		else if (!bMemoryless)
 		{
 			// an MSAA render target takes NumSamples more space, in addition to the resolve texture
 			TotalTextureSize += TotalTextureSize * CreateDesc.NumSamples;
diff --git a/Engine/Source/Runtime/Apple/MetalRHI/Public/MetalResources.h b/Engine/Source/Runtime/Apple/MetalRHI/Public/MetalResources.h
--- a/Engine/Source/Runtime/Apple/MetalRHI/Public/MetalResources.h
+++ b/Engine/Source/Runtime/Apple/MetalRHI/Public/MetalResources.h
@@ -133,29 +133,16 @@ public:
 	}
 };
 
-struct FMetalTextureDesc
+struct FMetalTextureCreateDesc : public FRHITextureCreateDesc
 {
-	FMetalTextureDesc(FRHITextureDesc const& InDesc);
-	
+	FMetalTextureCreateDesc(FRHITextureCreateDesc const& CreateDesc);
+
 	mtlpp::TextureDescriptor Desc;
 	mtlpp::PixelFormat MTLFormat;
-	bool bMemoryless = false;
 	bool bIsRenderTarget = false;
 	uint8 FormatKey = 0;
 };
 
-struct FMetalTextureCreateDesc : public FRHITextureCreateDesc, public FMetalTextureDesc
-{
-	FMetalTextureCreateDesc(FRHITextureCreateDesc const& CreateDesc)
-		: FRHITextureCreateDesc(CreateDesc)
-		, FMetalTextureDesc(CreateDesc)
-	{
-		// @todo: texture type unification - Metal can override NumSamples based on command line options.
-		// We should instead require the renderer to do this.
-		NumSamples = Desc.GetSampleCount();
-	}
-};
-
 // Metal RHI texture resource
 class METALRHI_API FMetalSurface : public FRHITexture
 {
diff --git a/Engine/Source/Runtime/Renderer/Private/PostProcess/PostProcessCompositeEditorPrimitives.h b/Engine/Source/Runtime/Renderer/Private/PostProcess/PostProcessCompositeEditorPrimitives.h
--- a/Engine/Source/Runtime/Renderer/Private/PostProcess/PostProcessCompositeEditorPrimitives.h
+++ b/Engine/Source/Runtime/Renderer/Private/PostProcess/PostProcessCompositeEditorPrimitives.h
@@ -21,12 +21,6 @@ public:
 		const FPermutationDomain PermutationVector(Parameters.PermutationId);
 		const int32 SampleCount = PermutationVector.Get<FSampleCountDimension>();
 
-		// Only SM5+ platforms supports MSAA.
-		if (!IsFeatureLevelSupported(Parameters.Platform, ERHIFeatureLevel::SM5) && SampleCount > 1)
-		{
-			return false;
-		}
-
 		// Only use permutations with valid MSAA sample counts.
 		if (!FMath::IsPowerOfTwo(SampleCount))
 		{
diff --git a/Engine/Source/Runtime/Renderer/Private/SceneHitProxyRendering.cpp b/Engine/Source/Runtime/Renderer/Private/SceneHitProxyRendering.cpp
--- a/Engine/Source/Runtime/Renderer/Private/SceneHitProxyRendering.cpp
+++ b/Engine/Source/Runtime/Renderer/Private/SceneHitProxyRendering.cpp
@@ -211,7 +211,7 @@ void InitHitProxyRender(FRDGBuilder& GraphBuilder, FSceneRenderer* SceneRenderer
 		const EShaderPlatform CurrentShaderPlatform = GShaderPlatformForFeatureLevel[FeatureLevel];
 		FRDGTextureDesc DepthDesc = SceneTextures.Depth.Target->Desc;
 
-		if (DepthDesc.NumSamples > 1 && RHISupportsSeparateMSAAAndResolveTextures(CurrentShaderPlatform))
+		if (DepthDesc.NumSamples > 1)
 		{
 			DepthDesc.NumSamples = 1;
 			OutHitProxyDepthTexture = GraphBuilder.CreateTexture(DepthDesc, TEXT("NoMSAASceneDepthZ"));
