Update enum used by the Metal RHI

diff --git a/Engine/Source/Runtime/Apple/MetalRHI/Private/MetalCommandList.cpp b/Engine/Source/Runtime/Apple/MetalRHI/Private/MetalCommandList.cpp
--- a/Engine/Source/Runtime/Apple/MetalRHI/Private/MetalCommandList.cpp
+++ b/Engine/Source/Runtime/Apple/MetalRHI/Private/MetalCommandList.cpp
@@ -193,9 +193,9 @@ static __attribute__ ((optnone)) void MetalCommandBufferFailurePageFault(mtlpp::
 	ReportMetalCommandBufferFailure(CompletedBuffer, TEXT("PageFault"));
 }
 
-static __attribute__ ((optnone)) void MetalCommandBufferFailureBlacklisted(mtlpp::CommandBuffer const& CompletedBuffer)
+static __attribute__ ((optnone)) void MetalCommandBufferFailureAccessRevoked(mtlpp::CommandBuffer const& CompletedBuffer)
 {
-	ReportMetalCommandBufferFailure(CompletedBuffer, TEXT("Blacklisted"));
+	ReportMetalCommandBufferFailure(CompletedBuffer, TEXT("AccessRevoked"));
 }
 
 static __attribute__ ((optnone)) void MetalCommandBufferFailureNotPermitted(mtlpp::CommandBuffer const& CompletedBuffer)
@@ -228,8 +228,8 @@ static void HandleMetalCommandBufferError(mtlpp::CommandBuffer const& CompletedB
 		case MTLCommandBufferErrorPageFault:
 			MetalCommandBufferFailurePageFault(CompletedBuffer);
 			break;
-		case MTLCommandBufferErrorBlacklisted:
-			MetalCommandBufferFailureBlacklisted(CompletedBuffer);
+		case MTLCommandBufferErrorAccessRevoked:
+			MetalCommandBufferFailureAccessRevoked(CompletedBuffer);
 			break;
 		case MTLCommandBufferErrorNotPermitted:
 			MetalCommandBufferFailureNotPermitted(CompletedBuffer);
