Make UpdateUniformBuffer thread safe by taking a command list.

diff --git a/Engine/Source/Runtime/Apple/AGXRHI/Private/AGXUniformBufferCommands.cpp b/Engine/Source/Runtime/Apple/AGXRHI/Private/AGXUniformBufferCommands.cpp
new file mode 100644
--- /dev/null
+++ b/Engine/Source/Runtime/Apple/AGXRHI/Private/AGXUniformBufferCommands.cpp
@@ -0,0 +1,83 @@
+// Copyright Epic Games, Inc. All Rights Reserved.
+
+#include "AGXRHIPrivate.h"
+#include "AGXShaderTypes.h"
+#include "AGXFrameAllocator.h"
+
+FUniformBufferRHIRef FAGXDynamicRHI::RHICreateUniformBuffer(const void* Contents, const FRHIUniformBufferLayout* Layout, EUniformBufferUsage Usage, EUniformBufferValidation Validation)
+{
+	return new FAGXUniformBuffer(Contents, Layout, Usage, Validation);
+}
+
+void FAGXDynamicRHI::RHIUpdateUniformBuffer(FRHICommandListBase& RHICmdList, FRHIUniformBuffer* UniformBufferRHI, const void* Contents)
+{
+    FAGXUniformBuffer* UniformBuffer = ResourceCast(UniformBufferRHI);
+
+	const void* SrcContents = Contents;
+
+	if (RHICmdList.IsTopOfPipe())
+	{
+		const FRHIUniformBufferLayout& Layout = UniformBuffer->GetLayout();
+
+		// Copy the contents memory region into the RHICmdList to allow deferred execution on the RHI thread.
+		void* DstContents = RHICmdList.Alloc(Layout.ConstantBufferSize, alignof(FRHIResource*));
+		FMemory::ParallelMemcpy(DstContents, Contents, Layout.ConstantBufferSize, EMemcpyCachePolicy::StoreUncached);
+
+		SrcContents = DstContents;
+	}
+
+	RHICmdList.EnqueueLambda([UniformBuffer, SrcContents](FRHICommandListBase& RHICmdList)
+	{
+		UniformBuffer->Update(SrcContents);
+	});
+
+	RHICmdList.RHIThreadFence(true);
+}
+
+template<EAGXShaderStages Stage, typename RHIShaderType>
+static void SetUniformBufferInternal(FAGXContext* Context, RHIShaderType* ShaderRHI, uint32 BufferIndex, FRHIUniformBuffer* UBRHI)
+{
+    @autoreleasepool
+    {
+        auto Shader = ResourceCast(ShaderRHI);
+        Context->GetCurrentState().BindUniformBuffer(Stage, BufferIndex, UBRHI);
+        
+        auto& Bindings = Shader->Bindings;
+        if((Bindings.ConstantBuffers) & (1 << BufferIndex))
+        {
+            FAGXUniformBuffer* UB = ResourceCast(UBRHI);
+            UB->PrepareToBind();
+            
+            FAGXBuffer Buf(UB->Backing, ns::Ownership::AutoRelease);
+            Context->GetCurrentState().SetShaderBuffer(Stage, Buf, nil, UB->Offset, UB->GetSize(), BufferIndex, MTLResourceUsageRead);
+        }
+    }
+}
+
+void FAGXRHICommandContext::RHISetShaderUniformBuffer(FRHIGraphicsShader* Shader, uint32 BufferIndex, FRHIUniformBuffer* Buffer)
+{
+	switch (Shader->GetFrequency())
+	{
+		case SF_Vertex:
+			SetUniformBufferInternal<EAGXShaderStages::Vertex, FRHIVertexShader>(Context, static_cast<FRHIVertexShader*>(Shader), BufferIndex, Buffer);
+			break;
+
+		case SF_Geometry:
+			NOT_SUPPORTED("RHISetShaderUniformBuffer-Geometry");
+			break;
+
+		case SF_Pixel:
+			SetUniformBufferInternal<EAGXShaderStages::Pixel, FRHIPixelShader>(Context, static_cast<FRHIPixelShader*>(Shader), BufferIndex, Buffer);
+			break;
+
+		default:
+			checkf(0, TEXT("FRHIShader Type %d is invalid or unsupported!"), (int32)Shader->GetFrequency());
+			NOT_SUPPORTED("RHIShaderStage");
+			break;
+	}
+}
+
+void FAGXRHICommandContext::RHISetShaderUniformBuffer(FRHIComputeShader* ComputeShaderRHI, uint32 BufferIndex, FRHIUniformBuffer* BufferRHI)
+{
+	SetUniformBufferInternal<EAGXShaderStages::Compute, FRHIComputeShader>(Context, ComputeShaderRHI, BufferIndex, BufferRHI);
+}
