Fix for texture CPU readback Lock.  Previous logic assumed it was going through the RHICopyToResolveTarget() path but we have to support the CopyTexturePath [and any other paths] and maybe even GPU shader write and render target.  If the texture storage mode is Managed then sync the requested Slice/Mip before continuing.  Optimise Apple Silicon Mac - use Shared storage mode for these types of textures to save the cost of Syncing the texture GPU->CPU since its unified memory.

#rnx

diff --git a/Engine/Source/Runtime/Apple/MetalRHI/Private/MetalRHIContext.cpp b/Engine/Source/Runtime/Apple/MetalRHI/Private/MetalRHIContext.cpp
--- a/Engine/Source/Runtime/Apple/MetalRHI/Private/MetalRHIContext.cpp
+++ b/Engine/Source/Runtime/Apple/MetalRHI/Private/MetalRHIContext.cpp
@@ -223,13 +223,6 @@ void FMetalRHICommandContext::ResolveTexture(UE::RHICore::FResolveTextureInfo In
 	for (int32 ArraySlice = ArraySliceBegin; ArraySlice < ArraySliceEnd; ArraySlice++)
 	{
 		Context->CopyFromTextureToTexture(Source->MSAAResolveTexture, ArraySlice, Info.MipLevel, Origin, Size, Destination->Texture, ArraySlice, Info.MipLevel, Origin);
-
-#if PLATFORM_MAC
-		if ((Destination->GPUReadback & FMetalSurface::EMetalGPUReadbackFlags::ReadbackRequested) != 0)
-		{
-			Context->GetCurrentRenderPass().SynchronizeTexture(Destination->Texture, ArraySlice, Info.MipLevel);
-		}
-#endif
 	}
 	}
 }
diff --git a/Engine/Source/Runtime/Apple/MetalRHI/Private/MetalRenderTarget.cpp b/Engine/Source/Runtime/Apple/MetalRHI/Private/MetalRenderTarget.cpp
--- a/Engine/Source/Runtime/Apple/MetalRHI/Private/MetalRenderTarget.cpp
+++ b/Engine/Source/Runtime/Apple/MetalRHI/Private/MetalRenderTarget.cpp
@@ -139,14 +139,6 @@ void FMetalRHICommandContext::RHICopyToResolveTarget(FRHITexture* SourceTextureR
 				}
 			}
 		}
-
-#if PLATFORM_MAC
-		if((Destination->GPUReadback & FMetalSurface::EMetalGPUReadbackFlags::ReadbackRequested) != 0)
-		{
-			Context->GetCurrentRenderPass().SynchronizeTexture(Destination->Texture, DestIndex, ResolveParams.MipIndex);
-		}
-#endif
-		
 	}
 	}
 }
@@ -377,21 +369,13 @@ void FMetalDynamicRHI::RHIReadSurfaceData(FRHITexture* TextureRHI, FIntRect Rect
 void FMetalDynamicRHI::RHIMapStagingSurface(FRHITexture* TextureRHI, FRHIGPUFence* FenceRHI, void*& OutData, int32& OutWidth, int32& OutHeight, uint32 GPUIndex)
 {
 	@autoreleasepool {
-    FMetalSurface* Surface = GetMetalSurfaceFromRHITexture(TextureRHI);
+		FMetalSurface* Surface = GetMetalSurfaceFromRHITexture(TextureRHI);
     
-#if PLATFORM_MAC
-	uint16 FencePoll = FenceRHI && FenceRHI->Poll() ? 1 : 0;
-    Surface->GPUReadback |= FencePoll << FMetalSurface::EMetalGPUReadbackFlags::ReadbackFenceCompleteShift;
-#endif
-    
-    uint32 Stride = 0;
-    OutWidth = Surface->GetSizeX();
-    OutHeight = Surface->GetSizeY();
-    OutData = Surface->Lock(0, 0, RLM_ReadOnly, Stride);
-    
-#if PLATFORM_MAC
-	Surface->GPUReadback = (Surface->Texture.GetPtr() && Surface->Texture.GetStorageMode() == mtlpp::StorageMode::Managed) ? FMetalSurface::EMetalGPUReadbackFlags::ReadbackRequested : 0;
-#endif
+		uint32 Stride = 0;
+		OutWidth = Surface->GetSizeX();
+		OutHeight = Surface->GetSizeY();
+		
+		OutData = Surface->Lock(0, 0, RLM_ReadOnly, Stride);
 	}
 }
 
diff --git a/Engine/Source/Runtime/Apple/MetalRHI/Private/MetalTexture.cpp b/Engine/Source/Runtime/Apple/MetalRHI/Private/MetalTexture.cpp
--- a/Engine/Source/Runtime/Apple/MetalRHI/Private/MetalTexture.cpp
+++ b/Engine/Source/Runtime/Apple/MetalRHI/Private/MetalTexture.cpp
@@ -460,27 +460,41 @@ FMetalTextureCreateDesc::FMetalTextureCreateDesc(FRHITextureCreateDesc const& In
 
 	{
 		Desc.SetUsage(ConvertFlagsToUsage(InDesc.Flags));
+		
+		const bool bAppleGPU = [GetMetalDeviceContext().GetDevice().GetPtr() supportsFamily:MTLGPUFamilyApple1];
 
 		if (EnumHasAnyFlags(InDesc.Flags, TexCreate_CPUReadback) && !EnumHasAnyFlags(InDesc.Flags, TexCreate_RenderTargetable | TexCreate_DepthStencilTargetable | TexCreate_FastVRAM))
 		{
 			Desc.SetCpuCacheMode(mtlpp::CpuCacheMode::DefaultCache);
+			
+			if(bAppleGPU)
+			{
+				Desc.SetStorageMode(mtlpp::StorageMode::Shared);
+				Desc.SetResourceOptions((mtlpp::ResourceOptions)(mtlpp::ResourceOptions::CpuCacheModeDefaultCache | mtlpp::ResourceOptions::StorageModeShared));
+			}
 #if PLATFORM_MAC
-			Desc.SetStorageMode(mtlpp::StorageMode::Managed);
-			Desc.SetResourceOptions((mtlpp::ResourceOptions)(mtlpp::ResourceOptions::CpuCacheModeDefaultCache | mtlpp::ResourceOptions::StorageModeManaged));
-#else
-			Desc.SetStorageMode(mtlpp::StorageMode::Shared);
-			Desc.SetResourceOptions((mtlpp::ResourceOptions)(mtlpp::ResourceOptions::CpuCacheModeDefaultCache | mtlpp::ResourceOptions::StorageModeShared));
+			else
+			{
+				Desc.SetStorageMode(mtlpp::StorageMode::Managed);
+				Desc.SetResourceOptions((mtlpp::ResourceOptions)(mtlpp::ResourceOptions::CpuCacheModeDefaultCache | mtlpp::ResourceOptions::StorageModeManaged));
+			}
 #endif
 		}
 		else if (EnumHasAnyFlags(InDesc.Flags, TexCreate_NoTiling) && !EnumHasAnyFlags(InDesc.Flags, TexCreate_FastVRAM | TexCreate_DepthStencilTargetable | TexCreate_RenderTargetable | TexCreate_UAV))
 		{
 			Desc.SetCpuCacheMode(mtlpp::CpuCacheMode::DefaultCache);
+			
+			if(bAppleGPU)
+			{
+				Desc.SetStorageMode(mtlpp::StorageMode::Shared);
+				Desc.SetResourceOptions((mtlpp::ResourceOptions)(mtlpp::ResourceOptions::CpuCacheModeDefaultCache | mtlpp::ResourceOptions::StorageModeShared));
+			}
 #if PLATFORM_MAC
-			Desc.SetStorageMode(mtlpp::StorageMode::Managed);
-			Desc.SetResourceOptions((mtlpp::ResourceOptions)(mtlpp::ResourceOptions::CpuCacheModeDefaultCache | mtlpp::ResourceOptions::StorageModeManaged));
-#else
-			Desc.SetStorageMode(mtlpp::StorageMode::Shared);
-			Desc.SetResourceOptions((mtlpp::ResourceOptions)(mtlpp::ResourceOptions::CpuCacheModeDefaultCache | mtlpp::ResourceOptions::StorageModeShared));
+			else
+			{
+				Desc.SetStorageMode(mtlpp::StorageMode::Managed);
+				Desc.SetResourceOptions((mtlpp::ResourceOptions)(mtlpp::ResourceOptions::CpuCacheModeDefaultCache | mtlpp::ResourceOptions::StorageModeManaged));
+			}
 #endif
 		}
 		else if (EnumHasAnyFlags(InDesc.Flags, TexCreate_RenderTargetable | TexCreate_DepthStencilTargetable | TexCreate_ResolveTargetable | TexCreate_DepthStencilResolveTarget))
@@ -1151,9 +1165,9 @@ void* FMetalSurface::Lock(uint32 MipIndex, uint32 ArrayIndex, EResourceLockMode
 			else
 			{
 #if PLATFORM_MAC
-				if((GPUReadback & EMetalGPUReadbackFlags::ReadbackRequestedAndComplete) != EMetalGPUReadbackFlags::ReadbackRequestedAndComplete)
+				if(this->Texture.GetStorageMode() == mtlpp::StorageMode::Managed)
 				{
-					// A previous texture sync has not been done, need the data now, request texture sync and kick the current command buffer.
+					// Managed texture - need to sync GPU -> CPU before access as it could have been written to by the GPU
 					auto SyncReadbackToCPU =
 					[this, &ArrayIndex, &MipIndex](FRHICommandListImmediate& RHICmdList)
 					{
diff --git a/Engine/Source/Runtime/Apple/MetalRHI/Public/MetalResources.h b/Engine/Source/Runtime/Apple/MetalRHI/Public/MetalResources.h
--- a/Engine/Source/Runtime/Apple/MetalRHI/Public/MetalResources.h
+++ b/Engine/Source/Runtime/Apple/MetalRHI/Public/MetalResources.h
@@ -212,8 +212,10 @@ public:
 
 	FMetalTexture Reallocate(FMetalTexture Texture, mtlpp::TextureUsage UsageModifier);
 	void MakeAliasable(void);
-
+	
+	int16 volatile Written;
 	uint8 const FormatKey;
+
 	//texture used for store actions and binding to shader params
 	FMetalTexture Texture;
 	//if surface is MSAA, texture used to bind for RT
@@ -227,17 +229,6 @@ public:
 	// iOS < A9 where depth resolve is unavailable.
 	FMetalTexture MSAAResolveTexture;
 
-	int16 volatile Written;
-	int16 GPUReadback = 0;
-	enum EMetalGPUReadbackFlags : int16
-	{
-		ReadbackRequestedShift 			= 0,
-		ReadbackFenceCompleteShift  	= 1,
-		ReadbackRequested 				= 1 << ReadbackRequestedShift,
-		ReadbackFenceComplete 			= 1 << ReadbackFenceCompleteShift,
-		ReadbackRequestedAndComplete 	= ReadbackRequested | ReadbackFenceComplete
-	};
-
 	// how much memory is allocated for this texture
 	uint64 TotalTextureSize;
 	
