Fix for potential crash in FMetalResourceViewBase buffer initialization constructor.  Source Buffer can be NULL (see FStaticMeshVertexBuffer::InitRHI()).

From UDN: https://epicgames.lightning.force.com/lightning/r/5004z00001g2DMkAAM/view

#rnx

diff --git a/Engine/Source/Runtime/Apple/MetalRHI/Private/MetalStateCache.cpp b/Engine/Source/Runtime/Apple/MetalRHI/Private/MetalStateCache.cpp
--- a/Engine/Source/Runtime/Apple/MetalRHI/Private/MetalStateCache.cpp
+++ b/Engine/Source/Runtime/Apple/MetalRHI/Private/MetalStateCache.cpp
@@ -1469,8 +1469,6 @@ void FMetalStateCache::SetShaderResourceView(FMetalContext* Context, EMetalShade
 		}
 		else
 		{
-			FMetalResourceMultiBuffer* Buffer = SRV->GetSourceBuffer();
-
 			if (IsLinearBuffer(ShaderStage, BindIndex) && SRV->GetLinearTexture())
 			{
 				ns::AutoReleased<FMetalTexture> Tex;
@@ -1478,9 +1476,17 @@ void FMetalStateCache::SetShaderResourceView(FMetalContext* Context, EMetalShade
 
 				SetShaderTexture(ShaderStage, Tex, BindIndex, mtlpp::ResourceUsage(mtlpp::ResourceUsage::Read | mtlpp::ResourceUsage::Sample));
 			}
-			else 
+			else
 			{
-				SetShaderBuffer(ShaderStage, Buffer->GetCurrentBufferOrNil(), Buffer->Data, SRV->Offset, Buffer->GetSize(), BindIndex, mtlpp::ResourceUsage::Read, (EPixelFormat)SRV->Format);
+				FMetalResourceMultiBuffer* Buffer = SRV->GetSourceBuffer();
+				if(Buffer != nullptr)
+				{
+					SetShaderBuffer(ShaderStage, Buffer->GetCurrentBufferOrNil(), Buffer->Data, SRV->Offset, Buffer->GetSize(), BindIndex, mtlpp::ResourceUsage::Read, (EPixelFormat)SRV->Format);
+				}
+				else
+				{
+					SetShaderBuffer(ShaderStage, nil, nullptr, 0, 0, BindIndex, mtlpp::ResourceUsage(0));
+				}
 			}
 		}
 	}
diff --git a/Engine/Source/Runtime/Apple/MetalRHI/Private/MetalUAV.cpp b/Engine/Source/Runtime/Apple/MetalRHI/Private/MetalUAV.cpp
--- a/Engine/Source/Runtime/Apple/MetalRHI/Private/MetalUAV.cpp
+++ b/Engine/Source/Runtime/Apple/MetalRHI/Private/MetalUAV.cpp
@@ -23,49 +23,51 @@ FMetalResourceViewBase::FMetalResourceViewBase(FRHIBuffer* InBuffer, uint32 InSt
 	check(!bTexture);
 
 	if (SourceBuffer)
+	{
 		SourceBuffer->AddRef();
 
-	EBufferUsageFlags Usage = SourceBuffer->GetUsage();
-	if (EnumHasAnyFlags(Usage, BUF_VertexBuffer))
-	{
-		if (!SourceBuffer)
+		EBufferUsageFlags Usage = SourceBuffer->GetUsage();
+		if (EnumHasAnyFlags(Usage, BUF_VertexBuffer))
 		{
-			Stride = 0;
-		}
-		else
-		{
-			check(SourceBuffer->GetUsage() & BUF_ShaderResource);
-			Stride = GPixelFormats[Format].BlockBytes;
+			if (!SourceBuffer)
+			{
+				Stride = 0;
+			}
+			else
+			{
+				check(SourceBuffer->GetUsage() & BUF_ShaderResource);
+				Stride = GPixelFormats[Format].BlockBytes;
 
-			LinearTextureDesc = MakeUnique<FMetalLinearTextureDescriptor>(InStartOffsetBytes, InNumElements, Stride);
-			SourceBuffer->CreateLinearTexture((EPixelFormat)Format, SourceBuffer, LinearTextureDesc.Get());
+				LinearTextureDesc = MakeUnique<FMetalLinearTextureDescriptor>(InStartOffsetBytes, InNumElements, Stride);
+				SourceBuffer->CreateLinearTexture((EPixelFormat)Format, SourceBuffer, LinearTextureDesc.Get());
+			}
 		}
-	}
-	else if (EnumHasAnyFlags(Usage, BUF_IndexBuffer))
-	{
-		if (!SourceBuffer)
+		else if (EnumHasAnyFlags(Usage, BUF_IndexBuffer))
 		{
-			Format = PF_R16_UINT;
-			Stride = 0;
+			if (!SourceBuffer)
+			{
+				Format = PF_R16_UINT;
+				Stride = 0;
+			}
+			else
+			{
+				Format = (SourceBuffer->IndexType == mtlpp::IndexType::UInt16) ? PF_R16_UINT : PF_R32_UINT;
+				Stride = SourceBuffer->GetStride();
+
+				check(Stride == ((Format == PF_R16_UINT) ? 2 : 4));
+
+				LinearTextureDesc = MakeUnique<FMetalLinearTextureDescriptor>(InStartOffsetBytes, InNumElements, Stride);
+				SourceBuffer->CreateLinearTexture((EPixelFormat)Format, SourceBuffer, LinearTextureDesc.Get());
+			}
 		}
 		else
 		{
-			Format = (SourceBuffer->IndexType == mtlpp::IndexType::UInt16) ? PF_R16_UINT : PF_R32_UINT;
-			Stride = SourceBuffer->GetStride();
+			check(EnumHasAnyFlags(Usage, BUF_StructuredBuffer));
 
-			check(Stride == ((Format == PF_R16_UINT) ? 2 : 4));
-
-			LinearTextureDesc = MakeUnique<FMetalLinearTextureDescriptor>(InStartOffsetBytes, InNumElements, Stride);
-			SourceBuffer->CreateLinearTexture((EPixelFormat)Format, SourceBuffer, LinearTextureDesc.Get());
+			Format = PF_Unknown;
+			Stride = SourceBuffer->GetStride();
 		}
 	}
-	else
-	{
-		check(EnumHasAnyFlags(Usage, BUF_StructuredBuffer));
-
-		Format = PF_Unknown;
-		Stride = SourceBuffer->GetStride();
-	}
 }
 
 // Constructor for textures
