Apply patch from Apple.   Enable WaveOps for macOS - tested on Vega 56 Monterey 12.5.1

#rnx

diff --git a/Engine/Config/Mac/DataDrivenPlatformInfo.ini b/Engine/Config/Mac/DataDrivenPlatformInfo.ini
--- a/Engine/Config/Mac/DataDrivenPlatformInfo.ini
+++ b/Engine/Config/Mac/DataDrivenPlatformInfo.ini
@@ -35,7 +35,7 @@ bSupportsDistanceFields=true
 bSupportsGen5TemporalAA=false
 bSupportsLumenGI=true
 bSupportsSSDIndirect=true
-bSupportsWaveOperations=Unsupported
+bSupportsWaveOperations=RuntimeDependent
 bSupportsGPUScene=true
 bSupportsFFTBloom=true
 bSupportsDiaphragmDOF=true
diff --git a/Engine/Source/Developer/Apple/MetalShaderFormat/Private/MetalDerivedData.cpp b/Engine/Source/Developer/Apple/MetalShaderFormat/Private/MetalDerivedData.cpp
--- a/Engine/Source/Developer/Apple/MetalShaderFormat/Private/MetalDerivedData.cpp
+++ b/Engine/Source/Developer/Apple/MetalShaderFormat/Private/MetalDerivedData.cpp
@@ -282,6 +282,8 @@ bool DoCompileMetalShader(
 		// Initialize compilation options for ShaderConductor
 		CrossCompiler::FShaderConductorOptions Options;
 
+		Options.TargetEnvironment = CrossCompiler::FShaderConductorOptions::ETargetEnvironment::Vulkan_1_1;
+
 		// Enable HLSL 2021 if specified
 		if (Input.Environment.CompilerFlags.Contains(CFLAG_HLSL2021))
 		{
diff --git a/Engine/Source/Runtime/Apple/MetalRHI/Private/MetalRHI.cpp b/Engine/Source/Runtime/Apple/MetalRHI/Private/MetalRHI.cpp
--- a/Engine/Source/Runtime/Apple/MetalRHI/Private/MetalRHI.cpp
+++ b/Engine/Source/Runtime/Apple/MetalRHI/Private/MetalRHI.cpp
@@ -335,6 +335,8 @@ FMetalDynamicRHI::FMetalDynamicRHI(ERHIFeatureLevel::Type RequestedFeatureLevel)
 		GRHIVendorId = 0x10DE;
 		bSupportsTiledReflections = true;
 		bSupportsDistanceFields = (FPlatformMisc::MacOSXVersionCompare(10,11,4) >= 0);
+
+		GRHISupportsWaveOperations = false;
 	}
 	else if(GRHIAdapterName.Contains("ATi") || GRHIAdapterName.Contains("AMD"))
 	{
@@ -349,6 +351,14 @@ FMetalDynamicRHI::FMetalDynamicRHI(ERHIFeatureLevel::Type RequestedFeatureLevel)
 		
 		// On AMD can also use completion handler time stamp if macOS < Catalina
 		GSupportsTimestampRenderQueries = true;
+
+		// Only tested on Vega.
+		GRHISupportsWaveOperations = GRHIAdapterName.Contains(TEXT("Vega"));
+		if (GRHISupportsWaveOperations)
+		{
+			GRHIMinimumWaveSize = 32;
+			GRHIMaximumWaveSize = 64;
+		}
 	}
 	else if(GRHIAdapterName.Contains("Intel"))
 	{
@@ -357,6 +367,8 @@ FMetalDynamicRHI::FMetalDynamicRHI(ERHIFeatureLevel::Type RequestedFeatureLevel)
 		GRHIVendorId = 0x8086;
 		bSupportsDistanceFields = (FPlatformMisc::MacOSXVersionCompare(10,12,2) >= 0);
 		bIsIntelHaswell = (GRHIAdapterName == TEXT("Intel HD Graphics 5000") || GRHIAdapterName == TEXT("Intel Iris Graphics") || GRHIAdapterName == TEXT("Intel Iris Pro Graphics"));
+
+		GRHISupportsWaveOperations = false;
 	}
 	else if(GRHIAdapterName.Contains("Apple"))
 	{
@@ -365,6 +377,10 @@ FMetalDynamicRHI::FMetalDynamicRHI(ERHIFeatureLevel::Type RequestedFeatureLevel)
 		bSupportsTiledReflections = true;
 		bSupportsDistanceFields = true;
 		GSupportsTimestampRenderQueries = true;
+
+		GRHISupportsWaveOperations = true;
+		GRHIMinimumWaveSize = 32;
+		GRHIMaximumWaveSize = 32;
 	}
 
 	bool const bRequestedSM5 = (RequestedFeatureLevel == ERHIFeatureLevel::SM5 || (!bRequestedFeatureLevel && (FParse::Param(FCommandLine::Get(),TEXT("metalsm5")) || FParse::Param(FCommandLine::Get(),TEXT("metalmrt")))));
