Fix for minor memory leak in TMetalBaseShader - release allocated mutable dictionary

diff --git a/Engine/Source/Runtime/Apple/MetalRHI/Private/Shaders/Types/Templates/MetalBaseShader.h b/Engine/Source/Runtime/Apple/MetalRHI/Private/Shaders/Types/Templates/MetalBaseShader.h
--- a/Engine/Source/Runtime/Apple/MetalRHI/Private/Shaders/Types/Templates/MetalBaseShader.h
+++ b/Engine/Source/Runtime/Apple/MetalRHI/Private/Shaders/Types/Templates/MetalBaseShader.h
@@ -297,6 +297,7 @@ void TMetalBaseShader<BaseResourceType, ShaderType>::Init(TArrayView<const uint8
 			[PreprocessorMacros addEntriesFromDictionary: @{ @"MTLSL_ENABLE_DEBUG_INFO" : @(1)}];
 #endif
 			CompileOptions.SetPreprocessorMacros(PreprocessorMacros);
+			[PreprocessorMacros release];
 #endif
 
 			mtlpp::LanguageVersion MetalVersion;
