Remove Catalina support and Metal 2.2 from Mac.

diff --git a/Engine/Source/Developer/Apple/MetalShaderFormat/Private/MetalDerivedData.cpp b/Engine/Source/Developer/Apple/MetalShaderFormat/Private/MetalDerivedData.cpp
--- a/Engine/Source/Developer/Apple/MetalShaderFormat/Private/MetalDerivedData.cpp
+++ b/Engine/Source/Developer/Apple/MetalShaderFormat/Private/MetalDerivedData.cpp
@@ -1040,15 +1040,10 @@ bool DoCompileMetalShader(
 					TargetDesc.Version = 20300;
 					break;
 				}
-                case 5:
-                {
-                    TargetDesc.Version = 20200;
-                    break;
-                }
 				default:
 				{
-					UE_LOG(LogShaders, Warning, TEXT("Metal Shader Version Unsupported, switching to default 2.2")); //EMacMetalShaderStandard::MacMetalSLStandard_Minimum
-					TargetDesc.Version = 20200;
+					UE_LOG(LogShaders, Warning, TEXT("Metal Shader Version Unsupported, switching to default 2.3")); //EMacMetalShaderStandard::MacMetalSLStandard_Minimum
+					TargetDesc.Version = 20300;
 					break;
 				}
 #else
diff --git a/Engine/Source/Developer/Mac/MacTargetPlatform/Classes/MacTargetSettings.h b/Engine/Source/Developer/Mac/MacTargetPlatform/Classes/MacTargetSettings.h
--- a/Engine/Source/Developer/Mac/MacTargetPlatform/Classes/MacTargetSettings.h
+++ b/Engine/Source/Developer/Mac/MacTargetPlatform/Classes/MacTargetSettings.h
@@ -13,10 +13,8 @@
 UENUM()
 enum class EMacMetalShaderStandard : uint8
 {
-    /** Metal Shader 2.2 is the minimum as of UE5.0*/
-    MacMetalSLStandard_Minimum = 0 UMETA(DisplayName="Minimum, Currently v2.2 (10.15.0+)"),
-    /** Metal Shaders Compatible With macOS 10.15 or later (std=osx-metal2.2) */
-    MacMetalSLStandard_2_2 = 5 UMETA(DisplayName="Metal v2.2 (10.15.0+)"),
+    /** Metal Shader 2.3 is the minimum as of UE5.0*/
+    MacMetalSLStandard_Minimum = 0 UMETA(DisplayName="Minimum, Currently v2.3 (11+)"),
     /** Metal Shaders Compatible With macOS 11.0 or later (std=osx-metal2.3) */
     MacMetalSLStandard_2_3 = 6 UMETA(DisplayName="Metal v2.3 (11.0+)"),
     /** Metal Shaders Compatible With macOS 12.0 or later (std=osx-metal2.4) */
diff --git a/Engine/Source/Editor/DetailCustomizations/Private/MacTargetSettingsDetails.cpp b/Engine/Source/Editor/DetailCustomizations/Private/MacTargetSettingsDetails.cpp
--- a/Engine/Source/Editor/DetailCustomizations/Private/MacTargetSettingsDetails.cpp
+++ b/Engine/Source/Editor/DetailCustomizations/Private/MacTargetSettingsDetails.cpp
@@ -122,19 +122,19 @@ static FText GetFriendlyNameFromRHINameMac(FName InRHIName)
 		FriendlyRHIName = LOCTEXT("OpenGLES31", "OpenGLES 3.1 (Mobile)");
 		break;
 	case SP_METAL:
-		FriendlyRHIName = LOCTEXT("Metal", "iOS Metal Mobile Renderer (ES3.1, Metal 1.1+, iOS 9.0 or later)");
+		FriendlyRHIName = LOCTEXT("Metal", "iOS Metal Mobile Renderer (ES3.1, Metal 2.3+, iOS 14.0 or later)");
 		break;
 	case SP_METAL_MRT:
-		FriendlyRHIName = LOCTEXT("MetalMRT", "iOS Metal Desktop Renderer (SM5, Metal 1.2+, iOS 10.0 or later)");
+		FriendlyRHIName = LOCTEXT("MetalMRT", "iOS Metal Desktop Renderer (SM5, Metal 2.3+, iOS 14.0 or later)");
 		break;
 	case SP_METAL_TVOS:
-		FriendlyRHIName = LOCTEXT("MetalTV", "tvOS Metal Mobile Renderer (ES3.1, Metal 1.1+, tvOS 9.0 or later)");
+		FriendlyRHIName = LOCTEXT("MetalTV", "tvOS Metal Mobile Renderer (ES3.1, Metal 2.3+, tvOS 14.0 or later)");
 		break;
 	case SP_METAL_MRT_TVOS:
-		FriendlyRHIName = LOCTEXT("MetalMRTTV", "tvOS Metal Desktop Renderer (SM5, Metal 1.2+, tvOS 10.0 or later)");
+		FriendlyRHIName = LOCTEXT("MetalMRTTV", "tvOS Metal Desktop Renderer (SM5, Metal 2.3+, tvOS 14.0 or later)");
 		break;
 	case SP_METAL_SM5:
-		FriendlyRHIName = LOCTEXT("MetalSM5", "Mac Metal Desktop Renderer (SM5, Metal 2.1+, macOS Catalina 10.15.7 or later)");
+		FriendlyRHIName = LOCTEXT("MetalSM5", "Mac Metal Desktop Renderer (SM5, Metal 2.3+, macOS BigSur 11.0 or later)");
 		break;
 	case SP_METAL_MACES3_1:
 		FriendlyRHIName = LOCTEXT("MetalES3.1", "Mac Metal High-End Mobile Preview (ES3.1)");
@@ -343,7 +343,6 @@ bool FMacTargetSettingsDetails::HandlePostExternalIconCopy(const FString& InChos
 }
 
 static uint32 GMacTargetSettingsMinOSVers[][3] = {
-	{10,15,0},
 	{11, 0,0},
     {12, 0,0}
 };
@@ -394,9 +393,9 @@ void FMacTargetSettingsDetails::SetShaderStandard(int32 Value)
     }
     
     ShaderVersionWarningTextBox->SetError(TEXT(""));
-    if (Value < 5) // EMacMetalShaderStandard::MacMetalSLStandard_Minimum
+    if (Value < 6) // EMacMetalShaderStandard::MacMetalSLStandard_Minimum
     {
-        ShaderVersionWarningTextBox->SetError(TEXT("Metal Shader Standard is 2.2 on UE5.0"));
+        ShaderVersionWarningTextBox->SetError(TEXT("Metal Shader Standard is 2.3 on UE5.1"));
     }
 }
 
@@ -408,9 +407,9 @@ void FMacTargetSettingsDetails::UpdateShaderStandardWarning()
     if (ShaderVersionPropertyHandle->IsValidHandle())
     {
         ShaderVersionPropertyHandle->GetValue(EnumValue);
-        if (EnumValue < 5)
+        if (EnumValue < 6)
         {
-            SetShaderStandard(5); // EMacMetalShaderStandard::MacMetalSLStandard_Minimum
+            SetShaderStandard(6); // EMacMetalShaderStandard::MacMetalSLStandard_Minimum
         }
     }
 }
diff --git a/Engine/Source/Editor/UnrealEd/Private/EditorEngine.cpp b/Engine/Source/Editor/UnrealEd/Private/EditorEngine.cpp
--- a/Engine/Source/Editor/UnrealEd/Private/EditorEngine.cpp
+++ b/Engine/Source/Editor/UnrealEd/Private/EditorEngine.cpp
@@ -6908,7 +6908,7 @@ void UEditorEngine::UpdateAutoLoadProject()
 #if PLATFORM_MAC
 	if ( !GIsBuildMachine )
 	{
-		if(FPlatformMisc::MacOSXVersionCompare(10,15,7) < 0)
+		if(FPlatformMisc::MacOSXVersionCompare(11, 0, 0) < 0)
 		{
 			if(FSlateApplication::IsInitialized())
 			{
diff --git a/Engine/Source/Programs/UnrealBuildTool/Platform/Mac/MacPlatformSDK.Versions.cs b/Engine/Source/Programs/UnrealBuildTool/Platform/Mac/MacPlatformSDK.Versions.cs
--- a/Engine/Source/Programs/UnrealBuildTool/Platform/Mac/MacPlatformSDK.Versions.cs
+++ b/Engine/Source/Programs/UnrealBuildTool/Platform/Mac/MacPlatformSDK.Versions.cs
@@ -14,7 +14,7 @@ namespace UnrealBuildTool
 
 		protected override void GetValidSoftwareVersionRange(out string MinVersion, out string? MaxVersion)
 		{
-			MinVersion = "10.15.7";
+			MinVersion = "11.0.0";
 			MaxVersion = null;
 		}
 	}
diff --git a/Engine/Source/Runtime/Apple/MetalRHI/Private/MetalCommandQueue.cpp b/Engine/Source/Runtime/Apple/MetalRHI/Private/MetalCommandQueue.cpp
--- a/Engine/Source/Runtime/Apple/MetalRHI/Private/MetalCommandQueue.cpp
+++ b/Engine/Source/Runtime/Apple/MetalRHI/Private/MetalCommandQueue.cpp
@@ -161,7 +161,7 @@ FMetalCommandQueue::FMetalCommandQueue(mtlpp::Device InDevice, uint32 const MaxN
 		Features |= EMetalFeaturesSetBytes;
 		
 		// On earlier OS versions Intel Broadwell couldn't suballocate properly
-		if (!(DeviceName.Contains(TEXT("Intel")) && (DeviceName.Contains(TEXT("5300")) || DeviceName.Contains(TEXT("6000")) || DeviceName.Contains(TEXT("6100")))) || FPlatformMisc::MacOSXVersionCompare(10,14,0) >= 0)
+		if (!(DeviceName.Contains(TEXT("Intel")) && (DeviceName.Contains(TEXT("5300")) || DeviceName.Contains(TEXT("6000")) || DeviceName.Contains(TEXT("6100")))))
 		{
 			// Using Private Memory & BlitEncoders for Vertex & Index data should be *much* faster.
 			Features |= EMetalFeaturesEfficientBufferBlits;
@@ -169,7 +169,7 @@ FMetalCommandQueue::FMetalCommandQueue(mtlpp::Device InDevice, uint32 const MaxN
 			Features |= EMetalFeaturesBufferSubAllocation;
 					
 			// On earlier OS versions Vega didn't like non-zero blit offsets
-	        if (!DeviceName.Contains(TEXT("Vega")) || FPlatformMisc::MacOSXVersionCompare(10,13,5) >= 0)
+	        if (!DeviceName.Contains(TEXT("Vega")))
 			{
 				Features |= EMetalFeaturesPrivateBufferSubAllocation;
 			}
@@ -219,7 +219,7 @@ FMetalCommandQueue::FMetalCommandQueue(mtlpp::Device InDevice, uint32 const MaxN
 		Features |= EMetalFeaturesSetBytes;
 	}
     
-    if(Device.SupportsFeatureSet(mtlpp::FeatureSet::macOS_GPUFamily1_v3) && FPlatformMisc::MacOSXVersionCompare(10,13,0) >= 0)
+    if(Device.SupportsFeatureSet(mtlpp::FeatureSet::macOS_GPUFamily1_v3))
 	{
 		Features |= EMetalFeaturesMultipleViewports | EMetalFeaturesPipelineBufferMutability | EMetalFeaturesGPUCaptureManager;
 		
diff --git a/Engine/Source/Runtime/Apple/MetalRHI/Private/MetalContext.cpp b/Engine/Source/Runtime/Apple/MetalRHI/Private/MetalContext.cpp
--- a/Engine/Source/Runtime/Apple/MetalRHI/Private/MetalContext.cpp
+++ b/Engine/Source/Runtime/Apple/MetalRHI/Private/MetalContext.cpp
@@ -161,28 +161,21 @@ static mtlpp::Device GetMTLDevice(uint32& DeviceIndex)
 	
 	ns::Array<mtlpp::Device> DeviceList;
 	
-	if (FPlatformMisc::MacOSXVersionCompare(10, 13, 4) >= 0)
-	{
-			DeviceList = mtlpp::Device::CopyAllDevicesWithObserver(GMetalDeviceObserver, ^(const mtlpp::Device & Device, const ns::String & Notification)
-			{
-				if ([Notification.GetPtr() isEqualToString:MTLDeviceWasAddedNotification])
-				{
-					FPlatformMisc::GPUChangeNotification(Device.GetRegistryID(), FPlatformMisc::EMacGPUNotification::Added);
-				}
-				else if ([Notification.GetPtr() isEqualToString:MTLDeviceRemovalRequestedNotification])
-				{
-					FPlatformMisc::GPUChangeNotification(Device.GetRegistryID(), FPlatformMisc::EMacGPUNotification::RemovalRequested);
-				}
-				else if ([Notification.GetPtr() isEqualToString:MTLDeviceWasRemovedNotification])
-				{
-					FPlatformMisc::GPUChangeNotification(Device.GetRegistryID(), FPlatformMisc::EMacGPUNotification::Removed);
-				}
-			});
-	}
-	else
-	{
-		DeviceList = mtlpp::Device::CopyAllDevices();
-	}
+    DeviceList = mtlpp::Device::CopyAllDevicesWithObserver(GMetalDeviceObserver, ^(const mtlpp::Device & Device, const ns::String & Notification)
+    {
+        if ([Notification.GetPtr() isEqualToString:MTLDeviceWasAddedNotification])
+        {
+            FPlatformMisc::GPUChangeNotification(Device.GetRegistryID(), FPlatformMisc::EMacGPUNotification::Added);
+        }
+        else if ([Notification.GetPtr() isEqualToString:MTLDeviceRemovalRequestedNotification])
+        {
+            FPlatformMisc::GPUChangeNotification(Device.GetRegistryID(), FPlatformMisc::EMacGPUNotification::RemovalRequested);
+        }
+        else if ([Notification.GetPtr() isEqualToString:MTLDeviceWasRemovedNotification])
+        {
+            FPlatformMisc::GPUChangeNotification(Device.GetRegistryID(), FPlatformMisc::EMacGPUNotification::Removed);
+        }
+    });
 	
 	const int32 NumDevices = DeviceList.GetSize();
 	
@@ -198,7 +191,7 @@ static mtlpp::Device GetMTLDevice(uint32& DeviceIndex)
  	int32 OverrideRendererId = FPlatformMisc::GetExplicitRendererIndex();
 	
 	int32 ExplicitRendererId = OverrideRendererId >= 0 ? OverrideRendererId : HmdGraphicsAdapter;
-	if(ExplicitRendererId < 0 && GPUs.Num() > 1 && FMacPlatformMisc::MacOSXVersionCompare(10, 11, 5) == 0)
+	if(ExplicitRendererId < 0 && GPUs.Num() > 1)
 	{
 		OverrideRendererId = -1;
 		bool bForceExplicitRendererId = false;
@@ -458,10 +451,7 @@ FMetalDeviceContext::~FMetalDeviceContext()
     delete UniformBufferAllocator;
 	
 #if PLATFORM_MAC
-	if (FPlatformMisc::MacOSXVersionCompare(10, 13, 4) >= 0)
-	{
-		mtlpp::Device::RemoveDeviceObserver(GMetalDeviceObserver);
-	}
+    mtlpp::Device::RemoveDeviceObserver(GMetalDeviceObserver);
 #endif
 }
 
diff --git a/Engine/Source/Runtime/Apple/MetalRHI/Private/MetalRHI.cpp b/Engine/Source/Runtime/Apple/MetalRHI/Private/MetalRHI.cpp
--- a/Engine/Source/Runtime/Apple/MetalRHI/Private/MetalRHI.cpp
+++ b/Engine/Source/Runtime/Apple/MetalRHI/Private/MetalRHI.cpp
@@ -326,32 +326,30 @@ FMetalDynamicRHI::FMetalDynamicRHI(ERHIFeatureLevel::Type RequestedFeatureLevel)
 	bool bSupportsSM5 = true;
 	bool bIsIntelHaswell = false;
 	
-	// All should work on Catalina+ using GPU end time
-	GSupportsTimestampRenderQueries = FPlatformMisc::MacOSXVersionCompare(10,15,0) >= 0;
+	GSupportsTimestampRenderQueries = true;
 	
 	if(GRHIAdapterName.Contains("Nvidia"))
 	{
 		bSupportsPointLights = true;
 		GRHIVendorId = 0x10DE;
 		bSupportsTiledReflections = true;
-		bSupportsDistanceFields = (FPlatformMisc::MacOSXVersionCompare(10,11,4) >= 0);
-
+		bSupportsDistanceFields = true;
 		GRHISupportsWaveOperations = false;
 	}
 	else if(GRHIAdapterName.Contains("ATi") || GRHIAdapterName.Contains("AMD"))
 	{
 		bSupportsPointLights = true;
 		GRHIVendorId = 0x1002;
-		if((FPlatformMisc::MacOSXVersionCompare(10,12,0) < 0) && GPUDesc.GPUVendorId == GRHIVendorId)
+		if(GPUDesc.GPUVendorId == GRHIVendorId)
 		{
 			GRHIAdapterName = FString(GPUDesc.GPUName);
 		}
 		bSupportsTiledReflections = true;
-		bSupportsDistanceFields = (FPlatformMisc::MacOSXVersionCompare(10,11,4) >= 0);
+		bSupportsDistanceFields = true;
 		
 		// On AMD can also use completion handler time stamp if macOS < Catalina
 		GSupportsTimestampRenderQueries = true;
-
+		
 		// Only tested on Vega.
 		GRHISupportsWaveOperations = GRHIAdapterName.Contains(TEXT("Vega"));
 		if (GRHISupportsWaveOperations)
@@ -363,11 +361,10 @@ FMetalDynamicRHI::FMetalDynamicRHI(ERHIFeatureLevel::Type RequestedFeatureLevel)
 	else if(GRHIAdapterName.Contains("Intel"))
 	{
 		bSupportsTiledReflections = false;
-		bSupportsPointLights = (FPlatformMisc::MacOSXVersionCompare(10,14,6) > 0);
+		bSupportsPointLights = true;
 		GRHIVendorId = 0x8086;
-		bSupportsDistanceFields = (FPlatformMisc::MacOSXVersionCompare(10,12,2) >= 0);
+		bSupportsDistanceFields = true;
 		bIsIntelHaswell = (GRHIAdapterName == TEXT("Intel HD Graphics 5000") || GRHIAdapterName == TEXT("Intel Iris Graphics") || GRHIAdapterName == TEXT("Intel Iris Pro Graphics"));
-
 		GRHISupportsWaveOperations = false;
 	}
 	else if(GRHIAdapterName.Contains("Apple"))
@@ -377,7 +374,7 @@ FMetalDynamicRHI::FMetalDynamicRHI(ERHIFeatureLevel::Type RequestedFeatureLevel)
 		bSupportsTiledReflections = true;
 		bSupportsDistanceFields = true;
 		GSupportsTimestampRenderQueries = true;
-
+		
 		GRHISupportsWaveOperations = true;
 		GRHIMinimumWaveSize = 32;
 		GRHIMaximumWaveSize = 32;
@@ -467,13 +464,9 @@ FMetalDynamicRHI::FMetalDynamicRHI(ERHIFeatureLevel::Type RequestedFeatureLevel)
 	
 #endif
 		
-	if(
-	   #if PLATFORM_MAC
-	   (Device.SupportsFeatureSet(mtlpp::FeatureSet::macOS_GPUFamily1_v3) && FPlatformMisc::MacOSXVersionCompare(10,13,0) >= 0)
-	   #elif PLATFORM_IOS || PLATFORM_TVOS
-	   FPlatformMisc::IOSVersionCompare(10,3,0)
-	   #endif
-	   )
+#if PLATFORM_MAC
+    if (Device.SupportsFeatureSet(mtlpp::FeatureSet::macOS_GPUFamily1_v3))
+#endif
 	{
 		GRHISupportsDynamicResolution = true;
 		GRHISupportsFrameCyclesBubblesRemoval = true;
@@ -557,14 +550,11 @@ FMetalDynamicRHI::FMetalDynamicRHI(ERHIFeatureLevel::Type RequestedFeatureLevel)
 	}
 
 #if PLATFORM_MAC
-	if (IsRHIDeviceIntel() && FPlatformMisc::MacOSXVersionCompare(10,13,5) < 0)
-	{
-		static auto CVarSGShadowQuality = IConsoleManager::Get().FindConsoleVariable((TEXT("sg.ShadowQuality")));
-		if (CVarSGShadowQuality && CVarSGShadowQuality->GetInt() != 0)
-		{
-			CVarSGShadowQuality->Set(0);
-		}
-	}
+    static auto CVarSGShadowQuality = IConsoleManager::Get().FindConsoleVariable((TEXT("sg.ShadowQuality")));
+    if (CVarSGShadowQuality && CVarSGShadowQuality->GetInt() != 0)
+    {
+        CVarSGShadowQuality->Set(0);
+    }
 
 	if (bIsIntelHaswell)
 	{
@@ -593,7 +583,7 @@ FMetalDynamicRHI::FMetalDynamicRHI(ERHIFeatureLevel::Type RequestedFeatureLevel)
 
 #if PLATFORM_MAC
 	check(Device.SupportsFeatureSet(mtlpp::FeatureSet::macOS_GPUFamily1_v1));
-	GRHISupportsBaseVertexIndex = FPlatformMisc::MacOSXVersionCompare(10,11,2) >= 0 || !IsRHIDeviceAMD(); // Supported on macOS & iOS but not tvOS - broken on AMD prior to 10.11.2
+	GRHISupportsBaseVertexIndex = true;
 	GRHISupportsFirstInstance = true; // Supported on macOS & iOS but not tvOS.
 	GMaxTextureDimensions = 16384;
 	GMaxCubeTextureDimensions = 16384;
@@ -601,7 +591,7 @@ FMetalDynamicRHI::FMetalDynamicRHI(ERHIFeatureLevel::Type RequestedFeatureLevel)
 	GMaxShadowDepthBufferSizeX = GMaxTextureDimensions;
 	GMaxShadowDepthBufferSizeY = GMaxTextureDimensions;
     bSupportsD16 = !FParse::Param(FCommandLine::Get(),TEXT("nometalv2")) && Device.SupportsFeatureSet(mtlpp::FeatureSet::macOS_GPUFamily1_v2);
-    GRHISupportsHDROutput = FPlatformMisc::MacOSXVersionCompare(10,14,4) >= 0 && Device.SupportsFeatureSet(mtlpp::FeatureSet::macOS_GPUFamily1_v2);
+    GRHISupportsHDROutput = Device.SupportsFeatureSet(mtlpp::FeatureSet::macOS_GPUFamily1_v2);
 	GRHIHDRDisplayOutputFormat = (GRHISupportsHDROutput) ? PF_PLATFORM_HDR_0 : PF_B8G8R8A8;
 	// Based on the spec below, the maxTotalThreadsPerThreadgroup is not a fixed number but calculated according to the device current ability, so the available threads could less than the maximum number.
 	// For safety and keep the consistency for all platform, reduce the maximum number to half of the device based.
diff --git a/Engine/Source/Runtime/Apple/MetalRHI/Private/MetalRHIRenderQuery.cpp b/Engine/Source/Runtime/Apple/MetalRHI/Private/MetalRHIRenderQuery.cpp
--- a/Engine/Source/Runtime/Apple/MetalRHI/Private/MetalRHIRenderQuery.cpp
+++ b/Engine/Source/Runtime/Apple/MetalRHI/Private/MetalRHIRenderQuery.cpp
@@ -282,15 +282,10 @@ void FMetalRHIRenderQuery::End(FMetalContext* Context)
 			{
 				Result = 0;
 				
-#if PLATFORM_MAC
-				if(FPlatformMisc::MacOSXVersionCompare(10,15,0) >= 0)
-#endif
-				{
-					// If there are no commands in the command buffer then this can be zero
-					// In this case GPU start time is also not correct - we need to fall back standard behaviour
-					// Only seen empty command buffers at the very end of a frame
-					Result = uint64((CmdBuffer.GetGpuEndTime() / 1000.0) / FPlatformTime::GetSecondsPerCycle64());
-				}
+                // If there are no commands in the command buffer then this can be zero
+                // In this case GPU start time is also not correct - we need to fall back standard behaviour
+                // Only seen empty command buffers at the very end of a frame
+                Result = uint64((CmdBuffer.GetGpuEndTime() / 1000.0) / FPlatformTime::GetSecondsPerCycle64());
 				
 				if(Result == 0)
 				{
diff --git a/Engine/Source/Runtime/Apple/MetalRHI/Private/MetalShaders.cpp b/Engine/Source/Runtime/Apple/MetalRHI/Private/MetalShaders.cpp
--- a/Engine/Source/Runtime/Apple/MetalRHI/Private/MetalShaders.cpp
+++ b/Engine/Source/Runtime/Apple/MetalRHI/Private/MetalShaders.cpp
@@ -58,16 +58,15 @@ mtlpp::LanguageVersion ValidateVersion(uint32 Version)
         case 6:
             Result = mtlpp::LanguageVersion::Version2_3;
             break;
-        case 5:
-            Result = mtlpp::LanguageVersion::Version2_2;
-            break;
-            Result = mtlpp::LanguageVersion::Version2_2; // minimum version as of UE5.0
+        case 0:
+            Version = 6;
+            Result = mtlpp::LanguageVersion::Version2_3; // minimum version as of UE5.1
             break;
         default:
-            //EMacMetalShaderStandard::MacMetalSLStandard_Minimum is currently 2.2
+            //EMacMetalShaderStandard::MacMetalSLStandard_Minimum is currently 2.3
             UE_LOG(LogTemp, Warning, TEXT("The Metal version currently set is not supported anymore. Set it in the Project Settings. Defaulting to the minimum version."));
-            Version = 5;
-            Result = mtlpp::LanguageVersion::Version2_2;
+            Version = 6;
+            Result = mtlpp::LanguageVersion::Version2_3;
             break;
     }
 #else
diff --git a/Engine/Source/Runtime/Apple/MetalRHI/Private/MetalViewport.cpp b/Engine/Source/Runtime/Apple/MetalRHI/Private/MetalViewport.cpp
--- a/Engine/Source/Runtime/Apple/MetalRHI/Private/MetalViewport.cpp
+++ b/Engine/Source/Runtime/Apple/MetalRHI/Private/MetalViewport.cpp
@@ -434,8 +434,7 @@ void FMetalViewport::Present(FMetalCommandQueue& CommandQueue, bool bLockToVsync
 	bIsLiveResize = View.inLiveResize;
 	{
 		FCAMetalLayer* CurrentLayer = (FCAMetalLayer*)[View layer];
-		static bool sVSyncSafe = FPlatformMisc::MacOSXVersionCompare(10,13,4) >= 0;
-		CurrentLayer.displaySyncEnabled = bLockToVsync || (!sVSyncSafe && !(IsRunningGame() && bIsFullScreen));
+		CurrentLayer.displaySyncEnabled = bLockToVsync || (!(IsRunningGame() && bIsFullScreen));
 	}
 #endif
 	
diff --git a/Engine/Source/Runtime/ApplicationCore/Private/Mac/MacApplication.cpp b/Engine/Source/Runtime/ApplicationCore/Private/Mac/MacApplication.cpp
--- a/Engine/Source/Runtime/ApplicationCore/Private/Mac/MacApplication.cpp
+++ b/Engine/Source/Runtime/ApplicationCore/Private/Mac/MacApplication.cpp
@@ -83,7 +83,7 @@ public:
 	
 	FMacControllerInterface(const TSharedRef< FGenericApplicationMessageHandler >& InMessageHandler)
 	{
-		if(FPlatformMisc::MacOSXVersionCompare(10,15,0) >= 0 && CVarMacControllerPreferGCImpl.GetValueOnAnyThread() > 0)
+		if(CVarMacControllerPreferGCImpl.GetValueOnAnyThread() > 0)
 		{
 			AppleControllerInterface = FAppleControllerInterface::Create(InMessageHandler);
 		}
diff --git a/Engine/Source/Runtime/Core/Private/Mac/MacPlatformProcess.cpp b/Engine/Source/Runtime/Core/Private/Mac/MacPlatformProcess.cpp
--- a/Engine/Source/Runtime/Core/Private/Mac/MacPlatformProcess.cpp
+++ b/Engine/Source/Runtime/Core/Private/Mac/MacPlatformProcess.cpp
@@ -78,28 +78,13 @@ void* FMacPlatformProcess::GetDllHandle( const TCHAR* Filename )
 	NSString* ExecutableFolder = [[[NSBundle mainBundle] executablePath] stringByDeletingLastPathComponent];
 	void* Handle = nullptr;
 
-	// On 11.0.0+, system-provided dynamic libraries do not exist on the
-	// filesystem, only in a built-in dynamic linker cache.
-	if (FPlatformMisc::MacOSXVersionCompare(10,16,0) >= 0)
-	{
-		Handle = GetDllHandleImpl(DylibPath, ExecutableFolder);
-		if (!Handle)
-		{
-			// If it's not a absolute or relative path, try to find the file in the app bundle
-			DylibPath = [ExecutableFolder stringByAppendingPathComponent:DylibPath];
-			Handle = GetDllHandleImpl(DylibPath, ExecutableFolder);
-		}
-	}
-	else
-	{
-		NSFileManager* FileManager = [NSFileManager defaultManager];
-		if (![FileManager fileExistsAtPath:DylibPath])
-		{
-			// If it's not a absolute or relative path, try to find the file in the app bundle
-			DylibPath = [ExecutableFolder stringByAppendingPathComponent:DylibPath];
-		}
-		Handle = GetDllHandleImpl(DylibPath, ExecutableFolder);
-	}
+    Handle = GetDllHandleImpl(DylibPath, ExecutableFolder);
+    if (!Handle)
+    {
+        // If it's not a absolute or relative path, try to find the file in the app bundle
+        DylibPath = [ExecutableFolder stringByAppendingPathComponent:DylibPath];
+        Handle = GetDllHandleImpl(DylibPath, ExecutableFolder);
+    }
 	return Handle;
 }
 
diff --git a/Engine/Source/Runtime/Launch/Private/Mac/LaunchMac.cpp b/Engine/Source/Runtime/Launch/Private/Mac/LaunchMac.cpp
--- a/Engine/Source/Runtime/Launch/Private/Mac/LaunchMac.cpp
+++ b/Engine/Source/Runtime/Launch/Private/Mac/LaunchMac.cpp
@@ -266,8 +266,8 @@ static int32 MacOSVersionCompare(const NSOperatingSystemVersion& VersionA, const
 	NSString* MinimumSystemVersionString = (NSString*)InfoDictionary[@"LSMinimumSystemVersion"];
 	NSOperatingSystemVersion MinimumSystemVersion = { 0 };
 	NSOperatingSystemVersion CurrentSystemVersion = [[NSProcessInfo processInfo] operatingSystemVersion];
-	NSOperatingSystemVersion MinSupportedMacOSVersion = { 10, 15, 7 };
-	NSString* MinSupportedMacOSVersionString = @"10.15.7";
+	NSOperatingSystemVersion MinSupportedMacOSVersion = { 11, 0, 0 };
+	NSString* MinSupportedMacOSVersionString = @"11.0.0";
 
 	NSArray<NSString*>* VersionComponents = [MinimumSystemVersionString componentsSeparatedByString:@"."];
 	MinimumSystemVersion.majorVersion = [[VersionComponents objectAtIndex:0] integerValue];
diff --git a/Engine/Source/Runtime/RHI/Private/DynamicRHI.cpp b/Engine/Source/Runtime/RHI/Private/DynamicRHI.cpp
--- a/Engine/Source/Runtime/RHI/Private/DynamicRHI.cpp
+++ b/Engine/Source/Runtime/RHI/Private/DynamicRHI.cpp
@@ -277,7 +277,7 @@ static void RHIDetectAndWarnOfBadDrivers(bool bHasEditorToken)
 		return;
 	}
 
-	if (FPlatformMisc::MacOSXVersionCompare(10,15,5) < 0)
+	if (FPlatformMisc::MacOSXVersionCompare(11, 0 ,0) < 0)
 	{
 		// this message can be suppressed with r.WarnOfBadDrivers=0
 		FPlatformMisc::MessageBoxExt(EAppMsgType::Ok,
