Catalina, iOS 14 removal
Metal 2.3 removed, Metal 3.0 added for Mac

diff --git a/Engine/Config/BaseEngine.ini b/Engine/Config/BaseEngine.ini
--- a/Engine/Config/BaseEngine.ini
+++ b/Engine/Config/BaseEngine.ini
@@ -2641,7 +2641,7 @@ DistroCodeSigningIdentity=iPhone Distribution:;Apple Distribution:
 VersionInfo=1.0
 FrameRateLock=PUFRL_30
 bEnableDynamicMaxFPS=False
-MinimumiOSVersion=IOS_14
+MinimumiOSVersion=IOS_15
 bSupportsIPad=True
 bSupportsIPhone=True
 AdditionalPlistData=
@@ -2658,7 +2658,7 @@ bAutomaticSigning=false
 UseFastIntrinsics=False
 ForceFloats=True
 EnableMathOptimisations=True
-MetalLanguageVersion=6
+MetalLanguageVersion=0
 bDisableMotionData=False
 bEnableAdvertisingIdentifier=True
 
@@ -2802,7 +2802,7 @@ bTarget32Bit=false
 bEnableRayTracing=true
 
 [/Script/MacTargetPlatform.MacTargetSettings]
-MetalLanguageVersion=5
+MetalLanguageVersion=0
 +TargetedRHIs=SF_METAL_SM5
 UseFastIntrinsics=False
 EnableMathOptimisations=True
@@ -3078,4 +3078,4 @@ FilePickerClass=/Script/InterchangeEditorUtilities.InterchangeFilePickerGeneric
 bStaticMeshUseSmoothEdgesIfSmoothingInformationIsMissing=True
 ContentImportSettings=(DefaultPipelineStackOverride=((Materials, "Materials"),(Textures, "Textures")),ShowPipelineStacksConfigurationDialogOverride=((Textures, False)),PipelineStacks=(("Assets", (Pipelines=("/Interchange/Pipelines/DefaultAssetsPipeline.DefaultAssetsPipeline"))),("Materials", (Pipelines=("/Interchange/Pipelines/DefaultMaterialPipeline.DefaultMaterialPipeline"))),("Textures", (Pipelines=("/Interchange/Pipelines/DefaultTexturePipeline.DefaultTexturePipeline")))),DefaultPipelineStack="Assets",PipelineConfigurationDialogClass="/Script/InterchangeEditorPipelines.InterchangePipelineConfigurationGeneric",bShowPipelineStacksConfigurationDialog=True)
 SceneImportSettings=(PipelineStacks=(("Scene", (Pipelines=("/Interchange/Pipelines/DefaultSceneAssetsPipeline.DefaultSceneAssetsPipeline","/Interchange/Pipelines/DefaultSceneLevelPipeline.DefaultSceneLevelPipeline")))),DefaultPipelineStack="Scene",PipelineConfigurationDialogClass="/Script/InterchangeEditorPipelines.InterchangePipelineConfigurationGeneric",bShowPipelineStacksConfigurationDialog=True)
-GenericPipelineClass=/Script/InterchangePipelines.InterchangeGenericAssetsPipeline
\ No newline at end of file
+GenericPipelineClass=/Script/InterchangePipelines.InterchangeGenericAssetsPipeline
diff --git a/Engine/Plugins/Media/ElectraUtil/Source/ElectraHTTPStream/Private/Apple/AppleElectraHTTPStream.cpp b/Engine/Plugins/Media/ElectraUtil/Source/ElectraHTTPStream/Private/Apple/AppleElectraHTTPStream.cpp
--- a/Engine/Plugins/Media/ElectraUtil/Source/ElectraHTTPStream/Private/Apple/AppleElectraHTTPStream.cpp
+++ b/Engine/Plugins/Media/ElectraUtil/Source/ElectraHTTPStream/Private/Apple/AppleElectraHTTPStream.cpp
@@ -1321,19 +1321,7 @@ FElectraHTTPStreamApple::FChallengeResponse FElectraHTTPStreamApple::ReceivedCha
 				SecTrustEvaluateWithError(CertTrust, nil);
 				TCFRef<SecKeyRef> CertPubKey;
 
-				if (@available(macOS 11, *))
-				{
-					CertPubKey = SecTrustCopyKey(CertTrust);
-				}
-				else
-				{
-					PRAGMA_DISABLE_DEPRECATION_WARNINGS
-					// warning: 'SecTrustCopyPublicKey' is deprecated: first deprecated in iOS 14.0 [-Wdeprecated-declarations]
-
-					CertPubKey = SecTrustCopyPublicKey(CertTrust);
-
-					PRAGMA_ENABLE_DEPRECATION_WARNINGS
-				}
+				CertPubKey = SecTrustCopyKey(CertTrust);
 
 				TCFRef<CFDataRef> CertPubKeyData = SecKeyCopyExternalRepresentation(CertPubKey, NULL);
 				if (!CertPubKeyData)
diff --git a/Engine/Source/Developer/Apple/MetalShaderFormat/Private/MetalDerivedData.cpp b/Engine/Source/Developer/Apple/MetalShaderFormat/Private/MetalDerivedData.cpp
--- a/Engine/Source/Developer/Apple/MetalShaderFormat/Private/MetalDerivedData.cpp
+++ b/Engine/Source/Developer/Apple/MetalShaderFormat/Private/MetalDerivedData.cpp
@@ -1030,20 +1030,20 @@ bool DoCompileMetalShader(
 			switch (VersionEnum)
 			{
 #if PLATFORM_MAC
+                case 8:
+                {
+                    TargetDesc.Version = 30000;
+                    break;
+                }
 				case 7:
 				{
 					TargetDesc.Version = 20400;
 					break;
 				}
-				case 6:
-				{
-					TargetDesc.Version = 20300;
-					break;
-				}
 				default:
 				{
-					UE_LOG(LogShaders, Warning, TEXT("Metal Shader Version Unsupported, switching to default 2.3")); //EMacMetalShaderStandard::MacMetalSLStandard_Minimum
-					TargetDesc.Version = 20300;
+					UE_LOG(LogShaders, Warning, TEXT("Metal Shader Version Unsupported, switching to default 2.4")); //EMacMetalShaderStandard::MacMetalSLStandard_Minimum
+					TargetDesc.Version = 20400;
 					break;
 				}
 #else
@@ -1057,15 +1057,10 @@ bool DoCompileMetalShader(
                     TargetDesc.Version = 20400;
                     break;
                 }
-                case 6:
-                {
-                    TargetDesc.Version = 20300;
-                    break;
-                }
                 default:
                 {
-                    UE_LOG(LogShaders, Warning, TEXT("Metal Shader Version Unsupported, switching to default 2.3")); //EIOSMetalShaderStandard::IOSMetalSLStandard_Minimum
-                    TargetDesc.Version = 20300;
+                    UE_LOG(LogShaders, Warning, TEXT("Metal Shader Version Unsupported, switching to default 2.4")); //EIOSMetalShaderStandard::IOSMetalSLStandard_Minimum
+                    TargetDesc.Version = 20400;
                     break;
                 }
 #endif
diff --git a/Engine/Source/Developer/IOS/IOSPlatformEditor/Private/IOSTargetSettingsCustomization.cpp b/Engine/Source/Developer/IOS/IOSPlatformEditor/Private/IOSTargetSettingsCustomization.cpp
--- a/Engine/Source/Developer/IOS/IOSPlatformEditor/Private/IOSTargetSettingsCustomization.cpp
+++ b/Engine/Source/Developer/IOS/IOSPlatformEditor/Private/IOSTargetSettingsCustomization.cpp
@@ -1906,7 +1906,7 @@ void FIOSTargetSettingsCustomization::SetShaderStandard(int32 Value)
 
     if (MinOSPropertyHandle.IsValid())
     {
-        uint8 IOSVersion = (uint8)EIOSVersion::IOS_14;
+        uint8 IOSVersion = (uint8)EIOSVersion::IOS_15;
         if (MinOSPropertyHandle.IsValid())
         {
             MinOSPropertyHandle->GetValue(IOSVersion);
@@ -1934,9 +1934,9 @@ void FIOSTargetSettingsCustomization::UpdateOSVersionWarning()
 		
 		if (bMRTEnabled)
 		{
-			if (EnumValue < (uint8)EIOSVersion::IOS_14)
+			if (EnumValue < (uint8)EIOSVersion::IOS_15)
 			{
-				SetMinVersion((int32)EIOSVersion::IOS_14);
+				SetMinVersion((int32)EIOSVersion::IOS_15);
 			}
 		}
 		else
@@ -1954,9 +1954,6 @@ void FIOSTargetSettingsCustomization::UpdateOSVersionWarning()
     switch (ShaderStandard)
     {
         case (int32)EIOSMetalShaderStandard::IOSMetalSLStandard_Minimum:
-        case (int32)EIOSMetalShaderStandard::IOSMetalSLStandard_2_3:
-            if (EnumValue < (uint8)EIOSVersion::IOS_14) {IOSVersionWarningTextBox->SetError(TEXT("iOS14 is the Minimum for Metal 2.3")); return;}
-            break;
         case (int32)EIOSMetalShaderStandard::IOSMetalSLStandard_2_4:
             if (EnumValue < (uint8)EIOSVersion::IOS_15) {IOSVersionWarningTextBox->SetError(TEXT("iOS15 is the Minimum for Metal 2.4")); return;}
             break;
@@ -1977,13 +1974,13 @@ void FIOSTargetSettingsCustomization::UpdateMetalMRTWarning()
 		{
 			uint8 EnumValue;
 			MinOSPropertyHandle->GetValue(EnumValue);
-			if (EnumValue < (uint8)EIOSVersion::IOS_14)
+			if (EnumValue < (uint8)EIOSVersion::IOS_15)
 			{
-				SetMinVersion((int32)EIOSVersion::IOS_14);
+				SetMinVersion((int32)EIOSVersion::IOS_15);
 			}
 			
 			ShaderVersionPropertyHandle->GetValue(EnumValue);
-			if (EnumValue < (int32)EIOSMetalShaderStandard::IOSMetalSLStandard_2_3)
+			if (EnumValue < (int32)EIOSMetalShaderStandard::IOSMetalSLStandard_2_4)
 			{
 				SetShaderStandard((int32)EIOSMetalShaderStandard::IOSMetalSLStandard_Minimum);
 			}
diff --git a/Engine/Source/Developer/Mac/MacTargetPlatform/Classes/MacTargetSettings.h b/Engine/Source/Developer/Mac/MacTargetPlatform/Classes/MacTargetSettings.h
--- a/Engine/Source/Developer/Mac/MacTargetPlatform/Classes/MacTargetSettings.h
+++ b/Engine/Source/Developer/Mac/MacTargetPlatform/Classes/MacTargetSettings.h
@@ -13,12 +13,13 @@
 UENUM()
 enum class EMacMetalShaderStandard : uint8
 {
-    /** Metal Shader 2.3 is the minimum as of UE5.0*/
-    MacMetalSLStandard_Minimum = 0 UMETA(DisplayName="Minimum, Currently v2.3 (11+)"),
-    /** Metal Shaders Compatible With macOS 11.0 or later (std=osx-metal2.3) */
-    MacMetalSLStandard_2_3 = 6 UMETA(DisplayName="Metal v2.3 (11.0+)"),
+    /** Metal Shader 2.4 is the minimum as of UE5.1*/
+    MacMetalSLStandard_Minimum = 0 UMETA(DisplayName="Minimum, Currently v2.4 (12+)"),
     /** Metal Shaders Compatible With macOS 12.0 or later (std=osx-metal2.4) */
     MacMetalSLStandard_2_4 = 7 UMETA(DisplayName="Metal v2.4 (12.0+)"),
+    /** Metal Shaders Compatible With macOS 13.0 or later (std=metal3.0) */
+    MacMetalSLStandard_3_0 = 8 UMETA(DisplayName="Metal v3.0 (13.0+)"),
+
 };
 
 UENUM()
diff --git a/Engine/Source/Editor/DetailCustomizations/Private/MacTargetSettingsDetails.cpp b/Engine/Source/Editor/DetailCustomizations/Private/MacTargetSettingsDetails.cpp
--- a/Engine/Source/Editor/DetailCustomizations/Private/MacTargetSettingsDetails.cpp
+++ b/Engine/Source/Editor/DetailCustomizations/Private/MacTargetSettingsDetails.cpp
@@ -122,19 +122,19 @@ static FText GetFriendlyNameFromRHINameMac(FName InRHIName)
 		FriendlyRHIName = LOCTEXT("OpenGLES31", "OpenGLES 3.1 (Mobile)");
 		break;
 	case SP_METAL:
-		FriendlyRHIName = LOCTEXT("Metal", "iOS Metal Mobile Renderer (ES3.1, Metal 2.3+, iOS 14.0 or later)");
+		FriendlyRHIName = LOCTEXT("Metal", "iOS Metal Mobile Renderer (ES3.1, Metal 2.4+, iOS 15.0 or later)");
 		break;
 	case SP_METAL_MRT:
-		FriendlyRHIName = LOCTEXT("MetalMRT", "iOS Metal Desktop Renderer (SM5, Metal 2.3+, iOS 14.0 or later)");
+		FriendlyRHIName = LOCTEXT("MetalMRT", "iOS Metal Desktop Renderer (SM5, Metal 2.4+, iOS 15.0 or later)");
 		break;
 	case SP_METAL_TVOS:
-		FriendlyRHIName = LOCTEXT("MetalTV", "tvOS Metal Mobile Renderer (ES3.1, Metal 2.3+, tvOS 14.0 or later)");
+		FriendlyRHIName = LOCTEXT("MetalTV", "tvOS Metal Mobile Renderer (ES3.1, Metal 2.4+, tvOS 15.0 or later)");
 		break;
 	case SP_METAL_MRT_TVOS:
-		FriendlyRHIName = LOCTEXT("MetalMRTTV", "tvOS Metal Desktop Renderer (SM5, Metal 2.3+, tvOS 14.0 or later)");
+		FriendlyRHIName = LOCTEXT("MetalMRTTV", "tvOS Metal Desktop Renderer (SM5, Metal 2.4+, tvOS 15.0 or later)");
 		break;
 	case SP_METAL_SM5:
-		FriendlyRHIName = LOCTEXT("MetalSM5", "Mac Metal Desktop Renderer (SM5, Metal 2.3+, macOS BigSur 11.0 or later)");
+		FriendlyRHIName = LOCTEXT("MetalSM5", "Mac Metal Desktop Renderer (SM5, Metal 2.4+, macOS Monterey 12.0 or later)");
 		break;
 	case SP_METAL_MACES3_1:
 		FriendlyRHIName = LOCTEXT("MetalES3.1", "Mac Metal High-End Mobile Preview (ES3.1)");
@@ -343,8 +343,8 @@ bool FMacTargetSettingsDetails::HandlePostExternalIconCopy(const FString& InChos
 }
 
 static uint32 GMacTargetSettingsMinOSVers[][3] = {
-	{11, 0,0},
-    {12, 0,0}
+	{12,0,0},
+    {13,0,0}
 };
 
 TSharedRef<SWidget> FMacTargetSettingsDetails::OnGetShaderVersionContent()
@@ -393,9 +393,9 @@ void FMacTargetSettingsDetails::SetShaderStandard(int32 Value)
     }
     
     ShaderVersionWarningTextBox->SetError(TEXT(""));
-    if (Value < 6) // EMacMetalShaderStandard::MacMetalSLStandard_Minimum
+    if (Value < 7) // EMacMetalShaderStandard::MacMetalSLStandard_Minimum
     {
-        ShaderVersionWarningTextBox->SetError(TEXT("Metal Shader Standard is 2.3 on UE5.1"));
+        ShaderVersionWarningTextBox->SetError(TEXT("Metal Shader Standard is 2.4 on UE5.1"));
     }
 }
 
@@ -407,9 +407,9 @@ void FMacTargetSettingsDetails::UpdateShaderStandardWarning()
     if (ShaderVersionPropertyHandle->IsValidHandle())
     {
         ShaderVersionPropertyHandle->GetValue(EnumValue);
-        if (EnumValue < 6)
+        if (EnumValue < 7)
         {
-            SetShaderStandard(6); // EMacMetalShaderStandard::MacMetalSLStandard_Minimum
+            SetShaderStandard(7); // EMacMetalShaderStandard::MacMetalSLStandard_Minimum
         }
     }
 }
diff --git a/Engine/Source/Editor/UnrealEd/Private/EditorEngine.cpp b/Engine/Source/Editor/UnrealEd/Private/EditorEngine.cpp
--- a/Engine/Source/Editor/UnrealEd/Private/EditorEngine.cpp
+++ b/Engine/Source/Editor/UnrealEd/Private/EditorEngine.cpp
@@ -6910,7 +6910,7 @@ void UEditorEngine::UpdateAutoLoadProject()
 #if PLATFORM_MAC
 	if ( !GIsBuildMachine )
 	{
-		if(FPlatformMisc::MacOSXVersionCompare(11, 0, 0) < 0)
+		if(FPlatformMisc::MacOSXVersionCompare(12, 0, 0) < 0)
 		{
 			if(FSlateApplication::IsInitialized())
 			{
diff --git a/Engine/Source/Programs/UnrealBuildTool/Platform/IOS/IOSPlatformSDK.Versions.cs b/Engine/Source/Programs/UnrealBuildTool/Platform/IOS/IOSPlatformSDK.Versions.cs
--- a/Engine/Source/Programs/UnrealBuildTool/Platform/IOS/IOSPlatformSDK.Versions.cs
+++ b/Engine/Source/Programs/UnrealBuildTool/Platform/IOS/IOSPlatformSDK.Versions.cs
@@ -14,7 +14,7 @@ namespace UnrealBuildTool
 		protected override void GetValidSoftwareVersionRange(out string MinVersion, out string? MaxVersion)
 		{
 			// what is our min IOS version?
-			MinVersion = "14.0";
+			MinVersion = "15.0";
 			MaxVersion = null;
 		}
 	}
diff --git a/Engine/Source/Programs/UnrealBuildTool/Platform/IOS/IOSToolChain.cs b/Engine/Source/Programs/UnrealBuildTool/Platform/IOS/IOSToolChain.cs
--- a/Engine/Source/Programs/UnrealBuildTool/Platform/IOS/IOSToolChain.cs
+++ b/Engine/Source/Programs/UnrealBuildTool/Platform/IOS/IOSToolChain.cs
@@ -34,7 +34,7 @@ namespace UnrealBuildTool
 		/// The version of the iOS to allow at build time.
 		/// </summary>
 		[XmlConfigFile(Category = "IOSToolChain")]
-		public string BuildIOSVersion = "13.0";
+		public string BuildIOSVersion = "15.0";
 
 		/// <summary>
 		/// Directory for the developer binaries
@@ -722,7 +722,7 @@ namespace UnrealBuildTool
 				Arguments.Append(" --launch-image \"Launch Image\"");
 				Arguments.Append(" --filter-for-device-model AppleTV5,3");
 				Arguments.Append(" --target-device tv");
-				Arguments.Append(" --minimum-deployment-target 14.0");
+				Arguments.Append(" --minimum-deployment-target 15.0");
 				Arguments.Append(" --platform appletvos");
 			}
 			else
@@ -731,7 +731,7 @@ namespace UnrealBuildTool
 				Arguments.Append(" --product-type com.apple.product-type.application");
 				Arguments.Append(" --target-device iphone");
 				Arguments.Append(" --target-device ipad");
-				Arguments.Append(" --minimum-deployment-target 14.0");
+				Arguments.Append(" --minimum-deployment-target 15.0");
 				Arguments.Append(" --platform iphoneos");
 			}
 			Arguments.Append(" --enable-on-demand-resources YES");
diff --git a/Engine/Source/Programs/UnrealBuildTool/Platform/Mac/MacPlatformSDK.Versions.cs b/Engine/Source/Programs/UnrealBuildTool/Platform/Mac/MacPlatformSDK.Versions.cs
--- a/Engine/Source/Programs/UnrealBuildTool/Platform/Mac/MacPlatformSDK.Versions.cs
+++ b/Engine/Source/Programs/UnrealBuildTool/Platform/Mac/MacPlatformSDK.Versions.cs
@@ -14,7 +14,7 @@ namespace UnrealBuildTool
 
 		protected override void GetValidSoftwareVersionRange(out string MinVersion, out string? MaxVersion)
 		{
-			MinVersion = "11.0.0";
+			MinVersion = "12.0.0";
 			MaxVersion = null;
 		}
 	}
diff --git a/Engine/Source/Runtime/Apple/MetalRHI/Private/MetalRHI.cpp b/Engine/Source/Runtime/Apple/MetalRHI/Private/MetalRHI.cpp
--- a/Engine/Source/Runtime/Apple/MetalRHI/Private/MetalRHI.cpp
+++ b/Engine/Source/Runtime/Apple/MetalRHI/Private/MetalRHI.cpp
@@ -637,14 +637,6 @@ FMetalDynamicRHI::FMetalDynamicRHI(ERHIFeatureLevel::Type RequestedFeatureLevel)
 	GMaxTextureMipCount = FPlatformMath::CeilLogTwo( GMaxTextureDimensions ) + 1;
 	GMaxTextureMipCount = FPlatformMath::Min<int32>( MAX_TEXTURE_MIP_COUNT, GMaxTextureMipCount );
 
-	// Using packed pixel formats (specifically MTLPixelFormatRG11B10Float) for texture_buffers on macOS <= 11.4 results in Metal validation errors:
-	// "pixel format (MTLPixelFormatRG11B10Float) cannot be written to from a shader on this device"
-	// macOS Big Sur 11.5.1 works; Big Sur 11.4.0 (and prior versions including Catalina) report the error; texture2d with buffer backing does not report this error.
-	bool bSupportsPackedPixelFormatsInTextureBuffers = true;
-#if PLATFORM_MAC
-	bSupportsPackedPixelFormatsInTextureBuffers = FPlatformMisc::MacOSXVersionCompare(11,5,1) >= 0;
-#endif
-
 	// Initialize the buffer format map - in such a way as to be able to validate it in non-shipping...
 #if METAL_DEBUG_OPTIONS
 	FMemory::Memset(GMetalBufferFormats, 255);
@@ -658,14 +650,7 @@ FMetalDynamicRHI::FMetalDynamicRHI(ERHIFeatureLevel::Type RequestedFeatureLevel)
 	GMetalBufferFormats[PF_DXT3                 ] = { mtlpp::PixelFormat::Invalid, (uint8)EMetalBufferFormat::Unknown };
 	GMetalBufferFormats[PF_DXT5                 ] = { mtlpp::PixelFormat::Invalid, (uint8)EMetalBufferFormat::Unknown };
 	GMetalBufferFormats[PF_UYVY                 ] = { mtlpp::PixelFormat::Invalid, (uint8)EMetalBufferFormat::Unknown };
-	if(bSupportsPackedPixelFormatsInTextureBuffers)
-	{
-		GMetalBufferFormats[PF_FloatRGB        	] = { mtlpp::PixelFormat::RG11B10Float, (uint8)EMetalBufferFormat::RG11B10Half };
-	}
-	else
-	{
-		GMetalBufferFormats[PF_FloatRGB        	] = { mtlpp::PixelFormat::RGBA16Float, (uint8)EMetalBufferFormat::RGBA16Half };
-	}
+	GMetalBufferFormats[PF_FloatRGB        	] = { mtlpp::PixelFormat::RG11B10Float, (uint8)EMetalBufferFormat::RG11B10Half };
 	GMetalBufferFormats[PF_FloatRGBA            ] = { mtlpp::PixelFormat::RGBA16Float, (uint8)EMetalBufferFormat::RGBA16Half };
 	GMetalBufferFormats[PF_DepthStencil         ] = { mtlpp::PixelFormat::Invalid, (uint8)EMetalBufferFormat::Unknown };
 	GMetalBufferFormats[PF_ShadowDepth          ] = { mtlpp::PixelFormat::Invalid, (uint8)EMetalBufferFormat::Unknown };
@@ -855,16 +840,9 @@ FMetalDynamicRHI::FMetalDynamicRHI(ERHIFeatureLevel::Type RequestedFeatureLevel)
     GPixelFormats[PF_DXT3				].PlatformFormat	= (uint32)mtlpp::PixelFormat::BC2_RGBA;
     GPixelFormats[PF_DXT5				].PlatformFormat	= (uint32)mtlpp::PixelFormat::BC3_RGBA;
 	
-	if(bSupportsPackedPixelFormatsInTextureBuffers)
-	{
-		GPixelFormats[PF_FloatRGB		].PlatformFormat	= (uint32)mtlpp::PixelFormat::RG11B10Float;
-		GPixelFormats[PF_FloatRGB		].BlockBytes		= 4;
-	}
-	else
-	{
-		GPixelFormats[PF_FloatRGB		].PlatformFormat	= (uint32)mtlpp::PixelFormat::RGBA16Float;
-		GPixelFormats[PF_FloatRGB		].BlockBytes		= 8;
-	}
+    GPixelFormats[PF_FloatRGB		].PlatformFormat	= (uint32)mtlpp::PixelFormat::RG11B10Float;
+    GPixelFormats[PF_FloatRGB		].BlockBytes		= 4;
+
 	
 	GPixelFormats[PF_FloatR11G11B10		].PlatformFormat	= (uint32)mtlpp::PixelFormat::RG11B10Float;
 	GPixelFormats[PF_FloatR11G11B10		].BlockBytes		= 4;
diff --git a/Engine/Source/Runtime/Apple/MetalRHI/Private/MetalShaders.cpp b/Engine/Source/Runtime/Apple/MetalRHI/Private/MetalShaders.cpp
--- a/Engine/Source/Runtime/Apple/MetalRHI/Private/MetalShaders.cpp
+++ b/Engine/Source/Runtime/Apple/MetalRHI/Private/MetalShaders.cpp
@@ -47,30 +47,30 @@ NSString* DecodeMetalSourceCode(uint32 CodeSize, TArray<uint8> const& Compressed
 
 mtlpp::LanguageVersion ValidateVersion(uint32 Version)
 {
-    mtlpp::LanguageVersion Result = mtlpp::LanguageVersion::Version2_2;
+    mtlpp::LanguageVersion Result = mtlpp::LanguageVersion::Version2_4;
 #if PLATFORM_MAC
-    Result = mtlpp::LanguageVersion::Version2_2;
+    Result = mtlpp::LanguageVersion::Version2_4;
     switch(Version)
     {
+        case 8:
+            Result = mtlpp::LanguageVersion::Version3_0;
+            break;
         case 7:
             Result = mtlpp::LanguageVersion::Version2_4;
             break;
-        case 6:
-            Result = mtlpp::LanguageVersion::Version2_3;
-            break;
         case 0:
-            Version = 6;
-            Result = mtlpp::LanguageVersion::Version2_3; // minimum version as of UE5.1
+            Version = 7;
+            Result = mtlpp::LanguageVersion::Version2_4; // minimum version as of UE5.1
             break;
         default:
-            //EMacMetalShaderStandard::MacMetalSLStandard_Minimum is currently 2.3
+            //EMacMetalShaderStandard::MacMetalSLStandard_Minimum is currently 2.4
             UE_LOG(LogTemp, Warning, TEXT("The Metal version currently set is not supported anymore. Set it in the Project Settings. Defaulting to the minimum version."));
-            Version = 6;
-            Result = mtlpp::LanguageVersion::Version2_3;
+            Version = 7;
+            Result = mtlpp::LanguageVersion::Version2_4;
             break;
     }
 #else
-    Result = mtlpp::LanguageVersion::Version2_3;
+    Result = mtlpp::LanguageVersion::Version2_4;
     switch(Version)
     {
         case 8:
@@ -79,18 +79,15 @@ mtlpp::LanguageVersion ValidateVersion(uint32 Version)
         case 7:
             Result = mtlpp::LanguageVersion::Version2_4;
             break;
-        case 6:
-            Result = mtlpp::LanguageVersion::Version2_3;
-            break;
         case 0:
-            Version = 6;
-            Result = mtlpp::LanguageVersion::Version2_3; // minimum version as of UE5.0
+            Version = 7;
+            Result = mtlpp::LanguageVersion::Version2_4; // minimum version as of UE5.1
             break;
         default:
-            //EMacMetalShaderStandard::MacMetalSLStandard_Minimum and EIOSMetalShaderStandard::IOSMetalSLStandard_Minimum is currently 2.3
+            //EMacMetalShaderStandard::MacMetalSLStandard_Minimum and EIOSMetalShaderStandard::IOSMetalSLStandard_Minimum is currently 2.4
             UE_LOG(LogTemp, Warning, TEXT("The Metal version currently set is not supported anymore. Set it in the Project Settings. Defaulting to the minimum version."));
-            Version = 6;
-            Result = mtlpp::LanguageVersion::Version2_3;
+            Version = 7;
+            Result = mtlpp::LanguageVersion::Version2_4;
             break;
     }
 #endif
diff --git a/Engine/Source/Runtime/Apple/MetalRHI/Private/Shaders/Types/Templates/MetalBaseShader.h b/Engine/Source/Runtime/Apple/MetalRHI/Private/Shaders/Types/Templates/MetalBaseShader.h
--- a/Engine/Source/Runtime/Apple/MetalRHI/Private/Shaders/Types/Templates/MetalBaseShader.h
+++ b/Engine/Source/Runtime/Apple/MetalRHI/Private/Shaders/Types/Templates/MetalBaseShader.h
@@ -303,18 +303,15 @@ void TMetalBaseShader<BaseResourceType, ShaderType>::Init(TArrayView<const uint8
 			mtlpp::LanguageVersion MetalVersion;
 			switch (Header.Version)
 			{
+                case 8:
+                    MetalVersion = mtlpp::LanguageVersion::Version3_0;
+                    break;
 				case 7:
 					MetalVersion = mtlpp::LanguageVersion::Version2_4;
 					break;
-				case 6:
-					MetalVersion = mtlpp::LanguageVersion::Version2_3;
-					break;
-				case 5:
-					MetalVersion = mtlpp::LanguageVersion::Version2_2;
-					break;
 				default:
 					UE_LOG(LogRHI, Fatal, TEXT("Failed to create shader with unknown version %d: %s"), Header.Version, *FString(NewShaderString));
-					MetalVersion = mtlpp::LanguageVersion::Version2_2;
+					MetalVersion = mtlpp::LanguageVersion::Version2_4;
 					break;
 			}
 			CompileOptions.SetLanguageVersion(MetalVersion);
diff --git a/Engine/Source/Runtime/ApplicationCore/Private/Apple/AppleControllerInterface.cpp b/Engine/Source/Runtime/ApplicationCore/Private/Apple/AppleControllerInterface.cpp
--- a/Engine/Source/Runtime/ApplicationCore/Private/Apple/AppleControllerInterface.cpp
+++ b/Engine/Source/Runtime/ApplicationCore/Private/Apple/AppleControllerInterface.cpp
@@ -176,7 +176,7 @@ void FAppleControllerInterface::HandleConnection(GCController* Controller)
         Controllers[ControllerIndex].Controller = [Controller retain];
         SetControllerType(ControllerIndex);
         
-        // Deprecated but buttonMenu behavior is unreliable in iOS/tvOS 14
+        // Deprecated but buttonMenu behavior is unreliable since iOS/tvOS 14
 		Controllers[ControllerIndex].bPauseWasPressed = false;
 PRAGMA_DISABLE_DEPRECATION_WARNINGS
 		Controller.controllerPausedHandler = ^(GCController* Cont)
@@ -245,11 +245,10 @@ void FAppleControllerInterface::SendControllerEvents()
 		FInputDeviceId DeviceId = INPUTDEVICEID_NONE;
 		DeviceMapper.RemapControllerIdToPlatformUserAndDevice(Controller.PlayerIndex, OUT UserId, OUT DeviceId);
 		
-		// Assumes iOS13, tvOS 13 & macOS 10.15
         GCExtendedGamepad* ExtendedGamepad = [ControllerImpl capture].extendedGamepad;
 		GCMotion* Motion = ControllerImpl.motion;
 		
-		// Workaround for unreliable buttonMenu behavior in iOS/tvOS 14
+		// Workaround for unreliable buttonMenu behavior since iOS/tvOS 14
 		if (Controller.bPauseWasPressed)
         {
             MessageHandler->OnControllerButtonPressed(FGamepadKeyNames::SpecialRight, UserId, DeviceId, false);
diff --git a/Engine/Source/Runtime/ApplicationCore/Private/Apple/AppleControllerInterface.h b/Engine/Source/Runtime/ApplicationCore/Private/Apple/AppleControllerInterface.h
--- a/Engine/Source/Runtime/ApplicationCore/Private/Apple/AppleControllerInterface.h
+++ b/Engine/Source/Runtime/ApplicationCore/Private/Apple/AppleControllerInterface.h
@@ -138,7 +138,7 @@ protected:
 		bool bHasReferenceAttitude;
 #endif
 
-		// Workaround for unreliable buttonMenu behavior in iOS/tvOS 14
+		// Workaround for unreliable buttonMenu behavior since iOS/tvOS 14
         bool bPauseWasPressed;
 	};
 	
diff --git a/Engine/Source/Runtime/Core/Private/Mac/MacPlatformProcess.cpp b/Engine/Source/Runtime/Core/Private/Mac/MacPlatformProcess.cpp
--- a/Engine/Source/Runtime/Core/Private/Mac/MacPlatformProcess.cpp
+++ b/Engine/Source/Runtime/Core/Private/Mac/MacPlatformProcess.cpp
@@ -527,10 +527,7 @@ FProcHandle FMacPlatformProcess::CreateProcInternal(const TCHAR* URL, const TCHA
 
 	if (OptionalWorkingDirectory)
 	{
-		if (@available(macOS 10.15, *))
-		{
-			posix_spawn_file_actions_addchdir_np(&FileActions, TCHAR_TO_UTF8(OptionalWorkingDirectory));
-		}
+		posix_spawn_file_actions_addchdir_np(&FileActions, TCHAR_TO_UTF8(OptionalWorkingDirectory));
 	}
 
 	posix_spawnattr_setflags(&SpawnAttr, SpawnFlags);
diff --git a/Engine/Source/Runtime/IOS/IOSRuntimeSettings/Classes/IOSRuntimeSettings.h b/Engine/Source/Runtime/IOS/IOSRuntimeSettings/Classes/IOSRuntimeSettings.h
--- a/Engine/Source/Runtime/IOS/IOSRuntimeSettings/Classes/IOSRuntimeSettings.h
+++ b/Engine/Source/Runtime/IOS/IOSRuntimeSettings/Classes/IOSRuntimeSettings.h
@@ -30,9 +30,6 @@ enum class EPowerUsageFrameRateLock : uint8
 UENUM()
 	enum class EIOSVersion : uint8
 {
-    /** iOS 14 */
-    IOS_14 = 14 UMETA(DisplayName = "14.0"),
-    
     /** iOS 15 */
     IOS_15 = 15 UMETA(DisplayName = "15.0"),
     
@@ -44,9 +41,7 @@ UENUM()
 enum class EIOSMetalShaderStandard : uint8
 {
     /** Metal Shader 2.3 is the minimum as of UE5.0*/
-    IOSMetalSLStandard_Minimum = 0 UMETA(DisplayName="Minimum, Currently v2.3 (iOS 14.0/tvOS 14.0)"),
-    /** Metal Shaders Compatible With iOS 14.0/tvOS 14.0 or later (std=ios-metal2.3) */
-    IOSMetalSLStandard_2_3 = 6 UMETA(DisplayName="Metal v2.3 (iOS 14.0/tvOS 14.0)"),
+    IOSMetalSLStandard_Minimum = 0 UMETA(DisplayName="Minimum, Currently v2.4 (iOS 15.0/tvOS 15.0)"),
     /** Metal Shaders Compatible With iOS 15.0/tvOS 15.0 or later (std=ios-metal2.4) */
     IOSMetalSLStandard_2_4 = 7 UMETA(DisplayName="Metal v2.4 (iOS 15.0/tvOS 15.0)"),
     /** Metal Shaders Compatible With iOS 16.0/tvOS 16.0 or later (std=metal3.0) */
diff --git a/Engine/Source/Runtime/IOS/IOSRuntimeSettings/Private/IOSRuntimeSettings.cpp b/Engine/Source/Runtime/IOS/IOSRuntimeSettings/Private/IOSRuntimeSettings.cpp
--- a/Engine/Source/Runtime/IOS/IOSRuntimeSettings/Private/IOSRuntimeSettings.cpp
+++ b/Engine/Source/Runtime/IOS/IOSRuntimeSettings/Private/IOSRuntimeSettings.cpp
@@ -31,7 +31,7 @@ UIOSRuntimeSettings::UIOSRuntimeSettings(const FObjectInitializer& ObjectInitial
 	bEnableDynamicMaxFPS = false;
 	bSupportsIPad = true;
 	bSupportsIPhone = true;
-	MinimumiOSVersion = EIOSVersion::IOS_14;
+	MinimumiOSVersion = EIOSVersion::IOS_15;
     bBuildAsFramework = true;
 	bGeneratedSYMFile = false;
 	bGeneratedSYMBundle = false;
@@ -126,9 +126,9 @@ void UIOSRuntimeSettings::PostInitProperties()
 		}
 	}
 
-	if (MinimumiOSVersion < EIOSVersion::IOS_14)
+	if (MinimumiOSVersion < EIOSVersion::IOS_15)
 	{
-		MinimumiOSVersion = EIOSVersion::IOS_14;
+		MinimumiOSVersion = EIOSVersion::IOS_15;
 		UpdateSinglePropertyInConfigFile(GetClass()->FindPropertyByName(GET_MEMBER_NAME_CHECKED(UIOSRuntimeSettings, MinimumiOSVersion)), GetDefaultConfigFilename());
 	}
 	if (!bSupportsMetal && !bSupportsMetalMRT)
diff --git a/Engine/Source/Runtime/Launch/Private/Mac/LaunchMac.cpp b/Engine/Source/Runtime/Launch/Private/Mac/LaunchMac.cpp
--- a/Engine/Source/Runtime/Launch/Private/Mac/LaunchMac.cpp
+++ b/Engine/Source/Runtime/Launch/Private/Mac/LaunchMac.cpp
@@ -266,8 +266,8 @@ static int32 MacOSVersionCompare(const NSOperatingSystemVersion& VersionA, const
 	NSString* MinimumSystemVersionString = (NSString*)InfoDictionary[@"LSMinimumSystemVersion"];
 	NSOperatingSystemVersion MinimumSystemVersion = { 0 };
 	NSOperatingSystemVersion CurrentSystemVersion = [[NSProcessInfo processInfo] operatingSystemVersion];
-	NSOperatingSystemVersion MinSupportedMacOSVersion = { 11, 0, 0 };
-	NSString* MinSupportedMacOSVersionString = @"11.0.0";
+	NSOperatingSystemVersion MinSupportedMacOSVersion = { 12, 0, 0 };
+	NSString* MinSupportedMacOSVersionString = @"12.0.0";
 
 	NSArray<NSString*>* VersionComponents = [MinimumSystemVersionString componentsSeparatedByString:@"."];
 	MinimumSystemVersion.majorVersion = [[VersionComponents objectAtIndex:0] integerValue];
diff --git a/Engine/Source/Runtime/Online/HTTP/Private/Apple/AppleHTTP.cpp b/Engine/Source/Runtime/Online/HTTP/Private/Apple/AppleHTTP.cpp
--- a/Engine/Source/Runtime/Online/HTTP/Private/Apple/AppleHTTP.cpp
+++ b/Engine/Source/Runtime/Online/HTTP/Private/Apple/AppleHTTP.cpp
@@ -713,19 +713,7 @@ static const unsigned char ecdsaSecp384r1Asn1Header[] =
                 SecTrustEvaluateWithError(CertTrust, nil);
                 TCFRef<SecKeyRef> CertPubKey;
 
-				if (@available(macOS 11, *))
-				{
-					CertPubKey = SecTrustCopyKey(CertTrust);
-				}
-				else
-				{
-					PRAGMA_DISABLE_DEPRECATION_WARNINGS
-					// warning: 'SecTrustCopyPublicKey' is deprecated: first deprecated in iOS 14.0 [-Wdeprecated-declarations]
-
-					CertPubKey = SecTrustCopyPublicKey(CertTrust);
-
-					PRAGMA_ENABLE_DEPRECATION_WARNINGS
-				}
+				CertPubKey = SecTrustCopyKey(CertTrust);
 
 				TCFRef<CFDataRef> CertPubKeyData = SecKeyCopyExternalRepresentation(CertPubKey, NULL);
                 if (!CertPubKeyData)
diff --git a/Engine/Source/Runtime/RHI/Private/DynamicRHI.cpp b/Engine/Source/Runtime/RHI/Private/DynamicRHI.cpp
--- a/Engine/Source/Runtime/RHI/Private/DynamicRHI.cpp
+++ b/Engine/Source/Runtime/RHI/Private/DynamicRHI.cpp
@@ -277,22 +277,19 @@ static void RHIDetectAndWarnOfBadDrivers(bool bHasEditorToken)
 		return;
 	}
 
-	if (FPlatformMisc::MacOSXVersionCompare(11, 0 ,0) < 0)
-	{
-		// this message can be suppressed with r.WarnOfBadDrivers=0
-		FPlatformMisc::MessageBoxExt(EAppMsgType::Ok,
-									 *NSLOCTEXT("MessageDialog", "UpdateMacOSX_Body", "Please update to the latest version of macOS for best performance and stability.").ToString(),
-									 *NSLOCTEXT("MessageDialog", "UpdateMacOSX_Title", "Update macOS").ToString());
+	// this message can be suppressed with r.WarnOfBadDrivers=0
+	FPlatformMisc::MessageBoxExt(EAppMsgType::Ok,
+								 *NSLOCTEXT("MessageDialog", "UpdateMacOSX_Body", "Please update to the latest version of macOS for best performance and stability.").ToString(),
+								 *NSLOCTEXT("MessageDialog", "UpdateMacOSX_Title", "Update macOS").ToString());
 		
 #if !UE_BUILD_SHIPPING
-		if (CVarBadDriverWarningIsFatal.GetValueOnGameThread())
-		{
-			// Force a fatal error depending on CVar
-			UE_LOG(LogRHI, Fatal, TEXT("Fatal crash requested when graphics drivers are out of date.\n")
-				TEXT("To prevent this crash, please update macOS."));
-		}
-#endif
+	if (CVarBadDriverWarningIsFatal.GetValueOnGameThread())
+	{
+		// Force a fatal error depending on CVar
+		UE_LOG(LogRHI, Fatal, TEXT("Fatal crash requested when graphics drivers are out of date.\n")
+			TEXT("To prevent this crash, please update macOS."));
 	}
+#endif
 }
 #endif // PLATFORM_WINDOWS
 
diff --git a/Templates/TP_AEC_HandheldARBP/Config/DefaultEngine.ini b/Templates/TP_AEC_HandheldARBP/Config/DefaultEngine.ini
--- a/Templates/TP_AEC_HandheldARBP/Config/DefaultEngine.ini
+++ b/Templates/TP_AEC_HandheldARBP/Config/DefaultEngine.ini
@@ -109,7 +109,6 @@ MobileProvision=iOS11.mobileprovision
 bSupportsPortraitOrientation=True
 bSupportsUpsideDownOrientation=False
 FrameRateLock=PUFRL_60
-MinimumiOSVersion=IOS_14
 
 [/Script/Engine.Engine]
 bUseFixedFrameRate=True
diff --git a/Templates/TP_FirstPerson/Config/DefaultEngine.ini b/Templates/TP_FirstPerson/Config/DefaultEngine.ini
--- a/Templates/TP_FirstPerson/Config/DefaultEngine.ini
+++ b/Templates/TP_FirstPerson/Config/DefaultEngine.ini
@@ -16,9 +16,6 @@ ServerDefaultMap=/Engine/Maps/Entry
 GlobalDefaultGameMode=/Script/TP_FirstPerson.TP_FirstPersonGameMode
 GlobalDefaultServerGameMode=None
 
-[/Script/IOSRuntimeSettings.IOSRuntimeSettings]
-MinimumiOSVersion=IOS_14
-
 [/Script/Engine.RendererSettings]
 r.ReflectionMethod=1
 r.GenerateMeshDistanceFields=True
diff --git a/Templates/TP_FirstPersonBP/Config/DefaultEngine.ini b/Templates/TP_FirstPersonBP/Config/DefaultEngine.ini
--- a/Templates/TP_FirstPersonBP/Config/DefaultEngine.ini
+++ b/Templates/TP_FirstPersonBP/Config/DefaultEngine.ini
@@ -16,9 +16,6 @@ ServerDefaultMap=/Engine/Maps/Entry
 GlobalDefaultGameMode=/Game/FirstPerson/Blueprints/BP_FirstPersonGameMode.BP_FirstPersonGameMode_C
 GlobalDefaultServerGameMode=None
 
-[/Script/IOSRuntimeSettings.IOSRuntimeSettings]
-MinimumiOSVersion=IOS_14
-
 [/Script/Engine.RendererSettings]
 r.ReflectionMethod=1
 r.GenerateMeshDistanceFields=True
diff --git a/Templates/TP_HandheldARBP/Config/DefaultEngine.ini b/Templates/TP_HandheldARBP/Config/DefaultEngine.ini
--- a/Templates/TP_HandheldARBP/Config/DefaultEngine.ini
+++ b/Templates/TP_HandheldARBP/Config/DefaultEngine.ini
@@ -107,7 +107,6 @@ MobileProvision=iOS11.mobileprovision
 bSupportsPortraitOrientation=True
 bSupportsUpsideDownOrientation=False
 FrameRateLock=PUFRL_60
-MinimumiOSVersion=IOS_14
 
 [/Script/Engine.Engine]
 bUseFixedFrameRate=True
diff --git a/Templates/TP_MFG_HandheldARBP/Config/DefaultEngine.ini b/Templates/TP_MFG_HandheldARBP/Config/DefaultEngine.ini
--- a/Templates/TP_MFG_HandheldARBP/Config/DefaultEngine.ini
+++ b/Templates/TP_MFG_HandheldARBP/Config/DefaultEngine.ini
@@ -104,12 +104,9 @@ DefaultGraphicsPerformance=Maximum
 AppliedDefaultGraphicsPerformance=Maximum
 
 [/Script/IOSRuntimeSettings.IOSRuntimeSettings]
-SigningCertificate=iPhone Developer: Chance Ivey (AP6BVS76S9)
-MobileProvision=iOS11.mobileprovision
 bSupportsPortraitOrientation=True
 bSupportsUpsideDownOrientation=False
 FrameRateLock=PUFRL_60
-MinimumiOSVersion=IOS_14
 
 [/Script/Engine.Engine]
 bUseFixedFrameRate=True
diff --git a/Templates/TP_Puzzle/Config/DefaultEngine.ini b/Templates/TP_Puzzle/Config/DefaultEngine.ini
--- a/Templates/TP_Puzzle/Config/DefaultEngine.ini
+++ b/Templates/TP_Puzzle/Config/DefaultEngine.ini
@@ -10,8 +10,4 @@ GameInstanceClass=/Script/Engine.GameInstance
 GameDefaultMap=/Game/PuzzleCPP/Maps/PuzzleExampleMap
 ServerDefaultMap=/Engine/Maps/Entry
 GlobalDefaultGameMode=/Script/TP_Puzzle.TP_PuzzleGameMode
-GlobalDefaultServerGameMode=None
-
-[/Script/IOSRuntimeSettings.IOSRuntimeSettings]
-MinimumiOSVersion=IOS_14
-
+GlobalDefaultServerGameMode=None
\ No newline at end of file
diff --git a/Templates/TP_PuzzleBP/Config/DefaultEngine.ini b/Templates/TP_PuzzleBP/Config/DefaultEngine.ini
--- a/Templates/TP_PuzzleBP/Config/DefaultEngine.ini
+++ b/Templates/TP_PuzzleBP/Config/DefaultEngine.ini
@@ -9,7 +9,4 @@ bUseSplitscreen=True
 TwoPlayerSplitscreenLayout=Horizontal
 ThreePlayerSplitscreenLayout=FavorTop
 GlobalDefaultGameMode=/Game/PuzzleBP/Blueprints/PuzzleGameMode.PuzzleGameMode_C
-GlobalDefaultServerGameMode=None
-
-[/Script/IOSRuntimeSettings.IOSRuntimeSettings]
-MinimumiOSVersion=IOS_14
+GlobalDefaultServerGameMode=None
\ No newline at end of file
diff --git a/Templates/TP_ThirdPerson/Config/DefaultEngine.ini b/Templates/TP_ThirdPerson/Config/DefaultEngine.ini
--- a/Templates/TP_ThirdPerson/Config/DefaultEngine.ini
+++ b/Templates/TP_ThirdPerson/Config/DefaultEngine.ini
@@ -3,9 +3,6 @@ GameDefaultMap=/Game/ThirdPerson/Maps/ThirdPersonMap.ThirdPersonMap
 EditorStartupMap=/Game/ThirdPerson/Maps/ThirdPersonMap.ThirdPersonMap
 GlobalDefaultGameMode="/Script/TP_ThirdPerson.TP_ThirdPersonGameMode"
 
-[/Script/IOSRuntimeSettings.IOSRuntimeSettings]
-MinimumiOSVersion=IOS_14
-
 [/Script/Engine.RendererSettings]
 r.Shadow.Virtual.Enable=1
 r.Mobile.EnableNoPrecomputedLightingCSMShader=1
diff --git a/Templates/TP_ThirdPersonBP/Config/DefaultEngine.ini b/Templates/TP_ThirdPersonBP/Config/DefaultEngine.ini
--- a/Templates/TP_ThirdPersonBP/Config/DefaultEngine.ini
+++ b/Templates/TP_ThirdPersonBP/Config/DefaultEngine.ini
@@ -11,9 +11,6 @@ ThreePlayerSplitscreenLayout=FavorTop
 GlobalDefaultGameMode=/Game/ThirdPerson/Blueprints/BP_ThirdPersonGameMode.BP_ThirdPersonGameMode_C
 GlobalDefaultServerGameMode=None
 
-[/Script/IOSRuntimeSettings.IOSRuntimeSettings]
-MinimumiOSVersion=IOS_14
-
 [/Script/Engine.RendererSettings]
 r.ReflectionMethod=1
 r.GenerateMeshDistanceFields=True
diff --git a/Templates/TP_TopDown/Config/DefaultEngine.ini b/Templates/TP_TopDown/Config/DefaultEngine.ini
--- a/Templates/TP_TopDown/Config/DefaultEngine.ini
+++ b/Templates/TP_TopDown/Config/DefaultEngine.ini
@@ -3,8 +3,5 @@ GameDefaultMap=/Game/TopDown/Maps/TopDownMap.TopDownMap
 EditorStartupMap=/Game/TopDown/Maps/TopDownMap.TopDownMap
 GlobalDefaultGameMode="/Script/TP_TopDown.TP_TopDownGameMode"
 
-[/Script/IOSRuntimeSettings.IOSRuntimeSettings]
-MinimumiOSVersion=IOS_14
-
 [/Script/Engine.RendererSettings]
 r.Mobile.EnableNoPrecomputedLightingCSMShader=1
\ No newline at end of file
diff --git a/Templates/TP_TopDownBP/Config/DefaultEngine.ini b/Templates/TP_TopDownBP/Config/DefaultEngine.ini
--- a/Templates/TP_TopDownBP/Config/DefaultEngine.ini
+++ b/Templates/TP_TopDownBP/Config/DefaultEngine.ini
@@ -49,9 +49,6 @@ HeuristicScale=0.999000
 VerticalDeviationFromGroundCompensation=0.000000
 bForceRebuildOnLoad=True
 
-[/Script/IOSRuntimeSettings.IOSRuntimeSettings]
-MinimumiOSVersion=IOS_14
-
 [/Script/Engine.RendererSettings]
 r.ReflectionMethod=1
 r.GenerateMeshDistanceFields=True
diff --git a/Templates/TP_VehicleAdv/Config/DefaultEngine.ini b/Templates/TP_VehicleAdv/Config/DefaultEngine.ini
--- a/Templates/TP_VehicleAdv/Config/DefaultEngine.ini
+++ b/Templates/TP_VehicleAdv/Config/DefaultEngine.ini
@@ -15,9 +15,6 @@ ServerDefaultMap=/Engine/Maps/Entry
 GlobalDefaultGameMode=/Script/TP_VehicleAdv.TP_VehicleAdvGameMode
 GlobalDefaultServerGameMode=None
 
-[/Script/IOSRuntimeSettings.IOSRuntimeSettings]
-MinimumiOSVersion=IOS_14
-
 [/Script/AndroidFileServerEditor.AndroidFileServerRuntimeSettings]
 bEnablePlugin=True
 bAllowNetworkConnection=True
diff --git a/Templates/TP_VehicleAdvBP/Config/DefaultEngine.ini b/Templates/TP_VehicleAdvBP/Config/DefaultEngine.ini
--- a/Templates/TP_VehicleAdvBP/Config/DefaultEngine.ini
+++ b/Templates/TP_VehicleAdvBP/Config/DefaultEngine.ini
@@ -18,9 +18,6 @@ ServerDefaultMap=/Engine/Maps/Entry.Entry
 GlobalDefaultGameMode=/Game/VehicleTemplate/Blueprints/VehicleAdvGameMode.VehicleAdvGameMode_C
 GlobalDefaultServerGameMode=None
 
-[/Script/IOSRuntimeSettings.IOSRuntimeSettings]
-MinimumiOSVersion=IOS_14
-
 [/Script/AndroidFileServerEditor.AndroidFileServerRuntimeSettings]
 bEnablePlugin=True
 bAllowNetworkConnection=True
