Fix bug where SceneColor\DepthCreateFlags were not updated after initializing variables
Changed iOS to use Memoryless flag to match other platforms when creating MSAA targets

diff --git a/Engine/Source/Runtime/Apple/MetalRHI/Private/MetalTexture.cpp b/Engine/Source/Runtime/Apple/MetalRHI/Private/MetalTexture.cpp
--- a/Engine/Source/Runtime/Apple/MetalRHI/Private/MetalTexture.cpp
+++ b/Engine/Source/Runtime/Apple/MetalRHI/Private/MetalTexture.cpp
@@ -583,6 +583,8 @@ FMetalSurface::FMetalSurface(FMetalTextureCreateDesc const& CreateDesc)
 	{
 		return;
 	}
+    
+    bool bIsMSAARequired = CreateDesc.NumSamples > 1 && !FParse::Param(FCommandLine::Get(), TEXT("nomsaa"));
 
 	FResourceBulkDataInterface* BulkData = CreateDesc.BulkData;
 
@@ -639,25 +641,46 @@ FMetalSurface::FMetalSurface(FMetalTextureCreateDesc const& CreateDesc)
 		const bool bBufferCompatibleOption = 	(CreateDesc.Desc.GetTextureType() == mtlpp::TextureType::Texture2D || CreateDesc.Desc.GetTextureType() == mtlpp::TextureType::TextureBuffer) &&
 												CreateDesc.NumMips == 1 && CreateDesc.ArraySize == 1 && CreateDesc.NumSamples == 1 && CreateDesc.Desc.GetDepth() == 1;
 
+        FMetalTextureCreateDesc NewCreateDesc = CreateDesc;
+        
+#if PLATFORM_IOS
+        // If we are attempting to create an MSAA texture the texture cannot be memoryless unless we are creating a depth texture
+        if(bIsMSAARequired && CreateDesc.Format != PF_DepthStencil && EnumHasAllFlags(CreateDesc.Flags, TexCreate_Memoryless))
+        {
+            NewCreateDesc.Flags &= ~TexCreate_Memoryless;
+            
+            if (GMetalForceIOSTexturesShared)
+            {
+                NewCreateDesc.Desc.SetStorageMode(mtlpp::StorageMode::Shared);
+                NewCreateDesc.Desc.SetResourceOptions((mtlpp::ResourceOptions)(mtlpp::ResourceOptions::CpuCacheModeDefaultCache | mtlpp::ResourceOptions::StorageModeShared));
+            }
+            else
+            {
+                NewCreateDesc.Desc.SetStorageMode(mtlpp::StorageMode::Private);
+                NewCreateDesc.Desc.SetResourceOptions((mtlpp::ResourceOptions)(mtlpp::ResourceOptions::CpuCacheModeDefaultCache | mtlpp::ResourceOptions::StorageModePrivate));
+            }
+        }
+#endif
+        
 		if(bBufferCompatibleOption && (EnumHasAllFlags(CreateDesc.Flags, TexCreate_UAV | TexCreate_NoTiling) || EnumHasAllFlags(CreateDesc.Flags, TexCreate_AtomicCompatible)))
 		{
 			mtlpp::Device Device = GetMetalDeviceContext().GetDevice();
 
 			const uint32 MinimumByteAlignment = Device.GetMinimumLinearTextureAlignmentForPixelFormat(CreateDesc.MTLFormat);
-			const NSUInteger BytesPerRow = Align(CreateDesc.Desc.GetWidth() * GPixelFormats[CreateDesc.Format].BlockBytes, MinimumByteAlignment);
+			const NSUInteger BytesPerRow = Align(NewCreateDesc.Desc.GetWidth() * GPixelFormats[NewCreateDesc.Format].BlockBytes, MinimumByteAlignment);
 
 			// Backing buffer resource options must match the texture we are going to create from it
-			FMetalPooledBufferArgs Args(Device, BytesPerRow * CreateDesc.Desc.GetHeight(), BUF_Dynamic, mtlpp::StorageMode::Private, CreateDesc.Desc.GetCpuCacheMode());
+			FMetalPooledBufferArgs Args(Device, BytesPerRow * NewCreateDesc.Desc.GetHeight(), BUF_Dynamic, mtlpp::StorageMode::Private, NewCreateDesc.Desc.GetCpuCacheMode());
 			FMetalBuffer Buffer = GetMetalDeviceContext().CreatePooledBuffer(Args);
 
-			Texture = Buffer.NewTexture(CreateDesc.Desc, 0, BytesPerRow);
+			Texture = Buffer.NewTexture(NewCreateDesc.Desc, 0, BytesPerRow);
 		}
 		else
 		{
 			// If we are in here then either the texture description is not buffer compatable or these flags were not set
 			// assert that these flag combinations are not set as they require a buffer backed texture and the texture description is not compatible with that
-			checkf(!EnumHasAllFlags(CreateDesc.Flags, TexCreate_AtomicCompatible), TEXT("Requested buffer backed texture that breaks Metal linear texture limitations: %s"), *FString([CreateDesc.Desc description]));
-			Texture = GetMetalDeviceContext().CreateTexture(this, CreateDesc.Desc);
+			checkf(!EnumHasAllFlags(CreateDesc.Flags, TexCreate_AtomicCompatible), TEXT("Requested buffer backed texture that breaks Metal linear texture limitations: %s"), *FString([NewCreateDesc.Desc description]));
+			Texture = GetMetalDeviceContext().CreateTexture(this, NewCreateDesc.Desc);
 		}
 		
 		METAL_FATAL_ASSERT(Texture, TEXT("Failed to create texture, desc %s"), *FString([CreateDesc.Desc description]));
@@ -687,7 +710,7 @@ FMetalSurface::FMetalSurface(FMetalTextureCreateDesc const& CreateDesc)
 	// calculate size of the texture
 	TotalTextureSize = GetMemorySize();
 
-	if (CreateDesc.NumSamples > 1 && !FParse::Param(FCommandLine::Get(), TEXT("nomsaa")))
+	if (bIsMSAARequired)
 	{
 		mtlpp::TextureDescriptor Desc = CreateDesc.Desc;
 		check(CreateDesc.bIsRenderTarget);
@@ -705,8 +728,9 @@ FMetalSurface::FMetalSurface(FMetalTextureCreateDesc const& CreateDesc)
 		}
 
 		bool bMemoryless = false;
+        
 #if PLATFORM_IOS
-		if (GMaxRHIFeatureLevel < ERHIFeatureLevel::SM5)
+		if (EnumHasAllFlags(CreateDesc.Flags, TexCreate_Memoryless))
 		{
 			bMemoryless = true;
 			Desc.SetStorageMode(mtlpp::StorageMode::Memoryless);
diff --git a/Engine/Source/Runtime/Engine/Private/SceneTexturesConfig.cpp b/Engine/Source/Runtime/Engine/Private/SceneTexturesConfig.cpp
--- a/Engine/Source/Runtime/Engine/Private/SceneTexturesConfig.cpp
+++ b/Engine/Source/Runtime/Engine/Private/SceneTexturesConfig.cpp
@@ -153,7 +153,7 @@ static ETextureCreateFlags GetSceneDepthStencilCreateFlags(uint32 NumSamples, bo
 	ETextureCreateFlags DepthCreateFlags = TexCreate_DepthStencilTargetable | TexCreate_ShaderResource | TexCreate_InputAttachmentRead | ExtraSceneDepthCreateFlags;
 	if (!bKeepDepthContent || (NumSamples > 1 && bMemorylessMSAA))
 	{
-		DepthCreateFlags = TexCreate_Memoryless;
+		DepthCreateFlags |= TexCreate_Memoryless;
 	}
 	if (NumSamples == 1 && GRHISupportsDepthUAV)
 	{
@@ -207,9 +207,10 @@ void FSceneTexturesConfig::Init(const FSceneTexturesConfigInitSettings& InitSett
 	bRequireMultiView		= InitSettings.bRequireMultiView;
 	bIsUsingGBuffers		= IsUsingGBuffers(ShaderPlatform);
 	bSupportsXRTargetManagerDepthAlloc = InitSettings.bSupportsXRTargetManagerDepthAlloc;
-
-	GetSceneColorFormatAndCreateFlags(FeatureLevel, InitSettings.bRequiresAlphaChannel, InitSettings.ExtraSceneColorCreateFlags, NumSamples, bMemorylessMSAA, ColorFormat, ColorCreateFlags);
-	DepthCreateFlags = GetSceneDepthStencilCreateFlags(NumSamples, bKeepDepthContent, bMemorylessMSAA, InitSettings.ExtraSceneDepthCreateFlags);
+    ExtraSceneColorCreateFlags = InitSettings.ExtraSceneColorCreateFlags;
+    ExtraSceneDepthCreateFlags = InitSettings.ExtraSceneDepthCreateFlags;
+    
+    BuildSceneColorAndDepthFlags();
 
 	if (bIsUsingGBuffers)
 	{
@@ -252,6 +253,12 @@ void FSceneTexturesConfig::Init(const FSceneTexturesConfigInitSettings& InitSett
 	}
 }
 
+void FSceneTexturesConfig::BuildSceneColorAndDepthFlags()
+{
+    GetSceneColorFormatAndCreateFlags(FeatureLevel, bRequiresAlphaChannel, ExtraSceneColorCreateFlags, NumSamples, bMemorylessMSAA, ColorFormat, ColorCreateFlags);
+    DepthCreateFlags = GetSceneDepthStencilCreateFlags(NumSamples, bKeepDepthContent, bMemorylessMSAA, ExtraSceneDepthCreateFlags);
+}
+
 uint32 FSceneTexturesConfig::GetGBufferRenderTargetsInfo(FGraphicsPipelineRenderTargetsInfo& RenderTargetsInfo) const 
 {
 	// Assume 1 sample for now
@@ -297,4 +304,4 @@ uint32 FSceneTexturesConfig::GetGBufferRenderTargetsInfo(FGraphicsPipelineRender
 	RenderTargetsInfo.bHasFragmentDensityAttachment = false; // RenderTargets.ShadingRateTexture != nullptr;
 
 	return RenderTargetCount;
-}
\ No newline at end of file
+}
diff --git a/Engine/Source/Runtime/Engine/Public/SceneTexturesConfig.h b/Engine/Source/Runtime/Engine/Public/SceneTexturesConfig.h
--- a/Engine/Source/Runtime/Engine/Public/SceneTexturesConfig.h
+++ b/Engine/Source/Runtime/Engine/Public/SceneTexturesConfig.h
@@ -122,6 +122,7 @@ struct ENGINE_API FSceneTexturesConfig
 	{}
 
 	void Init(const FSceneTexturesConfigInitSettings& InitSettings);
+    void BuildSceneColorAndDepthFlags();
 	uint32 GetGBufferRenderTargetsInfo(FGraphicsPipelineRenderTargetsInfo& RenderTargetsInfo) const;
 
 	FORCEINLINE bool IsValid() const
@@ -159,6 +160,10 @@ struct ENGINE_API FSceneTexturesConfig
 	ETextureCreateFlags ColorCreateFlags = ETextureCreateFlags::None;
 	ETextureCreateFlags DepthCreateFlags = ETextureCreateFlags::None;
 
+    // Flags passed in from initializer
+    ETextureCreateFlags ExtraSceneColorCreateFlags = ETextureCreateFlags::None;
+    ETextureCreateFlags ExtraSceneDepthCreateFlags = ETextureCreateFlags::None;
+    
 	// Optimized clear values to use for color / depth textures.
 	FClearValueBinding ColorClearValue = FClearValueBinding::Black;
 	FClearValueBinding DepthClearValue = FClearValueBinding::DepthFar;
@@ -192,6 +197,9 @@ struct ENGINE_API FSceneTexturesConfig
 
 	// (XR) True if we can request an XR depth swapchain
 	uint32 bSupportsXRTargetManagerDepthAlloc : 1;
+    
+    // True if we require an alpha channel for scene color
+    bool bRequiresAlphaChannel = false;
 
 private:
 	static FSceneTexturesConfig GlobalInstance;
diff --git a/Engine/Source/Runtime/RHI/Public/RHIResources.h b/Engine/Source/Runtime/RHI/Public/RHIResources.h
--- a/Engine/Source/Runtime/RHI/Public/RHIResources.h
+++ b/Engine/Source/Runtime/RHI/Public/RHIResources.h
@@ -1454,7 +1454,7 @@ struct RHI_API FRHITextureDesc
 
 		// Remove UAV flag for textures that don't need it (some formats are incompatible).
 		Flags |= TexCreate_RenderTargetable;
-		Flags &= ~(TexCreate_UAV | TexCreate_ResolveTargetable | TexCreate_DepthStencilResolveTarget);
+		Flags &= ~(TexCreate_UAV | TexCreate_ResolveTargetable | TexCreate_DepthStencilResolveTarget | TexCreate_Memoryless);
 	}
 
 	/** Returns whether this descriptor conforms to requirements. */
@@ -4333,4 +4333,4 @@ public:
 private:
 	TArray<TPair<FRHIBufferUAVCreateInfo, FUnorderedAccessViewRHIRef>, TInlineAllocator<1>> UAVs;
 	TArray<TPair<FRHIBufferSRVCreateInfo, FShaderResourceViewRHIRef>, TInlineAllocator<1>> SRVs;
-};
\ No newline at end of file
+};
diff --git a/Engine/Source/Runtime/RenderCore/Private/RenderGraphUtils.cpp b/Engine/Source/Runtime/RenderCore/Private/RenderGraphUtils.cpp
--- a/Engine/Source/Runtime/RenderCore/Private/RenderGraphUtils.cpp
+++ b/Engine/Source/Runtime/RenderCore/Private/RenderGraphUtils.cpp
@@ -213,6 +213,8 @@ RENDERCORE_API FRDGTextureMSAA CreateTextureMSAA(
 		{
 			ResolveFlags |= TexCreate_ResolveTargetable;
 		}
+        ResolveFlags &= ~(TexCreate_Memoryless);
+        
 		Desc.Flags = ResolveFlags | ResolveFlagsToAdd;
 		Texture.Resolve = GraphBuilder.CreateTexture(Desc, Name);
 	}
@@ -1097,4 +1099,4 @@ bool AllocatePooledTexture(const FRDGTextureDesc& Desc, TRefCountPtr<IPooledRend
 TRefCountPtr<IPooledRenderTarget> AllocatePooledTexture(const FRDGTextureDesc& Desc, const TCHAR* Name)
 {
 	return GRenderTargetPool.FindFreeElement(Desc, Name);
-}
\ No newline at end of file
+}
diff --git a/Engine/Source/Runtime/Renderer/Private/MobileShadingRenderer.cpp b/Engine/Source/Runtime/Renderer/Private/MobileShadingRenderer.cpp
--- a/Engine/Source/Runtime/Renderer/Private/MobileShadingRenderer.cpp
+++ b/Engine/Source/Runtime/Renderer/Private/MobileShadingRenderer.cpp
@@ -539,8 +539,11 @@ void FMobileSceneRenderer::InitViews(FRDGBuilder& GraphBuilder, FSceneTexturesCo
 	// Update the bKeepDepthContent based on the mobile renderer status.
 	SceneTexturesConfig.bKeepDepthContent = bKeepDepthContent;
 	// If we render in a single pass MSAA targets can be memoryless
-	SceneTexturesConfig.bMemorylessMSAA = !(bRequiresMultiPass || bShouldCompositeEditorPrimitives);
-	
+    SceneTexturesConfig.bMemorylessMSAA = !(bRequiresMultiPass || bShouldCompositeEditorPrimitives);
+    SceneTexturesConfig.NumSamples = NumMSAASamples;
+    
+    SceneTexturesConfig.BuildSceneColorAndDepthFlags();
+    
 	if (bDeferredShading) 
 	{
 		SetupGBufferFlags(SceneTexturesConfig, bRequiresMultiPass || GraphBuilder.IsDumpingFrame());
