Speculative fix for missing include in Metal.

#rnx

diff --git a/Engine/Source/Runtime/Apple/MetalRHI/Private/MetalRHI.cpp b/Engine/Source/Runtime/Apple/MetalRHI/Private/MetalRHI.cpp
--- a/Engine/Source/Runtime/Apple/MetalRHI/Private/MetalRHI.cpp
+++ b/Engine/Source/Runtime/Apple/MetalRHI/Private/MetalRHI.cpp
@@ -22,6 +22,7 @@
 #include "MetalLLM.h"
 #include "Engine/RendererSettings.h"
 #include "MetalTransitionData.h"
+#include "EngineGlobals.h"
 
 DEFINE_LOG_CATEGORY(LogMetal)
 
