Fix for CIS nopch error

diff --git a/Engine/Source/Runtime/Apple/AudioMixerAudioUnit/Private/AudioMixerPlatformAudioUnit.cpp b/Engine/Source/Runtime/Apple/AudioMixerAudioUnit/Private/AudioMixerPlatformAudioUnit.cpp
--- a/Engine/Source/Runtime/Apple/AudioMixerAudioUnit/Private/AudioMixerPlatformAudioUnit.cpp
+++ b/Engine/Source/Runtime/Apple/AudioMixerAudioUnit/Private/AudioMixerPlatformAudioUnit.cpp
@@ -12,6 +12,7 @@
 #include "Misc/CoreDelegates.h"
 #include "ADPCMAudioInfo.h"
 #include "BinkAudioInfo.h"
+#include "Misc/CommandLine.h"
 
 /*
  This implementation only depends on the audio units API which allows it to run on MacOS, iOS and tvOS.
@@ -26,6 +27,7 @@
 #include <AudioUnit/AudioUnit.h>
 #include <AVFoundation/AVAudioSession.h>
 
+
 static int32 SuspendCounter = 0;
 
 DECLARE_LOG_CATEGORY_EXTERN(LogAudioMixerAudioUnit, Log, All);
