Fix for texture pool reuse before the GPU has finished with it.  Don't place the metal texture in the context heap until it has gone through the DelayedFreeLists update path like buffers etc.

#rnx

diff --git a/Engine/Source/Runtime/Apple/MetalRHI/Private/MetalContext.cpp b/Engine/Source/Runtime/Apple/MetalRHI/Private/MetalContext.cpp
--- a/Engine/Source/Runtime/Apple/MetalRHI/Private/MetalContext.cpp
+++ b/Engine/Source/Runtime/Apple/MetalRHI/Private/MetalContext.cpp
@@ -721,18 +721,7 @@ void FMetalDeviceContext::ReleaseTexture(FMetalTexture& Texture)
 	{
 		check(Texture);
 		FreeListMutex.Lock();
-        if (Texture.GetStorageMode() == mtlpp::StorageMode::Private)
-        {
-            Heap.ReleaseTexture(nullptr, Texture);
-			
-			// Ensure that the Objective-C handle can't disappear prior to the GPU being done with it without racing with the above
-			if(!ObjectFreeList.Contains(Texture.GetPtr()))
-			{
-				[Texture.GetPtr() retain];
-				ObjectFreeList.Add(Texture.GetPtr());
-        }
-        }
-		else if(!UsedTextures.Contains(Texture))
+		if(!UsedTextures.Contains(Texture))
 		{
 			UsedTextures.Add(MoveTemp(Texture));
 		}
