Fixing bad removal of an if condition.

diff --git a/Engine/Source/Runtime/Apple/MetalRHI/Private/MetalContext.cpp b/Engine/Source/Runtime/Apple/MetalRHI/Private/MetalContext.cpp
--- a/Engine/Source/Runtime/Apple/MetalRHI/Private/MetalContext.cpp
+++ b/Engine/Source/Runtime/Apple/MetalRHI/Private/MetalContext.cpp
@@ -1408,14 +1408,6 @@ void FMetalContext::SetRenderPassInfo(const FRHIRenderPassInfo& RenderTargetsInf
 			bSet = StateCache.SetRenderPassInfo(RenderTargetsInfo, QueryBuffer->GetCurrentQueryBuffer(), bRestart);
 		}
 	}
-	else
-	{
-		if (NULL != StateCache.GetVisibilityResultsBuffer())
-		{
-			RenderPass.EndRenderPass();
-		}
-		bSet = StateCache.SetRenderPassInfo(RenderTargetsInfo, NULL, bRestart);
-	}
 	
 	if (bSet && StateCache.GetHasValidRenderTarget())
 	{
