Improve launch notification on Mac OS when Xcode can't be found.

diff --git a/Engine/Source/Runtime/Apple/MetalRHI/Private/MetalRHI.cpp b/Engine/Source/Runtime/Apple/MetalRHI/Private/MetalRHI.cpp
--- a/Engine/Source/Runtime/Apple/MetalRHI/Private/MetalRHI.cpp
+++ b/Engine/Source/Runtime/Apple/MetalRHI/Private/MetalRHI.cpp
@@ -115,7 +115,8 @@ static void VerifyMetalCompiler()
 		
 		if(!bFoundXCode)
 		{
-			FMessageDialog::Open(EAppMsgType::Ok, FText(NSLOCTEXT("MetalRHI", "XCodeMissingInstall", "Can't find Xcode install for Metal compiler. Please install Xcode and run Xcode.app to accept license or ensure active developer directory is set to current Xcode installation using xcode-select.")));
+			const FText Title = NSLOCTEXT("MetalRHI", "XCodeMissingInstallTitle", "Xcode Not Found");
+			FMessageDialog::Open(EAppMsgType::Ok, NSLOCTEXT("MetalRHI", "XCodeMissingInstall", "Unreal Engine requires Xcode to compile shaders for Metal. To continue, install Xcode and open it to accept the license agreement. If you install Xcode to any location other than Applications/Xcode, also run the xcode-select command-line tool to specify its location."), &Title);
 			FPlatformMisc::RequestExit(true);
 			return;
 		}
