Round up to block sizes to avoid blit encoder asserts.  D3D appears to follow this pattern (Engine Test Rendering test group would throw a blit encoder assert during a test without this fix if ran with the validation layer).

UDN: 00495381

diff --git a/Engine/Source/Runtime/Apple/MetalRHI/Private/MetalTexture.cpp b/Engine/Source/Runtime/Apple/MetalRHI/Private/MetalTexture.cpp
--- a/Engine/Source/Runtime/Apple/MetalRHI/Private/MetalTexture.cpp
+++ b/Engine/Source/Runtime/Apple/MetalRHI/Private/MetalTexture.cpp
@@ -1746,7 +1746,7 @@ static FMetalBuffer Internal_CreateBufferAndCopyTexture2DUpdateRegionData(FRHITe
 		
 		uint8* pDestRow = (uint8*)OutBuffer.GetContents();
 		uint8* pSourceRow = (uint8*)SourceData;
-		const uint32 NumRows = UpdateRegion.Height / (uint32)FormatInfo.BlockSizeY;
+		const uint32 NumRows = FMath::DivideAndRoundUp(UpdateRegion.Height, (uint32)FormatInfo.BlockSizeY);
 		
 		// Limit copy to line by line by update region pitch otherwise we can go off the end of source data on the last row
 		for (uint32 i = 0;i < NumRows;++i)
@@ -1772,7 +1772,7 @@ static void InternalUpdateTexture2D(FMetalContext& Context, FRHITexture2D* Textu
 		SCOPED_AUTORELEASE_POOL;
 		
 		const FPixelFormatInfo& FormatInfo = GPixelFormats[TextureRHI->GetFormat()];
-		const uint32 NumRows = UpdateRegion.Height / (uint32)FormatInfo.BlockSizeY;
+		const uint32 NumRows = FMath::DivideAndRoundUp(UpdateRegion.Height, (uint32)FormatInfo.BlockSizeY);
 		uint32 BytesPerImage = SourcePitch * NumRows;
 		
 		mtlpp::BlitOption Options = mtlpp::BlitOption::None;
@@ -1825,7 +1825,7 @@ static FMetalBuffer Internal_CreateBufferAndCopyTexture3DUpdateRegionData(FRHITe
 	check(CopyPitch <= SourceRowPitch);
 		
 	uint8_t* DestData = (uint8_t*)OutBuffer.GetContents();
-	const uint32 NumRows = UpdateRegion.Height / (uint32)FormatInfo.BlockSizeY;
+	const uint32 NumRows = FMath::DivideAndRoundUp(UpdateRegion.Height, (uint32)FormatInfo.BlockSizeY);
 		
 	// Perform safe line copy
 	for (uint32 i = 0;i < UpdateRegion.Depth;++i)
@@ -1855,7 +1855,7 @@ static void InternalUpdateTexture3D(FMetalContext& Context, FRHITexture3D* Textu
 	if(Tex.GetStorageMode() == mtlpp::StorageMode::Private)
 	{
 		const FPixelFormatInfo& FormatInfo = GPixelFormats[TextureRHI->GetFormat()];
-		const uint32 NumRows = UpdateRegion.Height / (uint32)FormatInfo.BlockSizeY;
+		const uint32 NumRows = FMath::DivideAndRoundUp(UpdateRegion.Height, (uint32)FormatInfo.BlockSizeY);
 		const uint32 BytesPerImage = SourceRowPitch * NumRows;
 		
 		mtlpp::BlitOption Options = mtlpp::BlitOption::None;
