Added flag for byte address buffer in UAV when creating MetalResourceView

diff --git a/Engine/Source/Runtime/Apple/MetalRHI/Private/MetalUAV.cpp b/Engine/Source/Runtime/Apple/MetalRHI/Private/MetalUAV.cpp
--- a/Engine/Source/Runtime/Apple/MetalRHI/Private/MetalUAV.cpp
+++ b/Engine/Source/Runtime/Apple/MetalRHI/Private/MetalUAV.cpp
@@ -62,7 +62,7 @@ FMetalResourceViewBase::FMetalResourceViewBase(FRHIBuffer* InBuffer, uint32 InSt
 		}
 		else
 		{
-			check(EnumHasAnyFlags(Usage, BUF_StructuredBuffer));
+			check(EnumHasAnyFlags(Usage, BUF_StructuredBuffer | BUF_ByteAddressBuffer));
 
 			Format = PF_Unknown;
 			Stride = SourceBuffer->GetStride();
