Fix for color grading white balance test for macOS.  Add missing sampler border color address mode when translating from engine to metal.

#rnx

diff --git a/Engine/Source/Runtime/Apple/MetalRHI/Private/MetalState.cpp b/Engine/Source/Runtime/Apple/MetalRHI/Private/MetalState.cpp
--- a/Engine/Source/Runtime/Apple/MetalRHI/Private/MetalState.cpp
+++ b/Engine/Source/Runtime/Apple/MetalRHI/Private/MetalState.cpp
@@ -33,7 +33,11 @@ static mtlpp::SamplerAddressMode TranslateWrapMode(ESamplerAddressMode AddressMo
 	{
 		case AM_Clamp:	return mtlpp::SamplerAddressMode::ClampToEdge;
 		case AM_Mirror: return mtlpp::SamplerAddressMode::MirrorRepeat;
+#if PLATFORM_MAC
+		case AM_Border: return mtlpp::SamplerAddressMode::ClampToBorderColor;
+#else
 		case AM_Border: return mtlpp::SamplerAddressMode::ClampToEdge;
+#endif
 		default:		return mtlpp::SamplerAddressMode::Repeat;
 	}
 }
