Fixing NonUnity error.

diff --git a/Engine/Source/Runtime/Apple/MetalRHI/Private/MetalState.cpp b/Engine/Source/Runtime/Apple/MetalRHI/Private/MetalState.cpp
--- a/Engine/Source/Runtime/Apple/MetalRHI/Private/MetalState.cpp
+++ b/Engine/Source/Runtime/Apple/MetalRHI/Private/MetalState.cpp
@@ -6,6 +6,7 @@
 
 #include "MetalRHIPrivate.h"
 #include "MetalProfiler.h"
+#include "RHIUtilities.h"
 
 static uint32 GetMetalMaxAnisotropy(ESamplerFilter Filter, uint32 MaxAniso)
 {
