Use Metal reported value for the max buffer size to assert on rather than a fixed value to avoid artificially limiting desktop platforms that are capable of much more..   Also make sure GMaxBufferDimensions is set to a sensisble value for Metal.

#rnx
#horde https://horde.devtools.epicgames.com/job/63e62b7ac6153a171c5a66cc

diff --git a/Engine/Source/Runtime/Apple/MetalRHI/Private/MetalRHI.cpp b/Engine/Source/Runtime/Apple/MetalRHI/Private/MetalRHI.cpp
--- a/Engine/Source/Runtime/Apple/MetalRHI/Private/MetalRHI.cpp
+++ b/Engine/Source/Runtime/Apple/MetalRHI/Private/MetalRHI.cpp
@@ -569,6 +569,13 @@ FMetalDynamicRHI::FMetalDynamicRHI(ERHIFeatureLevel::Type RequestedFeatureLevel)
 	GSupportsSeparateRenderTargetBlendState = (GMaxRHIFeatureLevel >= ERHIFeatureLevel::SM5);
 
 	GRHISupportsPipelineFileCache = true;
+	
+	// Appears to be no queryable value for max texture_buffer size and its not specified in the docs.  However,
+	// current testing across Apple Silicon macs, AMD and iPhone all quote a max value of 268435456 (1 << 28)
+	// from the Metal validation layer.  For certain pixel formats this appears to be larger than max buffer
+	// size on some devices.  Due to the way we are using texture_buffers allocated from a backing buffer,
+	// keep this safe by using a step down.  This is the current default value anyway but set here too in case that changes.
+	GMaxBufferDimensions = 1 << 27;
 
 #if PLATFORM_MAC
 	check(Device.SupportsFeatureSet(mtlpp::FeatureSet::macOS_GPUFamily1_v1));
diff --git a/Engine/Source/Runtime/Apple/MetalRHI/Private/MetalVertexBuffer.cpp b/Engine/Source/Runtime/Apple/MetalRHI/Private/MetalVertexBuffer.cpp
--- a/Engine/Source/Runtime/Apple/MetalRHI/Private/MetalVertexBuffer.cpp
+++ b/Engine/Source/Runtime/Apple/MetalRHI/Private/MetalVertexBuffer.cpp
@@ -156,7 +156,7 @@ FMetalRHIBuffer::FMetalRHIBuffer(uint32 InSize, EBufferUsageFlags InUsage, EMeta
 	
 	if (InSize)
 	{
-		checkf(InSize <= 1024 * 1024 * 1024, TEXT("Metal doesn't support buffers > 1GB"));
+		checkf(InSize <= [GetMetalDeviceContext().GetDevice().GetPtr() maxBufferLength], TEXT("Requested buffer size larger than supported by device."));
 		
 		// Temporary buffers less than the buffer page size - currently 4Kb - is better off going through the set*Bytes API if available.
 		// These can't be used for shader resources or UAVs if we want to use the 'Linear Texture' code path
@@ -881,4 +881,4 @@ void FMetalRHIBuffer::Init(FRHICommandListBase& RHICmdList, uint32 InSize, EBuff
 		// Discard the resource array's contents.
 		CreateInfo.ResourceArray->Discard();
 	}
-}
\ No newline at end of file
+}
