Disable unsupported PixelFormats BC4, BC6H & BC7 for iOS/tvOS and enable D24

diff --git a/Engine/Source/Runtime/Apple/MetalRHI/Private/MetalRHI.cpp b/Engine/Source/Runtime/Apple/MetalRHI/Private/MetalRHI.cpp
--- a/Engine/Source/Runtime/Apple/MetalRHI/Private/MetalRHI.cpp
+++ b/Engine/Source/Runtime/Apple/MetalRHI/Private/MetalRHI.cpp
@@ -745,8 +745,14 @@ FMetalDynamicRHI::FMetalDynamicRHI(ERHIFeatureLevel::Type RequestedFeatureLevel)
 	GPixelFormats[PF_DXT3				].Supported			= false;
     GPixelFormats[PF_DXT5				].PlatformFormat	= (uint32)mtlpp::PixelFormat::Invalid;
 	GPixelFormats[PF_DXT5				].Supported			= false;
+	GPixelFormats[PF_BC4				].PlatformFormat	= (uint32)mtlpp::PixelFormat::Invalid;
+	GPixelFormats[PF_BC4				].Supported			= false;
 	GPixelFormats[PF_BC5				].PlatformFormat	= (uint32)mtlpp::PixelFormat::Invalid;
 	GPixelFormats[PF_BC5				].Supported			= false;
+	GPixelFormats[PF_BC6H				].PlatformFormat	= (uint32)mtlpp::PixelFormat::Invalid;
+	GPixelFormats[PF_BC6H				].Supported			= false;
+	GPixelFormats[PF_BC7				].PlatformFormat	= (uint32)mtlpp::PixelFormat::Invalid;
+	GPixelFormats[PF_BC7				].Supported			= false;
 	GPixelFormats[PF_PVRTC2				].PlatformFormat	= (uint32)mtlpp::PixelFormat::PVRTC_RGBA_2BPP;
 	GPixelFormats[PF_PVRTC2				].Supported			= true;
 	GPixelFormats[PF_PVRTC4				].PlatformFormat	= (uint32)mtlpp::PixelFormat::PVRTC_RGBA_4BPP;
@@ -826,6 +832,8 @@ FMetalDynamicRHI::FMetalDynamicRHI(ERHIFeatureLevel::Type RequestedFeatureLevel)
 	GPixelFormats[PF_ShadowDepth		].PlatformFormat	= (uint32)mtlpp::PixelFormat::Depth16Unorm;
 	GPixelFormats[PF_ShadowDepth		].BlockBytes		= 2;
 	GPixelFormats[PF_ShadowDepth		].Supported			= true;
+	GPixelFormats[PF_D24				].PlatformFormat	= (uint32)mtlpp::PixelFormat::Depth32Float;
+	GPixelFormats[PF_D24				].Supported			= true;
 		
 	GPixelFormats[PF_BC5				].PlatformFormat	= (uint32)mtlpp::PixelFormat::Invalid;
 	GPixelFormats[PF_R5G6B5_UNORM		].PlatformFormat	= (uint32)mtlpp::PixelFormat::B5G6R5Unorm;
