Fix EngineTest launch fail through Xcode, it was using unavailable Metal feature (Memory Barrier) on Apple Silicon Macs

diff --git a/Engine/Source/Runtime/Apple/MetalRHI/Private/MetalRenderPass.cpp b/Engine/Source/Runtime/Apple/MetalRHI/Private/MetalRenderPass.cpp
--- a/Engine/Source/Runtime/Apple/MetalRHI/Private/MetalRenderPass.cpp
+++ b/Engine/Source/Runtime/Apple/MetalRHI/Private/MetalRenderPass.cpp
@@ -1700,8 +1700,14 @@ void FMetalRenderPass::InsertDebugDraw(FMetalCommandData& Data)
 			}
 
 		#if PLATFORM_MAC
+			// UE-182622: Apple Silicon will assert on "Memory Barrier With Resources Validation afterStages (0x2) can not contain MTLRenderStageTile or MTLRenderStageFragment on this device"
+			bool bCanUseMemoryBarrier = !GRHIAdapterName.Contains("Apple");
+			
 			id<MTLBuffer> DebugBufferPtr = State.GetDebugBuffer().GetPtr();
-			[(id<IMTLRenderCommandEncoder>)CurrentEncoder.GetRenderCommandEncoder().GetPtr() memoryBarrierWithResources:&DebugBufferPtr count:1 afterStages:MTLRenderStageFragment beforeStages:MTLRenderStageVertex];
+			if (bCanUseMemoryBarrier)
+			{
+				[(id<IMTLRenderCommandEncoder>)CurrentEncoder.GetRenderCommandEncoder().GetPtr() memoryBarrierWithResources:&DebugBufferPtr count:1 afterStages:MTLRenderStageFragment beforeStages:MTLRenderStageVertex];
+			}
 			
 			CurrentEncoder.SetShaderBytes(mtlpp::FunctionType::Vertex, (uint8 const*)&DebugInfo, sizeof(DebugInfo), 0);
 			State.SetShaderBufferDirty(EMetalShaderStages::Vertex, 0);
@@ -1711,7 +1717,10 @@ void FMetalRenderPass::InsertDebugDraw(FMetalCommandData& Data)
 
 			CurrentEncoder.GetRenderCommandEncoder().Draw(mtlpp::PrimitiveType::Point, 0, 1);
 			
-			[(id<IMTLRenderCommandEncoder>)CurrentEncoder.GetRenderCommandEncoder().GetPtr() memoryBarrierWithResources:&DebugBufferPtr count:1 afterStages:MTLRenderStageVertex beforeStages:MTLRenderStageVertex];
+			if (bCanUseMemoryBarrier)
+			{
+				[(id<IMTLRenderCommandEncoder>)CurrentEncoder.GetRenderCommandEncoder().GetPtr() memoryBarrierWithResources:&DebugBufferPtr count:1 afterStages:MTLRenderStageVertex beforeStages:MTLRenderStageVertex];
+			}
 		#else
 			CurrentEncoder.GetRenderCommandEncoder().SetTileData((uint8 const*)&DebugInfo, sizeof(DebugInfo), 0);
 			CurrentEncoder.GetRenderCommandEncoder().SetTileBuffer(State.GetDebugBuffer(), 0, 1);
