Fixed inconsistencies on Mac when rendering Niagara landscapes and using PreintegratedGF

diff --git a/Engine/Source/Runtime/Apple/MetalRHI/Private/MetalRHI.cpp b/Engine/Source/Runtime/Apple/MetalRHI/Private/MetalRHI.cpp
--- a/Engine/Source/Runtime/Apple/MetalRHI/Private/MetalRHI.cpp
+++ b/Engine/Source/Runtime/Apple/MetalRHI/Private/MetalRHI.cpp
@@ -947,11 +947,22 @@ FMetalDynamicRHI::FMetalDynamicRHI(ERHIFeatureLevel::Type RequestedFeatureLevel)
 	GPixelFormats[PF_FloatRGBA			].BlockBytes		= 8;
     GPixelFormats[PF_X24_G8				].PlatformFormat	= (uint32)mtlpp::PixelFormat::Stencil8;
     GPixelFormats[PF_X24_G8				].BlockBytes		= 1;
-    GPixelFormats[PF_X24_G8             ].Supported         = true;
-	GPixelFormats[PF_R32_FLOAT			].PlatformFormat	= (uint32)mtlpp::PixelFormat::R32Float;
+	GPixelFormats[PF_X24_G8             ].Supported         = true;
+	
+    GPixelFormats[PF_R32_FLOAT			].PlatformFormat	= (uint32)mtlpp::PixelFormat::R32Float;
+    if(Device.SupportsFeatureSet(mtlpp::FeatureSet::macOS_GPUFamily2_v1))
+    {
+        EnumAddFlags(GPixelFormats[PF_R32_FLOAT].Capabilities, EPixelFormatCapabilities::TextureFilterable);
+    }
+        
 	GPixelFormats[PF_G16R16				].PlatformFormat	= (uint32)mtlpp::PixelFormat::RG16Unorm;
 	GPixelFormats[PF_G16R16				].Supported			= true;
-	GPixelFormats[PF_G16R16F			].PlatformFormat	= (uint32)mtlpp::PixelFormat::RG16Float;
+    if(Device.SupportsFeatureSet(mtlpp::FeatureSet::macOS_GPUFamily2_v1))
+    {
+        EnumAddFlags(GPixelFormats[PF_G16R16].Capabilities, EPixelFormatCapabilities::TextureFilterable);
+    }
+        
+    GPixelFormats[PF_G16R16F			].PlatformFormat	= (uint32)mtlpp::PixelFormat::RG16Float;
 	GPixelFormats[PF_G16R16F_FILTER		].PlatformFormat	= (uint32)mtlpp::PixelFormat::RG16Float;
 	GPixelFormats[PF_G32R32F			].PlatformFormat	= (uint32)mtlpp::PixelFormat::RG32Float;
 	GPixelFormats[PF_A2B10G10R10		].PlatformFormat    = (uint32)mtlpp::PixelFormat::RGB10A2Unorm;
@@ -1003,6 +1014,8 @@ FMetalDynamicRHI::FMetalDynamicRHI(ERHIFeatureLevel::Type RequestedFeatureLevel)
 	{
 		checkf((NSUInteger)GMetalBufferFormats[i].LinearTextureFormat != NSUIntegerMax, TEXT("Metal linear texture format for pixel-format %s (%d) is not configured!"), GPixelFormats[i].Name, i);
 		checkf(GMetalBufferFormats[i].DataFormat != 255, TEXT("Metal data buffer format for pixel-format %s (%d) is not configured!"), GPixelFormats[i].Name, i);
+        
+        EnumAddFlags(GPixelFormats[i].Capabilities, EPixelFormatCapabilities::TextureFilterable);
 	}
 #endif
 
