Fixed compiler error on tvOS and iOS when checking enum

diff --git a/Engine/Source/Runtime/Apple/MetalRHI/Private/MetalRHI.cpp b/Engine/Source/Runtime/Apple/MetalRHI/Private/MetalRHI.cpp
--- a/Engine/Source/Runtime/Apple/MetalRHI/Private/MetalRHI.cpp
+++ b/Engine/Source/Runtime/Apple/MetalRHI/Private/MetalRHI.cpp
@@ -950,17 +950,21 @@ FMetalDynamicRHI::FMetalDynamicRHI(ERHIFeatureLevel::Type RequestedFeatureLevel)
 	GPixelFormats[PF_X24_G8             ].Supported         = true;
 	
     GPixelFormats[PF_R32_FLOAT			].PlatformFormat	= (uint32)mtlpp::PixelFormat::R32Float;
+#if PLATFORM_MAC
     if(Device.SupportsFeatureSet(mtlpp::FeatureSet::macOS_GPUFamily2_v1))
     {
         EnumAddFlags(GPixelFormats[PF_R32_FLOAT].Capabilities, EPixelFormatCapabilities::TextureFilterable);
     }
+#endif
         
 	GPixelFormats[PF_G16R16				].PlatformFormat	= (uint32)mtlpp::PixelFormat::RG16Unorm;
 	GPixelFormats[PF_G16R16				].Supported			= true;
+#if PLATFORM_MAC
     if(Device.SupportsFeatureSet(mtlpp::FeatureSet::macOS_GPUFamily2_v1))
     {
         EnumAddFlags(GPixelFormats[PF_G16R16].Capabilities, EPixelFormatCapabilities::TextureFilterable);
     }
+#endif
         
     GPixelFormats[PF_G16R16F			].PlatformFormat	= (uint32)mtlpp::PixelFormat::RG16Float;
 	GPixelFormats[PF_G16R16F_FILTER		].PlatformFormat	= (uint32)mtlpp::PixelFormat::RG16Float;
@@ -1014,8 +1018,6 @@ FMetalDynamicRHI::FMetalDynamicRHI(ERHIFeatureLevel::Type RequestedFeatureLevel)
 	{
 		checkf((NSUInteger)GMetalBufferFormats[i].LinearTextureFormat != NSUIntegerMax, TEXT("Metal linear texture format for pixel-format %s (%d) is not configured!"), GPixelFormats[i].Name, i);
 		checkf(GMetalBufferFormats[i].DataFormat != 255, TEXT("Metal data buffer format for pixel-format %s (%d) is not configured!"), GPixelFormats[i].Name, i);
-        
-        EnumAddFlags(GPixelFormats[i].Capabilities, EPixelFormatCapabilities::TextureFilterable);
 	}
 #endif
 
