Use the FMessageDialog overloads that pass the optional title by-value

#rnx

diff --git a/Engine/Plugins/Media/AvfMedia/Source/AvfMediaCapture/Private/Player/AvfMediaCapturePlayer.cpp b/Engine/Plugins/Media/AvfMedia/Source/AvfMediaCapture/Private/Player/AvfMediaCapturePlayer.cpp
--- a/Engine/Plugins/Media/AvfMedia/Source/AvfMediaCapture/Private/Player/AvfMediaCapturePlayer.cpp
+++ b/Engine/Plugins/Media/AvfMedia/Source/AvfMediaCapture/Private/Player/AvfMediaCapturePlayer.cpp
@@ -180,7 +180,7 @@ void FAvfMediaCapturePlayer::HandleAuthStatusError(EAvfMediaCaptureAuthStatus Au
 		ErrorString += FString(TEXT(" (")) + URL + FString(TEXT(")"));
 		FText Msg = FText::FromString(ErrorString);
 		FText Title = NSLOCTEXT("AvfMediaCapture", "AvfMediaCapture_DeviceFail", "Capture Device Start Failed");
-		FMessageDialog::Open(EAppMsgType::Ok, Msg, &Title);
+		FMessageDialog::Open(EAppMsgType::Ok, Msg, Title);
 	}
 }
 
diff --git a/Engine/Source/Runtime/Apple/MetalRHI/Private/MetalRHI.cpp b/Engine/Source/Runtime/Apple/MetalRHI/Private/MetalRHI.cpp
--- a/Engine/Source/Runtime/Apple/MetalRHI/Private/MetalRHI.cpp
+++ b/Engine/Source/Runtime/Apple/MetalRHI/Private/MetalRHI.cpp
@@ -115,8 +115,9 @@ static void VerifyMetalCompiler()
 		
 		if(!bFoundXCode)
 		{
-			const FText Title = NSLOCTEXT("MetalRHI", "XCodeMissingInstallTitle", "Xcode Not Found");
-			FMessageDialog::Open(EAppMsgType::Ok, NSLOCTEXT("MetalRHI", "XCodeMissingInstall", "Unreal Engine requires Xcode to compile shaders for Metal. To continue, install Xcode and open it to accept the license agreement. If you install Xcode to any location other than Applications/Xcode, also run the xcode-select command-line tool to specify its location."), &Title);
+			FMessageDialog::Open(EAppMsgType::Ok, 
+				NSLOCTEXT("MetalRHI", "XCodeMissingInstall", "Unreal Engine requires Xcode to compile shaders for Metal. To continue, install Xcode and open it to accept the license agreement. If you install Xcode to any location other than Applications/Xcode, also run the xcode-select command-line tool to specify its location."), 
+				NSLOCTEXT("MetalRHI", "XCodeMissingInstallTitle", "Xcode Not Found"));
 			FPlatformMisc::RequestExit(true);
 			return;
 		}
