Disable parallel occlusion queries on Vulkan since it doesn't actually support them.

diff --git a/Engine/Source/Runtime/Apple/MetalRHI/Private/MetalRHI.cpp b/Engine/Source/Runtime/Apple/MetalRHI/Private/MetalRHI.cpp
--- a/Engine/Source/Runtime/Apple/MetalRHI/Private/MetalRHI.cpp
+++ b/Engine/Source/Runtime/Apple/MetalRHI/Private/MetalRHI.cpp
@@ -542,12 +542,10 @@ FMetalDynamicRHI::FMetalDynamicRHI(ERHIFeatureLevel::Type RequestedFeatureLevel)
 	if (GMaxRHIFeatureLevel >= ERHIFeatureLevel::SM5)
 	{
 		GRHISupportsRHIThread = true;
-		GSupportsParallelOcclusionQueries = GRHISupportsRHIThread;
 	}
 	else
 	{
 		GRHISupportsRHIThread = FParse::Param(FCommandLine::Get(),TEXT("rhithread")) || (CVarUseIOSRHIThread.GetValueOnAnyThread() > 0);
-		GSupportsParallelOcclusionQueries = false;
 	}
 
 	if (FPlatformMisc::IsDebuggerPresent() && UE_BUILD_DEBUG)
diff --git a/Engine/Source/Runtime/VulkanRHI/Private/VulkanRHI.cpp b/Engine/Source/Runtime/VulkanRHI/Private/VulkanRHI.cpp
--- a/Engine/Source/Runtime/VulkanRHI/Private/VulkanRHI.cpp
+++ b/Engine/Source/Runtime/VulkanRHI/Private/VulkanRHI.cpp
@@ -910,7 +910,6 @@ void FVulkanDynamicRHI::InitInstance()
 		GRHISupportsRHIThread = GRHIThreadCvar->GetInt() != 0;
 		GRHISupportsParallelRHIExecute = GRHIThreadCvar->GetInt() > 1;
 #endif
-		GSupportsParallelOcclusionQueries = true;
 
 		// Some platforms might only have CPU for an RHI thread, but not for parallel tasks
 		GSupportsParallelRenderingTasksWithSeparateRHIThread = GRHISupportsRHIThread ? FVulkanPlatform::SupportParallelRenderingTasks() : false;
