Return the lock size in RHILockTexture2D(). Ran ReplayRun on all supported platforms, and manually ran tests on ones that don't.

diff --git a/Engine/Source/Runtime/Apple/MetalRHI/Private/MetalDynamicRHI.h b/Engine/Source/Runtime/Apple/MetalRHI/Private/MetalDynamicRHI.h
--- a/Engine/Source/Runtime/Apple/MetalRHI/Private/MetalDynamicRHI.h
+++ b/Engine/Source/Runtime/Apple/MetalRHI/Private/MetalDynamicRHI.h
@@ -78,7 +78,7 @@ public:
 	virtual FTexture2DRHIRef RHIAsyncReallocateTexture2D(FRHITexture2D* Texture2D, int32 NewMipCount, int32 NewSizeX, int32 NewSizeY, FThreadSafeCounter* RequestStatus) final override;
 	virtual ETextureReallocationStatus RHIFinalizeAsyncReallocateTexture2D(FRHITexture2D* Texture2D, bool bBlockUntilCompleted) final override;
 	virtual ETextureReallocationStatus RHICancelAsyncReallocateTexture2D(FRHITexture2D* Texture2D, bool bBlockUntilCompleted) final override;
-	virtual void* RHILockTexture2D(FRHITexture2D* Texture, uint32 MipIndex, EResourceLockMode LockMode, uint32& DestStride, bool bLockWithinMiptail) final override;
+	virtual void* RHILockTexture2D(FRHITexture2D* Texture, uint32 MipIndex, EResourceLockMode LockMode, uint32& DestStride, bool bLockWithinMiptail, uint64* OutLockedByteCount = nullptr) final override;
 	virtual void RHIUnlockTexture2D(FRHITexture2D* Texture, uint32 MipIndex, bool bLockWithinMiptail) final override;
 	virtual void* RHILockTexture2DArray(FRHITexture2DArray* Texture, uint32 TextureIndex, uint32 MipIndex, EResourceLockMode LockMode, uint32& DestStride, bool bLockWithinMiptail) final override;
 	virtual void RHIUnlockTexture2DArray(FRHITexture2DArray* Texture, uint32 TextureIndex, uint32 MipIndex, bool bLockWithinMiptail) final override;
@@ -123,7 +123,7 @@ public:
 	virtual ETextureReallocationStatus FinalizeAsyncReallocateTexture2D_RenderThread(class FRHICommandListImmediate& RHICmdList, FRHITexture2D* Texture2D, bool bBlockUntilCompleted) final override;
 	virtual ETextureReallocationStatus CancelAsyncReallocateTexture2D_RenderThread(class FRHICommandListImmediate& RHICmdList, FRHITexture2D* Texture2D, bool bBlockUntilCompleted) final override;
 
-	virtual void* LockTexture2D_RenderThread(class FRHICommandListImmediate& RHICmdList, FRHITexture2D* Texture, uint32 MipIndex, EResourceLockMode LockMode, uint32& DestStride, bool bLockWithinMiptail, bool bNeedsDefaultRHIFlush = true) final override;
+	virtual void* LockTexture2D_RenderThread(class FRHICommandListImmediate& RHICmdList, FRHITexture2D* Texture, uint32 MipIndex, EResourceLockMode LockMode, uint32& DestStride, bool bLockWithinMiptail, bool bNeedsDefaultRHIFlush = true, uint64* OutLockedByteCount = nullptr) final override;
 	virtual void UnlockTexture2D_RenderThread(class FRHICommandListImmediate& RHICmdList, FRHITexture2D* Texture, uint32 MipIndex, bool bLockWithinMiptail, bool bNeedsDefaultRHIFlush = true) final override;
 
 	virtual FUpdateTexture3DData RHIBeginUpdateTexture3D(FRHICommandListBase& RHICmdList, FRHITexture3D* Texture, uint32 MipIndex, const struct FUpdateTextureRegion3D& UpdateRegion) final override;
diff --git a/Engine/Source/Runtime/Apple/MetalRHI/Private/MetalTexture.cpp b/Engine/Source/Runtime/Apple/MetalRHI/Private/MetalTexture.cpp
--- a/Engine/Source/Runtime/Apple/MetalRHI/Private/MetalTexture.cpp
+++ b/Engine/Source/Runtime/Apple/MetalRHI/Private/MetalTexture.cpp
@@ -1044,10 +1044,14 @@ void FMetalSurface::UpdateSurfaceAndDestroySourceBuffer(id <MTLBuffer> SourceBuf
 #endif
 }
 
-void* FMetalSurface::Lock(uint32 MipIndex, uint32 ArrayIndex, EResourceLockMode LockMode, uint32& DestStride, bool SingleLayer /*= false*/)
+void* FMetalSurface::Lock(uint32 MipIndex, uint32 ArrayIndex, EResourceLockMode LockMode, uint32& DestStride, bool SingleLayer /*= false*/, uint64* OutLockedByteCount /* =nullptr */)
 {
 	// get size and stride
 	uint32 MipBytes = GetMipSize(MipIndex, &DestStride, false);
+	if (OutLockedByteCount)
+	{
+		*OutLockedByteCount = (uint64)MipBytes;
+	}
 	
 	// allocate some temporary memory
 	id <MTLBuffer> Buffer = AllocSurface(MipIndex, ArrayIndex, LockMode, DestStride, SingleLayer);
@@ -1205,7 +1209,7 @@ void FMetalSurface::Unlock(uint32 MipIndex, uint32 ArrayIndex, bool bTryAsync)
 	}
 }
 
-void* FMetalSurface::AsyncLock(class FRHICommandListImmediate& RHICmdList, uint32 MipIndex, uint32 ArrayIndex, EResourceLockMode LockMode, uint32& DestStride, bool bNeedsDefaultRHIFlush)
+void* FMetalSurface::AsyncLock(class FRHICommandListImmediate& RHICmdList, uint32 MipIndex, uint32 ArrayIndex, EResourceLockMode LockMode, uint32& DestStride, bool bNeedsDefaultRHIFlush, uint64* OutLockedByteCount)
 {
 	bool bDirectLock = (LockMode == RLM_ReadOnly || !GIsRHIInitialized);
 	
@@ -1220,13 +1224,17 @@ void* FMetalSurface::AsyncLock(class FRHICommandListImmediate& RHICmdList, uint3
 			QUICK_SCOPE_CYCLE_COUNTER(STAT_RHIMETHOD_LockTexture2D_Flush);
 			RHICmdList.ImmediateFlush(EImmediateFlushType::FlushRHIThread);
 		}
-		BufferData = Lock(MipIndex, ArrayIndex, LockMode, DestStride);
+		BufferData = Lock(MipIndex, ArrayIndex, LockMode, DestStride, OutLockedByteCount);
 	}
 	else
 	{
 		id <MTLBuffer> Buffer = AllocSurface(MipIndex, ArrayIndex, LockMode, DestStride);
 		check(Buffer);
 		
+		if (OutLockedByteCount)
+		{
+			*OutLockedByteCount = (uint64)GetMipSize(MipIndex, &DestStride, false);
+		}
 		BufferData = Buffer.contents;
 	}
 	
@@ -1558,14 +1566,14 @@ ETextureReallocationStatus FMetalDynamicRHI::RHICancelAsyncReallocateTexture2D(F
 	return TexRealloc_Failed;
 }
 
-void* FMetalDynamicRHI::LockTexture2D_RenderThread(class FRHICommandListImmediate& RHICmdList, FRHITexture2D* Texture, uint32 MipIndex, EResourceLockMode LockMode, uint32& DestStride, bool bLockWithinMiptail, bool bNeedsDefaultRHIFlush)
+void* FMetalDynamicRHI::LockTexture2D_RenderThread(class FRHICommandListImmediate& RHICmdList, FRHITexture2D* Texture, uint32 MipIndex, EResourceLockMode LockMode, uint32& DestStride, bool bLockWithinMiptail, bool bNeedsDefaultRHIFlush, uint64* OutLockedByteCount)
 {
 	@autoreleasepool {
 		check(IsInRenderingThread());
 		
 		FMetalSurface* TextureMTL = ResourceCast(Texture);
 		
-		void* BufferData = TextureMTL->AsyncLock(RHICmdList, MipIndex, 0, LockMode, DestStride, bNeedsDefaultRHIFlush);
+		void* BufferData = TextureMTL->AsyncLock(RHICmdList, MipIndex, 0, LockMode, DestStride, bNeedsDefaultRHIFlush, OutLockedByteCount);
 		
 		return BufferData;
 	}
@@ -1583,11 +1591,11 @@ void FMetalDynamicRHI::UnlockTexture2D_RenderThread(class FRHICommandListImmedia
 }
 
 
-void* FMetalDynamicRHI::RHILockTexture2D(FRHITexture2D* TextureRHI,uint32 MipIndex,EResourceLockMode LockMode,uint32& DestStride,bool bLockWithinMiptail)
+void* FMetalDynamicRHI::RHILockTexture2D(FRHITexture2D* TextureRHI,uint32 MipIndex,EResourceLockMode LockMode,uint32& DestStride,bool bLockWithinMiptail, uint64* OutLockedByteCount)
 {
 	@autoreleasepool {
 		FMetalSurface* Texture = ResourceCast(TextureRHI);
-		return Texture->Lock(MipIndex, 0, LockMode, DestStride);
+		return Texture->Lock(MipIndex, 0, LockMode, DestStride, OutLockedByteCount);
 	}
 }
 
diff --git a/Engine/Source/Runtime/Apple/MetalRHI/Public/MetalResources.h b/Engine/Source/Runtime/Apple/MetalRHI/Public/MetalResources.h
--- a/Engine/Source/Runtime/Apple/MetalRHI/Public/MetalResources.h
+++ b/Engine/Source/Runtime/Apple/MetalRHI/Public/MetalResources.h
@@ -197,7 +197,7 @@ public:
 	 * @param ArrayIndex Index of the texture array/face in the form Index*6+Face
 	 * @return A pointer to the specified texture data.
 	 */
-	void* Lock(uint32 MipIndex, uint32 ArrayIndex, EResourceLockMode LockMode, uint32& DestStride, bool SingleLayer = false);
+	void* Lock(uint32 MipIndex, uint32 ArrayIndex, EResourceLockMode LockMode, uint32& DestStride, bool SingleLayer = false, uint64* OutLockedByteCount = nullptr);
 	
 	/** Unlocks a previously locked mip-map.
 	 * @param ArrayIndex Index of the texture array/face in the form Index*6+Face
@@ -209,7 +209,7 @@ public:
 	 * @param ArrayIndex Index of the texture array/face in the form Index*6+Face
 	 * @return A pointer to the specified texture data.
 	 */
-	void* AsyncLock(class FRHICommandListImmediate& RHICmdList, uint32 MipIndex, uint32 ArrayIndex, EResourceLockMode LockMode, uint32& DestStride, bool bNeedsDefaultRHIFlush);
+	void* AsyncLock(class FRHICommandListImmediate& RHICmdList, uint32 MipIndex, uint32 ArrayIndex, EResourceLockMode LockMode, uint32& DestStride, bool bNeedsDefaultRHIFlush, uint64* OutLockedByteCount = nullptr);
 	
 	/** Unlocks a previously locked mip-map.
 	 * @param ArrayIndex Index of the texture array/face in the form Index*6+Face
diff --git a/Engine/Source/Runtime/D3D12RHI/Private/D3D12RHIPrivate.h b/Engine/Source/Runtime/D3D12RHI/Private/D3D12RHIPrivate.h
--- a/Engine/Source/Runtime/D3D12RHI/Private/D3D12RHIPrivate.h
+++ b/Engine/Source/Runtime/D3D12RHI/Private/D3D12RHIPrivate.h
@@ -451,7 +451,7 @@ public:
 	virtual FTexture2DRHIRef RHIAsyncReallocateTexture2D(FRHITexture2D* Texture2D, int32 NewMipCount, int32 NewSizeX, int32 NewSizeY, FThreadSafeCounter* RequestStatus) override;
 	virtual ETextureReallocationStatus RHIFinalizeAsyncReallocateTexture2D(FRHITexture2D* Texture2D, bool bBlockUntilCompleted) final override;
 	virtual ETextureReallocationStatus RHICancelAsyncReallocateTexture2D(FRHITexture2D* Texture2D, bool bBlockUntilCompleted) final override;
-	virtual void* RHILockTexture2D(FRHITexture2D* Texture, uint32 MipIndex, EResourceLockMode LockMode, uint32& DestStride, bool bLockWithinMiptail) final override;
+	virtual void* RHILockTexture2D(FRHITexture2D* Texture, uint32 MipIndex, EResourceLockMode LockMode, uint32& DestStride, bool bLockWithinMiptail, uint64* OutLockedByteCount = nullptr) final override;
 	virtual void RHIUnlockTexture2D(FRHITexture2D* Texture, uint32 MipIndex, bool bLockWithinMiptail) final override;
 	virtual void* RHILockTexture2DArray(FRHITexture2DArray* Texture, uint32 TextureIndex, uint32 MipIndex, EResourceLockMode LockMode, uint32& DestStride, bool bLockWithinMiptail) final override;
 	virtual void RHIUnlockTexture2DArray(FRHITexture2DArray* Texture, uint32 TextureIndex, uint32 MipIndex, bool bLockWithinMiptail) final override;
@@ -551,7 +551,7 @@ public:
 	// These will be un-commented as they are implemented.
 	//
 
-	virtual void* LockTexture2D_RenderThread(class FRHICommandListImmediate& RHICmdList, FRHITexture2D* Texture, uint32 MipIndex, EResourceLockMode LockMode, uint32& DestStride, bool bLockWithinMiptail, bool bNeedsDefaultRHIFlush = true) override final;
+	virtual void* LockTexture2D_RenderThread(class FRHICommandListImmediate& RHICmdList, FRHITexture2D* Texture, uint32 MipIndex, EResourceLockMode LockMode, uint32& DestStride, bool bLockWithinMiptail, bool bNeedsDefaultRHIFlush = true, uint64* OutLockedByteCount = nullptr) override final;
 	virtual void UnlockTexture2D_RenderThread(class FRHICommandListImmediate& RHICmdList, FRHITexture2D* Texture, uint32 MipIndex, bool bLockWithinMiptail, bool bNeedsDefaultRHIFlush = true) override final;
 
 	virtual void* LockTexture2DArray_RenderThread(class FRHICommandListImmediate& RHICmdList, FRHITexture2DArray* Texture, uint32 ArrayIndex, uint32 MipIndex, EResourceLockMode LockMode, uint32& DestStride, bool bLockWithinMiptail) override final;
@@ -619,7 +619,7 @@ public:
 	virtual void* CreateVirtualTexture(ETextureCreateFlags InFlags, D3D12_RESOURCE_DESC& ResourceDesc, const struct FD3D12TextureLayout& TextureLayout, FD3D12Resource** ppResource, FPlatformMemory::FPlatformVirtualMemoryBlock& RawTextureBlock, D3D12_RESOURCE_STATES InitialUsage = D3D12_RESOURCE_STATE_PIXEL_SHADER_RESOURCE) = 0;
 	virtual void DestroyVirtualTexture(ETextureCreateFlags InFlags, void* RawTextureMemory, FPlatformMemory::FPlatformVirtualMemoryBlock& RawTextureBlock, uint64 CommittedTextureSize) = 0;
 #endif
-	virtual bool HandleSpecialLock(void*& MemoryOut, uint32 MipIndex, uint32 ArrayIndex, FD3D12Texture* InTexture, EResourceLockMode LockMode, uint32& DestStride) { return false; }
+	virtual bool HandleSpecialLock(void*& MemoryOut, uint32 MipIndex, uint32 ArrayIndex, FD3D12Texture* InTexture, EResourceLockMode LockMode, uint32& DestStride, uint64* OutLockedByteCount = nullptr) { return false; }
 	virtual bool HandleSpecialUnlock(FRHICommandListBase* RHICmdList, uint32 MipIndex, FD3D12Texture* InTexture) { return false; }
 
 	FD3D12Adapter& GetAdapter(uint32_t Index = 0) { return *ChosenAdapters[Index]; }
diff --git a/Engine/Source/Runtime/D3D12RHI/Private/D3D12Texture.cpp b/Engine/Source/Runtime/D3D12RHI/Private/D3D12Texture.cpp
--- a/Engine/Source/Runtime/D3D12RHI/Private/D3D12Texture.cpp
+++ b/Engine/Source/Runtime/D3D12RHI/Private/D3D12Texture.cpp
@@ -1790,7 +1790,7 @@ void FD3D12Texture::AliasResources(FD3D12Texture* Texture)
 	}
 }
 
-void* FD3D12Texture::Lock(class FRHICommandListImmediate* RHICmdList, uint32 MipIndex, uint32 ArrayIndex, EResourceLockMode LockMode, uint32& DestStride)
+void* FD3D12Texture::Lock(class FRHICommandListImmediate* RHICmdList, uint32 MipIndex, uint32 ArrayIndex, EResourceLockMode LockMode, uint32& DestStride, uint64* OutLockedByteCount)
 {
 	SCOPE_CYCLE_COUNTER(STAT_D3D12LockTextureTime);
 
@@ -1818,6 +1818,11 @@ void* FD3D12Texture::Lock(class FRHICommandListImmediate* RHICmdList, uint32 Mip
 	const uint32 XBytesAligned = Align(NumBlocksX * BlockBytes, FD3D12_TEXTURE_DATA_PITCH_ALIGNMENT);
 	const uint32 MipBytesAligned = XBytesAligned * NumBlocksY;
 
+	if (OutLockedByteCount)
+	{
+		*OutLockedByteCount = MipBytesAligned;
+	}
+
 	FD3D12CommandContext& Context = Device->GetDefaultCommandContext();
 
 #if !PLATFORM_SUPPORTS_VIRTUAL_TEXTURES
@@ -1826,7 +1831,7 @@ void* FD3D12Texture::Lock(class FRHICommandListImmediate* RHICmdList, uint32 Mip
 
 	void* Data = nullptr;
 
-	if (FD3D12DynamicRHI::GetD3DRHI()->HandleSpecialLock(Data, MipIndex, ArrayIndex, this, LockMode, DestStride))
+	if (FD3D12DynamicRHI::GetD3DRHI()->HandleSpecialLock(Data, MipIndex, ArrayIndex, this, LockMode, DestStride, OutLockedByteCount))
 	{
 		// nothing left to do...
 		check(Data != nullptr);
@@ -2236,25 +2241,25 @@ void FD3D12Texture::GetReadBackHeapDesc(D3D12_PLACED_SUBRESOURCE_FOOTPRINT& OutF
 	}
 }
 
-void* FD3D12DynamicRHI::LockTexture2D_RenderThread(class FRHICommandListImmediate& RHICmdList, FRHITexture2D* TextureRHI, uint32 MipIndex, EResourceLockMode LockMode, uint32& DestStride, bool bLockWithinMiptail, bool bNeedsDefaultRHIFlush)
+void* FD3D12DynamicRHI::LockTexture2D_RenderThread(class FRHICommandListImmediate& RHICmdList, FRHITexture2D* TextureRHI, uint32 MipIndex, EResourceLockMode LockMode, uint32& DestStride, bool bLockWithinMiptail, bool bNeedsDefaultRHIFlush, uint64* OutLockedByteCount)
 {
 	if (CVarD3D12Texture2DRHIFlush.GetValueOnRenderThread() && bNeedsDefaultRHIFlush)
 	{
 		QUICK_SCOPE_CYCLE_COUNTER(STAT_RHIMETHOD_LockTexture2D_Flush);
 		RHICmdList.ImmediateFlush(EImmediateFlushType::FlushRHIThread);
-		return RHILockTexture2D(TextureRHI, MipIndex, LockMode, DestStride, bLockWithinMiptail);
+		return RHILockTexture2D(TextureRHI, MipIndex, LockMode, DestStride, bLockWithinMiptail, OutLockedByteCount);
 	}
 
 	check(TextureRHI);
 	FD3D12Texture* Texture = FD3D12DynamicRHI::ResourceCast(TextureRHI);
-	return Texture->Lock(&RHICmdList, MipIndex, 0, LockMode, DestStride);
+	return Texture->Lock(&RHICmdList, MipIndex, 0, LockMode, DestStride, OutLockedByteCount);
 }
 
-void* FD3D12DynamicRHI::RHILockTexture2D(FRHITexture2D* TextureRHI, uint32 MipIndex, EResourceLockMode LockMode, uint32& DestStride, bool bLockWithinMiptail)
+void* FD3D12DynamicRHI::RHILockTexture2D(FRHITexture2D* TextureRHI, uint32 MipIndex, EResourceLockMode LockMode, uint32& DestStride, bool bLockWithinMiptail, uint64* OutLockedByteCount)
 {
 	check(TextureRHI);
 	FD3D12Texture*  Texture = FD3D12DynamicRHI::ResourceCast(TextureRHI);
-	return Texture->Lock(nullptr, MipIndex, 0, LockMode, DestStride);
+	return Texture->Lock(nullptr, MipIndex, 0, LockMode, DestStride, OutLockedByteCount);
 }
 
 void FD3D12DynamicRHI::UnlockTexture2D_RenderThread(class FRHICommandListImmediate& RHICmdList, FRHITexture2D* TextureRHI, uint32 MipIndex, bool bLockWithinMiptail, bool bNeedsDefaultRHIFlush)
diff --git a/Engine/Source/Runtime/D3D12RHI/Private/D3D12Texture.h b/Engine/Source/Runtime/D3D12RHI/Private/D3D12Texture.h
--- a/Engine/Source/Runtime/D3D12RHI/Private/D3D12Texture.h
+++ b/Engine/Source/Runtime/D3D12RHI/Private/D3D12Texture.h
@@ -118,7 +118,7 @@ public:
 	}
 
 	// Locking/update functions
-	void* Lock(class FRHICommandListImmediate* RHICmdList, uint32 MipIndex, uint32 ArrayIndex, EResourceLockMode LockMode, uint32& DestStride);
+	void* Lock(class FRHICommandListImmediate* RHICmdList, uint32 MipIndex, uint32 ArrayIndex, EResourceLockMode LockMode, uint32& DestStride, uint64* OutLockedByteCount = nullptr);
 	void Unlock(class FRHICommandListImmediate* RHICmdList, uint32 MipIndex, uint32 ArrayIndex);
 	void UpdateTexture2D(FRHICommandListBase& RHICmdList, uint32 MipIndex, const struct FUpdateTextureRegion2D& UpdateRegion, uint32 SourcePitch, const uint8* SourceData);
 	void UpdateTexture(uint32 MipIndex, uint32 DestX, uint32 DestY, uint32 DestZ, const D3D12_TEXTURE_COPY_LOCATION& SourceCopyLocation);
diff --git a/Engine/Source/Runtime/Engine/Private/Rendering/Texture2DResource.cpp b/Engine/Source/Runtime/Engine/Private/Rendering/Texture2DResource.cpp
--- a/Engine/Source/Runtime/Engine/Private/Rendering/Texture2DResource.cpp
+++ b/Engine/Source/Runtime/Engine/Private/Rendering/Texture2DResource.cpp
@@ -140,8 +140,11 @@ void FTexture2DResource::CreateTexture()
 			if (MipData[ResourceMipIdx])
 			{
 				uint32 DestPitch = -1;
-				void* TheMipData = RHILockTexture2D(TextureRHI, RHIMipIdx, RLM_WriteOnly, DestPitch, false );
-				GetData( ResourceMipIdx, TheMipData, DestPitch );
+				uint64 Size = ~0ULL;
+				void* TheMipData = RHILockTexture2D(TextureRHI, RHIMipIdx, RLM_WriteOnly, DestPitch, false /* bLockWithinMiptail */, true /* bFlushRHIThread */, &Size);
+				check(Size != ~0ULL);
+
+				GetData( ResourceMipIdx, TheMipData, DestPitch, Size );
 				RHIUnlockTexture2D(TextureRHI, RHIMipIdx, false );
 			}
 		}
@@ -176,8 +179,11 @@ void FTexture2DResource::CreatePartiallyResidentTexture()
 		if ( MipData[MipIndex] != NULL )
 		{
 			uint32 DestPitch = -1;
-			void* TheMipData = RHILockTexture2D(TextureRHI, MipIndex, RLM_WriteOnly, DestPitch, false );
-			GetData( MipIndex, TheMipData, DestPitch );
+			uint64 Size = ~0ULL;
+			void* TheMipData = RHILockTexture2D(TextureRHI, MipIndex, RLM_WriteOnly, DestPitch, false /* bLockWithinMiptail */, true /* bFlushRHIThread */, &Size);
+			check(Size != ~0ULL);
+
+			GetData( MipIndex, TheMipData, DestPitch, Size);
 			RHIUnlockTexture2D(TextureRHI, MipIndex, false );
 		}
 	}
@@ -248,8 +254,9 @@ void FTexture2DResource::WarnRequiresTightPackedMip(int32 SizeX,int32 SizeY,EPix
  * @param MipIndex		Index of the mip-level to read.
  * @param Dest			Address of the destination buffer to receive the mip-level's data.
  * @param DestPitch		Number of bytes per row
+ * @param DestSize		Number of bytes locked by RHILockTexture2D.
  */
-void FTexture2DResource::GetData( uint32 MipIndex, void* Dest, uint32 DestPitch )
+void FTexture2DResource::GetData( uint32 MipIndex, void* Dest, uint32 DestPitch, uint64 DestSize )
 {
 	const FTexture2DMipMap& MipMap = *GetPlatformMip(MipIndex);
 	check( MipData[MipIndex] != nullptr );
@@ -259,10 +266,12 @@ void FTexture2DResource::GetData( uint32 MipIndex, void* Dest, uint32 DestPitch
 
 	int64 BulkDataSize = MipMap.BulkData.GetBulkDataSize();
 
-	UE_LOG(LogTextureUpload,Verbose,TEXT("Size: %dx%d , EffectiveSize=%d BulkDataSize=%d , SrcPitch=%d DestPitch=%d"),
+	UE_LOG(LogTextureUpload,Verbose,TEXT("Size: %dx%d , EffectiveSize=%d BulkDataSize=%d , SrcPitch=%d DestPitch=%d Format=%d, DestSize=%lld"),
 		MipMap.SizeX,MipMap.SizeY,
 		EffectiveSize,(int)BulkDataSize,
-		SrcPitch,DestPitch);
+		SrcPitch,DestPitch, PixelFormat, DestSize);
+
+	check(DestSize >= (uint64)BulkDataSize);
 
 	// for platforms that returned 0 pitch from Lock, we need to just use the bulk data directly, never do 
 	// runtime block size checking, conversion, or the like
diff --git a/Engine/Source/Runtime/Engine/Public/Rendering/Texture2DResource.h b/Engine/Source/Runtime/Engine/Public/Rendering/Texture2DResource.h
--- a/Engine/Source/Runtime/Engine/Public/Rendering/Texture2DResource.h
+++ b/Engine/Source/Runtime/Engine/Public/Rendering/Texture2DResource.h
@@ -80,5 +80,5 @@ private:
 	 * @param Dest		The address of the destination buffer to receive the mip-level's data.
 	 * @param DestPitch	Number of bytes per row
 	 */
-	void GetData( uint32 MipIndex,void* Dest,uint32 DestPitch );
+	void GetData( uint32 MipIndex,void* Dest,uint32 DestPitch, uint64 DestSize );
 };
diff --git a/Engine/Source/Runtime/NullDrv/Private/NullRHI.cpp b/Engine/Source/Runtime/NullDrv/Private/NullRHI.cpp
--- a/Engine/Source/Runtime/NullDrv/Private/NullRHI.cpp
+++ b/Engine/Source/Runtime/NullDrv/Private/NullRHI.cpp
@@ -76,9 +76,13 @@ void* FNullDynamicRHI::GetStaticBuffer(size_t Size)
 	return MemoryBuffer.GetData();
 }
 
-void* FNullDynamicRHI::GetStaticTextureBuffer(int32 SizeX, int32 SizeY, EPixelFormat Format, uint32& DestStride)
+void* FNullDynamicRHI::GetStaticTextureBuffer(int32 SizeX, int32 SizeY, EPixelFormat Format, uint32& DestStride, uint64* OutLockedByteCount)
 {
 	size_t Size = CalculateImageBytes(SizeX, SizeY, 0, Format);
+	if (OutLockedByteCount)
+	{
+		*OutLockedByteCount = Size;
+	}
 	DestStride = Size / SizeY;
 	return GetStaticBuffer(Size);
 }
diff --git a/Engine/Source/Runtime/NullDrv/Public/NullRHI.h b/Engine/Source/Runtime/NullDrv/Public/NullRHI.h
--- a/Engine/Source/Runtime/NullDrv/Public/NullRHI.h
+++ b/Engine/Source/Runtime/NullDrv/Public/NullRHI.h
@@ -232,10 +232,10 @@ public:
 	{ 
 		return TexRealloc_Succeeded; 
 	}
-	virtual void* RHILockTexture2D(FRHITexture2D* Texture, uint32 MipIndex, EResourceLockMode LockMode, uint32& DestStride, bool bLockWithinMiptail) final override
+	virtual void* RHILockTexture2D(FRHITexture2D* Texture, uint32 MipIndex, EResourceLockMode LockMode, uint32& DestStride, bool bLockWithinMiptail, uint64* OutLockedByteCount) final override
 	{ 
 		DestStride = 0; 
-		return GetStaticTextureBuffer(Texture->GetSizeX(), Texture->GetSizeY(), Texture->GetFormat(), DestStride);
+		return GetStaticTextureBuffer(Texture->GetSizeX(), Texture->GetSizeY(), Texture->GetFormat(), DestStride, OutLockedByteCount);
 	}
 	virtual void RHIUnlockTexture2D(FRHITexture2D* Texture, uint32 MipIndex, bool bLockWithinMiptail) final override
 	{
@@ -608,5 +608,5 @@ private:
 
 	/** Allocates a static buffer for RHI functions to return as a write destination. */
 	void* GetStaticBuffer(size_t Size);
-	void* GetStaticTextureBuffer(int32 SizeX, int32 SizeY, EPixelFormat Format, uint32& DestStride);
+	void* GetStaticTextureBuffer(int32 SizeX, int32 SizeY, EPixelFormat Format, uint32& DestStride, uint64* OutLockedByteCount = nullptr);
 };
diff --git a/Engine/Source/Runtime/OpenGLDrv/Private/OpenGLTexture.cpp b/Engine/Source/Runtime/OpenGLDrv/Private/OpenGLTexture.cpp
--- a/Engine/Source/Runtime/OpenGLDrv/Private/OpenGLTexture.cpp
+++ b/Engine/Source/Runtime/OpenGLDrv/Private/OpenGLTexture.cpp
@@ -1583,9 +1583,15 @@ ETextureReallocationStatus FOpenGLDynamicRHI::RHICancelAsyncReallocateTexture2D(
 	return TexRealloc_Succeeded;
 }
 
-void* FOpenGLDynamicRHI::RHILockTexture2D(FRHITexture2D* TextureRHI,uint32 MipIndex,EResourceLockMode LockMode,uint32& DestStride,bool bLockWithinMiptail)
+void* FOpenGLDynamicRHI::RHILockTexture2D(FRHITexture2D* TextureRHI,uint32 MipIndex,EResourceLockMode LockMode,uint32& DestStride,bool bLockWithinMiptail, uint64* OutLockedByteCount)
 {
-	return ResourceCast(TextureRHI)->Lock(MipIndex, 0, LockMode, DestStride);
+	FOpenGLTexture* Texture = ResourceCast(TextureRHI);
+	if (OutLockedByteCount)
+	{
+		*OutLockedByteCount = Texture->GetLockSize(MipIndex, 0, LockMode, DestStride);
+	}
+
+	return Texture->Lock(MipIndex, 0, LockMode, DestStride);
 }
 
 void FOpenGLDynamicRHI::RHIUnlockTexture2D(FRHITexture2D* TextureRHI, uint32 MipIndex, bool bLockWithinMiptail)
@@ -1982,7 +1988,7 @@ FTextureRHIRef FOpenGLDynamicRHI::RHICreateAliasedTexture(FTextureRHIRef& Source
 	return new FOpenGLTexture(*ResourceCast(SourceTexture), *Name, FOpenGLTexture::AliasResource);
 }
 
-void* FOpenGLDynamicRHI::LockTexture2D_RenderThread(class FRHICommandListImmediate& RHICmdList, FRHITexture2D* Texture, uint32 MipIndex, EResourceLockMode LockMode, uint32& DestStride, bool bLockWithinMiptail, bool bNeedsDefaultRHIFlush)
+void* FOpenGLDynamicRHI::LockTexture2D_RenderThread(class FRHICommandListImmediate& RHICmdList, FRHITexture2D* Texture, uint32 MipIndex, EResourceLockMode LockMode, uint32& DestStride, bool bLockWithinMiptail, bool bNeedsDefaultRHIFlush, uint64* OutLockedByteCount)
 {
 	check(IsInRenderingThread());
 	static auto* CVarRHICmdBufferWriteLocks = IConsoleManager::Get().FindTConsoleVariableDataInt(TEXT("r.RHICmdBufferWriteLocks"));
@@ -1992,7 +1998,7 @@ void* FOpenGLDynamicRHI::LockTexture2D_RenderThread(class FRHICommandListImmedia
 	if (!bBuffer || LockMode != RLM_WriteOnly || RHICmdList.Bypass() || !IsRunningRHIInSeparateThread())
 	{
 		RHITHREAD_GLCOMMAND_PROLOGUE();
-		return this->RHILockTexture2D(Texture, MipIndex, LockMode, DestStride, bLockWithinMiptail);
+		return this->RHILockTexture2D(Texture, MipIndex, LockMode, DestStride, bLockWithinMiptail, OutLockedByteCount);
 		RHITHREAD_GLCOMMAND_EPILOGUE_GET_RETURN(void *);
 		Result = ReturnValue;
 		MipBytes = ResourceCast(Texture)->GetLockSize(MipIndex, 0, LockMode, DestStride);
@@ -2004,6 +2010,11 @@ void* FOpenGLDynamicRHI::LockTexture2D_RenderThread(class FRHICommandListImmedia
 	}
 	check(Result);
 
+	if (OutLockedByteCount)
+	{
+		*OutLockedByteCount = MipBytes;
+	}
+
 	GLLockTracker.Lock(Texture, Result, 0, MipIndex, DestStride, MipBytes, LockMode);
 	return Result;
 }
@@ -2026,9 +2037,12 @@ void FOpenGLDynamicRHI::UnlockTexture2D_RenderThread(class FRHICommandListImmedi
 		auto GLCommand = [=]()
 		{
 			uint32 DestStride;
-			uint8* TexMem = (uint8*)this->RHILockTexture2D(Texture, MipIndex, Params.LockMode, DestStride, bLockWithinMiptail);
+			uint64 LockedByteCount = ~0ULL;
+			uint8* TexMem = (uint8*)this->RHILockTexture2D(Texture, MipIndex, Params.LockMode, DestStride, bLockWithinMiptail, &LockedByteCount);
+			check(LockedByteCount != ~0ULL);
 			uint8* BuffMem = (uint8*)Params.Buffer;
 			check(DestStride == Params.Stride);
+			check(LockedByteCount >= Params.BufferSize);
 			FMemory::Memcpy(TexMem, BuffMem, Params.BufferSize);
 			FMemory::Free(Params.Buffer);
 			this->RHIUnlockTexture2D(Texture, MipIndex, bLockWithinMiptail);
diff --git a/Engine/Source/Runtime/OpenGLDrv/Public/OpenGLDrv.h b/Engine/Source/Runtime/OpenGLDrv/Public/OpenGLDrv.h
--- a/Engine/Source/Runtime/OpenGLDrv/Public/OpenGLDrv.h
+++ b/Engine/Source/Runtime/OpenGLDrv/Public/OpenGLDrv.h
@@ -374,7 +374,7 @@ public:
 	virtual FTexture2DRHIRef RHIAsyncReallocateTexture2D(FRHITexture2D* Texture2D, int32 NewMipCount, int32 NewSizeX, int32 NewSizeY, FThreadSafeCounter* RequestStatus) final override;
 	virtual ETextureReallocationStatus RHIFinalizeAsyncReallocateTexture2D(FRHITexture2D* Texture2D, bool bBlockUntilCompleted) final override;
 	virtual ETextureReallocationStatus RHICancelAsyncReallocateTexture2D(FRHITexture2D* Texture2D, bool bBlockUntilCompleted) final override;
-	virtual void* RHILockTexture2D(FRHITexture2D* Texture, uint32 MipIndex, EResourceLockMode LockMode, uint32& DestStride, bool bLockWithinMiptail) final override;
+	virtual void* RHILockTexture2D(FRHITexture2D* Texture, uint32 MipIndex, EResourceLockMode LockMode, uint32& DestStride, bool bLockWithinMiptail, uint64* OutLockedByteCount = nullptr) final override;
 	virtual void RHIUnlockTexture2D(FRHITexture2D* Texture, uint32 MipIndex, bool bLockWithinMiptail) final override;
 	virtual void* RHILockTexture2DArray(FRHITexture2DArray* Texture, uint32 TextureIndex, uint32 MipIndex, EResourceLockMode LockMode, uint32& DestStride, bool bLockWithinMiptail) final override;
 	virtual void RHIUnlockTexture2DArray(FRHITexture2DArray* Texture, uint32 TextureIndex, uint32 MipIndex, bool bLockWithinMiptail) final override;
@@ -733,7 +733,7 @@ public:
 		}
 	};
 
-	virtual void* LockTexture2D_RenderThread(class FRHICommandListImmediate& RHICmdList, FRHITexture2D* Texture, uint32 MipIndex, EResourceLockMode LockMode, uint32& DestStride, bool bLockWithinMiptail, bool bNeedsDefaultRHIFlush) final override;
+	virtual void* LockTexture2D_RenderThread(class FRHICommandListImmediate& RHICmdList, FRHITexture2D* Texture, uint32 MipIndex, EResourceLockMode LockMode, uint32& DestStride, bool bLockWithinMiptail, bool bNeedsDefaultRHIFlush, uint64* OutLockedByteCount = nullptr) final override;
 	virtual void UnlockTexture2D_RenderThread(class FRHICommandListImmediate& RHICmdList, FRHITexture2D* Texture, uint32 MipIndex, bool bLockWithinMiptail, bool bNeedsDefaultRHIFlush) final override;
 	virtual void* RHILockTextureCubeFace_RenderThread(class FRHICommandListImmediate& RHICmdList, FRHITextureCube* Texture, uint32 FaceIndex, uint32 ArrayIndex, uint32 MipIndex, EResourceLockMode LockMode, uint32& DestStride, bool bLockWithinMiptail) final override;
 	virtual void RHIUnlockTextureCubeFace_RenderThread(class FRHICommandListImmediate& RHICmdList, FRHITextureCube* Texture, uint32 FaceIndex, uint32 ArrayIndex, uint32 MipIndex, bool bLockWithinMiptail) final override;
diff --git a/Engine/Source/Runtime/RHI/Private/RHICommandList.cpp b/Engine/Source/Runtime/RHI/Private/RHICommandList.cpp
--- a/Engine/Source/Runtime/RHI/Private/RHICommandList.cpp
+++ b/Engine/Source/Runtime/RHI/Private/RHICommandList.cpp
@@ -1665,18 +1665,18 @@ void FDynamicRHI::RHIEndMultiUpdateTexture3D(FRHICommandListBase& RHICmdList, TA
 	}
 }
 
-void* FDynamicRHI::LockTexture2D_RenderThread(class FRHICommandListImmediate& RHICmdList, FRHITexture2D* Texture, uint32 MipIndex, EResourceLockMode LockMode, uint32& DestStride, bool bLockWithinMiptail, bool bNeedsDefaultRHIFlush)
+void* FDynamicRHI::LockTexture2D_RenderThread(class FRHICommandListImmediate& RHICmdList, FRHITexture2D* Texture, uint32 MipIndex, EResourceLockMode LockMode, uint32& DestStride, bool bLockWithinMiptail, bool bNeedsDefaultRHIFlush, uint64* OutLockedByteCount)
 {
 	if (bNeedsDefaultRHIFlush) 
 	{
 		QUICK_SCOPE_CYCLE_COUNTER(STAT_RHIMETHOD_LockTexture2D_Flush);
 		CSV_SCOPED_TIMING_STAT(RHITFlushes, LockTexture2D_RenderThread);
 		RHICmdList.ImmediateFlush(EImmediateFlushType::FlushRHIThread);
-		return GDynamicRHI->RHILockTexture2D(Texture, MipIndex, LockMode, DestStride, bLockWithinMiptail);
+		return GDynamicRHI->RHILockTexture2D(Texture, MipIndex, LockMode, DestStride, bLockWithinMiptail, OutLockedByteCount);
 	}
 	CSV_SCOPED_TIMING_STAT(RHITStalls, LockTexture2D_RenderThread);
 	FScopedRHIThreadStaller StallRHIThread(RHICmdList);
-	return GDynamicRHI->RHILockTexture2D(Texture, MipIndex, LockMode, DestStride, bLockWithinMiptail);
+	return GDynamicRHI->RHILockTexture2D(Texture, MipIndex, LockMode, DestStride, bLockWithinMiptail, OutLockedByteCount);
 }
 
 void FDynamicRHI::UnlockTexture2D_RenderThread(class FRHICommandListImmediate& RHICmdList, FRHITexture2D* Texture, uint32 MipIndex, bool bLockWithinMiptail, bool bNeedsDefaultRHIFlush)
diff --git a/Engine/Source/Runtime/RHI/Public/DynamicRHI.h b/Engine/Source/Runtime/RHI/Public/DynamicRHI.h
--- a/Engine/Source/Runtime/RHI/Public/DynamicRHI.h
+++ b/Engine/Source/Runtime/RHI/Public/DynamicRHI.h
@@ -506,7 +506,7 @@ public:
 	* @return pointer to the CPU accessible resource data
 	*/
 	// FlushType: Flush RHI Thread
-	virtual void* RHILockTexture2D(FRHITexture2D* Texture, uint32 MipIndex, EResourceLockMode LockMode, uint32& DestStride, bool bLockWithinMiptail) = 0;
+	virtual void* RHILockTexture2D(FRHITexture2D* Texture, uint32 MipIndex, EResourceLockMode LockMode, uint32& DestStride, bool bLockWithinMiptail, uint64* OutLockedByteCount = nullptr) = 0;
 
 	/**
 	* Unlocks a previously locked RHI texture resource
@@ -900,7 +900,7 @@ public:
 	RHI_API virtual ETextureReallocationStatus FinalizeAsyncReallocateTexture2D_RenderThread(class FRHICommandListImmediate& RHICmdList, FRHITexture2D* Texture2D, bool bBlockUntilCompleted);
 	RHI_API virtual ETextureReallocationStatus CancelAsyncReallocateTexture2D_RenderThread(class FRHICommandListImmediate& RHICmdList, FRHITexture2D* Texture2D, bool bBlockUntilCompleted);
 
-	RHI_API virtual void* LockTexture2D_RenderThread(class FRHICommandListImmediate& RHICmdList, FRHITexture2D* Texture, uint32 MipIndex, EResourceLockMode LockMode, uint32& DestStride, bool bLockWithinMiptail, bool bNeedsDefaultRHIFlush = true);
+	RHI_API virtual void* LockTexture2D_RenderThread(class FRHICommandListImmediate& RHICmdList, FRHITexture2D* Texture, uint32 MipIndex, EResourceLockMode LockMode, uint32& DestStride, bool bLockWithinMiptail, bool bNeedsDefaultRHIFlush = true, uint64* OutLockedByteCount = nullptr);
 	RHI_API virtual void UnlockTexture2D_RenderThread(class FRHICommandListImmediate& RHICmdList, FRHITexture2D* Texture, uint32 MipIndex, bool bLockWithinMiptail, bool bNeedsDefaultRHIFlush = true);
 
 	RHI_API virtual void* LockTexture2DArray_RenderThread(class FRHICommandListImmediate& RHICmdList, FRHITexture2DArray* Texture, uint32 ArrayIndex, uint32 MipIndex, EResourceLockMode LockMode, uint32& DestStride, bool bLockWithinMiptail);
diff --git a/Engine/Source/Runtime/RHI/Public/RHICommandList.h b/Engine/Source/Runtime/RHI/Public/RHICommandList.h
--- a/Engine/Source/Runtime/RHI/Public/RHICommandList.h
+++ b/Engine/Source/Runtime/RHI/Public/RHICommandList.h
@@ -4117,10 +4117,10 @@ public:
 		return GDynamicRHI->CancelAsyncReallocateTexture2D_RenderThread(*this, Texture2D, bBlockUntilCompleted);
 	}
 	
-	FORCEINLINE void* LockTexture2D(FRHITexture2D* Texture, uint32 MipIndex, EResourceLockMode LockMode, uint32& DestStride, bool bLockWithinMiptail, bool bFlushRHIThread = true)
+	FORCEINLINE void* LockTexture2D(FRHITexture2D* Texture, uint32 MipIndex, EResourceLockMode LockMode, uint32& DestStride, bool bLockWithinMiptail, bool bFlushRHIThread = true, uint64* OutLockedByteCount = nullptr)
 	{
 		LLM_SCOPE(ELLMTag::Textures);
-		return GDynamicRHI->LockTexture2D_RenderThread(*this, Texture, MipIndex, LockMode, DestStride, bLockWithinMiptail, bFlushRHIThread);
+		return GDynamicRHI->LockTexture2D_RenderThread(*this, Texture, MipIndex, LockMode, DestStride, bLockWithinMiptail, bFlushRHIThread, OutLockedByteCount);
 	}
 	
 	FORCEINLINE void UnlockTexture2D(FRHITexture2D* Texture, uint32 MipIndex, bool bLockWithinMiptail, bool bFlushRHIThread = true)
@@ -5054,9 +5054,9 @@ FORCEINLINE ETextureReallocationStatus RHICancelAsyncReallocateTexture2D(FRHITex
 	return FRHICommandListExecutor::GetImmediateCommandList().CancelAsyncReallocateTexture2D(Texture2D, bBlockUntilCompleted);
 }
 
-FORCEINLINE void* RHILockTexture2D(FRHITexture2D* Texture, uint32 MipIndex, EResourceLockMode LockMode, uint32& DestStride, bool bLockWithinMiptail, bool bFlushRHIThread = true)
+FORCEINLINE void* RHILockTexture2D(FRHITexture2D* Texture, uint32 MipIndex, EResourceLockMode LockMode, uint32& DestStride, bool bLockWithinMiptail, bool bFlushRHIThread = true, uint64* OutLockedByteCount = nullptr)
 {
-	return FRHICommandListExecutor::GetImmediateCommandList().LockTexture2D(Texture, MipIndex, LockMode, DestStride, bLockWithinMiptail, bFlushRHIThread);
+	return FRHICommandListExecutor::GetImmediateCommandList().LockTexture2D(Texture, MipIndex, LockMode, DestStride, bLockWithinMiptail, bFlushRHIThread, OutLockedByteCount);
 }
 
 FORCEINLINE void RHIUnlockTexture2D(FRHITexture2D* Texture, uint32 MipIndex, bool bLockWithinMiptail, bool bFlushRHIThread = true)
diff --git a/Engine/Source/Runtime/RHI/Public/RHIValidation.h b/Engine/Source/Runtime/RHI/Public/RHIValidation.h
--- a/Engine/Source/Runtime/RHI/Public/RHIValidation.h
+++ b/Engine/Source/Runtime/RHI/Public/RHIValidation.h
@@ -624,9 +624,9 @@ public:
 	* @return pointer to the CPU accessible resource data
 	*/
 	// FlushType: Flush RHI Thread
-	virtual void* RHILockTexture2D(FRHITexture2D* Texture, uint32 MipIndex, EResourceLockMode LockMode, uint32& DestStride, bool bLockWithinMiptail) override final
+	virtual void* RHILockTexture2D(FRHITexture2D* Texture, uint32 MipIndex, EResourceLockMode LockMode, uint32& DestStride, bool bLockWithinMiptail, uint64* OutLockedByteCount) override final
 	{
-		return RHI->RHILockTexture2D(Texture, MipIndex, LockMode, DestStride, bLockWithinMiptail);
+		return RHI->RHILockTexture2D(Texture, MipIndex, LockMode, DestStride, bLockWithinMiptail, OutLockedByteCount);
 	}
 
 	/**
@@ -1148,9 +1148,9 @@ public:
 		RHI->UnlockBuffer_BottomOfPipe(RHICmdList, Buffer);
 	}
 
-	virtual void* LockTexture2D_RenderThread(class FRHICommandListImmediate& RHICmdList, FRHITexture2D* Texture, uint32 MipIndex, EResourceLockMode LockMode, uint32& DestStride, bool bLockWithinMiptail, bool bNeedsDefaultRHIFlush = true) override final
+	virtual void* LockTexture2D_RenderThread(class FRHICommandListImmediate& RHICmdList, FRHITexture2D* Texture, uint32 MipIndex, EResourceLockMode LockMode, uint32& DestStride, bool bLockWithinMiptail, bool bNeedsDefaultRHIFlush = true, uint64* OutLockedByteCount = nullptr) override final
 	{
-		return RHI->LockTexture2D_RenderThread(RHICmdList, Texture, MipIndex, LockMode, DestStride, bLockWithinMiptail, bNeedsDefaultRHIFlush);
+		return RHI->LockTexture2D_RenderThread(RHICmdList, Texture, MipIndex, LockMode, DestStride, bLockWithinMiptail, bNeedsDefaultRHIFlush, OutLockedByteCount);
 	}
 
 	virtual void UnlockTexture2D_RenderThread(class FRHICommandListImmediate& RHICmdList, FRHITexture2D* Texture, uint32 MipIndex, bool bLockWithinMiptail, bool bNeedsDefaultRHIFlush = true) override final
diff --git a/Engine/Source/Runtime/VulkanRHI/Private/VulkanTexture.cpp b/Engine/Source/Runtime/VulkanRHI/Private/VulkanTexture.cpp
--- a/Engine/Source/Runtime/VulkanRHI/Private/VulkanTexture.cpp
+++ b/Engine/Source/Runtime/VulkanRHI/Private/VulkanTexture.cpp
@@ -1120,7 +1120,7 @@ ETextureReallocationStatus FVulkanDynamicRHI::RHICancelAsyncReallocateTexture2D(
 	return TexRealloc_Succeeded;
 }
 
-void* FVulkanDynamicRHI::RHILockTexture2D(FRHITexture2D* TextureRHI, uint32 MipIndex, EResourceLockMode LockMode, uint32& DestStride, bool bLockWithinMiptail)
+void* FVulkanDynamicRHI::RHILockTexture2D(FRHITexture2D* TextureRHI, uint32 MipIndex, EResourceLockMode LockMode, uint32& DestStride, bool bLockWithinMiptail, uint64* OutLockedByteCount)
 {
 	LLM_SCOPE_VULKAN(ELLMTagVulkan::VulkanTextures);
 	FVulkanTexture* Texture = ResourceCast(TextureRHI);
@@ -1142,6 +1142,11 @@ void* FVulkanDynamicRHI::RHILockTexture2D(FRHITexture2D* TextureRHI, uint32 MipI
 	Texture->GetMipStride(MipIndex, DestStride);
 	*StagingBuffer = Device->GetStagingManager().AcquireBuffer(BufferSize);
 
+	if (OutLockedByteCount)
+	{
+		*OutLockedByteCount = BufferSize;
+	}
+
 	void* Data = (*StagingBuffer)->GetMappedPointer();
 	return Data;
 }
diff --git a/Engine/Source/Runtime/VulkanRHI/Public/VulkanDynamicRHI.h b/Engine/Source/Runtime/VulkanRHI/Public/VulkanDynamicRHI.h
--- a/Engine/Source/Runtime/VulkanRHI/Public/VulkanDynamicRHI.h
+++ b/Engine/Source/Runtime/VulkanRHI/Public/VulkanDynamicRHI.h
@@ -135,7 +135,7 @@ public:
 	}
 
 	virtual ETextureReallocationStatus RHICancelAsyncReallocateTexture2D(FRHITexture2D* Texture2D, bool bBlockUntilCompleted) final override;
-	virtual void* RHILockTexture2D(FRHITexture2D* Texture, uint32 MipIndex, EResourceLockMode LockMode, uint32& DestStride, bool bLockWithinMiptail) final override;
+	virtual void* RHILockTexture2D(FRHITexture2D* Texture, uint32 MipIndex, EResourceLockMode LockMode, uint32& DestStride, bool bLockWithinMiptail, uint64* OutLockedByteCount = nullptr) final override;
 	virtual void RHIUnlockTexture2D(FRHITexture2D* Texture, uint32 MipIndex, bool bLockWithinMiptail) final override
 	{
 		InternalUnlockTexture2D(false, Texture, MipIndex, bLockWithinMiptail);
@@ -200,9 +200,9 @@ public:
 	virtual FTexture2DRHIRef AsyncReallocateTexture2D_RenderThread(class FRHICommandListImmediate& RHICmdList, FRHITexture2D* Texture2D, int32 NewMipCount, int32 NewSizeX, int32 NewSizeY, FThreadSafeCounter* RequestStatus) override final;
 	//virtual ETextureReallocationStatus CancelAsyncReallocateTexture2D_RenderThread(class FRHICommandListImmediate& RHICmdList, FRHITexture2D* Texture2D, bool bBlockUntilCompleted) override final;
 	
-	virtual void* LockTexture2D_RenderThread(class FRHICommandListImmediate& RHICmdList, FRHITexture2D* Texture, uint32 MipIndex, EResourceLockMode LockMode, uint32& DestStride, bool bLockWithinMiptail, bool bNeedsDefaultRHIFlush = true) override final
+	virtual void* LockTexture2D_RenderThread(class FRHICommandListImmediate& RHICmdList, FRHITexture2D* Texture, uint32 MipIndex, EResourceLockMode LockMode, uint32& DestStride, bool bLockWithinMiptail, bool bNeedsDefaultRHIFlush = true, uint64* OutLockedByteCount = nullptr) override final
 	{
-		return RHILockTexture2D(Texture, MipIndex, LockMode, DestStride, bLockWithinMiptail);
+		return RHILockTexture2D(Texture, MipIndex, LockMode, DestStride, bLockWithinMiptail, OutLockedByteCount);
 	}
 	virtual void UnlockTexture2D_RenderThread(class FRHICommandListImmediate& RHICmdList, FRHITexture2D* Texture, uint32 MipIndex, bool bLockWithinMiptail, bool bNeedsDefaultRHIFlush = true) override final
 	{
diff --git a/Engine/Source/Runtime/Windows/D3D11RHI/Private/D3D11RHIPrivate.h b/Engine/Source/Runtime/Windows/D3D11RHI/Private/D3D11RHIPrivate.h
--- a/Engine/Source/Runtime/Windows/D3D11RHI/Private/D3D11RHIPrivate.h
+++ b/Engine/Source/Runtime/Windows/D3D11RHI/Private/D3D11RHIPrivate.h
@@ -506,7 +506,7 @@ public:
 	virtual FTexture2DRHIRef RHIAsyncReallocateTexture2D(FRHITexture2D* Texture2D, int32 NewMipCount, int32 NewSizeX, int32 NewSizeY, FThreadSafeCounter* RequestStatus) final override;
 	virtual ETextureReallocationStatus RHIFinalizeAsyncReallocateTexture2D(FRHITexture2D* Texture2D, bool bBlockUntilCompleted) final override;
 	virtual ETextureReallocationStatus RHICancelAsyncReallocateTexture2D(FRHITexture2D* Texture2D, bool bBlockUntilCompleted) final override;
-	virtual void* RHILockTexture2D(FRHITexture2D* Texture, uint32 MipIndex, EResourceLockMode LockMode, uint32& DestStride, bool bLockWithinMiptail) final override;
+	virtual void* RHILockTexture2D(FRHITexture2D* Texture, uint32 MipIndex, EResourceLockMode LockMode, uint32& DestStride, bool bLockWithinMiptail, uint64* OutLockedByteCount) final override;
 	virtual void RHIUnlockTexture2D(FRHITexture2D* Texture, uint32 MipIndex, bool bLockWithinMiptail) final override;
 	virtual void* RHILockTexture2DArray(FRHITexture2DArray* Texture, uint32 TextureIndex, uint32 MipIndex, EResourceLockMode LockMode, uint32& DestStride, bool bLockWithinMiptail) final override;
 	virtual void RHIUnlockTexture2DArray(FRHITexture2DArray* Texture, uint32 TextureIndex, uint32 MipIndex, bool bLockWithinMiptail) final override;
@@ -640,7 +640,7 @@ public:
 	virtual FTexture2DRHIRef AsyncReallocateTexture2D_RenderThread(class FRHICommandListImmediate& RHICmdList, FRHITexture2D* Texture2D, int32 NewMipCount, int32 NewSizeX, int32 NewSizeY, FThreadSafeCounter* RequestStatus) final override;
 	virtual ETextureReallocationStatus FinalizeAsyncReallocateTexture2D_RenderThread(class FRHICommandListImmediate& RHICmdList, FRHITexture2D* Texture2D, bool bBlockUntilCompleted) final override;
 	virtual ETextureReallocationStatus CancelAsyncReallocateTexture2D_RenderThread(class FRHICommandListImmediate& RHICmdList, FRHITexture2D* Texture2D, bool bBlockUntilCompleted) final override;
-	virtual void* LockTexture2D_RenderThread(class FRHICommandListImmediate& RHICmdList, FRHITexture2D* Texture, uint32 MipIndex, EResourceLockMode LockMode, uint32& DestStride, bool bLockWithinMiptail, bool bNeedsDefaultRHIFlush = true) final override;
+	virtual void* LockTexture2D_RenderThread(class FRHICommandListImmediate& RHICmdList, FRHITexture2D* Texture, uint32 MipIndex, EResourceLockMode LockMode, uint32& DestStride, bool bLockWithinMiptail, bool bNeedsDefaultRHIFlush = true, uint64* OutLockedByteCount = nullptr) final override;
 	virtual void UnlockTexture2D_RenderThread(class FRHICommandListImmediate& RHICmdList, FRHITexture2D* Texture, uint32 MipIndex, bool bLockWithinMiptail, bool bNeedsDefaultRHIFlush = true) final override;
 	virtual void RHIEndUpdateTexture3D(FRHICommandListBase& RHICmdList, FUpdateTexture3DData& UpdateData) final override;
 	virtual void* RHILockTextureCubeFace_RenderThread(class FRHICommandListImmediate& RHICmdList, FRHITextureCube* Texture, uint32 FaceIndex, uint32 ArrayIndex, uint32 MipIndex, EResourceLockMode LockMode, uint32& DestStride, bool bLockWithinMiptail) final override;
diff --git a/Engine/Source/Runtime/Windows/D3D11RHI/Private/D3D11Texture.cpp b/Engine/Source/Runtime/Windows/D3D11RHI/Private/D3D11Texture.cpp
--- a/Engine/Source/Runtime/Windows/D3D11RHI/Private/D3D11Texture.cpp
+++ b/Engine/Source/Runtime/Windows/D3D11RHI/Private/D3D11Texture.cpp
@@ -1180,7 +1180,7 @@ ETextureReallocationStatus FD3D11DynamicRHI::CancelAsyncReallocateTexture2D_Rend
 	return RHICancelAsyncReallocateTexture2D(Texture2D, bBlockUntilCompleted);
 }
 
-void* FD3D11Texture::Lock(FD3D11DynamicRHI* D3DRHI, uint32 MipIndex, uint32 ArrayIndex, EResourceLockMode LockMode, uint32& DestStride, bool bForceLockDeferred)
+void* FD3D11Texture::Lock(FD3D11DynamicRHI* D3DRHI, uint32 MipIndex, uint32 ArrayIndex, EResourceLockMode LockMode, uint32& DestStride, bool bForceLockDeferred, uint64* OutLockedByteCount)
 {
 	check(!IsTexture3D()); // Only 2D texture locks are implemented
 
@@ -1199,6 +1199,11 @@ void* FD3D11Texture::Lock(FD3D11DynamicRHI* D3DRHI, uint32 MipIndex, uint32 Arra
 	const uint32 NumBlocksY = (MipSizeY + BlockSizeY - 1) / BlockSizeY;
 	const uint32 MipBytes = NumBlocksX * NumBlocksY * BlockBytes;
 
+	if (OutLockedByteCount)
+	{
+		*OutLockedByteCount = MipBytes;
+	}
+
 	FD3D11LockedData LockedData;
 	if( LockMode == RLM_WriteOnly )
 	{
@@ -1310,12 +1315,12 @@ void FD3D11Texture::Unlock(FD3D11DynamicRHI* D3DRHI, uint32 MipIndex, uint32 Arr
 	}
 }
 
-void* FD3D11DynamicRHI::RHILockTexture2D(FRHITexture2D* TextureRHI, uint32 MipIndex, EResourceLockMode LockMode, uint32& DestStride, bool bLockWithinMiptail)
+void* FD3D11DynamicRHI::RHILockTexture2D(FRHITexture2D* TextureRHI, uint32 MipIndex, EResourceLockMode LockMode, uint32& DestStride, bool bLockWithinMiptail, uint64* OutLockedByteCount)
 {
 	check(TextureRHI);
 	FD3D11Texture* Texture = ResourceCast(TextureRHI);
 	ConditionalClearShaderResource(Texture, false);
-	return Texture->Lock(this, MipIndex, 0, LockMode, DestStride);
+	return Texture->Lock(this, MipIndex, 0, LockMode, DestStride, false /* bForceLockDeferred */, OutLockedByteCount);
 }
 
 void* FD3D11DynamicRHI::LockTexture2D_RenderThread(
@@ -1325,23 +1330,24 @@ void* FD3D11DynamicRHI::LockTexture2D_RenderThread(
 	EResourceLockMode LockMode,
 	uint32& DestStride,
 	bool bLockWithinMiptail,
-	bool bNeedsDefaultRHIFlush)
+	bool bNeedsDefaultRHIFlush,
+	uint64* OutLockedByteCount)
 {
 	void *LockedTexture = nullptr;
 
 	if (ShouldNotEnqueueRHICommand())
 	{
-		LockedTexture = RHILockTexture2D(Texture, MipIndex, LockMode, DestStride, bLockWithinMiptail);
+		LockedTexture = RHILockTexture2D(Texture, MipIndex, LockMode, DestStride, bLockWithinMiptail, OutLockedByteCount);
 	}
 	else if (LockMode == RLM_ReadOnly)
 	{
 		FRHICommandListExecutor::GetImmediateCommandList().ImmediateFlush(EImmediateFlushType::FlushRHIThread);
-		LockedTexture = RHILockTexture2D(Texture, MipIndex, LockMode, DestStride, bLockWithinMiptail);
+		LockedTexture = RHILockTexture2D(Texture, MipIndex, LockMode, DestStride, bLockWithinMiptail, OutLockedByteCount);
 	}
 	else
 	{
 		FD3D11Texture* TextureD3D11 = ResourceCast(Texture);
-		LockedTexture = TextureD3D11->Lock(this, MipIndex, 0, LockMode, DestStride, true);
+		LockedTexture = TextureD3D11->Lock(this, MipIndex, 0, LockMode, DestStride, true, OutLockedByteCount);
 	}
 	return LockedTexture;
 }
diff --git a/Engine/Source/Runtime/Windows/D3D11RHI/Public/D3D11Resources.h b/Engine/Source/Runtime/Windows/D3D11RHI/Public/D3D11Resources.h
--- a/Engine/Source/Runtime/Windows/D3D11RHI/Public/D3D11Resources.h
+++ b/Engine/Source/Runtime/Windows/D3D11RHI/Public/D3D11Resources.h
@@ -279,7 +279,7 @@ public:
 	* Locks one of the texture's mip-maps.
 	* @return A pointer to the specified texture data.
 	*/
-	void* Lock(class FD3D11DynamicRHI* D3DRHI, uint32 MipIndex, uint32 ArrayIndex, EResourceLockMode LockMode, uint32& DestStride, bool bForceLockDeferred = false);
+	void* Lock(class FD3D11DynamicRHI* D3DRHI, uint32 MipIndex, uint32 ArrayIndex, EResourceLockMode LockMode, uint32& DestStride, bool bForceLockDeferred = false, uint64* OutLockedByteCount = nullptr);
 
 	/** Unlocks a previously locked mip-map. */
 	void Unlock(class FD3D11DynamicRHI* D3DRHI, uint32 MipIndex, uint32 ArrayIndex);
