MetalRHI: Fix for incorrect timestamp queries calc on Apple Silicon.
Cleaned up hard coded GPU Vendor ID's

diff --git a/Engine/Source/Runtime/Apple/MetalRHI/Private/MetalCommandQueue.h b/Engine/Source/Runtime/Apple/MetalRHI/Private/MetalCommandQueue.h
--- a/Engine/Source/Runtime/Apple/MetalRHI/Private/MetalCommandQueue.h
+++ b/Engine/Source/Runtime/Apple/MetalRHI/Private/MetalCommandQueue.h
@@ -141,7 +141,7 @@ public:
 	static inline bool SupportsSeparateMSAAAndResolveTarget() { return (PLATFORM_MAC != 0 || GMaxRHIFeatureLevel >= ERHIFeatureLevel::SM5); }
 
 	/** @returns True on UMA system; false otherwise.  */
-	static inline bool IsUMASystem() { return (GRHIVendorId == 0x106B); }
+	static inline bool IsUMASystem() { return IsRHIDeviceApple(); }
 
 #pragma mark - Public Debug Support -
 
diff --git a/Engine/Source/Runtime/Apple/MetalRHI/Private/MetalContext.cpp b/Engine/Source/Runtime/Apple/MetalRHI/Private/MetalContext.cpp
--- a/Engine/Source/Runtime/Apple/MetalRHI/Private/MetalContext.cpp
+++ b/Engine/Source/Runtime/Apple/MetalRHI/Private/MetalContext.cpp
@@ -198,12 +198,12 @@ static mtlpp::Device GetMTLDevice(uint32& DeviceIndex)
 		for(uint32 i = 0; i < GPUs.Num(); i++)
 		{
 			FMacPlatformMisc::FGPUDescriptor const& GPU = GPUs[i];
-			if((GPU.GPUVendorId == 0x10DE))
+			if((GPU.GPUVendorId == EGpuVendorId::Nvidia))
 			{
 				OverrideRendererId = i;
 				bForceExplicitRendererId = (GPU.GPUMetalBundle && ![GPU.GPUMetalBundle isEqualToString:@"GeForceMTLDriverWeb"]);
 			}
-			else if(!GPU.GPUHeadless && GPU.GPUVendorId != 0x8086)
+			else if(!GPU.GPUHeadless && GPU.GPUVendorId != EGpuVendorId::Intel)
 			{
 				OverrideRendererId = i;
 			}
@@ -229,16 +229,16 @@ static mtlpp::Device GetMTLDevice(uint32& DeviceIndex)
 				DeviceIndex = ExplicitRendererId;
 				SelectedDevice = Device;
 			}
-			else if(([Device.GetName().GetPtr() rangeOfString:@"Nvidia" options:NSCaseInsensitiveSearch].location != NSNotFound && GPU.GPUVendorId == 0x10DE)
-			   || ([Device.GetName().GetPtr() rangeOfString:@"AMD" options:NSCaseInsensitiveSearch].location != NSNotFound && GPU.GPUVendorId == 0x1002)
-			   || ([Device.GetName().GetPtr() rangeOfString:@"Intel" options:NSCaseInsensitiveSearch].location != NSNotFound && GPU.GPUVendorId == 0x8086))
+			else if(([Device.GetName().GetPtr() rangeOfString:@"Nvidia" options:NSCaseInsensitiveSearch].location != NSNotFound && GPU.GPUVendorId == EGpuVendorId::Nvidia)
+			   || ([Device.GetName().GetPtr() rangeOfString:@"AMD" options:NSCaseInsensitiveSearch].location != NSNotFound && GPU.GPUVendorId == EGpuVendorId::Amd)
+			   || ([Device.GetName().GetPtr() rangeOfString:@"Intel" options:NSCaseInsensitiveSearch].location != NSNotFound && GPU.GPUVendorId == EGpuVendorId::Intel))
 			{
 				bool bMatchesName = (NameComponents.Num() > 0);
 				for (FString& Component : NameComponents)
 				{
 					bMatchesName &= FString(Device.GetName().GetPtr()).Contains(Component);
 				}
-				if((Device.IsHeadless() == GPU.GPUHeadless || GPU.GPUVendorId != 0x1002) && bMatchesName)
+				if((Device.IsHeadless() == GPU.GPUHeadless || GPU.GPUVendorId != EGpuVendorId::Amd) && bMatchesName)
                 {
 					DeviceIndex = ExplicitRendererId;
 					SelectedDevice = Device;
@@ -265,9 +265,9 @@ static mtlpp::Device GetMTLDevice(uint32& DeviceIndex)
 				bFoundDefault = true;
 				break;
 			}
-			else if(([SelectedDevice.GetName().GetPtr() rangeOfString:@"Nvidia" options:NSCaseInsensitiveSearch].location != NSNotFound && GPU.GPUVendorId == 0x10DE)
-					|| ([SelectedDevice.GetName().GetPtr() rangeOfString:@"AMD" options:NSCaseInsensitiveSearch].location != NSNotFound && GPU.GPUVendorId == 0x1002)
-					|| ([SelectedDevice.GetName().GetPtr() rangeOfString:@"Intel" options:NSCaseInsensitiveSearch].location != NSNotFound && GPU.GPUVendorId == 0x8086))
+			else if(([SelectedDevice.GetName().GetPtr() rangeOfString:@"Nvidia" options:NSCaseInsensitiveSearch].location != NSNotFound && GPU.GPUVendorId == EGpuVendorId::Nvidia)
+					|| ([SelectedDevice.GetName().GetPtr() rangeOfString:@"AMD" options:NSCaseInsensitiveSearch].location != NSNotFound && GPU.GPUVendorId == EGpuVendorId::Amd)
+					|| ([SelectedDevice.GetName().GetPtr() rangeOfString:@"Intel" options:NSCaseInsensitiveSearch].location != NSNotFound && GPU.GPUVendorId == EGpuVendorId::Intel))
 			{
 				NameComponents.Empty();
 				bool bMatchesName = FString(GPU.GPUName).TrimStart().ParseIntoArray(NameComponents, TEXT(" ")) > 0;
@@ -275,7 +275,7 @@ static mtlpp::Device GetMTLDevice(uint32& DeviceIndex)
 				{
 					bMatchesName &= FString(SelectedDevice.GetName().GetPtr()).Contains(Component);
 				}
-				if((SelectedDevice.IsHeadless() == GPU.GPUHeadless || GPU.GPUVendorId != 0x1002) && bMatchesName)
+				if((SelectedDevice.IsHeadless() == GPU.GPUHeadless || GPU.GPUVendorId != EGpuVendorId::Amd) && bMatchesName)
                 {
 					DeviceIndex = i;
 					bFoundDefault = true;
diff --git a/Engine/Source/Runtime/Apple/MetalRHI/Private/MetalRHI.cpp b/Engine/Source/Runtime/Apple/MetalRHI/Private/MetalRHI.cpp
--- a/Engine/Source/Runtime/Apple/MetalRHI/Private/MetalRHI.cpp
+++ b/Engine/Source/Runtime/Apple/MetalRHI/Private/MetalRHI.cpp
@@ -333,7 +333,7 @@ FMetalDynamicRHI::FMetalDynamicRHI(ERHIFeatureLevel::Type RequestedFeatureLevel)
 	if(GRHIAdapterName.Contains("Nvidia"))
 	{
 		bSupportsPointLights = true;
-		GRHIVendorId = 0x10DE;
+		GRHIVendorId = (uint32)EGpuVendorId::Nvidia;
 		bSupportsTiledReflections = true;
 		bSupportsDistanceFields = true;
 		GRHISupportsWaveOperations = false;
@@ -341,7 +341,7 @@ FMetalDynamicRHI::FMetalDynamicRHI(ERHIFeatureLevel::Type RequestedFeatureLevel)
 	else if(GRHIAdapterName.Contains("ATi") || GRHIAdapterName.Contains("AMD"))
 	{
 		bSupportsPointLights = true;
-		GRHIVendorId = 0x1002;
+		GRHIVendorId = (uint32)EGpuVendorId::Amd;
 		if(GPUDesc.GPUVendorId == GRHIVendorId)
 		{
 			GRHIAdapterName = FString(GPUDesc.GPUName);
@@ -364,7 +364,7 @@ FMetalDynamicRHI::FMetalDynamicRHI(ERHIFeatureLevel::Type RequestedFeatureLevel)
 	{
 		bSupportsTiledReflections = false;
 		bSupportsPointLights = true;
-		GRHIVendorId = 0x8086;
+		GRHIVendorId = (uint32)EGpuVendorId::Intel;
 		bSupportsDistanceFields = true;
 		bIsIntelHaswell = (GRHIAdapterName == TEXT("Intel HD Graphics 5000") || GRHIAdapterName == TEXT("Intel Iris Graphics") || GRHIAdapterName == TEXT("Intel Iris Pro Graphics"));
 		GRHISupportsWaveOperations = false;
@@ -372,7 +372,7 @@ FMetalDynamicRHI::FMetalDynamicRHI(ERHIFeatureLevel::Type RequestedFeatureLevel)
 	else if(GRHIAdapterName.Contains("Apple"))
 	{
 		bSupportsPointLights = true;
-		GRHIVendorId = 0x106B;
+		GRHIVendorId = (uint32)EGpuVendorId::Apple;
 		bSupportsTiledReflections = true;
 		bSupportsDistanceFields = true;
 		GSupportsTimestampRenderQueries = true;
diff --git a/Engine/Source/Runtime/Apple/MetalRHI/Private/MetalRHIRenderQuery.cpp b/Engine/Source/Runtime/Apple/MetalRHI/Private/MetalRHIRenderQuery.cpp
--- a/Engine/Source/Runtime/Apple/MetalRHI/Private/MetalRHIRenderQuery.cpp
+++ b/Engine/Source/Runtime/Apple/MetalRHI/Private/MetalRHIRenderQuery.cpp
@@ -285,7 +285,14 @@ void FMetalRHIRenderQuery::End(FMetalContext* Context)
                 // If there are no commands in the command buffer then this can be zero
                 // In this case GPU start time is also not correct - we need to fall back standard behaviour
                 // Only seen empty command buffers at the very end of a frame
-                Result = uint64((CmdBuffer.GetGpuEndTime() / 1000.0) / FPlatformTime::GetSecondsPerCycle64());
+                if (IsRHIDeviceApple())
+                {
+                    Result = uint64(CmdBuffer.GetGpuEndTime() - CmdBuffer.GetGpuStartTime()) * 1000;
+                }
+                else
+                {
+                    Result = uint64((CmdBuffer.GetGpuEndTime() / 1000.0) / FPlatformTime::GetSecondsPerCycle64());
+                }
 				
 				if(Result == 0)
 				{
diff --git a/Engine/Source/Runtime/RHI/Private/RHI.cpp b/Engine/Source/Runtime/RHI/Private/RHI.cpp
--- a/Engine/Source/Runtime/RHI/Private/RHI.cpp
+++ b/Engine/Source/Runtime/RHI/Private/RHI.cpp
@@ -1400,7 +1400,6 @@ ERHIFeatureLevel::Type GMaxRHIFeatureLevel = ERHIFeatureLevel::SM5;
 bool IsRHIDeviceAMD()
 {
 	check(GRHIVendorId != 0);
-	// AMD's drivers tested on July 11 2013 have hitching problems with async resource streaming, setting single threaded for now until fixed.
 	return GRHIVendorId == 0x1002;
 }
 
@@ -1413,10 +1412,15 @@ bool IsRHIDeviceIntel()
 bool IsRHIDeviceNVIDIA()
 {
 	check(GRHIVendorId != 0);
-	// NVIDIA GPUs are discrete and use DedicatedVideoMemory only.
 	return GRHIVendorId == 0x10DE;
 }
 
+bool IsRHIDeviceApple()
+{
+    check(GRHIVendorId != 0);
+    return GRHIVendorId == (uint32) EGpuVendorId::Apple;
+}
+
 uint32 RHIGetMetalShaderLanguageVersion(const FStaticShaderPlatform Platform)
 {
     if (IsMetalPlatform(Platform))
diff --git a/Engine/Source/Runtime/RHI/Public/RHI.h b/Engine/Source/Runtime/RHI/Public/RHI.h
--- a/Engine/Source/Runtime/RHI/Public/RHI.h
+++ b/Engine/Source/Runtime/RHI/Public/RHI.h
@@ -55,6 +55,9 @@ RHI_API bool IsRHIDeviceIntel();
 // to trigger GPU specific optimizations and fallbacks
 RHI_API bool IsRHIDeviceNVIDIA();
 
+// to trigger GPU specific optimizations and fallbacks
+RHI_API bool IsRHIDeviceApple();
+
 // helper to return the shader language version for Metal shader.
 RHI_API uint32 RHIGetMetalShaderLanguageVersion(const FStaticShaderPlatform Platform);
 
