MetalRHI: Fix for compiler errors on Intel Macs

diff --git a/Engine/Source/Runtime/Apple/MetalRHI/Private/MetalContext.cpp b/Engine/Source/Runtime/Apple/MetalRHI/Private/MetalContext.cpp
--- a/Engine/Source/Runtime/Apple/MetalRHI/Private/MetalContext.cpp
+++ b/Engine/Source/Runtime/Apple/MetalRHI/Private/MetalContext.cpp
@@ -198,12 +198,12 @@ static mtlpp::Device GetMTLDevice(uint32& DeviceIndex)
 		for(uint32 i = 0; i < GPUs.Num(); i++)
 		{
 			FMacPlatformMisc::FGPUDescriptor const& GPU = GPUs[i];
-			if((GPU.GPUVendorId == EGpuVendorId::Nvidia))
+			if((GPU.GPUVendorId == (uint32)EGpuVendorId::Nvidia))
 			{
 				OverrideRendererId = i;
 				bForceExplicitRendererId = (GPU.GPUMetalBundle && ![GPU.GPUMetalBundle isEqualToString:@"GeForceMTLDriverWeb"]);
 			}
-			else if(!GPU.GPUHeadless && GPU.GPUVendorId != EGpuVendorId::Intel)
+			else if(!GPU.GPUHeadless && GPU.GPUVendorId != (uint32)EGpuVendorId::Intel)
 			{
 				OverrideRendererId = i;
 			}
@@ -229,16 +229,16 @@ static mtlpp::Device GetMTLDevice(uint32& DeviceIndex)
 				DeviceIndex = ExplicitRendererId;
 				SelectedDevice = Device;
 			}
-			else if(([Device.GetName().GetPtr() rangeOfString:@"Nvidia" options:NSCaseInsensitiveSearch].location != NSNotFound && GPU.GPUVendorId == EGpuVendorId::Nvidia)
-			   || ([Device.GetName().GetPtr() rangeOfString:@"AMD" options:NSCaseInsensitiveSearch].location != NSNotFound && GPU.GPUVendorId == EGpuVendorId::Amd)
-			   || ([Device.GetName().GetPtr() rangeOfString:@"Intel" options:NSCaseInsensitiveSearch].location != NSNotFound && GPU.GPUVendorId == EGpuVendorId::Intel))
+			else if(([Device.GetName().GetPtr() rangeOfString:@"Nvidia" options:NSCaseInsensitiveSearch].location != NSNotFound && GPU.GPUVendorId == (uint32)EGpuVendorId::Nvidia)
+			   || ([Device.GetName().GetPtr() rangeOfString:@"AMD" options:NSCaseInsensitiveSearch].location != NSNotFound && GPU.GPUVendorId == (uint32)EGpuVendorId::Amd)
+			   || ([Device.GetName().GetPtr() rangeOfString:@"Intel" options:NSCaseInsensitiveSearch].location != NSNotFound && GPU.GPUVendorId == (uint32)EGpuVendorId::Intel))
 			{
 				bool bMatchesName = (NameComponents.Num() > 0);
 				for (FString& Component : NameComponents)
 				{
 					bMatchesName &= FString(Device.GetName().GetPtr()).Contains(Component);
 				}
-				if((Device.IsHeadless() == GPU.GPUHeadless || GPU.GPUVendorId != EGpuVendorId::Amd) && bMatchesName)
+				if((Device.IsHeadless() == GPU.GPUHeadless || GPU.GPUVendorId != (uint32)EGpuVendorId::Amd) && bMatchesName)
                 {
 					DeviceIndex = ExplicitRendererId;
 					SelectedDevice = Device;
@@ -265,9 +265,9 @@ static mtlpp::Device GetMTLDevice(uint32& DeviceIndex)
 				bFoundDefault = true;
 				break;
 			}
-			else if(([SelectedDevice.GetName().GetPtr() rangeOfString:@"Nvidia" options:NSCaseInsensitiveSearch].location != NSNotFound && GPU.GPUVendorId == EGpuVendorId::Nvidia)
-					|| ([SelectedDevice.GetName().GetPtr() rangeOfString:@"AMD" options:NSCaseInsensitiveSearch].location != NSNotFound && GPU.GPUVendorId == EGpuVendorId::Amd)
-					|| ([SelectedDevice.GetName().GetPtr() rangeOfString:@"Intel" options:NSCaseInsensitiveSearch].location != NSNotFound && GPU.GPUVendorId == EGpuVendorId::Intel))
+			else if(([SelectedDevice.GetName().GetPtr() rangeOfString:@"Nvidia" options:NSCaseInsensitiveSearch].location != NSNotFound && GPU.GPUVendorId == (uint32)EGpuVendorId::Nvidia)
+					|| ([SelectedDevice.GetName().GetPtr() rangeOfString:@"AMD" options:NSCaseInsensitiveSearch].location != NSNotFound && GPU.GPUVendorId == (uint32)EGpuVendorId::Amd)
+					|| ([SelectedDevice.GetName().GetPtr() rangeOfString:@"Intel" options:NSCaseInsensitiveSearch].location != NSNotFound && GPU.GPUVendorId == (uint32)EGpuVendorId::Intel))
 			{
 				NameComponents.Empty();
 				bool bMatchesName = FString(GPU.GPUName).TrimStart().ParseIntoArray(NameComponents, TEXT(" ")) > 0;
@@ -275,7 +275,7 @@ static mtlpp::Device GetMTLDevice(uint32& DeviceIndex)
 				{
 					bMatchesName &= FString(SelectedDevice.GetName().GetPtr()).Contains(Component);
 				}
-				if((SelectedDevice.IsHeadless() == GPU.GPUHeadless || GPU.GPUVendorId != EGpuVendorId::Amd) && bMatchesName)
+				if((SelectedDevice.IsHeadless() == GPU.GPUHeadless || GPU.GPUVendorId != (uint32)EGpuVendorId::Amd) && bMatchesName)
                 {
 					DeviceIndex = i;
 					bFoundDefault = true;
