MetalRHI: Fix for incorrect alignment on linear textures when using buffer backed SRV/UAV

diff --git a/Engine/Source/Runtime/Apple/MetalRHI/Private/MetalUAV.cpp b/Engine/Source/Runtime/Apple/MetalRHI/Private/MetalUAV.cpp
--- a/Engine/Source/Runtime/Apple/MetalRHI/Private/MetalUAV.cpp
+++ b/Engine/Source/Runtime/Apple/MetalRHI/Private/MetalUAV.cpp
@@ -146,9 +146,14 @@ void FMetalShaderResourceView::UpdateView()
 					mtlpp::PixelFormat Format = (mtlpp::PixelFormat)GMetalBufferFormats[Info.Format].LinearTextureFormat;
 					NSUInteger Options = ((NSUInteger)Buffer->Mode) << mtlpp::ResourceStorageModeShift;
 
+                    const uint32 MinimumByteAlignment = GetMetalDeviceContext().GetDevice().GetMinimumLinearTextureAlignmentForPixelFormat(Format);
+                    const uint32 MinimumElementAlignment = MinimumByteAlignment / Info.StrideInBytes;
+                    uint32 NumElements = Align(Info.NumElements, MinimumElementAlignment);
+                    uint32 SizeInBytes = NumElements * Info.StrideInBytes;
+                    
 					auto Desc = mtlpp::TextureDescriptor::TextureBufferDescriptor(
 						  Format
-						, Info.NumElements
+						, NumElements
 						, mtlpp::ResourceOptions(Options)
 						, mtlpp::TextureUsage::ShaderRead
 					);
@@ -156,7 +161,7 @@ void FMetalShaderResourceView::UpdateView()
 					Desc.SetAllowGPUOptimisedContents(false);
 
 					FMetalTexture& View = InitAsTextureView();
-					View = Buffer->GetCurrentBuffer().NewTexture(Desc, Info.OffsetInBytes, Info.SizeInBytes);
+					View = Buffer->GetCurrentBuffer().NewTexture(Desc, Info.OffsetInBytes, SizeInBytes);
 				}
 				break;
 
@@ -296,17 +301,22 @@ void FMetalUnorderedAccessView::UpdateView()
 				mtlpp::PixelFormat Format = (mtlpp::PixelFormat)GMetalBufferFormats[Info.Format].LinearTextureFormat;
 				NSUInteger Options = ((NSUInteger)Buffer->Mode) << mtlpp::ResourceStorageModeShift;
 
+                const uint32 MinimumByteAlignment = GetMetalDeviceContext().GetDevice().GetMinimumLinearTextureAlignmentForPixelFormat(Format);
+                const uint32 MinimumElementAlignment = MinimumByteAlignment / Info.StrideInBytes;
+                uint32 NumElements = Align(Info.NumElements, MinimumElementAlignment);
+                uint32 SizeInBytes = NumElements * Info.StrideInBytes;
+                
 				auto Desc = mtlpp::TextureDescriptor::TextureBufferDescriptor(
 					Format
-					, Info.NumElements
+					, NumElements
 					, mtlpp::ResourceOptions(Options)
 					, mtlpp::TextureUsage(mtlpp::TextureUsage::ShaderRead | mtlpp::TextureUsage::ShaderWrite)
 				);
 
 				Desc.SetAllowGPUOptimisedContents(false);
 
-                FMetalTexture MetalTexture(Buffer->GetCurrentBuffer().NewTexture(Desc, Info.OffsetInBytes, Info.SizeInBytes));
-                InitAsTextureBufferBacked(MetalTexture, Buffer->GetCurrentBuffer(), Info.OffsetInBytes, Info.SizeInBytes, Info.Format);
+                FMetalTexture MetalTexture(Buffer->GetCurrentBuffer().NewTexture(Desc, Info.OffsetInBytes, SizeInBytes));
+                InitAsTextureBufferBacked(MetalTexture, Buffer->GetCurrentBuffer(), Info.OffsetInBytes, SizeInBytes, Info.Format);
 			}
 			break;
 
