MetalRHI: Fixed Screenshot crash when ConvertSurfaceDataToFColor is called with HDR format

diff --git a/Engine/Source/Runtime/Apple/MetalRHI/Private/MetalRenderTarget.cpp b/Engine/Source/Runtime/Apple/MetalRHI/Private/MetalRenderTarget.cpp
--- a/Engine/Source/Runtime/Apple/MetalRHI/Private/MetalRenderTarget.cpp
+++ b/Engine/Source/Runtime/Apple/MetalRHI/Private/MetalRenderTarget.cpp
@@ -91,7 +91,7 @@ static void ConvertSurfaceDataToFColor(EPixelFormat Format, uint32 Width, uint32
 	{
 		ConvertRawR10G10B10A2DataToFColor(Width, Height, In, SrcPitch, Out);
 	}
-	else if (Format == PF_FloatRGBA)
+	else if (Format == PF_FloatRGBA || Format == PF_PLATFORM_HDR_0)
 	{
 		ConvertRawR16G16B16A16FDataToFColor(Width, Height, In, SrcPitch, Out, bLinearToGamma);
 	}
