Add a CVar to allow disabling memoryless textures on Metal iOS devices.  In some low memory conditions, lower spec'd iOS devices appear to crash if it is enabled.

diff --git a/Engine/Source/Runtime/Apple/MetalRHI/Private/MetalTexture.cpp b/Engine/Source/Runtime/Apple/MetalRHI/Private/MetalTexture.cpp
--- a/Engine/Source/Runtime/Apple/MetalRHI/Private/MetalTexture.cpp
+++ b/Engine/Source/Runtime/Apple/MetalRHI/Private/MetalTexture.cpp
@@ -30,6 +30,13 @@ FAutoConsoleVariableRef CVarMetalForceIOSTexturesShared(
 														TEXT("If true, forces all textures to be Shared on iOS"),
 														ECVF_RenderThreadSafe);
 
+int32 GMetalDisableIOSMemoryless = 0;
+FAutoConsoleVariableRef CVarMetalDisableIOSMemoryless(
+													  TEXT("rhi.Metal.DisableIOSMemoryless"),
+													  GMetalDisableIOSMemoryless,
+													  TEXT("If true, disabled the use of Memoryless textures on iOS"),
+													  ECVF_ReadOnly|ECVF_RenderThreadSafe);
+
 /** Given a pointer to a RHI texture that was created by the Metal RHI, returns a pointer to the FMetalTextureBase it encapsulates. */
 FMetalSurface* GetMetalSurfaceFromRHITexture(FRHITexture* Texture)
 {
@@ -442,7 +449,7 @@ FMetalTextureCreateDesc::FMetalTextureCreateDesc(FRHITextureCreateDesc const& In
 		}
 
 #if PLATFORM_IOS
-		if (EnumHasAnyFlags(InDesc.Flags, TexCreate_Memoryless))
+		if (!GMetalDisableIOSMemoryless && EnumHasAnyFlags(InDesc.Flags, TexCreate_Memoryless))
 		{
 			ensure(EnumHasAnyFlags(InDesc.Flags, (TexCreate_RenderTargetable | TexCreate_DepthStencilTargetable)));
 			ensure(!EnumHasAnyFlags(InDesc.Flags, (TexCreate_CPUReadback | TexCreate_CPUWritable)));
@@ -537,7 +544,7 @@ FMetalSurface::FMetalSurface(FMetalTextureCreateDesc const& CreateDesc)
         
 #if PLATFORM_IOS
         // If we are attempting to create an MSAA texture the texture cannot be memoryless unless we are creating a depth texture
-        if(bIsMSAARequired && CreateDesc.Format != PF_DepthStencil && EnumHasAllFlags(CreateDesc.Flags, TexCreate_Memoryless))
+        if(bIsMSAARequired && CreateDesc.Format != PF_DepthStencil && !GMetalDisableIOSMemoryless && EnumHasAllFlags(CreateDesc.Flags, TexCreate_Memoryless))
         {
             NewCreateDesc.Flags &= ~TexCreate_Memoryless;
             
@@ -639,7 +646,7 @@ FMetalSurface::FMetalSurface(FMetalTextureCreateDesc const& CreateDesc)
 		bool bMemoryless = false;
         
 #if PLATFORM_IOS
-		if (EnumHasAllFlags(CreateDesc.Flags, TexCreate_Memoryless))
+		if (!GMetalDisableIOSMemoryless && EnumHasAllFlags(CreateDesc.Flags, TexCreate_Memoryless))
 		{
 			bMemoryless = true;
 			Desc.SetStorageMode(mtlpp::StorageMode::Memoryless);
