Refactor resource transitions on Metal.
Removed PrologueEncoder since it's not synchronized properly.
It fixed the flickering/crash issue on Iphone14 with mobile GPUScene enabled.

diff --git a/Engine/Source/Runtime/Apple/MetalRHI/Private/MetalCommandEncoder.cpp b/Engine/Source/Runtime/Apple/MetalRHI/Private/MetalCommandEncoder.cpp
--- a/Engine/Source/Runtime/Apple/MetalRHI/Private/MetalCommandEncoder.cpp
+++ b/Engine/Source/Runtime/Apple/MetalRHI/Private/MetalCommandEncoder.cpp
@@ -220,12 +220,7 @@ FMetalCommandEncoder::FMetalCommandEncoder(FMetalCommandList& CmdList, EMetalCom
 #if ENABLE_METAL_GPUPROFILE
 , CommandBufferStats(nullptr)
 #endif
-#if METAL_DEBUG_OPTIONS
-, WaitCount(0)
-, UpdateCount(0)
-#endif
 , DebugGroups([NSMutableArray new])
-, FenceStage(mtlpp::RenderStages::Fragment)
 , EncoderNum(0)
 , CmdBufIndex(0)
 , Type(InType)
@@ -487,9 +482,6 @@ void FMetalCommandEncoder::CommitCommandBuffer(uint32 const Flags)
 		check(CommandBuffer);
 	}
 	
-	BufferBindingHistory.Empty();
-	TextureBindingHistory.Empty();
-	
 	EncoderNum = 0;
 }
 
@@ -570,7 +562,29 @@ void FMetalCommandEncoder::BeginRenderCommandEncoding(void)
 #endif // METAL_RHI_RAYTRACING
 	);
 	
-	FenceResources.Append(TransitionedResources);
+	static bool bSupportsFences = CommandList.GetCommandQueue().SupportsFeature(EMetalFeaturesFences);
+	if (bSupportsFences METAL_DEBUG_OPTION(|| CommandList.GetCommandQueue().GetRuntimeDebuggingLevel() >= EMetalDebugLevelValidation))
+	{
+		CommandEncoderFence.FenceResources = MoveTemp(TransitionedResources);
+
+		// Update fence state if current pass and prologue pass render to the same render targets
+		ns::Array<mtlpp::RenderPassColorAttachmentDescriptor> ColorAttachments = RenderPassDesc.GetColorAttachments();
+		for (uint32 i = 0; i < MaxSimultaneousRenderTargets; i++)
+		{
+			if (ColorAttachments[i].GetTexture())
+			{
+				FenceResource(ColorAttachments[i].GetTexture(), mtlpp::FunctionType::Fragment, true);
+			}
+		}
+		if (RenderPassDesc.GetDepthAttachment().GetTexture())
+		{
+			FenceResource(RenderPassDesc.GetDepthAttachment().GetTexture(), mtlpp::FunctionType::Fragment, true);
+		}
+		if (RenderPassDesc.GetStencilAttachment().GetTexture() && RenderPassDesc.GetStencilAttachment().GetTexture() != RenderPassDesc.GetDepthAttachment().GetTexture())
+		{
+			FenceResource(RenderPassDesc.GetStencilAttachment().GetTexture(), mtlpp::FunctionType::Fragment, true);
+		}
+	}
 	
 	RenderCommandEncoder = MTLPP_VALIDATE(mtlpp::CommandBuffer, CommandBuffer, SafeGetRuntimeDebuggingLevel() >= EMetalDebugLevelValidation, RenderCommandEncoder(RenderPassDesc));
 	METAL_DEBUG_LAYER(EMetalDebugLevelFastValidation, RenderEncoderDebug = FMetalRenderCommandEncoderDebugging(RenderCommandEncoder, RenderPassDesc, CommandBufferDebug));
@@ -606,8 +620,11 @@ void FMetalCommandEncoder::BeginComputeCommandEncoding(mtlpp::DispatchType Dispa
 #endif // METAL_RHI_RAYTRACING
 	);
 	
-	FenceResources.Append(TransitionedResources);
-	TransitionedResources.Empty();
+	static bool bSupportsFences = CommandList.GetCommandQueue().SupportsFeature(EMetalFeaturesFences);
+	if (bSupportsFences METAL_DEBUG_OPTION(|| CommandList.GetCommandQueue().GetRuntimeDebuggingLevel() >= EMetalDebugLevelValidation))
+	{
+		CommandEncoderFence.FenceResources = MoveTemp(TransitionedResources);
+	}
 	
 	if (DispatchType == mtlpp::DispatchType::Serial)
 	{
@@ -651,8 +668,11 @@ void FMetalCommandEncoder::BeginBlitCommandEncoding(void)
 #endif // METAL_RHI_RAYTRACING
 	);
 	
-	FenceResources.Append(TransitionedResources);
-	TransitionedResources.Empty();
+	static bool bSupportsFences = CommandList.GetCommandQueue().SupportsFeature(EMetalFeaturesFences);
+	if (bSupportsFences METAL_DEBUG_OPTION(|| CommandList.GetCommandQueue().GetRuntimeDebuggingLevel() >= EMetalDebugLevelValidation))
+	{
+		CommandEncoderFence.FenceResources = MoveTemp(TransitionedResources);
+	}
 	
 	BlitCommandEncoder = MTLPP_VALIDATE(mtlpp::CommandBuffer, CommandBuffer, SafeGetRuntimeDebuggingLevel() >= EMetalDebugLevelValidation, BlitCommandEncoder());
 	METAL_DEBUG_LAYER(EMetalDebugLevelFastValidation, BlitEncoderDebug = FMetalBlitCommandEncoderDebugging(BlitCommandEncoder, CommandBufferDebug));
@@ -686,9 +706,11 @@ void FMetalCommandEncoder::BeginAccelerationStructureCommandEncoding(void)
 	check(CommandBuffer);
 	check(IsRenderCommandEncoderActive() == false && IsComputeCommandEncoderActive() == false && IsBlitCommandEncoderActive() == false && IsAccelerationStructureCommandEncoderActive() == false);
 
-	FenceResources.Append(TransitionedResources);
-	TransitionedResources.Empty();
-
+	static bool bSupportsFences = CommandList.GetCommandQueue().SupportsFeature(EMetalFeaturesFences);
+	if (bSupportsFences METAL_DEBUG_OPTION(|| CommandList.GetCommandQueue().GetRuntimeDebuggingLevel() >= EMetalDebugLevelValidation))
+	{
+		CommandEncoderFence.FenceResources = MoveTemp(TransitionedResources);
+	}
 
 	AccelerationStructureCommandEncoder = CommandBuffer.AccelerationStructureCommandEncoder();
 	EncoderNum++;
@@ -731,61 +753,64 @@ TRefCountPtr<FMetalFence> FMetalCommandEncoder::EndEncoding(void)
 				ns::Array<mtlpp::RenderPassColorAttachmentDescriptor> ColorAttachments = RenderPassDesc.GetColorAttachments();
 				for (uint32 i = 0; i < MaxSimultaneousRenderTargets; i++)
 				{
-					if (ColorAttachments[i].GetTexture() && ColorAttachments[i].GetStoreAction() == mtlpp::StoreAction::Unknown)
+					if (ColorAttachments[i].GetTexture())
 					{
-						mtlpp::StoreAction Action = ColorStoreActions[i];
-						check(Action != mtlpp::StoreAction::Unknown);
-						RenderCommandEncoder.SetColorStoreAction((mtlpp::StoreAction)Action, i);
+						if (ColorAttachments[i].GetStoreAction() == mtlpp::StoreAction::Unknown)
+						{
+							mtlpp::StoreAction Action = ColorStoreActions[i];
+							check(Action != mtlpp::StoreAction::Unknown);
+							RenderCommandEncoder.SetColorStoreAction((mtlpp::StoreAction)Action, i);
+						}
+						// Recorded in case the epilogue pass renders to the same render targets
+						TransitionResources(ColorAttachments[i].GetTexture());
 					}
 				}
-				if (RenderPassDesc.GetDepthAttachment().GetTexture() && RenderPassDesc.GetDepthAttachment().GetStoreAction() == mtlpp::StoreAction::Unknown)
-				{
-					mtlpp::StoreAction Action = DepthStoreAction;
-					check(Action != mtlpp::StoreAction::Unknown);
-					RenderCommandEncoder.SetDepthStoreAction((mtlpp::StoreAction)Action);
-				}
-				if (RenderPassDesc.GetStencilAttachment().GetTexture() && RenderPassDesc.GetStencilAttachment().GetStoreAction() == mtlpp::StoreAction::Unknown)
-				{
-					mtlpp::StoreAction Action = StencilStoreAction;
-					check(Action != mtlpp::StoreAction::Unknown);
-					RenderCommandEncoder.SetStencilStoreAction((mtlpp::StoreAction)Action);
-				}
-				
-				for (FMetalFence* FragFence : FragmentFences)
+				if (RenderPassDesc.GetDepthAttachment().GetTexture())
 				{
-					if (FragFence->NeedsWait(mtlpp::RenderStages::Fragment))
+					if (RenderPassDesc.GetDepthAttachment().GetStoreAction() == mtlpp::StoreAction::Unknown)
 					{
-						mtlpp::Fence FragmentFence = FragFence->Get(mtlpp::RenderStages::Fragment);
-						mtlpp::Fence FragInnerFence = METAL_DEBUG_OPTION(CommandList.GetCommandQueue().GetRuntimeDebuggingLevel() >= EMetalDebugLevelValidation ? mtlpp::Fence(((FMetalDebugFence*)FragmentFence.GetPtr()).Inner) :) FragmentFence;
-						
-						RenderCommandEncoder.WaitForFence(FragInnerFence, FenceStage);
-						METAL_DEBUG_LAYER(EMetalDebugLevelFastValidation, RenderEncoderDebug.AddWaitFence(FragmentFence));
-						FragFence->Wait(mtlpp::RenderStages::Fragment);
+						mtlpp::StoreAction Action = DepthStoreAction;
+						check(Action != mtlpp::StoreAction::Unknown);
+						RenderCommandEncoder.SetDepthStoreAction((mtlpp::StoreAction)Action);
 					}
+					// Recorded in case the epilogue pass renders to the same render targets
+					TransitionResources(RenderPassDesc.GetDepthAttachment().GetTexture());
 				}
-				FragmentFences.Empty();
-				
-				if (FenceStage == mtlpp::RenderStages::Vertex)
+				if (RenderPassDesc.GetStencilAttachment().GetTexture())
 				{
-					FenceResources.Empty();
-					FenceStage = mtlpp::RenderStages::Fragment;
+					if (RenderPassDesc.GetStencilAttachment().GetStoreAction() == mtlpp::StoreAction::Unknown)
+					{
+						mtlpp::StoreAction Action = StencilStoreAction;
+						check(Action != mtlpp::StoreAction::Unknown);
+						RenderCommandEncoder.SetStencilStoreAction((mtlpp::StoreAction)Action);
+					}
+					// Recorded in case the epilogue pass renders to the same render targets
+					if (RenderPassDesc.GetStencilAttachment().GetTexture() != RenderPassDesc.GetDepthAttachment().GetTexture())
+					{
+						TransitionResources(RenderPassDesc.GetStencilAttachment().GetTexture());
+					}
 				}
-				
-				if (EncoderFence && EncoderFence->NeedsWrite(mtlpp::RenderStages::Fragment))
+
+				// Wait the prologue fence
 				{
-					Fence = EncoderFence;
+					EMetalFenceWaitStage::Type FenceWaitStage = CommandEncoderFence.BarrierScope.GetFenceWaitStage();
+					FMetalFence* PrologueFence = CommandEncoderFence.Fence;
+					if (PrologueFence)
+					{
+						mtlpp::RenderStages FenceStage = FenceWaitStage == EMetalFenceWaitStage::BeforeVertex ? mtlpp::RenderStages::Vertex : mtlpp::RenderStages::Fragment;
+						mtlpp::Fence MTLFence = PrologueFence->Get();
+						mtlpp::Fence InnerFence = METAL_DEBUG_OPTION(CommandList.GetCommandQueue().GetRuntimeDebuggingLevel() >= EMetalDebugLevelValidation ? mtlpp::Fence(((FMetalDebugFence*)MTLFence.GetPtr()).Inner) : ) MTLFence;
+
+						RenderCommandEncoder.WaitForFence(InnerFence, FenceStage);
+						METAL_DEBUG_LAYER(EMetalDebugLevelFastValidation, RenderEncoderDebug.AddWaitFence(MTLFence));
+						PrologueFence->Wait();
+					}
+					CommandEncoderFence.Reset();
 				}
+
+				Fence = EncoderFence;
 				UpdateFence(EncoderFence);
 				
-#if METAL_DEBUG_OPTIONS
-				if (bSupportsFences && SafeGetRuntimeDebuggingLevel() >= EMetalDebugLevelFastValidation && (!WaitCount || !UpdateCount))
-				{
-					UE_LOG(LogMetal, Error, TEXT("%s has incorrect fence waits (%u) vs. updates (%u)."), *FString(RenderCommandEncoder.GetLabel()), WaitCount, UpdateCount);
-					
-				}
-				WaitCount = 0;
-				UpdateCount = 0;
-#endif
 				RenderCommandEncoder.EndEncoding();
 				METAL_DEBUG_LAYER(EMetalDebugLevelFastValidation, RenderEncoderDebug.EndEncoder());
 				RenderCommandEncoder = nil;
@@ -795,38 +820,26 @@ TRefCountPtr<FMetalFence> FMetalCommandEncoder::EndEncoding(void)
 		else if(IsComputeCommandEncoderActive())
 		{
 			check(!bSupportsFences || EncoderFence);
-			
-			for (FMetalFence* FragFence : FragmentFences)
+
+			// Wait the prologue fence
 			{
-				if (FragFence->NeedsWait(mtlpp::RenderStages::Fragment))
+				EMetalFenceWaitStage::Type FenceWaitStage = CommandEncoderFence.BarrierScope.GetFenceWaitStage();
+				FMetalFence* PrologueFence = CommandEncoderFence.Fence;
+				if (PrologueFence)
 				{
-					mtlpp::Fence FragmentFence = FragFence->Get(mtlpp::RenderStages::Fragment);
-					mtlpp::Fence FragInnerFence = METAL_DEBUG_OPTION(CommandList.GetCommandQueue().GetRuntimeDebuggingLevel() >= EMetalDebugLevelValidation ? mtlpp::Fence(((FMetalDebugFence*)FragmentFence.GetPtr()).Inner) :) FragmentFence;
-					
-					ComputeCommandEncoder.WaitForFence(FragInnerFence);
-					METAL_DEBUG_LAYER(EMetalDebugLevelFastValidation, ComputeEncoderDebug.AddWaitFence(FragmentFence));
-					FragFence->Wait(mtlpp::RenderStages::Fragment);
+					mtlpp::Fence MTLFence = PrologueFence->Get();
+					mtlpp::Fence InnerFence = METAL_DEBUG_OPTION(CommandList.GetCommandQueue().GetRuntimeDebuggingLevel() >= EMetalDebugLevelValidation ? mtlpp::Fence(((FMetalDebugFence*)MTLFence.GetPtr()).Inner) : ) MTLFence;
+
+					ComputeCommandEncoder.WaitForFence(InnerFence);
+					METAL_DEBUG_LAYER(EMetalDebugLevelFastValidation, ComputeEncoderDebug.AddWaitFence(MTLFence));
+					PrologueFence->Wait();
 				}
+				CommandEncoderFence.Reset();
 			}
-			FragmentFences.Empty();
-			FenceResources.Empty();
-			FenceStage = mtlpp::RenderStages::Fragment;
-			
-			if (EncoderFence && EncoderFence->NeedsWrite(mtlpp::RenderStages::Fragment))
-			{
-				Fence = EncoderFence;
-			}
+
+			Fence = EncoderFence;
 			UpdateFence(EncoderFence);
 			
-#if METAL_DEBUG_OPTIONS
-			if (bSupportsFences && SafeGetRuntimeDebuggingLevel() >= EMetalDebugLevelFastValidation && (!WaitCount || !UpdateCount))
-			{
-				UE_LOG(LogMetal, Error, TEXT("%s has incorrect fence waits (%u) vs. updates (%u)."), *FString(ComputeCommandEncoder.GetLabel()), WaitCount, UpdateCount);
-				
-			}
-			WaitCount = 0;
-			UpdateCount = 0;
-#endif
 			ComputeCommandEncoder.EndEncoding();
 			METAL_DEBUG_LAYER(EMetalDebugLevelFastValidation, ComputeEncoderDebug.EndEncoder());
 			ComputeCommandEncoder = nil;
@@ -835,38 +848,25 @@ TRefCountPtr<FMetalFence> FMetalCommandEncoder::EndEncoding(void)
 		else if(IsBlitCommandEncoderActive())
 		{
 			// check(!bSupportsFences || EncoderFence);
-			
-			for (FMetalFence* FragFence : FragmentFences)
+			// Wait the prologue fence
 			{
-				if (FragFence->NeedsWait(mtlpp::RenderStages::Fragment))
+				EMetalFenceWaitStage::Type FenceWaitStage = CommandEncoderFence.BarrierScope.GetFenceWaitStage();
+				FMetalFence* PrologueFence = CommandEncoderFence.Fence;
+				if (PrologueFence)
 				{
-					mtlpp::Fence FragmentFence = FragFence->Get(mtlpp::RenderStages::Fragment);
-					mtlpp::Fence FragInnerFence = METAL_DEBUG_OPTION(CommandList.GetCommandQueue().GetRuntimeDebuggingLevel() >= EMetalDebugLevelValidation ? mtlpp::Fence(((FMetalDebugFence*)FragmentFence.GetPtr()).Inner) :) FragmentFence;
-					
-					BlitCommandEncoder.WaitForFence(FragInnerFence);
-					METAL_DEBUG_LAYER(EMetalDebugLevelFastValidation, BlitEncoderDebug.AddWaitFence(FragmentFence));
-					FragFence->Wait(mtlpp::RenderStages::Fragment);
+					mtlpp::Fence MTLFence = PrologueFence->Get();
+					mtlpp::Fence InnerFence = METAL_DEBUG_OPTION(CommandList.GetCommandQueue().GetRuntimeDebuggingLevel() >= EMetalDebugLevelValidation ? mtlpp::Fence(((FMetalDebugFence*)MTLFence.GetPtr()).Inner) : ) MTLFence;
+
+					BlitCommandEncoder.WaitForFence(InnerFence);
+					METAL_DEBUG_LAYER(EMetalDebugLevelFastValidation, BlitEncoderDebug.AddWaitFence(MTLFence));
+					PrologueFence->Wait();
 				}
+				CommandEncoderFence.Reset();
 			}
-			FragmentFences.Empty();
-			FenceResources.Empty();
-			FenceStage = mtlpp::RenderStages::Fragment;
-			
-			if (EncoderFence && EncoderFence->NeedsWrite(mtlpp::RenderStages::Fragment))
-			{
-				Fence = EncoderFence;
-			}
+
+			Fence = EncoderFence;
 			UpdateFence(EncoderFence);
 			
-#if METAL_DEBUG_OPTIONS
-			if (bSupportsFences && SafeGetRuntimeDebuggingLevel() >= EMetalDebugLevelFastValidation && (!WaitCount || !UpdateCount))
-			{
-				UE_LOG(LogMetal, Error, TEXT("%s has incorrect fence waits (%u) vs. updates (%u)."), *FString(BlitCommandEncoder.GetLabel()), WaitCount, UpdateCount);
-				
-			}
-			WaitCount = 0;
-			UpdateCount = 0;
-#endif
 			BlitCommandEncoder.EndEncoding();
 			METAL_DEBUG_LAYER(EMetalDebugLevelFastValidation, BlitEncoderDebug.EndEncoder());
 			BlitCommandEncoder = nil;
@@ -936,42 +936,25 @@ void FMetalCommandEncoder::UpdateFence(FMetalFence* Fence)
 	static bool bSupportsFences = CommandList.GetCommandQueue().SupportsFeature(EMetalFeaturesFences);
 	if ((bSupportsFences METAL_DEBUG_OPTION(|| CommandList.GetCommandQueue().GetRuntimeDebuggingLevel() >= EMetalDebugLevelValidation)) && Fence)
 	{
-		mtlpp::Fence VertexFence = Fence->Get(mtlpp::RenderStages::Vertex);
-		mtlpp::Fence InnerFence = METAL_DEBUG_OPTION(CommandList.GetCommandQueue().GetRuntimeDebuggingLevel() >= EMetalDebugLevelValidation ? mtlpp::Fence(((FMetalDebugFence*)VertexFence.GetPtr()).Inner) :) VertexFence;
+		mtlpp::Fence MTLFence = Fence->Get();
+		mtlpp::Fence InnerFence = METAL_DEBUG_OPTION(CommandList.GetCommandQueue().GetRuntimeDebuggingLevel() >= EMetalDebugLevelValidation ? mtlpp::Fence(((FMetalDebugFence*)MTLFence.GetPtr()).Inner) :) MTLFence;
 		if (RenderCommandEncoder)
 		{
-			mtlpp::Fence FragmentFence = Fence->Get(mtlpp::RenderStages::Fragment);
-			mtlpp::Fence FragInnerFence = METAL_DEBUG_OPTION(CommandList.GetCommandQueue().GetRuntimeDebuggingLevel() >= EMetalDebugLevelValidation ? mtlpp::Fence(((FMetalDebugFence*)FragmentFence.GetPtr()).Inner) :) FragmentFence;
-			
-			if (Fence->NeedsWrite(mtlpp::RenderStages::Vertex))
-			{
-				RenderCommandEncoder.UpdateFence(InnerFence, (mtlpp::RenderStages)(mtlpp::RenderStages::Vertex));
-				METAL_DEBUG_LAYER(EMetalDebugLevelFastValidation, RenderEncoderDebug.AddUpdateFence(VertexFence));
-				Fence->Write(mtlpp::RenderStages::Vertex);
-				METAL_DEBUG_LAYER(EMetalDebugLevelFastValidation, UpdateCount++);
-			}
-
-			if (Fence->NeedsWrite(mtlpp::RenderStages::Fragment))
-			{
-				RenderCommandEncoder.UpdateFence(FragInnerFence, (mtlpp::RenderStages)(mtlpp::RenderStages::Fragment));
-				METAL_DEBUG_LAYER(EMetalDebugLevelFastValidation, RenderEncoderDebug.AddUpdateFence(FragmentFence));
-				Fence->Write(mtlpp::RenderStages::Fragment);
-				METAL_DEBUG_LAYER(EMetalDebugLevelFastValidation, UpdateCount++);
-			}
+			RenderCommandEncoder.UpdateFence(InnerFence, mtlpp::RenderStages::Fragment);
+			METAL_DEBUG_LAYER(EMetalDebugLevelFastValidation, RenderEncoderDebug.AddUpdateFence(MTLFence));
+			Fence->Write();
 		}
-		else if (ComputeCommandEncoder && Fence->NeedsWrite(mtlpp::RenderStages::Vertex))
+		else if (ComputeCommandEncoder)
 		{
 			ComputeCommandEncoder.UpdateFence(InnerFence);
-			METAL_DEBUG_LAYER(EMetalDebugLevelFastValidation, ComputeEncoderDebug.AddUpdateFence(VertexFence));
-			Fence->Write(mtlpp::RenderStages::Vertex);
-			METAL_DEBUG_LAYER(EMetalDebugLevelFastValidation, UpdateCount++);
+			METAL_DEBUG_LAYER(EMetalDebugLevelFastValidation, ComputeEncoderDebug.AddUpdateFence(MTLFence));
+			Fence->Write();
 		}
-		else if (BlitCommandEncoder && Fence->NeedsWrite(mtlpp::RenderStages::Vertex))
+		else if (BlitCommandEncoder)
 		{
 			BlitCommandEncoder.UpdateFence(InnerFence);
-			METAL_DEBUG_LAYER(EMetalDebugLevelFastValidation, BlitEncoderDebug.AddUpdateFence(VertexFence));
-			Fence->Write(mtlpp::RenderStages::Vertex);
-			METAL_DEBUG_LAYER(EMetalDebugLevelFastValidation, UpdateCount++);
+			METAL_DEBUG_LAYER(EMetalDebugLevelFastValidation, BlitEncoderDebug.AddUpdateFence(MTLFence));
+			Fence->Write();
 		}
 	}
 }
@@ -986,121 +969,8 @@ void FMetalCommandEncoder::WaitForFence(FMetalFence* Fence)
 	static bool bSupportsFences = CommandList.GetCommandQueue().SupportsFeature(EMetalFeaturesFences);
 	if ((bSupportsFences METAL_DEBUG_OPTION(|| CommandList.GetCommandQueue().GetRuntimeDebuggingLevel() >= EMetalDebugLevelValidation)) && Fence)
 	{
-		if (Fence->NeedsWait(mtlpp::RenderStages::Vertex))
-		{
-			METAL_DEBUG_LAYER(EMetalDebugLevelFastValidation, WaitCount++);
-			
-			mtlpp::Fence VertexFence = Fence->Get(mtlpp::RenderStages::Vertex);
-			mtlpp::Fence InnerFence = METAL_DEBUG_OPTION(CommandList.GetCommandQueue().GetRuntimeDebuggingLevel() >= EMetalDebugLevelValidation ? mtlpp::Fence(((FMetalDebugFence*)VertexFence.GetPtr()).Inner) :) VertexFence;
-			if (RenderCommandEncoder)
-			{
-				RenderCommandEncoder.WaitForFence(InnerFence, (mtlpp::RenderStages)(mtlpp::RenderStages::Vertex));
-				METAL_DEBUG_LAYER(EMetalDebugLevelFastValidation, RenderEncoderDebug.AddWaitFence(VertexFence));
-				Fence->Wait(mtlpp::RenderStages::Vertex);
-			}
-			else if (ComputeCommandEncoder)
-			{
-				ComputeCommandEncoder.WaitForFence(InnerFence);
-				METAL_DEBUG_LAYER(EMetalDebugLevelFastValidation, ComputeEncoderDebug.AddWaitFence(VertexFence));
-				Fence->Wait(mtlpp::RenderStages::Vertex);
-			}
-			else if (BlitCommandEncoder)
-			{
-				BlitCommandEncoder.WaitForFence(InnerFence);
-				METAL_DEBUG_LAYER(EMetalDebugLevelFastValidation, BlitEncoderDebug.AddWaitFence(VertexFence));
-				Fence->Wait(mtlpp::RenderStages::Vertex);
-			}
-		}
-		if (Fence->NeedsWait(mtlpp::RenderStages::Fragment))
-		{
-			if (FenceStage == mtlpp::RenderStages::Vertex || BlitCommandEncoder)
-			{
-				mtlpp::Fence FragmentFence = Fence->Get(mtlpp::RenderStages::Fragment);
-				mtlpp::Fence FragInnerFence = METAL_DEBUG_OPTION(CommandList.GetCommandQueue().GetRuntimeDebuggingLevel() >= EMetalDebugLevelValidation ? mtlpp::Fence(((FMetalDebugFence*)FragmentFence.GetPtr()).Inner) :) FragmentFence;
-				if (RenderCommandEncoder)
-				{
-					RenderCommandEncoder.WaitForFence(FragInnerFence, (mtlpp::RenderStages)(mtlpp::RenderStages::Vertex));
-					METAL_DEBUG_LAYER(EMetalDebugLevelFastValidation, RenderEncoderDebug.AddWaitFence(FragmentFence));
-					Fence->Wait(mtlpp::RenderStages::Fragment);
-				}
-				else if (ComputeCommandEncoder)
-				{
-					ComputeCommandEncoder.WaitForFence(FragInnerFence);
-					METAL_DEBUG_LAYER(EMetalDebugLevelFastValidation, ComputeEncoderDebug.AddWaitFence(FragmentFence));
-					Fence->Wait(mtlpp::RenderStages::Fragment);
-				}
-				else if (BlitCommandEncoder)
-				{
-					BlitCommandEncoder.WaitForFence(FragInnerFence);
-					METAL_DEBUG_LAYER(EMetalDebugLevelFastValidation, BlitEncoderDebug.AddWaitFence(FragmentFence));
-					Fence->Wait(mtlpp::RenderStages::Fragment);
-				}
-				METAL_DEBUG_LAYER(EMetalDebugLevelFastValidation, WaitCount++);
-			}
-			else
-			{
-				METAL_DEBUG_LAYER(EMetalDebugLevelFastValidation, WaitCount++);
-				FragmentFences.Add(Fence);
-			}
-		}
-	}
-}
-
-void FMetalCommandEncoder::WaitAndUpdateFence(FMetalFence* Fence)
-{
-	check(IsRenderCommandEncoderActive() || IsComputeCommandEncoderActive() || IsBlitCommandEncoderActive());
-	static bool bSupportsFences = CommandList.GetCommandQueue().SupportsFeature(EMetalFeaturesFences);
-	if ((bSupportsFences METAL_DEBUG_OPTION(|| CommandList.GetCommandQueue().GetRuntimeDebuggingLevel() >= EMetalDebugLevelValidation)) && Fence)
-	{
-		METAL_DEBUG_LAYER(EMetalDebugLevelFastValidation, WaitCount++);
-		METAL_DEBUG_LAYER(EMetalDebugLevelFastValidation, UpdateCount++);
-		
-		mtlpp::Fence VertexFence = Fence->Get(mtlpp::RenderStages::Vertex);
-		mtlpp::Fence InnerFence = METAL_DEBUG_OPTION(CommandList.GetCommandQueue().GetRuntimeDebuggingLevel() >= EMetalDebugLevelValidation ? mtlpp::Fence(((FMetalDebugFence*)VertexFence.GetPtr()).Inner) :) VertexFence;
-		if (RenderCommandEncoder)
-		{
-			mtlpp::Fence FragmentFence = Fence->Get(mtlpp::RenderStages::Fragment);
-			mtlpp::Fence FragInnerFence = METAL_DEBUG_OPTION(CommandList.GetCommandQueue().GetRuntimeDebuggingLevel() >= EMetalDebugLevelValidation ? mtlpp::Fence(((FMetalDebugFence*)FragmentFence.GetPtr()).Inner) :) FragmentFence;
-			
-			METAL_DEBUG_LAYER(EMetalDebugLevelFastValidation, WaitCount++);
-			METAL_DEBUG_LAYER(EMetalDebugLevelFastValidation, UpdateCount++);
-			
-			RenderCommandEncoder.WaitForFence(InnerFence, (mtlpp::RenderStages)(mtlpp::RenderStages::Vertex));
-			METAL_DEBUG_LAYER(EMetalDebugLevelFastValidation, RenderEncoderDebug.AddWaitFence(VertexFence));
-			Fence->Wait(mtlpp::RenderStages::Vertex);
-			
-			RenderCommandEncoder.WaitForFence(FragInnerFence, (mtlpp::RenderStages)(mtlpp::RenderStages::Fragment));
-			METAL_DEBUG_LAYER(EMetalDebugLevelFastValidation, RenderEncoderDebug.AddWaitFence(FragmentFence));
-			Fence->Wait(mtlpp::RenderStages::Fragment);
-			
-			RenderCommandEncoder.UpdateFence(InnerFence, (mtlpp::RenderStages)(mtlpp::RenderStages::Fragment));
-			METAL_DEBUG_LAYER(EMetalDebugLevelFastValidation, RenderEncoderDebug.AddUpdateFence(VertexFence));
-			Fence->Write(mtlpp::RenderStages::Vertex);
-			
-			RenderCommandEncoder.UpdateFence(FragInnerFence, (mtlpp::RenderStages)(mtlpp::RenderStages::Vertex));
-			METAL_DEBUG_LAYER(EMetalDebugLevelFastValidation, RenderEncoderDebug.AddUpdateFence(FragmentFence));
-			Fence->Write(mtlpp::RenderStages::Fragment);
-		}
-		else if (ComputeCommandEncoder)
-		{
-			ComputeCommandEncoder.WaitForFence(InnerFence);
-			METAL_DEBUG_LAYER(EMetalDebugLevelFastValidation, ComputeEncoderDebug.AddWaitFence(VertexFence));
-			Fence->Wait(mtlpp::RenderStages::Vertex);
-			
-			ComputeCommandEncoder.UpdateFence(InnerFence);
-			METAL_DEBUG_LAYER(EMetalDebugLevelFastValidation, ComputeEncoderDebug.AddUpdateFence(VertexFence));
-			Fence->Write(mtlpp::RenderStages::Vertex);
-		}
-		else if (BlitCommandEncoder)
-		{
-			BlitCommandEncoder.WaitForFence(InnerFence);
-			METAL_DEBUG_LAYER(EMetalDebugLevelFastValidation, BlitEncoderDebug.AddWaitFence(VertexFence));
-			Fence->Wait(mtlpp::RenderStages::Vertex);
-			
-			BlitCommandEncoder.UpdateFence(InnerFence);
-			METAL_DEBUG_LAYER(EMetalDebugLevelFastValidation, BlitEncoderDebug.AddUpdateFence(VertexFence));
-			Fence->Write(mtlpp::RenderStages::Vertex);
-		}
+		// Will be waited at encoder end recording
+		CommandEncoderFence.Fence = Fence;
 	}
 }
 
@@ -1405,13 +1275,10 @@ void FMetalCommandEncoder::SetShaderAccelerationStructure(mtlpp::FunctionType co
 
 void FMetalCommandEncoder::SetShaderBuffer(mtlpp::FunctionType const FunctionType, FMetalBuffer const& Buffer, NSUInteger const Offset, NSUInteger const Length, NSUInteger index, mtlpp::ResourceUsage const Usage, EPixelFormat const Format, NSUInteger const ElementRowPitch, TArray<TTuple<ns::AutoReleased<mtlpp::Resource>, mtlpp::ResourceUsage>> ReferencedResources)
 {
+	FenceResource(Buffer, FunctionType);
 	check(index < ML_MaxBuffers);
     if(GetMetalDeviceContext().SupportsFeature(EMetalFeaturesSetBufferOffset) && Buffer && (ShaderBuffers[uint32(FunctionType)].Bound & (1 << index)) && ShaderBuffers[uint32(FunctionType)].Buffers[index] == Buffer)
     {
-		if (FunctionType == mtlpp::FunctionType::Vertex || FunctionType == mtlpp::FunctionType::Kernel)
-		{
-			FenceResource(Buffer);
-		}
 		SetShaderBufferOffset(FunctionType, Offset, Length, index);
 		ShaderBuffers[uint32(FunctionType)].Usage[index] = Usage;
 		ShaderBuffers[uint32(FunctionType)].ReferencedResources[index] = ReferencedResources;
@@ -1439,11 +1306,6 @@ void FMetalCommandEncoder::SetShaderBuffer(mtlpp::FunctionType const FunctionTyp
 		
 		SetShaderBufferInternal(FunctionType, index);
     }
-    
-	if (Buffer)
-	{
-		BufferBindingHistory.Add(ns::AutoReleased<FMetalBuffer>(Buffer));
-	}
 }
 
 void FMetalCommandEncoder::SetShaderData(mtlpp::FunctionType const FunctionType, FMetalBufferData* Data, NSUInteger const Offset, NSUInteger const Index, EPixelFormat const Format, NSUInteger const ElementRowPitch)
@@ -1570,12 +1432,12 @@ void FMetalCommandEncoder::SetShaderBufferOffset(mtlpp::FunctionType FunctionTyp
 
 void FMetalCommandEncoder::SetShaderTexture(mtlpp::FunctionType FunctionType, FMetalTexture const& Texture, NSUInteger index, mtlpp::ResourceUsage Usage)
 {
+	FenceResource(Texture, FunctionType);
 	check(index < ML_MaxTextures);
 	switch (FunctionType)
 	{
 		case mtlpp::FunctionType::Vertex:
 			check (RenderCommandEncoder);
-			FenceResource(Texture);
 			// MTLPP_VALIDATE(mtlpp::RenderCommandEncoder, RenderCommandEncoder, SafeGetRuntimeDebuggingLevel() >= EMetalDebugLevelValidation, UseResource(Texture, mtlpp::ResourceUsage::Read));
 			METAL_DEBUG_LAYER(EMetalDebugLevelFastValidation, RenderEncoderDebug.SetTexture(EMetalShaderVertex, Texture, index));
 			RenderCommandEncoder.SetVertexTexture(Texture, index);
@@ -1588,7 +1450,6 @@ void FMetalCommandEncoder::SetShaderTexture(mtlpp::FunctionType FunctionType, FM
 			break;
 		case mtlpp::FunctionType::Kernel:
 			check (ComputeCommandEncoder);
-			FenceResource(Texture);
 			// MTLPP_VALIDATE(mtlpp::ComputeCommandEncoder, ComputeCommandEncoder, SafeGetRuntimeDebuggingLevel() >= EMetalDebugLevelValidation, UseResource(Texture, mtlpp::ResourceUsage::Read));
 			METAL_DEBUG_LAYER(EMetalDebugLevelFastValidation, ComputeEncoderDebug.SetTexture(Texture, index));
 			ComputeCommandEncoder.SetTexture(Texture, index);
@@ -1612,7 +1473,6 @@ void FMetalCommandEncoder::SetShaderTexture(mtlpp::FunctionType FunctionType, FM
 		}
 		
 		ShaderBuffers[uint32(FunctionType)].SetTextureSwizzle(index, Swizzle);
-		TextureBindingHistory.Add(ns::AutoReleased<FMetalTexture>(Texture));
 	}
 }
 
@@ -1652,16 +1512,14 @@ void FMetalCommandEncoder::SetShaderSideTable(mtlpp::FunctionType const Function
 
 void FMetalCommandEncoder::UseIndirectArgumentResource(FMetalTexture const& Texture, mtlpp::ResourceUsage const Usage)
 {
-	FenceResource(Texture);
+	FenceResource(Texture, mtlpp::FunctionType::Vertex);
 	UseResource(Texture, Usage);
-	TextureBindingHistory.Add(ns::AutoReleased<FMetalTexture>(Texture));
 }
 
 void FMetalCommandEncoder::UseIndirectArgumentResource(FMetalBuffer const& Buffer, mtlpp::ResourceUsage const Usage)
 {
-	FenceResource(Buffer);
+	FenceResource(Buffer, mtlpp::FunctionType::Vertex);
 	UseResource(Buffer, Usage);
-	BufferBindingHistory.Add(ns::AutoReleased<FMetalBuffer>(Buffer));
 }
 
 void FMetalCommandEncoder::TransitionResources(mtlpp::Resource const& Resource)
@@ -1689,19 +1547,9 @@ FMetalSubBufferRing& FMetalCommandEncoder::GetRingBuffer(void)
 
 #pragma mark - Public Resource query Access -
 
-bool FMetalCommandEncoder::HasTextureBindingHistory(FMetalTexture const& Texture) const
-{
-	return TextureBindingHistory.Contains(ns::AutoReleased<FMetalTexture>(Texture));
-}
-
-bool FMetalCommandEncoder::HasBufferBindingHistory(FMetalBuffer const& Buffer) const
-{
-	return BufferBindingHistory.Contains(ns::AutoReleased<FMetalBuffer>(Buffer));
-}
-
 #pragma mark - Private Functions -
 
-void FMetalCommandEncoder::FenceResource(mtlpp::Texture const& Resource)
+void FMetalCommandEncoder::FenceResource(mtlpp::Texture const& Resource, mtlpp::FunctionType Function, bool bIsRenderTarget/* = false*/)
 {
 	mtlpp::Resource::Type Res = Resource.GetPtr();
 	ns::AutoReleased<mtlpp::Texture> Parent = Resource.GetParentTexture();
@@ -1714,77 +1562,59 @@ void FMetalCommandEncoder::FenceResource(mtlpp::Texture const& Resource)
 	{
 		Res = Buffer.GetPtr();
 	}
-	if (FenceStage == mtlpp::RenderStages::Vertex || FenceResources.Contains(Res))
+
+	FMetalBarrierScope& BarrierScope = CommandEncoderFence.BarrierScope;
+	if (CommandEncoderFence.FenceResources.Contains(Res))
 	{
-		FenceStage = mtlpp::RenderStages::Vertex;
-		
-		for (FMetalFence* FragFence : FragmentFences)
+		switch (Function)
 		{
-			if (FragFence->NeedsWait(mtlpp::RenderStages::Fragment))
+		case mtlpp::FunctionType::Kernel:
+			BarrierScope.TexturesWaitStage = EMetalFenceWaitStage::BeforeVertex;
+			break;
+		case mtlpp::FunctionType::Vertex:
+			BarrierScope.TexturesWaitStage = EMetalFenceWaitStage::BeforeVertex;
+			break;
+		case mtlpp::FunctionType::Fragment:
+			if (bIsRenderTarget)
 			{
-				mtlpp::Fence FragmentFence = FragFence->Get(mtlpp::RenderStages::Fragment);
-				mtlpp::Fence FragInnerFence = METAL_DEBUG_OPTION(CommandList.GetCommandQueue().GetRuntimeDebuggingLevel() >= EMetalDebugLevelValidation ? mtlpp::Fence(((FMetalDebugFence*)FragmentFence.GetPtr()).Inner) :) FragmentFence;
-				
-				if (RenderCommandEncoder)
-				{
-					RenderCommandEncoder.WaitForFence(FragInnerFence, (mtlpp::RenderStages)(mtlpp::RenderStages::Vertex));
-					METAL_DEBUG_LAYER(EMetalDebugLevelFastValidation, RenderEncoderDebug.AddWaitFence(FragmentFence));
-					FragFence->Wait(mtlpp::RenderStages::Fragment);
-				}
-				else if (ComputeCommandEncoder)
-				{
-					ComputeCommandEncoder.WaitForFence(FragInnerFence);
-					METAL_DEBUG_LAYER(EMetalDebugLevelFastValidation, ComputeEncoderDebug.AddWaitFence(FragmentFence));
-					FragFence->Wait(mtlpp::RenderStages::Fragment);
-				}
-				else if (BlitCommandEncoder)
-				{
-					BlitCommandEncoder.WaitForFence(FragInnerFence);
-					METAL_DEBUG_LAYER(EMetalDebugLevelFastValidation, BlitEncoderDebug.AddWaitFence(FragmentFence));
-					FragFence->Wait(mtlpp::RenderStages::Fragment);
-				}
-				METAL_DEBUG_LAYER(EMetalDebugLevelFastValidation, WaitCount++);
+				BarrierScope.RenderTargetsWaitStage = EMetalFenceWaitStage::BeforeFragment;
 			}
-		}
-		FragmentFences.Empty();
+			else if (BarrierScope.TexturesWaitStage == EMetalFenceWaitStage::None)
+			{
+				BarrierScope.TexturesWaitStage = EMetalFenceWaitStage::BeforeFragment;
+			}
+			break;
+		default:
+			checkNoEntry();
+			break;
+		};
 	}
 }
 
-void FMetalCommandEncoder::FenceResource(mtlpp::Buffer const& Resource)
+void FMetalCommandEncoder::FenceResource(mtlpp::Buffer const& Resource, mtlpp::FunctionType Function)
 {
 	mtlpp::Resource::Type Res = Resource.GetPtr();
-	if (FenceStage == mtlpp::RenderStages::Vertex || FenceResources.Contains(Res))
+	FMetalBarrierScope& BarrierScope = CommandEncoderFence.BarrierScope;
+	if (CommandEncoderFence.FenceResources.Contains(Res))
 	{
-		FenceStage = mtlpp::RenderStages::Vertex;
-		
-		for (FMetalFence* FragFence : FragmentFences)
+		switch (Function)
 		{
-			if (FragFence->NeedsWait(mtlpp::RenderStages::Fragment))
+		case mtlpp::FunctionType::Kernel:
+			BarrierScope.BuffersWaitStage = EMetalFenceWaitStage::BeforeVertex;
+			break;
+		case mtlpp::FunctionType::Vertex:
+			BarrierScope.BuffersWaitStage = EMetalFenceWaitStage::BeforeVertex;
+			break;
+		case mtlpp::FunctionType::Fragment:
+			if (BarrierScope.BuffersWaitStage == EMetalFenceWaitStage::None)
 			{
-				mtlpp::Fence FragmentFence = FragFence->Get(mtlpp::RenderStages::Fragment);
-				mtlpp::Fence FragInnerFence = METAL_DEBUG_OPTION(CommandList.GetCommandQueue().GetRuntimeDebuggingLevel() >= EMetalDebugLevelValidation ? mtlpp::Fence(((FMetalDebugFence*)FragmentFence.GetPtr()).Inner) :) FragmentFence;
-				
-				if (RenderCommandEncoder)
-				{
-					RenderCommandEncoder.WaitForFence(FragInnerFence, (mtlpp::RenderStages)(mtlpp::RenderStages::Vertex));
-					METAL_DEBUG_LAYER(EMetalDebugLevelFastValidation, RenderEncoderDebug.AddWaitFence(FragmentFence));
-					FragFence->Wait(mtlpp::RenderStages::Fragment);
-				}
-				else if (ComputeCommandEncoder)
-				{
-					ComputeCommandEncoder.WaitForFence(FragInnerFence);
-					METAL_DEBUG_LAYER(EMetalDebugLevelFastValidation, ComputeEncoderDebug.AddWaitFence(FragmentFence));
-					FragFence->Wait(mtlpp::RenderStages::Fragment);
-				}
-				else if (BlitCommandEncoder)
-				{
-					BlitCommandEncoder.WaitForFence(FragInnerFence);
-					METAL_DEBUG_LAYER(EMetalDebugLevelFastValidation, BlitEncoderDebug.AddWaitFence(FragmentFence));
-					FragFence->Wait(mtlpp::RenderStages::Fragment);
-				}
+				BarrierScope.BuffersWaitStage = EMetalFenceWaitStage::BeforeFragment;
 			}
-		}
-		FragmentFences.Empty();
+			break;
+		default:
+			checkNoEntry();
+			break;
+		};
 	}
 }
 
@@ -1892,13 +1722,12 @@ void FMetalCommandEncoder::SetShaderBufferInternal(mtlpp::FunctionType Function,
 			ActiveBuffers.Add(Buffer);
 		}
 #endif
-        
+        	FenceResource(Buffer, Function);
 		switch (Function)
 		{
 			case mtlpp::FunctionType::Vertex:
 				Binding.Bound |= (1 << Index);
 				check(RenderCommandEncoder);
-				FenceResource(Buffer);
 				// MTLPP_VALIDATE(mtlpp::RenderCommandEncoder, RenderCommandEncoder, SafeGetRuntimeDebuggingLevel() >= EMetalDebugLevelValidation, UseResource(Buffer, mtlpp::ResourceUsage::Read));
 				METAL_DEBUG_LAYER(EMetalDebugLevelFastValidation, RenderEncoderDebug.SetBuffer(EMetalShaderVertex, Buffer, Offset, Index));
 				RenderCommandEncoder.SetVertexBuffer(Buffer, Offset, Index);
@@ -1915,7 +1744,6 @@ void FMetalCommandEncoder::SetShaderBufferInternal(mtlpp::FunctionType Function,
 			case mtlpp::FunctionType::Kernel:
 				Binding.Bound |= (1 << Index);
 				check(ComputeCommandEncoder);
-				FenceResource(Buffer);
 				// MTLPP_VALIDATE(mtlpp::ComputeCommandEncoder, ComputeCommandEncoder, SafeGetRuntimeDebuggingLevel() >= EMetalDebugLevelValidation, UseResource(Buffer, mtlpp::ResourceUsage::Read));
 				METAL_DEBUG_LAYER(EMetalDebugLevelFastValidation, ComputeEncoderDebug.SetBuffer(Buffer, Offset, Index));
 				ComputeCommandEncoder.SetBuffer(Buffer, Offset, Index);
diff --git a/Engine/Source/Runtime/Apple/MetalRHI/Private/MetalCommandEncoder.h b/Engine/Source/Runtime/Apple/MetalRHI/Private/MetalCommandEncoder.h
--- a/Engine/Source/Runtime/Apple/MetalRHI/Private/MetalCommandEncoder.h
+++ b/Engine/Source/Runtime/Apple/MetalRHI/Private/MetalCommandEncoder.h
@@ -82,6 +82,84 @@ struct FMetalCommandDebug
 	FMetalCommandData Data;
 };
 
+namespace EMetalFenceWaitStage
+{
+	enum Type
+	{
+		None			= 0,
+		BeforeVertex	= 1 << 0, // Wait fence before VS/CS is executed
+		BeforeFragment	= 1 << 1, // Wait fence before PS is executed
+	};
+};
+
+// An emulator of MTLBarrierScope that is supported since IOS 16.0
+struct FMetalBarrierScope
+{
+	EMetalFenceWaitStage::Type BuffersWaitStage;
+	EMetalFenceWaitStage::Type RenderTargetsWaitStage;
+	EMetalFenceWaitStage::Type TexturesWaitStage;
+
+	FMetalBarrierScope()
+	{
+		Reset();
+	}
+
+	void Reset()
+	{
+		BuffersWaitStage = EMetalFenceWaitStage::None;
+		RenderTargetsWaitStage = EMetalFenceWaitStage::None;
+		TexturesWaitStage = EMetalFenceWaitStage::None;
+	}
+
+	EMetalFenceWaitStage::Type GetFenceWaitStage()
+	{
+		if (BuffersWaitStage == EMetalFenceWaitStage::BeforeVertex
+			|| TexturesWaitStage == EMetalFenceWaitStage::BeforeVertex)
+		{
+			return EMetalFenceWaitStage::BeforeVertex;
+		}
+		else
+		{
+			return EMetalFenceWaitStage::BeforeFragment;
+		}
+	}
+};
+
+struct FMetalCommandEncoderFence
+{
+	// Resources produced in the prologue pass
+	TSet<mtlpp::Resource::Type> FenceResources;
+	// Fence written by the prologue pass
+	TRefCountPtr<FMetalFence> Fence;
+	// The barrier state for the current pass
+	FMetalBarrierScope BarrierScope;
+
+	FMetalCommandEncoderFence() = default;
+
+	FMetalCommandEncoderFence(FMetalCommandEncoderFence&& In)
+	{
+		FenceResources = MoveTemp(In.FenceResources);
+		Fence = In.Fence;
+	}
+
+	FMetalCommandEncoderFence& operator=(FMetalCommandEncoderFence&& lhs)
+	{
+		if (this != &lhs)
+		{
+			FenceResources = MoveTemp(lhs.FenceResources);
+			Fence = lhs.Fence;
+		}
+		return *this;
+	}
+
+	void Reset()
+	{
+		FenceResources.Empty();
+		Fence = nullptr;
+		BarrierScope.Reset();
+	}
+};
+
 class FMetalCommandBufferMarkers : public ns::Object<FMetalCommandBufferDebug*, ns::CallingConvention::ObjectiveC>
 {
 public:
@@ -235,12 +313,6 @@ public:
 	
 	/** Prevent further GPU work until the event is reached. */
 	void WaitForFence(FMetalFence* Fence);
-	
-	/**
-	 * Prevent further GPU work until the event is reached and then kick the fence again so future encoders may wait on it.
-	 * Unlike the functions above this is not guarded and the caller must ensure proper matching waits & updates.
-	 */
-	void WaitAndUpdateFence(FMetalFence* Fence);
 
 #pragma mark - Public Debug Support -
 	
@@ -481,13 +553,6 @@ public:
 	FMetalSubBufferRing& GetRingBuffer(void);
 	
 #pragma mark - Public Resource query Access -
-
-	/*
-	 * Returns True if the Resource has been bound to a command encoder, otherwise false.  History will be cleared after a commit operation
-	 * @returns True if the Resource has been bound to a command encoder, otherwise false.
-	 */
-	bool HasTextureBindingHistory(FMetalTexture const& Texture) const;
-	bool HasBufferBindingHistory(FMetalBuffer const& Buffer) const;
 	
 private:
 #pragma mark - Private Functions -
@@ -502,8 +567,8 @@ private:
 	
 	void SetShaderBufferInternal(mtlpp::FunctionType Function, uint32 Index);
 	
-	void FenceResource(mtlpp::Texture const& Resource);
-	void FenceResource(mtlpp::Buffer const& Resource);
+	void FenceResource(mtlpp::Texture const& Resource, mtlpp::FunctionType Function, bool bIsRenderTarget = false);
+	void FenceResource(mtlpp::Buffer const& Resource, mtlpp::FunctionType Function);
 
 	void UseResource(mtlpp::Resource const& Resource, mtlpp::ResourceUsage const Usage);
 	
@@ -591,27 +656,16 @@ public:
 #if ENABLE_METAL_GPUPROFILE
 	FMetalCommandBufferStats* CommandBufferStats;
 #endif
-#if METAL_DEBUG_OPTIONS
-	uint8 WaitCount;
-	uint8 UpdateCount;
-#endif
 
 	TArray<ns::Object<mtlpp::CommandBufferHandler>> CompletionHandlers;
 	NSMutableArray* DebugGroups;
     
     TSet<ns::AutoReleased<FMetalBuffer>> ActiveBuffers;
-    
-	TSet<ns::AutoReleased<FMetalBuffer>> BufferBindingHistory;
-	TSet<ns::AutoReleased<FMetalTexture>> TextureBindingHistory;
 	
 	TMap<mtlpp::Resource::Type, mtlpp::ResourceUsage> ResourceUsage;
 	
 	TSet<mtlpp::Resource::Type> TransitionedResources;
-	TSet<mtlpp::Resource::Type> FenceResources;
-	
-	TSet<TRefCountPtr<FMetalFence>> FragmentFences;
-	
-	mtlpp::RenderStages FenceStage;
+	FMetalCommandEncoderFence CommandEncoderFence;
 	uint32 EncoderNum;
 	uint32 CmdBufIndex;
 	EMetalCommandEncoderType Type;
diff --git a/Engine/Source/Runtime/Apple/MetalRHI/Private/MetalCommandQueue.cpp b/Engine/Source/Runtime/Apple/MetalRHI/Private/MetalCommandQueue.cpp
--- a/Engine/Source/Runtime/Apple/MetalRHI/Private/MetalCommandQueue.cpp
+++ b/Engine/Source/Runtime/Apple/MetalRHI/Private/MetalCommandQueue.cpp
@@ -325,13 +325,12 @@ FMetalFence* FMetalCommandQueue::CreateFence(ns::String const& Label) const
 	if ((Features & EMetalFeaturesFences) != 0)
 	{
 		FMetalFence* InternalFence = FMetalFencePool::Get().AllocateFence();
-		for (uint32 i = mtlpp::RenderStages::Vertex; InternalFence && i <= mtlpp::RenderStages::Fragment; i++)
 		{
-			mtlpp::Fence InnerFence = InternalFence->Get((mtlpp::RenderStages)i);
+			mtlpp::Fence InnerFence = InternalFence->Get();
 			NSString* String = nil;
 			if (GetEmitDrawEvents())
 			{
-				String = [NSString stringWithFormat:@"%u %p: %@", i, InnerFence.GetPtr(), Label.GetPtr()];
+				String = [NSString stringWithFormat:@"%p: %@", InnerFence.GetPtr(), Label.GetPtr()];
 			}
 	#if METAL_DEBUG_OPTIONS
 			if (RuntimeDebuggingLevel >= EMetalDebugLevelValidation)
diff --git a/Engine/Source/Runtime/Apple/MetalRHI/Private/MetalFence.cpp b/Engine/Source/Runtime/Apple/MetalRHI/Private/MetalFence.cpp
--- a/Engine/Source/Runtime/Apple/MetalRHI/Private/MetalFence.cpp
+++ b/Engine/Source/Runtime/Apple/MetalRHI/Private/MetalFence.cpp
@@ -119,20 +119,6 @@ uint32 FMetalFence::Release() const
 	return Refs;
 }
 
-#if METAL_DEBUG_OPTIONS
-void FMetalFence::Validate(void) const
-{
-	if (GetMetalDeviceContext().GetCommandQueue().GetRuntimeDebuggingLevel() >= EMetalDebugLevelValidation && Get(mtlpp::RenderStages::Vertex))
-	{
-		[(FMetalDebugFence*)Get(mtlpp::RenderStages::Vertex).GetPtr() validate];
-		if (Get(mtlpp::RenderStages::Fragment))
-		{
-			[(FMetalDebugFence*)Get(mtlpp::RenderStages::Fragment).GetPtr() validate];
-		}
-	}
-}
-#endif
-
 void FMetalFencePool::Initialise(mtlpp::Device const& InDevice)
 {
 	Device = InDevice;
@@ -141,13 +127,10 @@ void FMetalFencePool::Initialise(mtlpp::Device const& InDevice)
 #if METAL_DEBUG_OPTIONS
 		if (GMetalRuntimeDebugLevel >= EMetalDebugLevelValidation)
 		{
-			FMetalDebugFence* VertexFence = [[FMetalDebugFence new] autorelease];
-			VertexFence.Inner = Device.NewFence();
-			FMetalDebugFence* FragmentFence = [[FMetalDebugFence new] autorelease];
-			FragmentFence.Inner = Device.NewFence();
+			FMetalDebugFence* DebugFence = [[FMetalDebugFence new] autorelease];
+			DebugFence.Inner = Device.NewFence();
 			FMetalFence* F = new FMetalFence;
-			F->Set(mtlpp::RenderStages::Vertex, VertexFence);
-			F->Set(mtlpp::RenderStages::Fragment, FragmentFence);
+			F->Set(DebugFence);
 			Fences.Add(F);
 			Lifo.Push(F);
 		}
@@ -155,8 +138,7 @@ void FMetalFencePool::Initialise(mtlpp::Device const& InDevice)
 #endif
 		{
 			FMetalFence* F = new FMetalFence;
-			F->Set(mtlpp::RenderStages::Vertex, Device.NewFence());
-			F->Set(mtlpp::RenderStages::Fragment, Device.NewFence());
+			F->Set(Device.NewFence());
 #if METAL_DEBUG_OPTIONS
 			if (GMetalRuntimeDebugLevel >= EMetalDebugLevelValidation)
 			{
@@ -192,21 +174,6 @@ FMetalFence* FMetalFencePool::AllocateFence()
 	return Fence;
 }
 
-void FMetalFence::ValidateUsage(FMetalFence* InFence)
-{
-	if (InFence)
-	{
-		if (InFence->NumWrites(mtlpp::RenderStages::Vertex) != InFence->NumWaits(mtlpp::RenderStages::Vertex))
-		{
-			UE_LOG(LogMetal, Warning, TEXT("%p (%s) writes %d waits %d"), InFence, *FString(InFence->Get(mtlpp::RenderStages::Vertex).GetLabel()), (uint32)InFence->NumWrites(mtlpp::RenderStages::Vertex), (uint32)InFence->NumWaits(mtlpp::RenderStages::Vertex));
-		}
-		if (InFence->NumWrites(mtlpp::RenderStages::Fragment) != InFence->NumWaits(mtlpp::RenderStages::Fragment))
-		{
-			UE_LOG(LogMetal, Warning, TEXT("%p (%s) writes %d waits %d"), InFence, *FString(InFence->Get(mtlpp::RenderStages::Fragment).GetLabel()), (uint32)InFence->NumWrites(mtlpp::RenderStages::Fragment), (uint32)InFence->NumWaits(mtlpp::RenderStages::Fragment));
-		}
-	}
-}
-
 void FMetalFencePool::ReleaseFence(FMetalFence* const InFence)
 {
 	if (InFence)
@@ -217,7 +184,6 @@ void FMetalFencePool::ReleaseFence(FMetalFence* const InFence)
 		if (GMetalRuntimeDebugLevel >= EMetalDebugLevelValidation)
 		{
 			FScopeLock Lock(&Mutex);
-			FMetalFence::ValidateUsage(InFence);
 			check(!Fences.Contains(InFence));
 			Fences.Add(InFence);
 		}
diff --git a/Engine/Source/Runtime/Apple/MetalRHI/Private/MetalFence.h b/Engine/Source/Runtime/Apple/MetalRHI/Private/MetalFence.h
--- a/Engine/Source/Runtime/Apple/MetalRHI/Private/MetalFence.h
+++ b/Engine/Source/Runtime/Apple/MetalRHI/Private/MetalFence.h
@@ -63,19 +63,11 @@
 
 class FMetalFence
 {
-	enum
-	{
-		NumFenceStages = 2
-	};
 public:
 	FMetalFence()
 	: NumRefs(0)
 	{
-		for (uint32 i = 0; i < NumFenceStages; i++)
-		{
-			Writes[i] = 0;
-			Waits[i] = 0;
-		}
+		Reset();
 	}
 	
 	explicit FMetalFence(FMetalFence const& Other)
@@ -103,73 +95,51 @@ public:
 	{
 		if (&Other != this)
 		{
-			for (uint32 i = 0; i < NumFenceStages; i++)
-			{
-				Fences[i] = Other.Fences[i];
-			}
+			Fence = Other.Fence;
 		}
 		return *this;
 	}
 	
-#if METAL_DEBUG_OPTIONS
-	void Validate(void) const;
-#endif
-	
 	void Reset(void)
 	{
-		for (uint32 i = 0; i < NumFenceStages; i++)
-		{
-			Writes[i] = 0;
-			Waits[i] = 0;
-		}
-	}
-	
-	void Write(mtlpp::RenderStages Stage)
-	{
-		Writes[(uint32)Stage - 1u]++;
+		WriteNum = 0;
+		WaitNum = 0;
 	}
 	
-	void Wait(mtlpp::RenderStages Stage)
+	void Write()
 	{
-		Waits[(uint32)Stage - 1u]++;
+		WriteNum++;
 	}
 	
-	int8 NumWrites(mtlpp::RenderStages Stage) const
+	void Wait()
 	{
-		return Writes[(uint32)Stage - 1u];
+		WaitNum++;
 	}
 	
-	int8 NumWaits(mtlpp::RenderStages Stage) const
+	int8 NumWrite() const
 	{
-		return Waits[(uint32)Stage - 1u];
+		return WriteNum;
 	}
 	
-	bool NeedsWrite(mtlpp::RenderStages Stage) const
+	int8 NumWait() const
 	{
-		return Writes[(uint32)Stage - 1u] == 0 || (Waits[(uint32)Stage - 1u] > Writes[(uint32)Stage - 1u]);
+		return WaitNum;
 	}
 	
-	bool NeedsWait(mtlpp::RenderStages Stage) const
+	mtlpp::Fence Get() const
 	{
-		return Waits[(uint32)Stage - 1u] == 0 || (Writes[(uint32)Stage - 1u] > Waits[(uint32)Stage - 1u]);
+		return Fence;
 	}
 	
-	mtlpp::Fence Get(mtlpp::RenderStages Stage) const
+	void Set(mtlpp::Fence InFence)
 	{
-		return Fences[(uint32)Stage - 1u];
+		Fence = InFence;
 	}
 	
-	void Set(mtlpp::RenderStages Stage, mtlpp::Fence Fence)
-	{
-		Fences[(uint32)Stage - 1u] = Fence;
-	}
-	
-	static void ValidateUsage(FMetalFence* Fence);
-	
 private:
-	mtlpp::Fence Fences[NumFenceStages];
-	int8 Writes[NumFenceStages];
-	int8 Waits[NumFenceStages];
+	mtlpp::Fence Fence;
+	int8 WriteNum;
+	int8 WaitNum;
 	mutable int32 NumRefs;	
 };
 
diff --git a/Engine/Source/Runtime/Apple/MetalRHI/Private/MetalRenderPass.cpp b/Engine/Source/Runtime/Apple/MetalRHI/Private/MetalRenderPass.cpp
--- a/Engine/Source/Runtime/Apple/MetalRHI/Private/MetalRenderPass.cpp
+++ b/Engine/Source/Runtime/Apple/MetalRHI/Private/MetalRenderPass.cpp
@@ -43,7 +43,6 @@ FMetalRenderPass::FMetalRenderPass(FMetalCommandList& InCmdList, FMetalStateCach
 	: CmdList(InCmdList)
 	, State(Cache)
 	, CurrentEncoder(InCmdList, EMetalCommandEncoderCurrent)
-	, PrologueEncoder(InCmdList, EMetalCommandEncoderPrologue)
 	, RenderPassDesc(nil)
 	, ComputeDispatchType(mtlpp::DispatchType::Serial)
 	, NumOutstandingOps(0)
@@ -54,8 +53,6 @@ FMetalRenderPass::FMetalRenderPass(FMetalCommandList& InCmdList, FMetalStateCach
 FMetalRenderPass::~FMetalRenderPass(void)
 {
 	check(!CurrentEncoder.GetCommandBuffer());
-	check(!PrologueEncoder.GetCommandBuffer());
-	check(!PassStartFence);
 }
 
 void FMetalRenderPass::SetDispatchType(mtlpp::DispatchType Type)
@@ -72,43 +69,8 @@ void FMetalRenderPass::Begin()
 	}
 }
 
-void FMetalRenderPass::Wait(FMetalFence* Fence)
-{
-	if (Fence)
-	{
-		if (PrologueEncoder.IsBlitCommandEncoderActive() || PrologueEncoder.IsComputeCommandEncoderActive())
-		{
-			PrologueEncoder.WaitForFence(Fence);
-			METAL_DEBUG_LAYER(EMetalDebugLevelValidation, FMetalFence::ValidateUsage(Fence));
-		}
-		else if (CurrentEncoder.IsRenderCommandEncoderActive() || CurrentEncoder.IsBlitCommandEncoderActive() || CurrentEncoder.IsComputeCommandEncoderActive())
-		{
-			CurrentEncoder.WaitForFence(Fence);
-			METAL_DEBUG_LAYER(EMetalDebugLevelValidation, FMetalFence::ValidateUsage(Fence));
-		}
-        else
-        {
-            PassStartFence = Fence;
-			PrologueStartEncoderFence = Fence;
-        }
-	}
-}
-
 TRefCountPtr<FMetalFence> const& FMetalRenderPass::Submit(EMetalSubmitFlags Flags)
 {
-	if (CurrentEncoder.GetCommandBuffer() || (Flags & EMetalSubmitFlagsAsyncCommandBuffer))
-	{
-		if (PrologueEncoder.IsBlitCommandEncoderActive() || PrologueEncoder.IsComputeCommandEncoderActive())
-		{
-			check(PrologueEncoder.GetCommandBuffer());
-			PrologueEncoderFence = PrologueEncoder.EndEncoding();
-		}
-		if (PrologueEncoder.GetCommandBuffer())
-		{
-			PrologueEncoder.CommitCommandBuffer((Flags & EMetalSubmitFlagsAsyncCommandBuffer) ? Flags : EMetalSubmitFlagsNone);
-        }
-    }
-
     // Must be on the render thread if there's no RHI thread
     // Must be on the RHI thread otherwise
     CheckMetalThread();
@@ -145,7 +107,6 @@ TRefCountPtr<FMetalFence> const& FMetalRenderPass::Submit(EMetalSubmitFlags Flag
 	OutstandingBufferUploads.Empty();
 	if (Flags & EMetalSubmitFlagsResetState)
 	{
-		PrologueEncoder.Reset();
 		CurrentEncoder.Reset();
 	}
 	
@@ -161,10 +122,6 @@ void FMetalRenderPass::BeginRenderPass(mtlpp::RenderPassDescriptor RenderPass)
 	check(CurrentEncoder.GetCommandBuffer());
 	
 	// EndEncoding should provide the encoder fence...
-	if (PrologueEncoder.IsBlitCommandEncoderActive() || PrologueEncoder.IsComputeCommandEncoderActive())
-	{
-		PrologueEncoderFence = PrologueEncoder.EndEncoding();
-	}
 	if (CurrentEncoder.IsRenderCommandEncoderActive() || CurrentEncoder.IsBlitCommandEncoderActive() || CurrentEncoder.IsComputeCommandEncoderActive())
 	{
 		State.FlushVisibilityResults(CurrentEncoder);
@@ -180,32 +137,16 @@ void FMetalRenderPass::BeginRenderPass(mtlpp::RenderPassDescriptor RenderPass)
 	if (!GMetalDeferRenderPasses || !State.CanRestartRenderPass())
 	{
 		CurrentEncoder.BeginRenderCommandEncoding();
-		if (PassStartFence)
-		{
-			CurrentEncoder.WaitForFence(PassStartFence);
-			PassStartFence = nullptr;
-		}
 		if (CurrentEncoderFence)
 		{
 			CurrentEncoder.WaitForFence(CurrentEncoderFence);
 			CurrentEncoderFence = nullptr;
 		}
-		if (PrologueEncoderFence)
-		{
-			// Consume on the current encoder but do not invalidate
-			CurrentEncoder.WaitForFence(PrologueEncoderFence);
-		}
-		if (PrologueEncoder.IsBlitCommandEncoderActive() || PrologueEncoder.IsComputeCommandEncoderActive())
-		{
-			CurrentEncoder.WaitForFence(PrologueEncoder.GetEncoderFence());
-		}
 		State.SetRenderStoreActions(CurrentEncoder, false);
 		check(CurrentEncoder.IsRenderCommandEncoderActive());
 	}
 	
 	bWithinRenderPass = true;
-	
-	check(!PrologueEncoder.IsBlitCommandEncoderActive() && !PrologueEncoder.IsComputeCommandEncoderActive());
 }
 
 void FMetalRenderPass::RestartRenderPass(mtlpp::RenderPassDescriptor RenderPass)
@@ -296,25 +237,11 @@ void FMetalRenderPass::RestartRenderPass(mtlpp::RenderPassDescriptor RenderPass)
 	
 	CurrentEncoder.SetRenderPassDescriptor(RenderPassDesc);
 	CurrentEncoder.BeginRenderCommandEncoding();
-	if (PassStartFence)
-	{
-		CurrentEncoder.WaitForFence(PassStartFence);
-		PassStartFence = nullptr;
-	}
 	if (CurrentEncoderFence)
 	{
 		CurrentEncoder.WaitForFence(CurrentEncoderFence);
 		CurrentEncoderFence = nullptr;
 	}
-	if (PrologueEncoderFence)
-	{
-		// Consume on the current encoder but do not invalidate
-		CurrentEncoder.WaitForFence(PrologueEncoderFence);
-	}
-	if (PrologueEncoder.IsBlitCommandEncoderActive() || PrologueEncoder.IsComputeCommandEncoderActive())
-	{
-		CurrentEncoder.WaitForFence(PrologueEncoder.GetEncoderFence());
-	}
 	State.SetRenderStoreActions(CurrentEncoder, false);
 	
 	check(CurrentEncoder.IsRenderCommandEncoderActive());
@@ -544,28 +471,6 @@ void FMetalRenderPass::DrawIndexedPrimitiveIndirect(uint32 PrimitiveType,FMetalR
 
 void FMetalRenderPass::Dispatch(uint32 ThreadGroupCountX, uint32 ThreadGroupCountY, uint32 ThreadGroupCountZ)
 {
-    if (CurrentEncoder.NumEncodedPasses() == 0)
-	{
-		ConditionalSwitchToAsyncCompute();
-		check(PrologueEncoder.GetCommandBuffer());
-		check(PrologueEncoder.IsComputeCommandEncoderActive());
-		
-		PrepareToAsyncDispatch();
-		
-		TRefCountPtr<FMetalComputeShader> ComputeShader = State.GetComputeShader();
-		check(ComputeShader);
-		
-		METAL_GPUPROFILE(FMetalProfiler::GetProfiler()->EncodeDispatch(PrologueEncoder.GetCommandBufferStats(), __FUNCTION__));
-		
-		mtlpp::Size ThreadgroupCounts = mtlpp::Size(ComputeShader->NumThreadsX, ComputeShader->NumThreadsY, ComputeShader->NumThreadsZ);
-		check(ComputeShader->NumThreadsX > 0 && ComputeShader->NumThreadsY > 0 && ComputeShader->NumThreadsZ > 0);
-		mtlpp::Size Threadgroups = mtlpp::Size(ThreadGroupCountX, ThreadGroupCountY, ThreadGroupCountZ);
-		PrologueEncoder.GetComputeCommandEncoder().DispatchThreadgroups(Threadgroups, ThreadgroupCounts);
-		METAL_DEBUG_LAYER(EMetalDebugLevelFastValidation, PrologueEncoder.GetComputeCommandEncoderDebugging().DispatchThreadgroups(Threadgroups, ThreadgroupCounts));
-		
-		ConditionalSubmit();
-	}
-	else
 	{
 		ConditionalSwitchToCompute();
 		check(CurrentEncoder.GetCommandBuffer());
@@ -602,29 +507,6 @@ void FMetalRenderPass::Dispatch(uint32 ThreadGroupCountX, uint32 ThreadGroupCoun
 void FMetalRenderPass::DispatchIndirect(FMetalRHIBuffer* ArgumentBuffer, uint32 ArgumentOffset)
 {
 	check(ArgumentBuffer);
-	
-    if (CurrentEncoder.NumEncodedPasses() == 0)
-	{
-		ConditionalSwitchToAsyncCompute();
-		check(PrologueEncoder.GetCommandBuffer());
-		check(PrologueEncoder.IsComputeCommandEncoderActive());
-		check(ArgumentBuffer->GetCurrentBuffer());
-		
-		PrepareToAsyncDispatch();
-		
-		TRefCountPtr<FMetalComputeShader> ComputeShader = State.GetComputeShader();
-		check(ComputeShader);
-		
-		METAL_GPUPROFILE(FMetalProfiler::GetProfiler()->EncodeDispatch(PrologueEncoder.GetCommandBufferStats(), __FUNCTION__));
-		mtlpp::Size ThreadgroupCounts = mtlpp::Size(ComputeShader->NumThreadsX, ComputeShader->NumThreadsY, ComputeShader->NumThreadsZ);
-		check(ComputeShader->NumThreadsX > 0 && ComputeShader->NumThreadsY > 0 && ComputeShader->NumThreadsZ > 0);
-		
-		PrologueEncoder.GetComputeCommandEncoder().DispatchThreadgroupsWithIndirectBuffer(ArgumentBuffer->GetCurrentBuffer(), ArgumentOffset, ThreadgroupCounts);
-		METAL_DEBUG_LAYER(EMetalDebugLevelFastValidation, PrologueEncoder.GetComputeCommandEncoderDebugging().DispatchThreadgroupsWithIndirectBuffer(ArgumentBuffer->GetCurrentBuffer(), ArgumentOffset, ThreadgroupCounts));
-		
-		ConditionalSubmit();
-	}
-	else
 	{
 		ConditionalSwitchToCompute();
 		check(CurrentEncoder.GetCommandBuffer());
@@ -786,15 +668,6 @@ void FMetalRenderPass::FillBuffer(FMetalBuffer const& Buffer, ns::Range Range, u
 	
 	mtlpp::BlitCommandEncoder TargetEncoder;
 	METAL_DEBUG_ONLY(FMetalBlitCommandEncoderDebugging Debugging);
-	bool bAsync = !CurrentEncoder.HasBufferBindingHistory(Buffer);
-	if(bAsync)
-	{
-		ConditionalSwitchToAsyncBlit();
-		TargetEncoder = PrologueEncoder.GetBlitCommandEncoder();
-		METAL_GPUPROFILE(FMetalProfiler::GetProfiler()->EncodeBlit(PrologueEncoder.GetCommandBufferStats(), FString::Printf(TEXT("FillBuffer: %p %llu %llu"), Buffer.GetPtr(), Buffer.GetOffset() + Range.Location, Range.Length)));
-		METAL_DEBUG_LAYER(EMetalDebugLevelFastValidation, Debugging = PrologueEncoder.GetBlitCommandEncoderDebugging());
-	}
-	else
 	{
 		ConditionalSwitchToBlit();
 		TargetEncoder = CurrentEncoder.GetBlitCommandEncoder();
@@ -805,9 +678,8 @@ void FMetalRenderPass::FillBuffer(FMetalBuffer const& Buffer, ns::Range Range, u
 	check(TargetEncoder.GetPtr());
 	
 	MTLPP_VALIDATE(mtlpp::BlitCommandEncoder, TargetEncoder, SafeGetRuntimeDebuggingLevel() >= EMetalDebugLevelValidation, Fill(Buffer, Range, Value));
-	METAL_DEBUG_LAYER(EMetalDebugLevelFastValidation, (bAsync ? PrologueEncoder.GetBlitCommandEncoderDebugging() : CurrentEncoder.GetBlitCommandEncoderDebugging()).Fill(Buffer, Range, Value));
+	METAL_DEBUG_LAYER(EMetalDebugLevelFastValidation, CurrentEncoder.GetBlitCommandEncoderDebugging().Fill(Buffer, Range, Value));
 	
-	if (!bAsync)
 	{
 		ConditionalSubmit();
 	}
@@ -817,15 +689,6 @@ bool FMetalRenderPass::AsyncCopyFromBufferToTexture(FMetalBuffer const& Buffer,
 {
 	mtlpp::BlitCommandEncoder TargetEncoder;
 	METAL_DEBUG_ONLY(FMetalBlitCommandEncoderDebugging Debugging);
-	bool bAsync = !CurrentEncoder.HasTextureBindingHistory(toTexture);
-	if(bAsync)
-	{
-		ConditionalSwitchToAsyncBlit();
-		TargetEncoder = PrologueEncoder.GetBlitCommandEncoder();
-		METAL_GPUPROFILE(FMetalProfiler::GetProfiler()->EncodeBlit(PrologueEncoder.GetCommandBufferStats(), __FUNCTION__));
-		METAL_DEBUG_LAYER(EMetalDebugLevelFastValidation, Debugging = PrologueEncoder.GetBlitCommandEncoderDebugging());
-	}
-	else
 	{
 		ConditionalSwitchToBlit();
 		TargetEncoder = CurrentEncoder.GetBlitCommandEncoder();
@@ -846,22 +709,13 @@ bool FMetalRenderPass::AsyncCopyFromBufferToTexture(FMetalBuffer const& Buffer,
 		METAL_DEBUG_LAYER(EMetalDebugLevelFastValidation, Debugging.Copy(Buffer, sourceOffset, sourceBytesPerRow, sourceBytesPerImage, sourceSize, toTexture, destinationSlice, destinationLevel, destinationOrigin, options));
 	}
 	
-	return bAsync;
+	return false;
 }
 
 bool FMetalRenderPass::AsyncCopyFromTextureToTexture(FMetalTexture const& Texture, uint32 sourceSlice, uint32 sourceLevel, mtlpp::Origin sourceOrigin, mtlpp::Size sourceSize, FMetalTexture const& toTexture, uint32 destinationSlice, uint32 destinationLevel, mtlpp::Origin destinationOrigin)
 {
 	mtlpp::BlitCommandEncoder TargetEncoder;
 	METAL_DEBUG_ONLY(FMetalBlitCommandEncoderDebugging Debugging);
-	bool bAsync = !CurrentEncoder.HasTextureBindingHistory(toTexture);
-	if(bAsync)
-	{
-		ConditionalSwitchToAsyncBlit();
-		TargetEncoder = PrologueEncoder.GetBlitCommandEncoder();
-		METAL_GPUPROFILE(FMetalProfiler::GetProfiler()->EncodeBlit(PrologueEncoder.GetCommandBufferStats(), __FUNCTION__));
-		METAL_DEBUG_LAYER(EMetalDebugLevelFastValidation, Debugging = PrologueEncoder.GetBlitCommandEncoderDebugging());
-	}
-	else
 	{
 		ConditionalSwitchToBlit();
 		TargetEncoder = CurrentEncoder.GetBlitCommandEncoder();
@@ -874,27 +728,18 @@ bool FMetalRenderPass::AsyncCopyFromTextureToTexture(FMetalTexture const& Textur
 	MTLPP_VALIDATE(mtlpp::BlitCommandEncoder, TargetEncoder, SafeGetRuntimeDebuggingLevel() >= EMetalDebugLevelValidation, Copy(Texture, sourceSlice, sourceLevel, sourceOrigin, sourceSize, toTexture, destinationSlice, destinationLevel, destinationOrigin));
 	METAL_DEBUG_LAYER(EMetalDebugLevelFastValidation, Debugging.Copy(Texture, sourceSlice, sourceLevel, sourceOrigin, sourceSize, toTexture, destinationSlice, destinationLevel, destinationOrigin));
 	
-	return bAsync;
+	return false;
 }
 
 bool FMetalRenderPass::CanAsyncCopyToBuffer(FMetalBuffer const& DestinationBuffer)
 {
-	return !CurrentEncoder.HasBufferBindingHistory(DestinationBuffer);
+	return false;
 }
 
 void FMetalRenderPass::AsyncCopyFromBufferToBuffer(FMetalBuffer const& SourceBuffer, NSUInteger SourceOffset, FMetalBuffer const& DestinationBuffer, NSUInteger DestinationOffset, NSUInteger Size)
 {
 	mtlpp::BlitCommandEncoder TargetEncoder;
 	METAL_DEBUG_ONLY(FMetalBlitCommandEncoderDebugging Debugging);
-	bool bAsync = !CurrentEncoder.HasBufferBindingHistory(DestinationBuffer);
-	if(bAsync)
-	{
-		ConditionalSwitchToAsyncBlit();
-		TargetEncoder = PrologueEncoder.GetBlitCommandEncoder();
-		METAL_GPUPROFILE(FMetalProfiler::GetProfiler()->EncodeBlit(PrologueEncoder.GetCommandBufferStats(), FString::Printf(TEXT("AsyncCopyFromBufferToBuffer: %p %llu %llu"), DestinationBuffer.GetPtr(), DestinationBuffer.GetOffset() + DestinationOffset, Size)));
-		METAL_DEBUG_LAYER(EMetalDebugLevelFastValidation, Debugging = PrologueEncoder.GetBlitCommandEncoderDebugging());
-	}
-	else
 	{
 		ConditionalSwitchToBlit();
 		TargetEncoder = CurrentEncoder.GetBlitCommandEncoder();
@@ -911,12 +756,6 @@ void FMetalRenderPass::AsyncCopyFromBufferToBuffer(FMetalBuffer const& SourceBuf
 FMetalBuffer FMetalRenderPass::AllocateTemporyBufferForCopy(FMetalBuffer const& DestinationBuffer, NSUInteger Size, NSUInteger Align)
 {
 	FMetalBuffer Buffer;
-	bool bAsync = !CurrentEncoder.HasBufferBindingHistory(DestinationBuffer);
-	if(bAsync)
-	{
-		Buffer = PrologueEncoder.GetRingBuffer().NewBuffer(Size, Align);
-	}
-	else
 	{
 		Buffer = CurrentEncoder.GetRingBuffer().NewBuffer(Size, Align);
 	}
@@ -926,39 +765,23 @@ FMetalBuffer FMetalRenderPass::AllocateTemporyBufferForCopy(FMetalBuffer const&
 void FMetalRenderPass::AsyncGenerateMipmapsForTexture(FMetalTexture const& Texture)
 {
 	// This must be a plain old error
-	check(!CurrentEncoder.HasTextureBindingHistory(Texture));
-	ConditionalSwitchToAsyncBlit();
-	mtlpp::BlitCommandEncoder Encoder = PrologueEncoder.GetBlitCommandEncoder();
+	ConditionalSwitchToBlit();
+	mtlpp::BlitCommandEncoder Encoder = CurrentEncoder.GetBlitCommandEncoder();
 	check(Encoder.GetPtr());
 	
 	METAL_GPUPROFILE(FMetalProfiler::GetProfiler()->EncodeBlit(CurrentEncoder.GetCommandBufferStats(), __FUNCTION__));
 	MTLPP_VALIDATE(mtlpp::BlitCommandEncoder, Encoder, SafeGetRuntimeDebuggingLevel() >= EMetalDebugLevelValidation, GenerateMipmaps(Texture));
-	METAL_DEBUG_LAYER(EMetalDebugLevelFastValidation, PrologueEncoder.GetBlitCommandEncoderDebugging().GenerateMipmaps(Texture));
+	METAL_DEBUG_LAYER(EMetalDebugLevelFastValidation, CurrentEncoder.GetBlitCommandEncoderDebugging().GenerateMipmaps(Texture));
 }
 
 TRefCountPtr<FMetalFence> const& FMetalRenderPass::End(void)
 {
-	// EndEncoding should provide the encoder fence...
-	if (PrologueEncoder.IsBlitCommandEncoderActive() || PrologueEncoder.IsComputeCommandEncoderActive())
-	{
-		PrologueEncoderFence = PrologueEncoder.EndEncoding();
-	}
-	
 	if (CurrentEncoder.IsRenderCommandEncoderActive() || CurrentEncoder.IsBlitCommandEncoderActive() || CurrentEncoder.IsComputeCommandEncoderActive())
 	{
 		State.FlushVisibilityResults(CurrentEncoder);
 		check(!CurrentEncoderFence);
-		check(!PassStartFence);
 		CurrentEncoderFence = CurrentEncoder.EndEncoding();
 	}
-	else if (PassStartFence)
-	{
-		ConditionalSwitchToBlit();
-		CurrentEncoderFence = CurrentEncoder.EndEncoding();
-		PassStartFence = nullptr;
-	}
-	
-	check(!PassStartFence);
 	
 	State.SetRenderTargetsActive(false);
 	
@@ -980,20 +803,11 @@ void FMetalRenderPass::AddCompletionHandler(mtlpp::CommandBufferHandler Handler)
 
 void FMetalRenderPass::AddAsyncCommandBufferHandlers(mtlpp::CommandBufferHandler Scheduled, mtlpp::CommandBufferHandler Completion)
 {
-	check(PrologueEncoder.GetCommandBuffer() && PrologueEncoder.IsBlitCommandEncoderActive());
-	if (Scheduled)
-	{
-		PrologueEncoder.GetCommandBuffer().AddScheduledHandler(Scheduled);
-	}
-	if (Completion)
-	{
-		PrologueEncoder.AddCompletionHandler(Completion);
-	}
+	checkf(false, TEXT("Async command buffer has not been supported yet"));
 }
 
 void FMetalRenderPass::TransitionResources(mtlpp::Resource const& Resource)
 {
-	PrologueEncoder.TransitionResources(Resource);
 	CurrentEncoder.TransitionResources(Resource);
 }
 
@@ -1023,19 +837,16 @@ void FMetalRenderPass::InsertDebugEncoder()
 void FMetalRenderPass::InsertDebugSignpost(ns::String const& String)
 {
 	CurrentEncoder.InsertDebugSignpost(String);
-	PrologueEncoder.InsertDebugSignpost(FString::Printf(TEXT("Prologue %s"), *FString(String.GetPtr())).GetNSString());
 }
 
 void FMetalRenderPass::PushDebugGroup(ns::String const& String)
 {
 	CurrentEncoder.PushDebugGroup(String);
-	PrologueEncoder.PushDebugGroup(FString::Printf(TEXT("Prologue %s"), *FString(String.GetPtr())).GetNSString());
 }
 
 void FMetalRenderPass::PopDebugGroup(void)
 {
 	CurrentEncoder.PopDebugGroup();
-	PrologueEncoder.PopDebugGroup();
 }
 
 #pragma mark - Public Accessors -
@@ -1057,7 +868,6 @@ FMetalSubBufferRing& FMetalRenderPass::GetRingBuffer(void)
 
 void FMetalRenderPass::ShrinkRingBuffers(void)
 {
-	PrologueEncoder.GetRingBuffer().Shrink();
 	CurrentEncoder.GetRingBuffer().Shrink();
 }
 
@@ -1103,28 +913,14 @@ void FMetalRenderPass::ConditionalSwitchToCompute(void)
 	{
 		State.SetStateDirty();
 		CurrentEncoder.BeginComputeCommandEncoding(ComputeDispatchType);
-		if (PassStartFence)
-		{
-			CurrentEncoder.WaitForFence(PassStartFence);
-			PassStartFence = nullptr;
-		}
 		if (CurrentEncoderFence)
 		{
 			CurrentEncoder.WaitForFence(CurrentEncoderFence);
 			CurrentEncoderFence = nullptr;
 		}
-		if (PrologueEncoderFence)
-		{
-			CurrentEncoder.WaitForFence(PrologueEncoderFence);
-		}
 	}
 	
 	check(CurrentEncoder.IsComputeCommandEncoderActive());
-	
-	if (PrologueEncoder.IsBlitCommandEncoderActive() || PrologueEncoder.IsComputeCommandEncoderActive())
-	{
-		CurrentEncoder.WaitForFence(PrologueEncoder.GetEncoderFence());
-	}
 }
 
 void FMetalRenderPass::ConditionalSwitchToBlit(void)
@@ -1147,157 +943,14 @@ void FMetalRenderPass::ConditionalSwitchToBlit(void)
 	if (!CurrentEncoder.IsBlitCommandEncoderActive())
 	{
 		CurrentEncoder.BeginBlitCommandEncoding();
-		if (PassStartFence)
-		{
-			CurrentEncoder.WaitForFence(PassStartFence);
-			PassStartFence = nullptr;
-		}
 		if (CurrentEncoderFence)
 		{
 			CurrentEncoder.WaitForFence(CurrentEncoderFence);
 			CurrentEncoderFence = nullptr;
 		}
-		if (PrologueEncoderFence)
-		{
-			CurrentEncoder.WaitForFence(PrologueEncoderFence);
-		}
 	}
 	
 	check(CurrentEncoder.IsBlitCommandEncoderActive());
-	
-	if (PrologueEncoder.IsBlitCommandEncoderActive() || PrologueEncoder.IsComputeCommandEncoderActive())
-	{
-		CurrentEncoder.WaitForFence(PrologueEncoder.GetEncoderFence());
-	}
-}
-
-void FMetalRenderPass::ConditionalSwitchToAsyncBlit(void)
-{
-	SCOPE_CYCLE_COUNTER(STAT_MetalSwitchToAsyncBlitTime);
-	
-	if (PrologueEncoder.IsComputeCommandEncoderActive() || PrologueEncoder.IsRenderCommandEncoderActive())
-	{
-		PrologueEncoderFence = PrologueEncoder.EndEncoding();
-	}
-	
-	if (!PrologueEncoder.IsBlitCommandEncoderActive())
-	{
-		if (!PrologueEncoder.GetCommandBuffer())
-		{
-			PrologueEncoder.StartCommandBuffer();
-		}
-		PrologueEncoder.BeginBlitCommandEncoding();
-		if (PrologueStartEncoderFence)
-		{
-			if (PrologueStartEncoderFence->NeedsWait(mtlpp::RenderStages::Vertex))
-			{
-				PrologueEncoder.WaitForFence(PrologueStartEncoderFence);
-			}
-			else
-			{
-				PrologueEncoder.WaitAndUpdateFence(PrologueStartEncoderFence);
-			}
-			PrologueStartEncoderFence = nullptr;
-		}
-		if (PrologueEncoderFence)
-		{
-			if (PrologueEncoderFence->NeedsWait(mtlpp::RenderStages::Vertex))
-			{
-				PrologueEncoder.WaitForFence(PrologueEncoderFence);
-			}
-			else
-			{
-				PrologueEncoder.WaitAndUpdateFence(PrologueEncoderFence);
-			}
-			PrologueEncoderFence = nullptr;
-		}
-#if METAL_DEBUG_OPTIONS
-//		if (GetEmitDrawEvents() && PrologueEncoder.GetEncoderFence())
-//		{
-//			for (uint32 i = mtlpp::RenderStages::Vertex; i <= mtlpp::RenderStages::Fragment && PrologueEncoder.GetEncoderFence()->Get((mtlpp::RenderStages)i).GetPtr(); i++)
-//			{
-//				if (CmdList.GetCommandQueue().GetRuntimeDebuggingLevel() >= EMetalDebugLevelValidation)
-//				{
-//					PrologueEncoder.GetEncoderFence()->Get((mtlpp::RenderStages)i).GetPtr().label = [NSString stringWithFormat:@"Prologue %@", PrologueEncoderFence->Get((mtlpp::RenderStages)i).GetLabel().GetPtr()];
-//				}
-//				else
-//				{
-//					PrologueEncoder.GetEncoderFence()->Get((mtlpp::RenderStages)i).SetLabel([NSString stringWithFormat:@"Prologue %@", PrologueEncoderFence->Get((mtlpp::RenderStages)i).GetLabel().GetPtr()]);
-//				}
-//			}
-//		}
-#endif
-
-		if (CurrentEncoder.IsRenderCommandEncoderActive() || CurrentEncoder.IsComputeCommandEncoderActive() || CurrentEncoder.IsBlitCommandEncoderActive())
-		{
-			CurrentEncoder.WaitForFence(PrologueEncoder.GetEncoderFence());
-		}
-	}
-	
-	check(PrologueEncoder.IsBlitCommandEncoderActive());
-}
-
-void FMetalRenderPass::ConditionalSwitchToAsyncCompute(void)
-{
-	SCOPE_CYCLE_COUNTER(STAT_MetalSwitchToComputeTime);
-	
-	if (PrologueEncoder.IsBlitCommandEncoderActive() || PrologueEncoder.IsRenderCommandEncoderActive())
-	{
-		PrologueEncoderFence = PrologueEncoder.EndEncoding();
-	}
-	
-	if (!PrologueEncoder.IsComputeCommandEncoderActive())
-	{
-		if (!PrologueEncoder.GetCommandBuffer())
-		{
-			PrologueEncoder.StartCommandBuffer();
-		}
-		State.SetStateDirty();
-		PrologueEncoder.BeginComputeCommandEncoding(ComputeDispatchType);
-		
-		if (PrologueStartEncoderFence)
-		{
-			if (PrologueStartEncoderFence->NeedsWait(mtlpp::RenderStages::Vertex))
-			{
-				PrologueEncoder.WaitForFence(PrologueStartEncoderFence);
-			}
-			else
-			{
-				PrologueEncoder.WaitAndUpdateFence(PrologueStartEncoderFence);
-			}
-			PrologueStartEncoderFence = nullptr;
-		}
-		if (PrologueEncoderFence)
-		{
-			if (PrologueEncoderFence->NeedsWait(mtlpp::RenderStages::Vertex))
-			{
-				PrologueEncoder.WaitForFence(PrologueEncoderFence);
-			}
-			else
-			{
-				PrologueEncoder.WaitAndUpdateFence(PrologueEncoderFence);
-			}
-			PrologueEncoderFence = nullptr;
-		}
-		if (PassStartFence)
-		{
-			if(PassStartFence->NeedsWait(mtlpp::RenderStages::Vertex))
-			{
-				PrologueEncoder.WaitForFence(PassStartFence);
-			}
-			else
-			{
-				PrologueEncoder.WaitAndUpdateFence(PassStartFence);
-			}
-			PassStartFence = nullptr;
-		}
-		if (CurrentEncoder.IsRenderCommandEncoderActive() || CurrentEncoder.IsComputeCommandEncoderActive() || CurrentEncoder.IsBlitCommandEncoderActive())
-		{
-			CurrentEncoder.WaitForFence(PrologueEncoder.GetEncoderFence());
-		}
-	}
-	
-	check(PrologueEncoder.IsComputeCommandEncoderActive());
 }
 
 void FMetalRenderPass::CommitRenderResourceTables(void)
@@ -1355,34 +1008,6 @@ void FMetalRenderPass::CommitDispatchResourceTables(void)
 #endif //METAL_RHI_RAYTRACING
 }
 
-void FMetalRenderPass::CommitAsyncDispatchResourceTables(void)
-{
-	State.CommitComputeResources(&PrologueEncoder);
-	
-	State.CommitResourceTable(EMetalShaderStages::Compute, mtlpp::FunctionType::Kernel, PrologueEncoder);
-	
-	FMetalComputeShader const* ComputeShader = State.GetComputeShader();
-	if (ComputeShader->SideTableBinding >= 0)
-	{
-		PrologueEncoder.SetShaderSideTable(mtlpp::FunctionType::Kernel, ComputeShader->SideTableBinding);
-		State.SetShaderBuffer(EMetalShaderStages::Compute, nil, nil, 0, 0, ComputeShader->SideTableBinding, mtlpp::ResourceUsage(0));
-	}
-
-#if METAL_RHI_RAYTRACING
-	// TODO: Crappy workaround for inline raytracing support.
-	if (ComputeShader->RayTracingBindings.InstanceIndexBuffer != UINT32_MAX && InstanceBufferSRV.IsValid())
-	{
-		FMetalRHIBuffer* SourceBuffer = ResourceCast(InstanceBufferSRV->GetBuffer());
-		check(SourceBuffer);
-		FMetalBuffer CurBuffer = SourceBuffer->GetCurrentBufferOrNil();
-		check(CurBuffer);
-
-		PrologueEncoder.SetShaderBuffer(mtlpp::FunctionType::Kernel, CurBuffer, InstanceBufferSRV->Offset, CurBuffer.GetLength(), ComputeShader->RayTracingBindings.InstanceIndexBuffer, mtlpp::ResourceUsage::Read);
-		State.SetShaderBuffer(EMetalShaderStages::Compute, CurBuffer, nil, InstanceBufferSRV->Offset, CurBuffer.GetLength(), ComputeShader->RayTracingBindings.InstanceIndexBuffer, mtlpp::ResourceUsage::Read);
-	}
-#endif //METAL_RHI_RAYTRACING
-}
-
 void FMetalRenderPass::PrepareToRender(uint32 PrimitiveType)
 {
 	SCOPE_CYCLE_COUNTER(STAT_MetalPrepareToRenderTime);
@@ -1391,12 +1016,12 @@ void FMetalRenderPass::PrepareToRender(uint32 PrimitiveType)
 	check(CurrentEncoder.IsRenderCommandEncoderActive());
 	
 	// Set raster state
-	State.SetRenderState(CurrentEncoder, nullptr);
+	State.SetRenderState(CurrentEncoder);
 	
 	// Bind shader resources
 	CommitRenderResourceTables();
     
-    State.SetRenderPipelineState(CurrentEncoder, nullptr);
+    State.SetRenderPipelineState(CurrentEncoder);
 }
 
 void FMetalRenderPass::PrepareToDispatch(void)
@@ -1412,19 +1037,6 @@ void FMetalRenderPass::PrepareToDispatch(void)
     State.SetComputePipelineState(CurrentEncoder);
 }
 
-void FMetalRenderPass::PrepareToAsyncDispatch(void)
-{
-	SCOPE_CYCLE_COUNTER(STAT_MetalPrepareToDispatchTime);
-	
-	check(PrologueEncoder.GetCommandBuffer());
-	check(PrologueEncoder.IsComputeCommandEncoderActive());
-	
-	// Bind shader resources
-	CommitAsyncDispatchResourceTables();
-	
-	State.SetComputePipelineState(PrologueEncoder);
-}
-
 void FMetalRenderPass::ConditionalSubmit()
 {
 	NumOutstandingOps++;
@@ -1506,7 +1118,7 @@ void FMetalRenderPass::ConditionalSubmit()
 
 uint32 FMetalRenderPass::GetEncoderIndex(void) const
 {
-	return PrologueEncoder.NumEncodedPasses() + CurrentEncoder.NumEncodedPasses();
+	return CurrentEncoder.NumEncodedPasses();
 }
 
 uint32 FMetalRenderPass::GetCommandBufferIndex(void) const
diff --git a/Engine/Source/Runtime/Apple/MetalRHI/Private/MetalRenderPass.h b/Engine/Source/Runtime/Apple/MetalRHI/Private/MetalRenderPass.h
--- a/Engine/Source/Runtime/Apple/MetalRHI/Private/MetalRenderPass.h
+++ b/Engine/Source/Runtime/Apple/MetalRHI/Private/MetalRenderPass.h
@@ -25,8 +25,6 @@ public:
 	void SetDispatchType(mtlpp::DispatchType Type);
 	
     void Begin();
-	
-	void Wait(FMetalFence* Fence);
 
     void BeginRenderPass(mtlpp::RenderPassDescriptor RenderPass);
 
@@ -150,16 +148,12 @@ private:
     void ConditionalSwitchToRender(void);
     void ConditionalSwitchToCompute(void);
 	void ConditionalSwitchToBlit(void);
-	void ConditionalSwitchToAsyncBlit(void);
-	void ConditionalSwitchToAsyncCompute(void);
 	
     void PrepareToRender(uint32 PrimType);
     void PrepareToDispatch(void);
-	void PrepareToAsyncDispatch(void);
 
     void CommitRenderResourceTables(void);
     void CommitDispatchResourceTables(void);
-	void CommitAsyncDispatchResourceTables(void);
     
     void ConditionalSubmit();
 	
@@ -180,19 +174,13 @@ private:
     uint16 BoundSamplers[EMetalShaderStages::Num];
     
     FMetalCommandEncoder CurrentEncoder;
-    FMetalCommandEncoder PrologueEncoder;
 	
 	// To ensure that buffer uploads aren't overwritten before they are used track what is in flight
 	// Disjoint ranges *are* permitted!
 	TMap<id<MTLBuffer>, TArray<NSRange>> OutstandingBufferUploads;
 
 	// Fences for the current command encoder chain
-	TRefCountPtr<FMetalFence> PassStartFence;
 	TRefCountPtr<FMetalFence> CurrentEncoderFence;
-
-	// Fences for the prologue command encoder chain
-	TRefCountPtr<FMetalFence> PrologueStartEncoderFence;
-	TRefCountPtr<FMetalFence> PrologueEncoderFence;
     
     mtlpp::RenderPassDescriptor RenderPassDesc;
     
diff --git a/Engine/Source/Runtime/Apple/MetalRHI/Private/MetalStateCache.cpp b/Engine/Source/Runtime/Apple/MetalRHI/Private/MetalStateCache.cpp
--- a/Engine/Source/Runtime/Apple/MetalRHI/Private/MetalStateCache.cpp
+++ b/Engine/Source/Runtime/Apple/MetalRHI/Private/MetalStateCache.cpp
@@ -1791,7 +1791,7 @@ void FMetalStateCache::FlushVisibilityResults(FMetalCommandEncoder& CommandEncod
 #endif
 }
 
-void FMetalStateCache::SetRenderState(FMetalCommandEncoder& CommandEncoder, FMetalCommandEncoder* PrologueEncoder)
+void FMetalStateCache::SetRenderState(FMetalCommandEncoder& CommandEncoder)
 {
 	SCOPE_CYCLE_COUNTER(STAT_MetalSetRenderStateTime);
 	
@@ -1878,7 +1878,7 @@ void FMetalStateCache::EnsureTextureAndType(EMetalShaderStages Stage, uint32 Ind
 #endif
 }
 
-void FMetalStateCache::SetRenderPipelineState(FMetalCommandEncoder& CommandEncoder, FMetalCommandEncoder* PrologueEncoder)
+void FMetalStateCache::SetRenderPipelineState(FMetalCommandEncoder& CommandEncoder)
 {
 	SCOPE_CYCLE_COUNTER(STAT_MetalSetRenderPipelineStateTime);
 	
@@ -1889,11 +1889,6 @@ void FMetalStateCache::SetRenderPipelineState(FMetalCommandEncoder& CommandEncod
 
 		check(Pipeline);
         CommandEncoder.SetRenderPipelineState(Pipeline);
-        if (Pipeline->ComputePipelineState)
-        {
-            check(PrologueEncoder);
-            PrologueEncoder->SetComputePipelineState(Pipeline);
-        }
         
         PipelineBits &= EMetalPipelineFlagComputeMask;
     }
diff --git a/Engine/Source/Runtime/Apple/MetalRHI/Private/MetalStateCache.h b/Engine/Source/Runtime/Apple/MetalRHI/Private/MetalStateCache.h
--- a/Engine/Source/Runtime/Apple/MetalRHI/Private/MetalStateCache.h
+++ b/Engine/Source/Runtime/Apple/MetalRHI/Private/MetalStateCache.h
@@ -129,7 +129,7 @@ public:
 	
 	void SetRenderStoreActions(FMetalCommandEncoder& CommandEncoder, bool const bConditionalSwitch);
 	
-	void SetRenderState(FMetalCommandEncoder& CommandEncoder, FMetalCommandEncoder* PrologueEncoder);
+	void SetRenderState(FMetalCommandEncoder& CommandEncoder);
 
 	void CommitRenderResources(FMetalCommandEncoder* Raster);
 
@@ -170,7 +170,7 @@ public:
 	FTexture2DRHIRef CreateFallbackDepthStencilSurface(uint32 Width, uint32 Height);
 	bool GetFallbackDepthStencilBound(void) const { return bFallbackDepthStencilBound; }
 	
-	void SetRenderPipelineState(FMetalCommandEncoder& CommandEncoder, FMetalCommandEncoder* PrologueEncoder);
+	void SetRenderPipelineState(FMetalCommandEncoder& CommandEncoder);
     void SetComputePipelineState(FMetalCommandEncoder& CommandEncoder);
 	void FlushVisibilityResults(FMetalCommandEncoder& CommandEncoder);
 
diff --git a/Engine/Source/Runtime/Apple/MetalRHI/Private/MetalTransitionData.cpp b/Engine/Source/Runtime/Apple/MetalRHI/Private/MetalTransitionData.cpp
--- a/Engine/Source/Runtime/Apple/MetalRHI/Private/MetalTransitionData.cpp
+++ b/Engine/Source/Runtime/Apple/MetalRHI/Private/MetalTransitionData.cpp
@@ -12,9 +12,6 @@
 
 #pragma mark - Metal RHI Resource Transition Data Definitions -
 
-#define UE_METAL_TRANSITION_DATA_USES_EXPLICIT_FENCING					0
-
-
 //------------------------------------------------------------------------------
 
 #pragma mark - Metal RHI Resource Transition Data Implementation -
@@ -32,52 +29,16 @@ FMetalTransitionData::FMetalTransitionData(ERHIPipeline
 	bCrossPipeline = (SrcPipelines != DstPipelines);
 
 	Infos.Append(InInfos.GetData(), InInfos.Num());
-
-	// TODO: Determine whether the Metal RHI needs to create a separate, per-transition fence.
-#if UE_METAL_TRANSITION_DATA_USES_EXPLICIT_FENCING
-	if (bCrossPipeline && !EnumHasAnyFlags(CreateFlags, ERHITransitionCreateFlags::NoFence))
-	{
-		// Get the current context pointer.
-		FMetalContext* Context = GetMetalDeviceContext().GetCurrentContext();
-
-		// Get the current render pass fence.
-		TRefCountPtr<FMetalFence> const& MetalFence = Context->GetCurrentRenderPass().End();
-
-		// Write it again as we may wait on this fence in two different encoders.
-		Context->GetCurrentRenderPass().Update(MetalFence);
-
-		// Write it into the transition data.
-		Fence = MetalFence;
-		Fence->AddRef();
-	}
-#endif // UE_METAL_TRANSITION_DATA_USES_EXPLICIT_FENCING
 }
 
 void FMetalTransitionData::BeginResourceTransitions() const
 {
-	// TODO: Determine whether the Metal RHI needs to do anything here.
-#if UE_METAL_TRANSITION_DATA_USES_EXPLICIT_FENCING
-	if (Fence.IsValid())
-	{
-		FMetalContext* Context = GetMetalDeviceContext().GetCurrentContext();
-
-		if (Context->GetCurrentCommandBuffer())
-		{
-			Context->SubmitCommandsHint(EMetalSubmitFlagsNone);
-		}
-
-		Context->GetCurrentRenderPass().Begin(Fence);
-	}
-#endif // UE_METAL_TRANSITION_DATA_USES_EXPLICIT_FENCING
 }
 
 void FMetalTransitionData::EndResourceTransitions() const
 {
-	// No action necessary for same pipe transitions
-	if (SrcPipelines == DstPipelines)
-	{
-		return;
-	}
+	//@ToDo: Support AysncCompute with MTLEvent
+	check(SrcPipelines == DstPipelines);
 
 	for (const auto& Info : Infos)
 	{
@@ -86,6 +47,14 @@ void FMetalTransitionData::EndResourceTransitions() const
 			continue;
 		}
 
+		if (Info.AccessAfter == ERHIAccess::Discard)
+		{
+			// Discard as a destination is a no-op
+			continue;
+		}
+
+		checkf(Info.AccessAfter != ERHIAccess::Unknown, TEXT("Transitioning a resource to an unknown state is not allowed."));
+
 		switch (Info.Type)
 		{
 			case FRHITransitionInfo::EType::UAV:
