Fixed build break.

diff --git a/Engine/Plugins/Runtime/nDisplay/Source/DisplayCluster/Private/Game/EngineClasses/Scene/DisplayClusterWidgetComponent.cpp b/Engine/Plugins/Runtime/nDisplay/Source/DisplayCluster/Private/Game/EngineClasses/Scene/DisplayClusterWidgetComponent.cpp
--- a/Engine/Plugins/Runtime/nDisplay/Source/DisplayCluster/Private/Game/EngineClasses/Scene/DisplayClusterWidgetComponent.cpp
+++ b/Engine/Plugins/Runtime/nDisplay/Source/DisplayCluster/Private/Game/EngineClasses/Scene/DisplayClusterWidgetComponent.cpp
@@ -185,7 +185,7 @@ public:
 				this->StaticMeshVertexBuffers.StaticMeshVertexBuffer.BindPackedTexCoordVertexBuffer(VertexFactoryPtr, Data);
 				this->StaticMeshVertexBuffers.StaticMeshVertexBuffer.BindLightMapVertexBuffer(VertexFactoryPtr, Data, 0);
 				this->StaticMeshVertexBuffers.ColorVertexBuffer.BindColorVertexBuffer(VertexFactoryPtr, Data);
-				VertexFactoryPtr->SetData(Data);
+				VertexFactoryPtr->SetData(RHICmdList, Data);
 
 				VertexFactoryPtr->UpdateRHI(RHICmdList);
 			});
diff --git a/Engine/Source/Runtime/Apple/MetalRHI/Private/MetalBuffer.cpp b/Engine/Source/Runtime/Apple/MetalRHI/Private/MetalBuffer.cpp
--- a/Engine/Source/Runtime/Apple/MetalRHI/Private/MetalBuffer.cpp
+++ b/Engine/Source/Runtime/Apple/MetalRHI/Private/MetalBuffer.cpp
@@ -1193,7 +1193,7 @@ uint32 FMetalBufferPoolPolicyData::GetPoolBucketSize(uint32 Bucket)
 	return BucketSizes[Index];
 }
 
-FMetalBuffer FMetalBufferPoolPolicyData::CreateResource(CreationArguments Args)
+FMetalBuffer FMetalBufferPoolPolicyData::CreateResource(FRHICommandListBase&, CreationArguments Args)
 {
 	check(Args.Device);	
 	uint32 BufferSize = GetPoolBucketSize(GetPoolBucketIndex(Args));
@@ -1522,7 +1522,7 @@ FMetalBuffer FMetalResourceHeap::CreateBuffer(uint32 Size, uint32 Alignment, EBu
 				 }
 				 else
 				 {
-                    Buffer = ManagedBuffers.CreatePooledResource(FMetalPooledBufferArgs(Queue->GetDevice(), BlockSize, Flags, StorageMode, CpuMode));
+                    Buffer = ManagedBuffers.CreatePooledResource(FRHICommandListExecutor::GetImmediateCommandList(), FMetalPooledBufferArgs(Queue->GetDevice(), BlockSize, Flags, StorageMode, CpuMode));
 					if (GMetalResourcePurgeInPool)
 					{
 						Buffer.SetPurgeableState(mtlpp::PurgeableState::NonVolatile);
@@ -1599,7 +1599,7 @@ FMetalBuffer FMetalResourceHeap::CreateBuffer(uint32 Size, uint32 Alignment, EBu
 				else
 				{
 					FScopeLock Lock(&Mutex);
-                    Buffer = Buffers[Storage].CreatePooledResource(FMetalPooledBufferArgs(Queue->GetDevice(), BlockSize, Flags, StorageMode, CpuMode));
+                    Buffer = Buffers[Storage].CreatePooledResource(FRHICommandListExecutor::GetImmediateCommandList(), FMetalPooledBufferArgs(Queue->GetDevice(), BlockSize, Flags, StorageMode, CpuMode));
 					if (GMetalResourcePurgeInPool)
 					{
                    		Buffer.SetPurgeableState(mtlpp::PurgeableState::NonVolatile);
diff --git a/Engine/Source/Runtime/Apple/MetalRHI/Private/MetalBuffer.h b/Engine/Source/Runtime/Apple/MetalRHI/Private/MetalBuffer.h
--- a/Engine/Source/Runtime/Apple/MetalRHI/Private/MetalBuffer.h
+++ b/Engine/Source/Runtime/Apple/MetalRHI/Private/MetalBuffer.h
@@ -241,7 +241,7 @@ public:
 	 * @param Args The buffer size in bytes.
 	 * @returns A suitably sized buffer or NULL on failure.
 	 */
-	FMetalBuffer CreateResource(CreationArguments Args);
+	FMetalBuffer CreateResource(FRHICommandListBase& RHICmdList, CreationArguments Args);
 	
 	/** Gets the arguments used to create resource
 	 * @param Resource The buffer to get data for.
