Fix a typo in MetalRHI.

diff --git a/Engine/Source/Runtime/Apple/MetalRHI/Private/MetalContext.cpp b/Engine/Source/Runtime/Apple/MetalRHI/Private/MetalContext.cpp
--- a/Engine/Source/Runtime/Apple/MetalRHI/Private/MetalContext.cpp
+++ b/Engine/Source/Runtime/Apple/MetalRHI/Private/MetalContext.cpp
@@ -75,7 +75,7 @@ static FAutoConsoleVariableRef CVarMetalResourcePurgeOnDelete(
 static int32 GMetalResourceDeferDeleteNumFrames = 0;
 static FAutoConsoleVariableRef CVarMetalResourceDeferDeleteNumFrames(
 	TEXT("rhi.Metal.ResourceDeferDeleteNumFrames"),
-	GMetalResourcePurgeOnDelete,
+	GMetalResourceDeferDeleteNumFrames,
 	TEXT("Debug option: set to the number of frames that must have passed before resource free-lists are processed and resources disposed of. (Default: 0, Off)"));
 #endif
 
