Fixes for Nanite on Mac

- Disable COMPILER_SUPPORTS_WAVE_VOTE for Mac
- Fixed bug where Texture2DArray was flattened due to atomics, but bound as Texture2D
- Fixed bug where defines were read incorrectly in shader compiler

diff --git a/Engine/Config/Mac/MacEngine.ini b/Engine/Config/Mac/MacEngine.ini
--- a/Engine/Config/Mac/MacEngine.ini
+++ b/Engine/Config/Mac/MacEngine.ini
@@ -15,10 +15,6 @@ fx.NiagaraAllowRuntimeScalabilityChanges=1
 ;Perplatform to PerQualityLevel conversion mapping for platform groups
 QualityLevelMapping="high"
 
-[/Script/Engine.RendererSettings]
-r.Nanite.AllowComputeMaterials=0
-r.Nanite.ComputeMaterials=0
-
 [TextureStreaming]
 ; PoolSizeVRAMPercentage is how much percentage of GPU Dedicated VRAM should be used as a TexturePool cache for streaming textures (0 - unlimited streaming)
 PoolSizeVRAMPercentage=70
diff --git a/Engine/Shaders/Public/Platform.ush b/Engine/Shaders/Public/Platform.ush
--- a/Engine/Shaders/Public/Platform.ush
+++ b/Engine/Shaders/Public/Platform.ush
@@ -507,10 +507,18 @@
 #endif
 
 #if PLATFORM_SUPPORTS_SM6_0_WAVE_OPERATIONS
-#define COMPILER_SUPPORTS_WAVE_ONCE			1
-#define COMPILER_SUPPORTS_WAVE_VOTE			1
-#define COMPILER_SUPPORTS_WAVE_MINMAX		1
-#define COMPILER_SUPPORTS_WAVE_BIT_ORAND	1
+#if !defined(COMPILER_SUPPORTS_WAVE_ONCE)
+	#define COMPILER_SUPPORTS_WAVE_ONCE		 1
+#endif
+#if !defined(COMPILER_SUPPORTS_WAVE_VOTE)
+    #define COMPILER_SUPPORTS_WAVE_VOTE		 1
+#endif
+#if !defined(COMPILER_SUPPORTS_WAVE_MINMAX)
+	#define COMPILER_SUPPORTS_WAVE_MINMAX	 1
+#endif
+#if !defined(COMPILER_SUPPORTS_WAVE_BIT_ORAND)
+	#define COMPILER_SUPPORTS_WAVE_BIT_ORAND 1
+#endif
 #endif
 
 // If compiler lane management in a wave.
diff --git a/Engine/Shaders/Public/Platform/Metal/MetalCommon.ush b/Engine/Shaders/Public/Platform/Metal/MetalCommon.ush
--- a/Engine/Shaders/Public/Platform/Metal/MetalCommon.ush
+++ b/Engine/Shaders/Public/Platform/Metal/MetalCommon.ush
@@ -5,7 +5,7 @@
 =============================================================================*/
 
 // Update this GUID to improve shader recompilation for Metal only shaders
-// GUID = A634A419-D6F0-460F-85E4-CB558AB61570
+// GUID = 7DDC23C5-DDBC-4A9F-9047-CC365A54D83C
 
 #pragma once
 
@@ -18,7 +18,15 @@
 #define STRONG_TYPE
 #define StrongTypedBuffer Buffer
 
-#define PLATFORM_SUPPORTS_SM6_0_WAVE_OPERATIONS 1
+#if METAL_SM6_PROFILE == 1
+    #define PLATFORM_SUPPORTS_SM6_0_WAVE_OPERATIONS 1
+    // Disable WaveVote due to bugs
+    #define COMPILER_SUPPORTS_WAVE_VOTE 0
+#endif
+
+#if IS_BASE_PASS && COMPUTE_SHADED && !defined(USE_FORCE_TEXTURE_MIP)
+	#define USE_FORCE_TEXTURE_MIP 1
+#endif
 
 /* Too many bugs to enable wave-intrinsics just yet, but we are getting there.
 #if COMPUTESHADER || PIXELSHADER
diff --git a/Engine/Source/Developer/Apple/MetalShaderFormat/Private/MetalDerivedData.cpp b/Engine/Source/Developer/Apple/MetalShaderFormat/Private/MetalDerivedData.cpp
--- a/Engine/Source/Developer/Apple/MetalShaderFormat/Private/MetalDerivedData.cpp
+++ b/Engine/Source/Developer/Apple/MetalShaderFormat/Private/MetalDerivedData.cpp
@@ -985,11 +985,11 @@ bool DoCompileMetalShader(
     		// Detect if we need to patch VSM shaders (flatten 2D array as regular 2D texture).
 			// Must be done as VSM uses 2DArray and requires atomics support. And Metal does not
 			// support atomics on 2Darray...
-			bool bShouldFlatten2DArray = Input.Environment.GetCompileArgument(TEXT("VIRTUAL_SHADOW_MAP"), 0)
-                                            || Input.Environment.GetCompileArgument(TEXT("VIRTUAL_TEXTURE_TARGET"), 0)
-                                            || Input.ShaderName.Find("PhysicalPage") != INDEX_NONE
-                                            || Input.ShaderName.Find("Virtual") != INDEX_NONE
-                                            || Input.ShaderName.Find("ClassifyMaterial") != INDEX_NONE;
+			bool bShouldFlatten2DArray = InPreprocessedShader.Find("VIRTUAL_SHADOW_MAP 1") > 0
+                                        || InPreprocessedShader.Find("VIRTUAL_TEXTURE_TARGET 1") > 0
+                                        || Input.ShaderName.Find("PhysicalPage") != INDEX_NONE
+                                        || Input.ShaderName.Find("Virtual") != INDEX_NONE
+                                        || Input.ShaderName.Find("ClassifyMaterial") != INDEX_NONE;
 
 			// Need to patch Clear/Memset CS too.
 			if (!bShouldFlatten2DArray)
diff --git a/Engine/Source/Runtime/Apple/MetalRHI/Private/MetalUAV.cpp b/Engine/Source/Runtime/Apple/MetalRHI/Private/MetalUAV.cpp
--- a/Engine/Source/Runtime/Apple/MetalRHI/Private/MetalUAV.cpp
+++ b/Engine/Source/Runtime/Apple/MetalRHI/Private/MetalUAV.cpp
@@ -377,13 +377,14 @@ void FMetalUnorderedAccessView::UpdateView()
         bool bUseSourceTexture = Info.bAllMips && Info.bAllSlices &&
                                 UAVDimensionToMetalTextureType(Info.Dimension) == TextureType && MetalFormat == Texture->Texture.GetPixelFormat();
         
+        bool bIsAtomicCompatible = EnumHasAllFlags(Texture->GetDesc().Flags, TexCreate_AtomicCompatible) ||
+                                            EnumHasAllFlags(Texture->GetDesc().Flags, ETextureCreateFlags::Atomic64Compatible);
+        
         // We can use the source texture directly if the view's format / mip count etc matches.
         if (bUseSourceTexture)
 		{
             // If we are using texture atomics then we need to bind them as buffers because Metal lacks texture atomics
-            if((EnumHasAllFlags(Texture->GetDesc().Flags, TexCreate_UAV | TexCreate_NoTiling) ||
-                EnumHasAllFlags(Texture->GetDesc().Flags, TexCreate_AtomicCompatible) ||
-                 EnumHasAllFlags(Texture->GetDesc().Flags, ETextureCreateFlags::Atomic64Compatible)) && Texture->Texture.GetBuffer())
+            if((EnumHasAllFlags(Texture->GetDesc().Flags, TexCreate_UAV | TexCreate_NoTiling) || bIsAtomicCompatible) && Texture->Texture.GetBuffer())
             {
                 FMetalBuffer MetalBuffer(Texture->Texture.GetBuffer(), false);
                 InitAsTextureBufferBacked(Texture->Texture, MetalBuffer,
@@ -410,6 +411,12 @@ void FMetalUnorderedAccessView::UpdateView()
                 ArraySize = Info.ArrayRange.Num * 6;
             }
             
+            // Metal doesn't support atomic Texture2DArray
+            if(bIsAtomicCompatible && Info.Dimension == FRHIViewDesc::EDimension::Texture2DArray)
+            {
+                TextureType = mtlpp::TextureType::Texture2D;
+            }
+            
 			FMetalTexture MetalTexture(Texture->Texture.NewTextureView(
 				MetalFormat,
                 TextureType,
@@ -418,9 +425,7 @@ void FMetalUnorderedAccessView::UpdateView()
 			);
             
             // If we are using texture atomics then we need to bind them as buffers because Metal lacks texture atomics
-            if((EnumHasAllFlags(Texture->GetDesc().Flags, TexCreate_UAV | TexCreate_NoTiling) ||
-                EnumHasAllFlags(Texture->GetDesc().Flags, TexCreate_AtomicCompatible) ||
-                EnumHasAllFlags(Texture->GetDesc().Flags, ETextureCreateFlags::Atomic64Compatible))
+            if((EnumHasAllFlags(Texture->GetDesc().Flags, TexCreate_UAV | TexCreate_NoTiling) || bIsAtomicCompatible)
 				 && Texture->Texture.GetBuffer())
             {
                 FMetalBuffer MetalBuffer(Texture->Texture.GetBuffer(), false);
