Bring back MTL argument for older MacOS and iOS

diff --git a/Engine/Source/Runtime/Apple/MetalRHI/Private/MetalComputeCommandEncoder.cpp b/Engine/Source/Runtime/Apple/MetalRHI/Private/MetalComputeCommandEncoder.cpp
--- a/Engine/Source/Runtime/Apple/MetalRHI/Private/MetalComputeCommandEncoder.cpp
+++ b/Engine/Source/Runtime/Apple/MetalRHI/Private/MetalComputeCommandEncoder.cpp
@@ -375,75 +375,132 @@ void FMetalComputeCommandEncoderDebugging::Validate()
 	switch(((FMetalDebugComputeCommandEncoder*)m_ptr)->Buffer.GetPtr()->DebugLevel)
 	{
 		case EMetalDebugLevelConditionalSubmit:
-		case EMetalDebugLevelWaitForComplete:
-		case EMetalDebugLevelLogOperations:
-		case EMetalDebugLevelValidation:
-		{
-			check(((FMetalDebugComputeCommandEncoder*)m_ptr)->Pipeline);
-			
-			MTLComputePipelineReflection* Reflection = ((FMetalDebugComputeCommandEncoder*)m_ptr)->Pipeline->ComputePipelineReflection;
-			check(Reflection);
-	
-            NSArray<id<MTLBinding>>* Bindings = Reflection.bindings;
-			for (uint32 i = 0; i < Bindings.count; i++)
-			{
-                id<MTLBinding> Binding = [Bindings objectAtIndex:i];
-				check(Binding);
-				switch(Binding.type)
-				{
-					case MTLBindingTypeBuffer:
-					{
-						checkf(Binding.index < ML_MaxBuffers, TEXT("Metal buffer index exceeded!"));
-						if ((((FMetalDebugComputeCommandEncoder*)m_ptr)->ShaderBuffers.Buffers[Binding.index] == nil && ((FMetalDebugComputeCommandEncoder*)m_ptr)->ShaderBuffers.Bytes[Binding.index] == nil))
-						{
-							UE_LOG(LogMetal, Warning, TEXT("Unbound buffer at Metal index %u which will crash the driver: %s"), (uint32)Binding.index, *FString([Binding description]));
-							bOK = false;
-						}
-						break;
-					}
-					case MTLBindingTypeThreadgroupMemory:
-					{
-						break;
-					}
-					case MTLBindingTypeTexture:
-					{
-                        id<MTLTextureBinding> TextureBinding = (id<MTLTextureBinding>)[Bindings objectAtIndex:i];
-						checkf(Binding.index < ML_MaxTextures, TEXT("Metal texture index exceeded!"));
-						if (((FMetalDebugComputeCommandEncoder*)m_ptr)->ShaderTextures.Textures[Binding.index] == nil)
-						{
-							UE_LOG(LogMetal, Warning, TEXT("Unbound texture at Metal index %u  which will crash the driver: %s"), (uint32)Binding.index, *FString([Binding description]));
-							bOK = false;
-						}
-						else if (((FMetalDebugComputeCommandEncoder*)m_ptr)->ShaderTextures.Textures[Binding.index].textureType != TextureBinding.textureType)
-						{
-							UE_LOG(LogMetal, Warning, TEXT("Incorrect texture type bound at Metal index %u which will crash the driver: %s\n%s"), (uint32)Binding.index, *FString([Binding description]), *FString([((FMetalDebugComputeCommandEncoder*)m_ptr)->ShaderTextures.Textures[Binding.index] description]));
-							bOK = false;
-						}
-						break;
-					}
-					case MTLBindingTypeSampler:
-					{
-						checkf(Binding.index < ML_MaxSamplers, TEXT("Metal sampler index exceeded!"));
-						if (((FMetalDebugComputeCommandEncoder*)m_ptr)->ShaderSamplers.Samplers[Binding.index] == nil)
-						{
-							UE_LOG(LogMetal, Warning, TEXT("Unbound sampler at Metal index %u which will crash the driver: %s"), (uint32)Binding.index, *FString([Binding description]));
-							bOK = false;
-						}
-						break;
-					}
-					default:
-						check(false);
-						break;
-				}
-			}
-			break;
-		}
-		default:
-		{
-			break;
-		}
-	}
-	
+        case EMetalDebugLevelWaitForComplete:
+        case EMetalDebugLevelLogOperations:
+        case EMetalDebugLevelValidation:
+        {
+            check(((FMetalDebugComputeCommandEncoder*)m_ptr)->Pipeline);
+            
+            MTLComputePipelineReflection* Reflection = ((FMetalDebugComputeCommandEncoder*)m_ptr)->Pipeline->ComputePipelineReflection;
+            check(Reflection);
+            if (@available(macOS 13.0, iOS 16.0, *))
+            {             
+                NSArray<id<MTLBinding>>* Bindings = Reflection.bindings;
+                for (uint32 i = 0; i < Bindings.count; i++)
+                {
+                    id<MTLBinding> Binding = [Bindings objectAtIndex:i];
+                    check(Binding);
+                    switch(Binding.type)
+                    {
+                        case MTLBindingTypeBuffer:
+                        {
+                            checkf(Binding.index < ML_MaxBuffers, TEXT("Metal buffer index exceeded!"));
+                            if ((((FMetalDebugComputeCommandEncoder*)m_ptr)->ShaderBuffers.Buffers[Binding.index] == nil && ((FMetalDebugComputeCommandEncoder*)m_ptr)->ShaderBuffers.Bytes[Binding.index] == nil))
+                            {
+                                UE_LOG(LogMetal, Warning, TEXT("Unbound buffer at Metal index %u which will crash the driver: %s"), (uint32)Binding.index, *FString([Binding description]));
+                                bOK = false;
+                            }
+                            break;
+                        }
+                        case MTLBindingTypeThreadgroupMemory:
+                        {
+                            break;
+                        }
+                        case MTLBindingTypeTexture:
+                        {
+                            id<MTLTextureBinding> TextureBinding = (id<MTLTextureBinding>)[Bindings objectAtIndex:i];
+                            checkf(Binding.index < ML_MaxTextures, TEXT("Metal texture index exceeded!"));
+                            if (((FMetalDebugComputeCommandEncoder*)m_ptr)->ShaderTextures.Textures[Binding.index] == nil)
+                            {
+                                UE_LOG(LogMetal, Warning, TEXT("Unbound texture at Metal index %u  which will crash the driver: %s"), (uint32)Binding.index, *FString([Binding description]));
+                                bOK = false;
+                            }
+                            else if (((FMetalDebugComputeCommandEncoder*)m_ptr)->ShaderTextures.Textures[Binding.index].textureType != TextureBinding.textureType)
+                            {
+                                UE_LOG(LogMetal, Warning, TEXT("Incorrect texture type bound at Metal index %u which will crash the driver: %s\n%s"), (uint32)Binding.index, *FString([Binding description]), *FString([((FMetalDebugComputeCommandEncoder*)m_ptr)->ShaderTextures.Textures[Binding.index] description]));
+                                bOK = false;
+                            }
+                            break;
+                        }
+                        case MTLBindingTypeSampler:
+                        {
+                            checkf(Binding.index < ML_MaxSamplers, TEXT("Metal sampler index exceeded!"));
+                            if (((FMetalDebugComputeCommandEncoder*)m_ptr)->ShaderSamplers.Samplers[Binding.index] == nil)
+                            {
+                                UE_LOG(LogMetal, Warning, TEXT("Unbound sampler at Metal index %u which will crash the driver: %s"), (uint32)Binding.index, *FString([Binding description]));
+                                bOK = false;
+                            }
+                            break;
+                        }
+                        default:
+                            check(false);
+                            break;
+                    }
+                }
+            }
+            else
+            {
+                NSArray<MTLArgument*>* Arguments = Reflection.arguments;
+                for (uint32 i = 0; i < Arguments.count; i++)
+                {
+                    MTLArgument* Arg = [Arguments objectAtIndex:i];
+                    check(Arg);
+                    switch(Arg.type)
+                    {
+                        case MTLArgumentTypeBuffer:
+                        {
+                            checkf(Arg.index < ML_MaxBuffers, TEXT("Metal buffer index exceeded!"));
+                            if ((((FMetalDebugComputeCommandEncoder*)m_ptr)->ShaderBuffers.Buffers[Arg.index] == nil && ((FMetalDebugComputeCommandEncoder*)m_ptr)->ShaderBuffers.Bytes[Arg.index] == nil))
+                            {
+                                UE_LOG(LogMetal, Warning, TEXT("Unbound buffer at Metal index %u which will crash the driver: %s"), (uint32)Arg.index, *FString([Arg description]));
+                                bOK = false;
+                            }
+                            break;
+                        }
+                        case MTLArgumentTypeThreadgroupMemory:
+                        {
+                            break;
+                        }
+                        case MTLArgumentTypeTexture:
+                        {
+                            checkf(Arg.index < ML_MaxTextures, TEXT("Metal texture index exceeded!"));
+                            if (((FMetalDebugComputeCommandEncoder*)m_ptr)->ShaderTextures.Textures[Arg.index] == nil)
+                            {
+                                UE_LOG(LogMetal, Warning, TEXT("Unbound texture at Metal index %u  which will crash the driver: %s"), (uint32)Arg.index, *FString([Arg description]));
+                                bOK = false;
+                            }
+                            else if (((FMetalDebugComputeCommandEncoder*)m_ptr)->ShaderTextures.Textures[Arg.index].textureType != Arg.textureType)
+                            {
+                                UE_LOG(LogMetal, Warning, TEXT("Incorrect texture type bound at Metal index %u which will crash the driver: %s\n%s"), (uint32)Arg.index, *FString([Arg description]), *FString([((FMetalDebugComputeCommandEncoder*)m_ptr)->ShaderTextures.Textures[Arg.index] description]));
+                                bOK = false;
+                            }
+                            break;
+                        }
+                        case MTLArgumentTypeSampler:
+                        {
+                            checkf(Arg.index < ML_MaxSamplers, TEXT("Metal sampler index exceeded!"));
+                            if (((FMetalDebugComputeCommandEncoder*)m_ptr)->ShaderSamplers.Samplers[Arg.index] == nil)
+                            {
+                                UE_LOG(LogMetal, Warning, TEXT("Unbound sampler at Metal index %u which will crash the driver: %s"), (uint32)Arg.index, *FString([Arg description]));
+                                bOK = false;
+                            }
+                            break;
+                        }
+                        default:
+                            check(false);
+                            break;
+                    }
+                }
+            }
+            
+            break;
+        }
+        default:
+        {
+            break;
+        }
+    }
+    
     if (!bOK)
     {
         UE_LOG(LogMetal, Error, TEXT("Metal Validation failures for compute shader:\n%s"), (((FMetalDebugComputeCommandEncoder*)m_ptr)->Pipeline && ((FMetalDebugComputeCommandEncoder*)m_ptr)->Pipeline->ComputeSource) ? *FString(((FMetalDebugComputeCommandEncoder*)m_ptr)->Pipeline->ComputeSource) : TEXT("nil"));
diff --git a/Engine/Source/Runtime/Apple/MetalRHI/Private/MetalPipeline.cpp b/Engine/Source/Runtime/Apple/MetalRHI/Private/MetalPipeline.cpp
--- a/Engine/Source/Runtime/Apple/MetalRHI/Private/MetalPipeline.cpp
+++ b/Engine/Source/Runtime/Apple/MetalRHI/Private/MetalPipeline.cpp
@@ -481,98 +481,193 @@ APPLE_PLATFORM_OBJECT_ALLOC_OVERRIDES(FMetalShaderPipeline)
 		}
 	}
 }
+
 - (void)initResourceMask:(EMetalShaderFrequency)Frequency
 {
-    NSArray<id<MTLBinding>>* Bindings = nil;
-	switch(Frequency)
-	{
-		case EMetalShaderVertex:
-		{
-			MTLRenderPipelineReflection* Reflection = RenderPipelineReflection;
-			check(Reflection);
-			
-            Bindings = Reflection.vertexBindings;
-			break;
-		}
-		case EMetalShaderFragment:
-		{
-			MTLRenderPipelineReflection* Reflection = RenderPipelineReflection;
-			check(Reflection);
-			
-            Bindings = Reflection.fragmentBindings;
-			break;
-		}
-		case EMetalShaderCompute:
-		{
-			MTLComputePipelineReflection* Reflection = ComputePipelineReflection;
-			check(Reflection);
-			
-            Bindings = Reflection.bindings;
-			break;
-		}
-		case EMetalShaderStream:
-		{
-			MTLRenderPipelineReflection* Reflection = StreamPipelineReflection;
-			check(Reflection);
-			
-            Bindings = Reflection.vertexBindings;
-			break;
-		}
-		default:
-			check(false);
-			break;
-	}
-	
-	for (uint32 i = 0; i < Bindings.count; i++)
-	{
-        id<MTLBinding> Binding = [Bindings objectAtIndex:i];
-		check(Binding);
+    if (@available(macOS 13.0, iOS 16.0, *))
+    {
+        NSArray<id<MTLBinding>>* Bindings = nil;
+        switch(Frequency)
+        {
+            case EMetalShaderVertex:
+            {
+                MTLRenderPipelineReflection* Reflection = RenderPipelineReflection;
+                check(Reflection);
+                
+                Bindings = Reflection.vertexBindings;
+                break;
+            }
+            case EMetalShaderFragment:
+            {
+                MTLRenderPipelineReflection* Reflection = RenderPipelineReflection;
+                check(Reflection);
+                
+                Bindings = Reflection.fragmentBindings;
+                break;
+            }
+            case EMetalShaderCompute:
+            {
+                MTLComputePipelineReflection* Reflection = ComputePipelineReflection;
+                check(Reflection);
+                
+                Bindings = Reflection.bindings;
+                break;
+            }
+            case EMetalShaderStream:
+            {
+                MTLRenderPipelineReflection* Reflection = StreamPipelineReflection;
+                check(Reflection);
+                
+                Bindings = Reflection.vertexBindings;
+                break;
+            }
+            default:
+                check(false);
+                break;
+        }
         
-        if (!Binding.isUsed)
+        for (uint32 i = 0; i < Bindings.count; i++)
         {
-            continue;
+            id<MTLBinding> Binding = [Bindings objectAtIndex:i];
+            check(Binding);
+            
+            if (!Binding.isUsed)
+            {
+                continue;
+            }
+            
+            switch(Binding.type)
+            {
+                case MTLBindingTypeBuffer:
+                {
+                    id<MTLBufferBinding> BufferBinding = (id<MTLBufferBinding>)[Bindings objectAtIndex:i];
+                    checkf(Binding.index < ML_MaxBuffers, TEXT("Metal buffer index exceeded!"));
+                    if (FString(Binding.name) != TEXT("BufferSizes") && FString(Binding.name) != TEXT("spvBufferSizeConstants"))
+                    {
+                        ResourceMask[Frequency].BufferMask |= (1 << Binding.index);
+                        
+                        if(BufferDataSizes[Frequency].Num() < 31)
+                            BufferDataSizes[Frequency].SetNumZeroed(31);
+                        
+                        BufferDataSizes[Frequency][Binding.index] = BufferBinding.bufferDataSize;
+                    }
+                    break;
+                }
+                case MTLBindingTypeThreadgroupMemory:
+                {
+                    break;
+                }
+                case MTLBindingTypeTexture:
+                {
+                    id<MTLTextureBinding> TextureBinding = (id<MTLTextureBinding>)[Bindings objectAtIndex:i];
+                    checkf(Binding.index < ML_MaxTextures, TEXT("Metal texture index exceeded!"));
+                    ResourceMask[Frequency].TextureMask |= (FMetalTextureMask(1) << Binding.index);
+                    TextureTypes[Frequency].Add(Binding.index, (uint8)TextureBinding.textureType);
+                    break;
+                }
+                case MTLBindingTypeSampler:
+                {
+                    checkf(Binding.index < ML_MaxSamplers, TEXT("Metal sampler index exceeded!"));
+                    ResourceMask[Frequency].SamplerMask |= (1 << Binding.index);
+                    break;
+                }
+                default:
+                    check(false);
+                    break;
+            }
         }
-
-		switch(Binding.type)
-		{
-			case MTLBindingTypeBuffer:
-			{
-                id<MTLBufferBinding> BufferBinding = (id<MTLBufferBinding>)[Bindings objectAtIndex:i];
-				checkf(Binding.index < ML_MaxBuffers, TEXT("Metal buffer index exceeded!"));
-				if (FString(Binding.name) != TEXT("BufferSizes") && FString(Binding.name) != TEXT("spvBufferSizeConstants"))
-				{
-					ResourceMask[Frequency].BufferMask |= (1 << Binding.index);
-				
-					if(BufferDataSizes[Frequency].Num() < 31)
-						BufferDataSizes[Frequency].SetNumZeroed(31);
-				
-					BufferDataSizes[Frequency][Binding.index] = BufferBinding.bufferDataSize;
-				}
-				break;
-			}
-			case MTLBindingTypeThreadgroupMemory:
-			{
-				break;
-			}
-			case MTLBindingTypeTexture:
-			{
-                id<MTLTextureBinding> TextureBinding = (id<MTLTextureBinding>)[Bindings objectAtIndex:i];
-				checkf(Binding.index < ML_MaxTextures, TEXT("Metal texture index exceeded!"));
-				ResourceMask[Frequency].TextureMask |= (FMetalTextureMask(1) << Binding.index);
-				TextureTypes[Frequency].Add(Binding.index, (uint8)TextureBinding.textureType);
-				break;
-			}
-			case MTLBindingTypeSampler:
-			{
-				checkf(Binding.index < ML_MaxSamplers, TEXT("Metal sampler index exceeded!"));
-				ResourceMask[Frequency].SamplerMask |= (1 << Binding.index);
-				break;
-			}
-			default:
-				check(false);
-				break;
-		}
-	}
+    }
+    else
+    {
+        NSArray<MTLArgument*>* Arguments = nil;
+        switch(Frequency)
+        {
+            case EMetalShaderVertex:
+            {
+                MTLRenderPipelineReflection* Reflection = RenderPipelineReflection;
+                check(Reflection);
+                
+                Arguments = Reflection.vertexArguments;
+                break;
+            }
+            case EMetalShaderFragment:
+            {
+                MTLRenderPipelineReflection* Reflection = RenderPipelineReflection;
+                check(Reflection);
+                
+                Arguments = Reflection.fragmentArguments;
+                break;
+            }
+            case EMetalShaderCompute:
+            {
+                MTLComputePipelineReflection* Reflection = ComputePipelineReflection;
+                check(Reflection);
+                
+                Arguments = Reflection.arguments;
+                break;
+            }
+            case EMetalShaderStream:
+            {
+                MTLRenderPipelineReflection* Reflection = StreamPipelineReflection;
+                check(Reflection);
+                
+                Arguments = Reflection.vertexArguments;
+                break;
+            }
+            default:
+                check(false);
+                break;
+        }
+        
+        for (uint32 i = 0; i < Arguments.count; i++)
+        {
+            MTLArgument* Arg = [Arguments objectAtIndex:i];
+            check(Arg);
+            
+            if (!Arg.active)
+            {
+                continue;
+            }
+            
+            switch(Arg.type)
+            {
+                case MTLArgumentTypeBuffer:
+                {
+                    checkf(Arg.index < ML_MaxBuffers, TEXT("Metal buffer index exceeded!"));
+                    if (FString(Arg.name) != TEXT("BufferSizes") && FString(Arg.name) != TEXT("spvBufferSizeConstants"))
+                    {
+                        ResourceMask[Frequency].BufferMask |= (1 << Arg.index);
+                        
+                        if(BufferDataSizes[Frequency].Num() < 31)
+                            BufferDataSizes[Frequency].SetNumZeroed(31);
+                        
+                        BufferDataSizes[Frequency][Arg.index] = Arg.bufferDataSize;
+                    }
+                    break;
+                }
+                case MTLArgumentTypeThreadgroupMemory:
+                {
+                    break;
+                }
+                case MTLArgumentTypeTexture:
+                {
+                    checkf(Arg.index < ML_MaxTextures, TEXT("Metal texture index exceeded!"));
+                    ResourceMask[Frequency].TextureMask |= (FMetalTextureMask(1) << Arg.index);
+                    TextureTypes[Frequency].Add(Arg.index, (uint8)Arg.textureType);
+                    break;
+                }
+                case MTLArgumentTypeSampler:
+                {
+                    checkf(Arg.index < ML_MaxSamplers, TEXT("Metal sampler index exceeded!"));
+                    ResourceMask[Frequency].SamplerMask |= (1 << Arg.index);
+                    break;
+                }
+                default:
+                    check(false);
+                    break;
+            }
+        }
+    }
 }
 @end
 
diff --git a/Engine/Source/Runtime/Apple/MetalRHI/Private/MetalRenderCommandEncoder.cpp b/Engine/Source/Runtime/Apple/MetalRHI/Private/MetalRenderCommandEncoder.cpp
--- a/Engine/Source/Runtime/Apple/MetalRHI/Private/MetalRenderCommandEncoder.cpp
+++ b/Engine/Source/Runtime/Apple/MetalRHI/Private/MetalRenderCommandEncoder.cpp
@@ -600,77 +600,151 @@ bool FMetalRenderCommandEncoderDebugging::ValidateFunctionBindings(EMetalShaderF
 			
 			MTLRenderPipelineReflection* Reflection = ((FMetalDebugRenderCommandEncoder*)m_ptr)->Pipeline->RenderPipelineReflection;
 			check(Reflection);
-			
-            NSArray<id<MTLBinding>>* Bindings = nil;
-			switch(Frequency)
-			{
-				case EMetalShaderVertex:
-				{
-                    Bindings = Reflection.vertexBindings;
-					break;
-				}
-				case EMetalShaderFragment:
-				{
-                    Bindings = Reflection.fragmentBindings;
-					break;
-				}
-				default:
-					check(false);
-					break;
-			}
-			
-			for (uint32 i = 0; i < Bindings.count; i++)
-			{
-                id<MTLBinding> Binding = [Bindings objectAtIndex:i];
-				check(Binding);
-				switch(Binding.type)
-				{
-					case MTLBindingTypeBuffer:
-					{
-						checkf(Binding.index < ML_MaxBuffers, TEXT("Metal buffer index exceeded!"));
-						if ((((FMetalDebugRenderCommandEncoder*)m_ptr)->ShaderBuffers[Frequency].Buffers[Binding.index] == nil && ((FMetalDebugRenderCommandEncoder*)m_ptr)->ShaderBuffers[Frequency].Bytes[Binding.index] == nil))
-						{
-							bOK = false;
-							UE_LOG(LogMetal, Warning, TEXT("Unbound buffer at Metal index %u which will crash the driver: %s"), (uint32)Binding.index, *FString([Binding description]));
-						}
-						break;
-					}
-					case MTLBindingTypeThreadgroupMemory:
-					{
-						break;
-					}
-					case MTLBindingTypeTexture:
-					{
-                        id<MTLTextureBinding> TextureBinding = (id<MTLTextureBinding>)[Bindings objectAtIndex:i];
-						checkf(Binding.index < ML_MaxTextures, TEXT("Metal texture index exceeded!"));
-						if (((FMetalDebugRenderCommandEncoder*)m_ptr)->ShaderTextures[Frequency].Textures[Binding.index] == nil)
-						{
-							bOK = false;
-							UE_LOG(LogMetal, Warning, TEXT("Unbound texture at Metal index %u which will crash the driver: %s"), (uint32)Binding.index, *FString([Binding description]));
-						}
-						else if (((FMetalDebugRenderCommandEncoder*)m_ptr)->ShaderTextures[Frequency].Textures[Binding.index].textureType != TextureBinding.textureType)
-						{
-							bOK = false;
-							UE_LOG(LogMetal, Warning, TEXT("Incorrect texture type bound at Metal index %u which will crash the driver: %s\n%s"), (uint32)Binding.index, *FString([Binding description]), *FString([((FMetalDebugRenderCommandEncoder*)m_ptr)->ShaderTextures[Frequency].Textures[Binding.index] description]));
-						}
-						break;
-					}
-					case MTLBindingTypeSampler:
-					{
-						checkf(Binding.index < ML_MaxSamplers, TEXT("Metal sampler index exceeded!"));
-						if (((FMetalDebugRenderCommandEncoder*)m_ptr)->ShaderSamplers[Frequency].Samplers[Binding.index] == nil)
-						{
-							bOK = false;
-							UE_LOG(LogMetal, Warning, TEXT("Unbound sampler at Metal index %u which will crash the driver: %s"), (uint32)Binding.index, *FString([Binding description]));
-						}
-						break;
-					}
-					default:
-						check(false);
-						break;
-				}
-			}
-			break;
+            if (@available(macOS 13.0, iOS 16.0, *))
+            {
+                NSArray<id<MTLBinding>>* Bindings = nil;
+                switch(Frequency)
+                {
+                    case EMetalShaderVertex:
+                    {
+                        Bindings = Reflection.vertexBindings;
+                        break;
+                    }
+                    case EMetalShaderFragment:
+                    {
+                        Bindings = Reflection.fragmentBindings;
+                        break;
+                    }
+                    default:
+                        check(false);
+                        break;
+                }
+                
+                for (uint32 i = 0; i < Bindings.count; i++)
+                {
+                    id<MTLBinding> Binding = [Bindings objectAtIndex:i];
+                    check(Binding);
+                    switch(Binding.type)
+                    {
+                        case MTLBindingTypeBuffer:
+                        {
+                            checkf(Binding.index < ML_MaxBuffers, TEXT("Metal buffer index exceeded!"));
+                            if ((((FMetalDebugRenderCommandEncoder*)m_ptr)->ShaderBuffers[Frequency].Buffers[Binding.index] == nil && ((FMetalDebugRenderCommandEncoder*)m_ptr)->ShaderBuffers[Frequency].Bytes[Binding.index] == nil))
+                            {
+                                bOK = false;
+                                UE_LOG(LogMetal, Warning, TEXT("Unbound buffer at Metal index %u which will crash the driver: %s"), (uint32)Binding.index, *FString([Binding description]));
+                            }
+                            break;
+                        }
+                        case MTLBindingTypeThreadgroupMemory:
+                        {
+                            break;
+                        }
+                        case MTLBindingTypeTexture:
+                        {
+                            id<MTLTextureBinding> TextureBinding = (id<MTLTextureBinding>)[Bindings objectAtIndex:i];
+                            checkf(Binding.index < ML_MaxTextures, TEXT("Metal texture index exceeded!"));
+                            if (((FMetalDebugRenderCommandEncoder*)m_ptr)->ShaderTextures[Frequency].Textures[Binding.index] == nil)
+                            {
+                                bOK = false;
+                                UE_LOG(LogMetal, Warning, TEXT("Unbound texture at Metal index %u which will crash the driver: %s"), (uint32)Binding.index, *FString([Binding description]));
+                            }
+                            else if (((FMetalDebugRenderCommandEncoder*)m_ptr)->ShaderTextures[Frequency].Textures[Binding.index].textureType != TextureBinding.textureType)
+                            {
+                                bOK = false;
+                                UE_LOG(LogMetal, Warning, TEXT("Incorrect texture type bound at Metal index %u which will crash the driver: %s\n%s"), (uint32)Binding.index, *FString([Binding description]), *FString([((FMetalDebugRenderCommandEncoder*)m_ptr)->ShaderTextures[Frequency].Textures[Binding.index] description]));
+                            }
+                            break;
+                        }
+                        case MTLBindingTypeSampler:
+                        {
+                            checkf(Binding.index < ML_MaxSamplers, TEXT("Metal sampler index exceeded!"));
+                            if (((FMetalDebugRenderCommandEncoder*)m_ptr)->ShaderSamplers[Frequency].Samplers[Binding.index] == nil)
+                            {
+                                bOK = false;
+                                UE_LOG(LogMetal, Warning, TEXT("Unbound sampler at Metal index %u which will crash the driver: %s"), (uint32)Binding.index, *FString([Binding description]));
+                            }
+                            break;
+                        }
+                        default:
+                            check(false);
+                            break;
+                    }
+                }
+                break;
+            }
+            else
+            {
+                NSArray<MTLArgument*>* Arguments = nil;
+                        switch(Frequency)
+                        {
+                            case EMetalShaderVertex:
+                            {
+                                Arguments = Reflection.vertexArguments;
+                                break;
+                            }
+                            case EMetalShaderFragment:
+                            {
+                                Arguments = Reflection.fragmentArguments;
+                                break;
+                            }
+                            default:
+                                check(false);
+                                break;
+                        }
+                        
+                        for (uint32 i = 0; i < Arguments.count; i++)
+                        {
+                            MTLArgument* Arg = [Arguments objectAtIndex:i];
+                            check(Arg);
+                            switch(Arg.type)
+                            {
+                                case MTLArgumentTypeBuffer:
+                                {
+                                    checkf(Arg.index < ML_MaxBuffers, TEXT("Metal buffer index exceeded!"));
+                                    if ((((FMetalDebugRenderCommandEncoder*)m_ptr)->ShaderBuffers[Frequency].Buffers[Arg.index] == nil && ((FMetalDebugRenderCommandEncoder*)m_ptr)->ShaderBuffers[Frequency].Bytes[Arg.index] == nil))
+                                    {
+                                        bOK = false;
+                                        UE_LOG(LogMetal, Warning, TEXT("Unbound buffer at Metal index %u which will crash the driver: %s"), (uint32)Arg.index, *FString([Arg description]));
+                                    }
+                                    break;
+                                }
+                                case MTLArgumentTypeThreadgroupMemory:
+                                {
+                                    break;
+                                }
+                                case MTLArgumentTypeTexture:
+                                {
+                                    checkf(Arg.index < ML_MaxTextures, TEXT("Metal texture index exceeded!"));
+                                    if (((FMetalDebugRenderCommandEncoder*)m_ptr)->ShaderTextures[Frequency].Textures[Arg.index] == nil)
+                                    {
+                                        bOK = false;
+                                        UE_LOG(LogMetal, Warning, TEXT("Unbound texture at Metal index %u which will crash the driver: %s"), (uint32)Arg.index, *FString([Arg description]));
+                                    }
+                                    else if (((FMetalDebugRenderCommandEncoder*)m_ptr)->ShaderTextures[Frequency].Textures[Arg.index].textureType != Arg.textureType)
+                                    {
+                                        bOK = false;
+                                        UE_LOG(LogMetal, Warning, TEXT("Incorrect texture type bound at Metal index %u which will crash the driver: %s\n%s"), (uint32)Arg.index, *FString([Arg description]), *FString([((FMetalDebugRenderCommandEncoder*)m_ptr)->ShaderTextures[Frequency].Textures[Arg.index] description]));
+                                    }
+                                    break;
+                                }
+                                case MTLArgumentTypeSampler:
+                                {
+                                    checkf(Arg.index < ML_MaxSamplers, TEXT("Metal sampler index exceeded!"));
+                                    if (((FMetalDebugRenderCommandEncoder*)m_ptr)->ShaderSamplers[Frequency].Samplers[Arg.index] == nil)
+                                    {
+                                        bOK = false;
+                                        UE_LOG(LogMetal, Warning, TEXT("Unbound sampler at Metal index %u which will crash the driver: %s"), (uint32)Arg.index, *FString([Arg description]));
+                                    }
+                                    break;
+                                }
+                                default:
+                                    check(false);
+                                    break;
+                            }
+                        }
+                        break;
+            }
 		}
 		default:
 		{
