iOS A8 devices do not support Virtual Textures, so if they have been enabled via r.Mobile.VirtualTextures, give a log Warning that it's not supported when running on an A8 device. Don't use Fatal as level might not use VTs.

diff --git a/Engine/Source/Runtime/Apple/MetalRHI/Private/MetalRHI.cpp b/Engine/Source/Runtime/Apple/MetalRHI/Private/MetalRHI.cpp
--- a/Engine/Source/Runtime/Apple/MetalRHI/Private/MetalRHI.cpp
+++ b/Engine/Source/Runtime/Apple/MetalRHI/Private/MetalRHI.cpp
@@ -249,7 +249,7 @@ FMetalDynamicRHI::FMetalDynamicRHI(ERHIFeatureLevel::Type RequestedFeatureLevel)
         static auto* CVarMobileVirtualTextures = IConsoleManager::Get().FindTConsoleVariableDataInt(TEXT("r.Mobile.VirtualTextures"));
         if(CVarMobileVirtualTextures->GetValueOnAnyThread() != 0)
         {
-            UE_LOG(LogMetal, Fatal, TEXT("Mobile Virtual Textures require a minimum of the Apple A9 feature set."));
+            UE_LOG(LogMetal, Warning, TEXT("Mobile Virtual Textures require a minimum of the Apple A9 feature set."));
         }
     }
         
