Disable the SM5 Renderer for A8 devices

diff --git a/Engine/Source/Runtime/Apple/MetalRHI/Private/MetalRHI.cpp b/Engine/Source/Runtime/Apple/MetalRHI/Private/MetalRHI.cpp
--- a/Engine/Source/Runtime/Apple/MetalRHI/Private/MetalRHI.cpp
+++ b/Engine/Source/Runtime/Apple/MetalRHI/Private/MetalRHI.cpp
@@ -258,8 +258,8 @@ FMetalDynamicRHI::FMetalDynamicRHI(ERHIFeatureLevel::Type RequestedFeatureLevel)
 
 	bool const bRequestedMetalMRT = ((RequestedFeatureLevel >= ERHIFeatureLevel::SM5) || (!bRequestedFeatureLevel && FParse::Param(FCommandLine::Get(),TEXT("metalmrt"))));
 
-    // only allow GBuffers, etc on A8s (A7s are just not going to cut it)
-    if (bProjectSupportsMRTs && bRequestedMetalMRT)
+    // Only allow SM5 MRT on A9 or above devices
+    if (bProjectSupportsMRTs && bRequestedMetalMRT && !bIsA8FeatureSet)
     {
 #if PLATFORM_TVOS
 		ValidateTargetedRHIFeatureLevelExists(SP_METAL_MRT);
@@ -274,7 +274,7 @@ FMetalDynamicRHI::FMetalDynamicRHI(ERHIFeatureLevel::Type RequestedFeatureLevel)
 	{
 		if (bRequestedMetalMRT)
 		{
-			UE_LOG(LogMetal, Warning, TEXT("Metal MRT support requires an iOS or tvOS device with an A8 processor or later. Falling back to Metal ES 3.1."));
+			UE_LOG(LogMetal, Warning, TEXT("Metal MRT support requires an iOS or tvOS device with an A8x/A9 processor or later. Falling back to Metal ES 3.1."));
 		}
 		
 #if PLATFORM_TVOS
