Required Metal changes to support iOS Simulator.
- add a new shader platform (SP_METAL_SIM) to support iOS Simulator specific compilation requirements
- update iOS sim #define WITH_SIMULATOR to WITH_IOS_SIMULATOR

diff --git a/Engine/Config/IOS/DataDrivenPlatformInfo.ini b/Engine/Config/IOS/DataDrivenPlatformInfo.ini
--- a/Engine/Config/IOS/DataDrivenPlatformInfo.ini
+++ b/Engine/Config/IOS/DataDrivenPlatformInfo.ini
@@ -69,3 +69,15 @@ bSupportsManualVertexFetch=false
 bSupportsVolumeTextureAtomics=false
 bSupportsDistanceFields=true
 FriendlyName=LOCTEXT("FriendlyShaderPlatformName_IOSSM5", "IOS Metal SM5")
+
+[ShaderPlatform METAL_SIM]
+Language=Metal
+MaxFeatureLevel=ES3_1
+ShaderFormat=SF_METAL_SIM
+bIsMobile=true
+bTargetsTiledGPU=true
+bSupports4ComponentUAVReadWrite=true
+bSupportsManualVertexFetch=false
+bSupportsVolumeTextureAtomics=false
+bSupportsShaderPipelines = false
+FriendlyName=LOCTEXT("FriendlyShaderPlatformName_IOSMobileSIM", "IOS Metal Simulator Mobile")
\ No newline at end of file
diff --git a/Engine/Source/Editor/DetailCustomizations/Private/MacTargetSettingsDetails.cpp b/Engine/Source/Editor/DetailCustomizations/Private/MacTargetSettingsDetails.cpp
--- a/Engine/Source/Editor/DetailCustomizations/Private/MacTargetSettingsDetails.cpp
+++ b/Engine/Source/Editor/DetailCustomizations/Private/MacTargetSettingsDetails.cpp
@@ -140,6 +140,9 @@ static FText GetFriendlyNameFromRHINameMac(FName InRHIName)
     case SP_METAL_SM6:
         FriendlyRHIName = LOCTEXT("MetalSM6", "Mac Metal Desktop Renderer Beta (SM6, Metal 2.4+, macOS 14.0 or later, M2+)");
         break;
+	case SP_METAL_SIM:
+		FriendlyRHIName = LOCTEXT("MetalSim", "iOS Metal Simulator Mobile Renderer (Simulator, Metal 2.4+, iOS 15.0 or later)");
+		break;
 	case SP_METAL_MACES3_1:
 		FriendlyRHIName = LOCTEXT("MetalMobile", "Mac Metal High-End Mobile Preview (Mobile Preview)");
 		break;
diff --git a/Engine/Source/Editor/MaterialEditor/Private/MaterialStats.cpp b/Engine/Source/Editor/MaterialEditor/Private/MaterialStats.cpp
--- a/Engine/Source/Editor/MaterialEditor/Private/MaterialStats.cpp
+++ b/Engine/Source/Editor/MaterialEditor/Private/MaterialStats.cpp
@@ -489,6 +489,7 @@ void FMaterialStats::BuildShaderPlatformDB()
 	AddShaderPlatform(EPlatformCategoryType::Desktop, SP_METAL_MRT_MAC, TEXT("Metal SM5 (MRT)"), bCanCompileMacDesktopSPs, true, TEXT("macOS, Metal, Shader Model 5"));
 	AddShaderPlatform(EPlatformCategoryType::IOS, SP_METAL, TEXT("Metal"), bCanCompileMacMobileSPs, true, TEXT("iOS, Metal, Mobile"));
 	AddShaderPlatform(EPlatformCategoryType::IOS, SP_METAL_MRT, TEXT("Metal MRT"), bCanCompileMacMobileSPs, true, TEXT("iOS, Metal, Shader Model 5"));
+	AddShaderPlatform(EPlatformCategoryType::IOS, SP_METAL_SIM, TEXT("Metal Simulator"), bCanCompileMacMobileSPs, true, TEXT("iOS, Metal, Mobile"));
 
 	ITargetPlatformManagerModule& TPM = GetTargetPlatformManagerRef();
 
diff --git a/Engine/Source/Programs/UnrealBuildTool/Platform/IOS/UEBuildIOS.cs b/Engine/Source/Programs/UnrealBuildTool/Platform/IOS/UEBuildIOS.cs
--- a/Engine/Source/Programs/UnrealBuildTool/Platform/IOS/UEBuildIOS.cs
+++ b/Engine/Source/Programs/UnrealBuildTool/Platform/IOS/UEBuildIOS.cs
@@ -1152,11 +1152,11 @@ namespace UnrealBuildTool
 
 			if (Target.Architecture == UnrealArch.IOSSimulator || Target.Architecture == UnrealArch.TVOSSimulator)
 			{
-				CompileEnvironment.Definitions.Add("WITH_SIMULATOR=1");
+				CompileEnvironment.Definitions.Add("WITH_IOS_SIMULATOR=1");
 			}
 			else
 			{
-				CompileEnvironment.Definitions.Add("WITH_SIMULATOR=0");
+				CompileEnvironment.Definitions.Add("WITH_IOS_SIMULATOR=0");
 			}
 
 			if (ProjectSettings.bEnableAdvertisingIdentifier)
diff --git a/Engine/Source/Runtime/Apple/MetalRHI/Private/MetalBuffer.cpp b/Engine/Source/Runtime/Apple/MetalRHI/Private/MetalBuffer.cpp
--- a/Engine/Source/Runtime/Apple/MetalRHI/Private/MetalBuffer.cpp
+++ b/Engine/Source/Runtime/Apple/MetalRHI/Private/MetalBuffer.cpp
@@ -948,7 +948,9 @@ FMetalSubBufferRing::FMetalSubBufferRing(NSUInteger Size, NSUInteger Alignment,
 , Storage((mtlpp::StorageMode)((Options & mtlpp::ResourceStorageModeMask) >> mtlpp::ResourceStorageModeShift))
 {
 	Options = (mtlpp::ResourceOptions)FMetalCommandQueue::GetCompatibleResourceOptions(Options);
+#if !defined(WITH_IOS_SIMULATOR) || !WITH_IOS_SIMULATOR
 	check(Storage != mtlpp::StorageMode::Private /* Private memory requires command-buffers and encoders to properly marshal! */);
+#endif
 	FMemory::Memzero(FrameSize);
 }
 
diff --git a/Engine/Source/Runtime/Apple/MetalRHI/Private/MetalCommandQueue.cpp b/Engine/Source/Runtime/Apple/MetalRHI/Private/MetalCommandQueue.cpp
--- a/Engine/Source/Runtime/Apple/MetalRHI/Private/MetalCommandQueue.cpp
+++ b/Engine/Source/Runtime/Apple/MetalRHI/Private/MetalCommandQueue.cpp
@@ -377,7 +377,11 @@ mtlpp::ResourceOptions FMetalCommandQueue::GetCompatibleResourceOptions(mtlpp::R
 #if PLATFORM_IOS // Swizzle Managed to Shared for iOS - we can do this as they are equivalent, unlike Shared -> Managed on Mac.
 	if ((Options & (1 /*mtlpp::StorageMode::Managed*/ << mtlpp::ResourceStorageModeShift)))
 	{
+#if WITH_IOS_SIMULATOR
+		NewOptions |= mtlpp::ResourceOptions::StorageModePrivate;
+#else
 		NewOptions |= mtlpp::ResourceOptions::StorageModeShared;
+#endif
 	}
 #endif
 	return (mtlpp::ResourceOptions)NewOptions;
diff --git a/Engine/Source/Runtime/Apple/MetalRHI/Private/MetalRHI.cpp b/Engine/Source/Runtime/Apple/MetalRHI/Private/MetalRHI.cpp
--- a/Engine/Source/Runtime/Apple/MetalRHI/Private/MetalRHI.cpp
+++ b/Engine/Source/Runtime/Apple/MetalRHI/Private/MetalRHI.cpp
@@ -58,7 +58,7 @@ static void ValidateTargetedRHIFeatureLevelExists(EShaderPlatform Platform)
 		}
 	}
 #else
-	if (Platform == SP_METAL || Platform == SP_METAL_TVOS)
+	if (Platform == SP_METAL || Platform == SP_METAL_TVOS || Platform == SP_METAL_SIM)
 	{
 		GConfig->GetBool(TEXT("/Script/IOSRuntimeSettings.IOSRuntimeSettings"), TEXT("bSupportsMetal"), bSupportsShaderPlatform, GEngineIni);
 	}
@@ -274,16 +274,21 @@ FMetalDynamicRHI::FMetalDynamicRHI(ERHIFeatureLevel::Type RequestedFeatureLevel)
 	{
 		if (bRequestedMetalMRT)
 		{
-			UE_LOG(LogMetal, Warning, TEXT("Metal MRT support requires an iOS or tvOS device with an A8x/A9 processor or later. Falling back to Metal ES 3.1."));
+			UE_LOG(LogMetal, Warning, TEXT("Metal MRT support requires an iOS or tvOS device with an A8 processor or later. Falling back to Metal ES 3.1."));
 		}
 		
 #if PLATFORM_TVOS
 		ValidateTargetedRHIFeatureLevelExists(SP_METAL_TVOS);
 		GMaxRHIShaderPlatform = SP_METAL_TVOS;
+#else
+#if WITH_IOS_SIMULATOR
+		ValidateTargetedRHIFeatureLevelExists(SP_METAL_SIM);
+		GMaxRHIShaderPlatform = SP_METAL_SIM;
 #else
 		ValidateTargetedRHIFeatureLevelExists(SP_METAL);
 		GMaxRHIShaderPlatform = SP_METAL;
-#endif
+#endif	// WITH_IOS_SIMULATOR
+#endif	// PLATFORM_TVOS
         GMaxRHIFeatureLevel = ERHIFeatureLevel::ES3_1;
 	}
 		
@@ -299,7 +304,11 @@ FMetalDynamicRHI::FMetalDynamicRHI(ERHIFeatureLevel::Type RequestedFeatureLevel)
 	GShaderPlatformForFeatureLevel[ERHIFeatureLevel::ES3_1] = SP_METAL_TVOS;
 #else
 	GShaderPlatformForFeatureLevel[ERHIFeatureLevel::ES2_REMOVED] = SP_NumPlatforms;
+#if WITH_IOS_SIMULATOR
+	GShaderPlatformForFeatureLevel[ERHIFeatureLevel::ES3_1] = SP_METAL_SIM;
+#else
 	GShaderPlatformForFeatureLevel[ERHIFeatureLevel::ES3_1] = SP_METAL;
+#endif
 #endif
 	GShaderPlatformForFeatureLevel[ERHIFeatureLevel::SM4_REMOVED] = SP_NumPlatforms;
 	GShaderPlatformForFeatureLevel[ERHIFeatureLevel::SM5] = (GMaxRHIFeatureLevel >= ERHIFeatureLevel::SM5) ? GMaxRHIShaderPlatform : SP_NumPlatforms;
diff --git a/Engine/Source/Runtime/Apple/MetalRHI/Private/MetalRHIPrivate.h b/Engine/Source/Runtime/Apple/MetalRHI/Private/MetalRHIPrivate.h
--- a/Engine/Source/Runtime/Apple/MetalRHI/Private/MetalRHIPrivate.h
+++ b/Engine/Source/Runtime/Apple/MetalRHI/Private/MetalRHIPrivate.h
@@ -59,8 +59,15 @@ const uint32 MetalBufferBytesSize = BufferOffsetAlignment * 32;
 const uint32 MaxMetalStreams = 31;
 #else
 #define BUFFER_MANAGED_MEM 0
+
+#if WITH_IOS_SIMULATOR
+#define BUFFER_STORAGE_MODE mtlpp::StorageMode::Private
+#define BUFFER_RESOURCE_STORAGE_MANAGED mtlpp::ResourceOptions::StorageModePrivate
+#else
 #define BUFFER_STORAGE_MODE mtlpp::StorageMode::Shared
 #define BUFFER_RESOURCE_STORAGE_MANAGED mtlpp::ResourceOptions::StorageModeShared
+#endif
+
 #define BUFFER_DYNAMIC_REALLOC BUF_AnyDynamic
 // How many possible vertex streams are allowed
 const uint32 MaxMetalStreams = 30;
diff --git a/Engine/Source/Runtime/Apple/MetalRHI/Private/MetalRenderTarget.cpp b/Engine/Source/Runtime/Apple/MetalRHI/Private/MetalRenderTarget.cpp
--- a/Engine/Source/Runtime/Apple/MetalRHI/Private/MetalRenderTarget.cpp
+++ b/Engine/Source/Runtime/Apple/MetalRHI/Private/MetalRenderTarget.cpp
@@ -163,7 +163,11 @@ void FMetalDynamicRHI::RHIReadSurfaceData(FRHITexture* TextureRHI, FIntRect Rect
 #if PLATFORM_MAC
 			mtlpp::StorageMode StorageMode = mtlpp::StorageMode::Managed;
 #else
-			mtlpp::StorageMode StorageMode = mtlpp::StorageMode::Shared;
+#if WITH_IOS_SIMULATOR
+			mtlpp::StorageMode StorageMode = mtlpp::StorageMode::Private; 
+#else
+			mtlpp::StorageMode StorageMode = mtlpp::StorageMode::Shared; 
+#endif
 #endif
 			mtlpp::PixelFormat MetalFormat = (mtlpp::PixelFormat)GPixelFormats[Surface->GetDesc().Format].PlatformFormat;
 			mtlpp::TextureDescriptor Desc;
diff --git a/Engine/Source/Runtime/Apple/MetalRHI/Private/MetalTexture.cpp b/Engine/Source/Runtime/Apple/MetalRHI/Private/MetalTexture.cpp
--- a/Engine/Source/Runtime/Apple/MetalRHI/Private/MetalTexture.cpp
+++ b/Engine/Source/Runtime/Apple/MetalRHI/Private/MetalTexture.cpp
@@ -24,7 +24,7 @@ FAutoConsoleVariableRef CVarMetalMaxOutstandingAsyncTexUploads(
 															   ECVF_ReadOnly|ECVF_RenderThreadSafe
 															   );
 
-int32 GMetalForceIOSTexturesShared = 1;
+int32 GMetalForceIOSTexturesShared = !WITH_IOS_SIMULATOR;
 FAutoConsoleVariableRef CVarMetalForceIOSTexturesShared(
 														TEXT("rhi.Metal.ForceIOSTexturesShared"),
 														GMetalForceIOSTexturesShared,
@@ -194,7 +194,11 @@ mtlpp::PixelFormat UEToMetalFormat(EPixelFormat UEFormat, bool bSRGB)
 			switch (MTLFormat)
 			{
 			default: break;
+#if WITH_IOS_SIMULATOR
+			case mtlpp::PixelFormat::R8Unorm        : MTLFormat = mtlpp::PixelFormat::R8Unorm    		  ; break;
+#else
 			case mtlpp::PixelFormat::R8Unorm        : MTLFormat = mtlpp::PixelFormat::R8Unorm_sRGB        ; break;
+#endif
 			case mtlpp::PixelFormat::PVRTC_RGBA_2BPP: MTLFormat = mtlpp::PixelFormat::PVRTC_RGBA_2BPP_sRGB; break;
 			case mtlpp::PixelFormat::PVRTC_RGBA_4BPP: MTLFormat = mtlpp::PixelFormat::PVRTC_RGBA_4BPP_sRGB; break;
 			case mtlpp::PixelFormat::ASTC_4x4_LDR   : MTLFormat = mtlpp::PixelFormat::ASTC_4x4_sRGB       ; break;
@@ -362,7 +366,11 @@ FMetalTextureCreateDesc::FMetalTextureCreateDesc(FRHITextureCreateDesc const& In
 	{
 		Desc.SetUsage(ConvertFlagsToUsage(InDesc.Flags));
 		
+#if WITH_IOS_SIMULATOR
+		const bool bAppleGPU = false;
+#else
 		const bool bAppleGPU = [GetMetalDeviceContext().GetDevice().GetPtr() supportsFamily:MTLGPUFamilyApple1];
+#endif
 
 		if (EnumHasAnyFlags(InDesc.Flags, TexCreate_CPUReadback) && !EnumHasAnyFlags(InDesc.Flags, TexCreate_RenderTargetable | TexCreate_DepthStencilTargetable | TexCreate_FastVRAM))
 		{
diff --git a/Engine/Source/Runtime/ApplicationCore/Private/IOS/IOSView.cpp b/Engine/Source/Runtime/ApplicationCore/Private/IOS/IOSView.cpp
--- a/Engine/Source/Runtime/ApplicationCore/Private/IOS/IOSView.cpp
+++ b/Engine/Source/Runtime/ApplicationCore/Private/IOS/IOSView.cpp
@@ -462,7 +462,7 @@ self.accessibilityElements = @[Window.accessibilityContainer];
 	}
 
 #if !UE_BUILD_SHIPPING
-#if WITH_SIMULATOR
+#if WITH_IOS_SIMULATOR
 	// use 2 on the simulator so that Option-Click will bring up console (option-click is for doing pinch gestures, which we don't care about, atm)
 	if( NumActiveTouches >= 2 )
 #else
diff --git a/Engine/Source/Runtime/Core/Public/Apple/ApplePlatformMisc.h b/Engine/Source/Runtime/Core/Public/Apple/ApplePlatformMisc.h
--- a/Engine/Source/Runtime/Core/Public/Apple/ApplePlatformMisc.h
+++ b/Engine/Source/Runtime/Core/Public/Apple/ApplePlatformMisc.h
@@ -16,8 +16,8 @@
 #define APPLE_PROFILING_ENABLED (UE_BUILD_DEBUG | UE_BUILD_DEVELOPMENT)
 #endif
 
-#ifndef WITH_SIMULATOR
-#define WITH_SIMULATOR 0
+#ifndef WITH_IOS_SIMULATOR
+#define WITH_IOS_SIMULATOR 0
 #endif
 
 #define UE_DEBUG_BREAK_IMPL() PLATFORM_BREAK()
diff --git a/Engine/Source/Runtime/RHI/Private/RHIStrings.cpp b/Engine/Source/Runtime/RHI/Private/RHIStrings.cpp
--- a/Engine/Source/Runtime/RHI/Private/RHIStrings.cpp
+++ b/Engine/Source/Runtime/RHI/Private/RHIStrings.cpp
@@ -95,6 +95,7 @@ FString LexToString(EShaderPlatform Platform, bool bError)
 	case SP_METAL_MRT_MAC: return TEXT("METAL_MRT_MAC");
 	case SP_METAL_SM5: return TEXT("METAL_SM5");
     case SP_METAL_SM6: return TEXT("METAL_SM6");
+    case SP_METAL_SIM: return TEXT("METAL_SIM");
 	case SP_METAL_MACES3_1: return TEXT("METAL_MACES3_1");
 	case SP_VULKAN_ES3_1_ANDROID: return TEXT("VULKAN_ES3_1_ANDROID");
 	case SP_VULKAN_PCES3_1: return TEXT("VULKAN_PCES3_1");
@@ -212,6 +213,7 @@ FName ShaderPlatformToPlatformName(EShaderPlatform Platform)
 		return NAME_PLATFORM_ANDROID;
 	case SP_METAL:
 	case SP_METAL_MRT:
+	case SP_METAL_SIM:
 		return NAME_PLATFORM_IOS;
 	case SP_METAL_SM5:
     case SP_METAL_SM6:
diff --git a/Engine/Source/Runtime/RHI/Public/RHIShaderFormatDefinitions.inl b/Engine/Source/Runtime/RHI/Public/RHIShaderFormatDefinitions.inl
--- a/Engine/Source/Runtime/RHI/Public/RHIShaderFormatDefinitions.inl
+++ b/Engine/Source/Runtime/RHI/Public/RHIShaderFormatDefinitions.inl
@@ -23,6 +23,7 @@ static FName NAME_SF_METAL_MRT_TVOS(TEXT("SF_METAL_MRT_TVOS"));
 static FName NAME_SF_METAL_MRT_MAC(TEXT("SF_METAL_MRT_MAC"));
 static FName NAME_SF_METAL_SM5(TEXT("SF_METAL_SM5"));
 static FName NAME_SF_METAL_SM6(TEXT("SF_METAL_SM6"));
+static FName NAME_SF_METAL_SIM(TEXT("SF_METAL_SIM"));
 static FName NAME_SF_METAL_MACES3_1(TEXT("SF_METAL_MACES3_1"));
 
 static FName NAME_VULKAN_ES3_1_ANDROID(TEXT("SF_VULKAN_ES31_ANDROID"));
@@ -47,6 +48,7 @@ static EShaderPlatform ShaderFormatNameToShaderPlatform(FName ShaderFormat)
 	if (ShaderFormat == NAME_SF_METAL_MRT_MAC)			return SP_METAL_MRT_MAC;
 	if (ShaderFormat == NAME_SF_METAL_SM5)				return SP_METAL_SM5;
     if (ShaderFormat == NAME_SF_METAL_SM6)              return SP_METAL_SM6;
+	if (ShaderFormat == NAME_SF_METAL_SIM)				return SP_METAL_SIM;
 	if (ShaderFormat == NAME_SF_METAL_MACES3_1)			return SP_METAL_MACES3_1;
 
 	if (ShaderFormat == NAME_VULKAN_ES3_1_ANDROID)		return SP_VULKAN_ES3_1_ANDROID;
diff --git a/Engine/Source/Runtime/RHI/Public/RHIShaderPlatform.h b/Engine/Source/Runtime/RHI/Public/RHIShaderPlatform.h
--- a/Engine/Source/Runtime/RHI/Public/RHIShaderPlatform.h
+++ b/Engine/Source/Runtime/RHI/Public/RHIShaderPlatform.h
@@ -46,6 +46,7 @@ enum EShaderPlatform : uint16
 	SP_PCD3D_SM6          = SP_StaticPlatform_Last + 2,
 	SP_VULKAN_SM6         = SP_StaticPlatform_Last + 4,
     SP_METAL_SM6          = SP_StaticPlatform_Last + 5,
+    SP_METAL_SIM          = SP_StaticPlatform_Last + 6,
 
 	SP_CUSTOM_PLATFORM_FIRST,
 	SP_CUSTOM_PLATFORM_LAST = (SP_CUSTOM_PLATFORM_FIRST + 100),
diff --git a/Engine/Source/Runtime/Renderer/Private/MobileShadingRenderer.cpp b/Engine/Source/Runtime/Renderer/Private/MobileShadingRenderer.cpp
--- a/Engine/Source/Runtime/Renderer/Private/MobileShadingRenderer.cpp
+++ b/Engine/Source/Runtime/Renderer/Private/MobileShadingRenderer.cpp
@@ -1336,7 +1336,10 @@ void FMobileSceneRenderer::RenderForwardSinglePass(FRDGBuilder& GraphBuilder, FM
 		RHICmdList.SetCurrentStat(GET_STATID(STAT_CLMM_Translucency));
 		RenderDecals(RHICmdList, View);
 		RenderModulatedShadowProjections(RHICmdList, ViewContext.ViewIndex, View);
-		RenderFog(RHICmdList, View);
+		if (GMaxRHIShaderPlatform != SP_METAL_SIM)
+		{
+			RenderFog(RHICmdList, View);
+		}
 		RenderLocalFogVolumeMobile(RHICmdList, View);
 		// Draw translucency.
 		RenderTranslucency(RHICmdList, View);
diff --git a/Engine/Source/Runtime/Renderer/Private/VT/RuntimeVirtualTextureRender.cpp b/Engine/Source/Runtime/Renderer/Private/VT/RuntimeVirtualTextureRender.cpp
--- a/Engine/Source/Runtime/Renderer/Private/VT/RuntimeVirtualTextureRender.cpp
+++ b/Engine/Source/Runtime/Renderer/Private/VT/RuntimeVirtualTextureRender.cpp
@@ -98,6 +98,7 @@ namespace RuntimeVirtualTexture
 		{
 		case SP_METAL:
 		case SP_METAL_MRT:
+		case SP_METAL_SIM:
 		case SP_METAL_TVOS:
 		case SP_METAL_MRT_TVOS:
 		case SP_VULKAN_ES3_1_ANDROID:
