Fixed bug when creating DepthStencil using the ImmediateContext which was unnecessary for surfaces without data.

diff --git a/Engine/Source/Runtime/Apple/MetalRHI/Private/MetalStateCache.cpp b/Engine/Source/Runtime/Apple/MetalRHI/Private/MetalStateCache.cpp
--- a/Engine/Source/Runtime/Apple/MetalRHI/Private/MetalStateCache.cpp
+++ b/Engine/Source/Runtime/Apple/MetalRHI/Private/MetalStateCache.cpp
@@ -2171,11 +2171,14 @@ FTexture2DRHIRef FMetalStateCache::CreateFallbackDepthStencilSurface(uint32 Widt
 	if (!IsValidRef(FallbackDepthStencilSurface) || FallbackDepthStencilSurface->GetSizeX() != Width || FallbackDepthStencilSurface->GetSizeY() != Height)
 #endif
 	{
-		const FRHITextureCreateDesc Desc =
+		FRHITextureCreateDesc Desc =
 			FRHITextureCreateDesc::Create2D(TEXT("FallbackDepthStencilSurface"), Width, Height, PF_DepthStencil)
-			.SetFlags(ETextureCreateFlags::DepthStencilTargetable);
+            .SetFlags(ETextureCreateFlags::DepthStencilTargetable);
 
-		FallbackDepthStencilSurface = RHICreateTexture(Desc);
+        Desc.SetInitialState(RHIGetDefaultResourceState(Desc.Flags, false));
+        
+        FMetalSurface* Surface = new FMetalSurface(nullptr, Desc);
+        FallbackDepthStencilSurface = Surface;
 	}
 	check(IsValidRef(FallbackDepthStencilSurface));
 	return FallbackDepthStencilSurface;
