- Forcing the logic of -nometalheap and -nometalfence for IOS, but allowing a cvar to re-enable them (rhi.Metal.ForceNoHeap, rhi.Metal.ForceNoFence - they default to 1)

diff --git a/Engine/Source/Runtime/Apple/MetalRHI/Private/MetalCommandQueue.cpp b/Engine/Source/Runtime/Apple/MetalRHI/Private/MetalCommandQueue.cpp
--- a/Engine/Source/Runtime/Apple/MetalRHI/Private/MetalCommandQueue.cpp
+++ b/Engine/Source/Runtime/Apple/MetalRHI/Private/MetalCommandQueue.cpp
@@ -19,6 +19,23 @@ uint64 FMetalCommandQueue::Features = 0;
 extern mtlpp::VertexFormat GMetalFColorVertexFormat;
 bool GMetalCommandBufferDebuggingEnabled = 0;
 
+
+static int32 GForceNoMetalHeap = 1;
+static FAutoConsoleVariableRef CVarMetalForceNoHeap(
+	TEXT("rhi.Metal.ForceNoHeap"),
+	GForceNoMetalHeap,
+	TEXT("[IOS] When enabled, act as if -nometalheap was on the commandline\n")
+	TEXT("(On by default (1))"));
+
+static int32 GForceNoMetalFence = 1;
+static FAutoConsoleVariableRef CVarMetalForceNoFence(
+	TEXT("rhi.Metal.ForceNoFence"),
+	GForceNoMetalFence,
+	TEXT("[IOS] When enabled, act as if -nometalfence was on the commandline\n")
+	TEXT("(On by default (1))"));
+
+
+
 #pragma mark - Public C++ Boilerplate -
 
 FMetalCommandQueue::FMetalCommandQueue(mtlpp::Device InDevice, uint32 const MaxNumCommandBuffers /* = 0 */)
@@ -116,12 +133,12 @@ FMetalCommandQueue::FMetalCommandQueue(mtlpp::Device InDevice, uint32 const MaxN
 	Features |= EMetalFeaturesPresentMinDuration | EMetalFeaturesGPUCaptureManager | EMetalFeaturesBufferSubAllocation | EMetalFeaturesParallelRenderEncoders | EMetalFeaturesPipelineBufferMutability;
 
 	Features |= EMetalFeaturesMaxThreadsPerThreadgroup;
-	if (!FParse::Param(FCommandLine::Get(), TEXT("nometalfence")))
+	if (!GForceNoMetalFence && !FParse::Param(FCommandLine::Get(), TEXT("nometalfence")))
 	{
 		Features |= EMetalFeaturesFences;
 	}
                     
-    if (!FParse::Param(FCommandLine::Get(),TEXT("nometalheap")))
+    if (!GForceNoMetalHeap && !FParse::Param(FCommandLine::Get(),TEXT("nometalheap")))
 	{
 		Features |= EMetalFeaturesHeaps;
 	}
