MetalRHI: Fixed crash on 3D texture readback

diff --git a/Engine/Source/Runtime/Apple/MetalRHI/Private/MetalTexture.cpp b/Engine/Source/Runtime/Apple/MetalRHI/Private/MetalTexture.cpp
--- a/Engine/Source/Runtime/Apple/MetalRHI/Private/MetalTexture.cpp
+++ b/Engine/Source/Runtime/Apple/MetalRHI/Private/MetalTexture.cpp
@@ -587,6 +587,8 @@ FMetalSurface::FMetalSurface(FRHICommandListBase* RHICmdList, FMetalTextureCreat
 		}
         else if (bTextureArrayWithAtomics)
         {
+            checkf(GMaxRHIFeatureLevel == ERHIFeatureLevel::SM6, TEXT("Requested texture array with atomics that is unsupported on this platform"));
+            
             mtlpp::Device Device = GetMetalDeviceContext().GetDevice();
 
             const uint32 MinimumByteAlignment = Device.GetMinimumLinearTextureAlignmentForPixelFormat(CreateDesc.MTLFormat);
@@ -1080,8 +1082,13 @@ void* FMetalSurface::Lock(uint32 MipIndex, uint32 ArrayIndex, EResourceLockMode
 					BytesPerRow = 0;
 					MipBytes = 0;
 				}
-				MTLPP_VALIDATE(mtlpp::Texture, Texture, SafeGetRuntimeDebuggingLevel() >= EMetalDebugLevelValidation, GetBytes(MTLPP_VALIDATE(mtlpp::Buffer, SourceData, SafeGetRuntimeDebuggingLevel() >= EMetalDebugLevelValidation, GetContents()), BytesPerRow, MipBytes, Region, MipIndex, ArrayIndex));
-			}
+                
+                uint32_t BytesPerImage = MipBytes;
+                BytesPerImage = MipBytes / Region.size.depth;
+                                
+                void * Contents = MTLPP_VALIDATE(mtlpp::Buffer, SourceData, SafeGetRuntimeDebuggingLevel() >= EMetalDebugLevelValidation, GetContents());
+                MTLPP_VALIDATE(mtlpp::Texture, Texture, SafeGetRuntimeDebuggingLevel() >= EMetalDebugLevelValidation, GetBytes(Contents, BytesPerRow, BytesPerImage, Region, MipIndex, ArrayIndex));
+                }
 			
 #if PLATFORM_MAC
 			// Pack RGBA8_sRGB into R8_sRGB for non Apple Silicon Mac.
