Fixed crash when resizing viewport on Mac

diff --git a/Engine/Source/Runtime/Apple/MetalRHI/Private/MetalViewport.cpp b/Engine/Source/Runtime/Apple/MetalRHI/Private/MetalViewport.cpp
--- a/Engine/Source/Runtime/Apple/MetalRHI/Private/MetalViewport.cpp
+++ b/Engine/Source/Runtime/Apple/MetalRHI/Private/MetalViewport.cpp
@@ -186,6 +186,12 @@ void FMetalViewport::Resize(uint32 InSizeX, uint32 InSizeY, bool bInIsFullscreen
 	
 	mtlpp::PixelFormat MetalFormat = (mtlpp::PixelFormat)GPixelFormats[Format].PlatformFormat;
 	
+    ENQUEUE_RENDER_COMMAND(FlushPendingRHICommands)(
+        [Viewport = this](FRHICommandListImmediate& RHICmdList)
+        {
+            GRHICommandList.GetImmediateCommandList().BlockUntilGPUIdle();
+        });
+    
 	if (IsValidRef(BackBuffer[Index]) && Format != BackBuffer[Index]->GetFormat())
 	{
 		// Really need to flush the RHI thread & GPU here...
@@ -193,17 +199,16 @@ void FMetalViewport::Resize(uint32 InSizeX, uint32 InSizeY, bool bInIsFullscreen
 		ENQUEUE_RENDER_COMMAND(FlushPendingRHICommands)(
 			[Viewport = this](FRHICommandListImmediate& RHICmdList)
 			{
-				GRHICommandList.GetImmediateCommandList().BlockUntilGPUIdle();
 				Viewport->ReleaseDrawable();
 				Viewport->Release();
 			});
-		
-		// Issue a fence command to the rendering thread and wait for it to complete.
-		FRenderCommandFence Fence;
-		Fence.BeginFence();	
-		Fence.Wait();
 	}
     
+    // Issue a fence command to the rendering thread and wait for it to complete.
+    FRenderCommandFence Fence;
+    Fence.BeginFence();
+    Fence.Wait();
+    
 #if PLATFORM_MAC
 	MainThreadCall(^
 	{
@@ -307,11 +312,7 @@ id<CAMetalDrawable> FMetalViewport::GetDrawable(EMetalViewportAccessFlag Accesso
 #endif
 	
 	SCOPE_CYCLE_COUNTER(STAT_MetalMakeDrawableTime);
-    if (!Drawable
-#if !PLATFORM_MAC
-        || (Drawable.texture.width != BackBuffer[GetViewportIndex(Accessor)]->GetSizeX() || Drawable.texture.height != BackBuffer[GetViewportIndex(Accessor)]->GetSizeY())
-#endif
-        )
+    if (!Drawable || (Drawable.texture.width != BackBuffer[GetViewportIndex(Accessor)]->GetSizeX() || Drawable.texture.height != BackBuffer[GetViewportIndex(Accessor)]->GetSizeY()))
 	{
 		// Drawable changed, release the previously retained object.
 		if (Drawable != nil)
