MetalViewport: Fix bad indexing into DrawableTextures
-EMetalViewportAccessFlag needs to be translated into an index, it can't just be used directly.

diff --git a/Engine/Source/Runtime/Apple/MetalRHI/Private/MetalViewport.cpp b/Engine/Source/Runtime/Apple/MetalRHI/Private/MetalViewport.cpp
--- a/Engine/Source/Runtime/Apple/MetalRHI/Private/MetalViewport.cpp
+++ b/Engine/Source/Runtime/Apple/MetalRHI/Private/MetalViewport.cpp
@@ -407,8 +407,9 @@ FMetalTexture FMetalViewport::GetDrawableTexture(EMetalViewportAccessFlag Access
 		}
 		
 		// get the color texture out and use that with the RHI
-		DrawableTextures[Accessor] = cp_drawable_get_color_texture(SwiftDrawable, 0);
-		return DrawableTextures[Accessor];
+		uint32 Index = GetViewportIndex(Accessor);
+		DrawableTextures[Index] = cp_drawable_get_color_texture(SwiftDrawable, 0);
+		return DrawableTextures[Index];
 	}
 #endif
 	
@@ -430,13 +431,15 @@ FMetalTexture FMetalViewport::GetDrawableTexture(EMetalViewportAccessFlag Access
 		}
 	}
 #endif
-	DrawableTextures[Accessor] = CurrentDrawable.texture;
+	uint32 Index = GetViewportIndex(Accessor);
+	DrawableTextures[Index] = CurrentDrawable.texture;
 	return CurrentDrawable.texture;
 }
 
 ns::AutoReleased<FMetalTexture> FMetalViewport::GetCurrentTexture(EMetalViewportAccessFlag Accessor)
 {
-	return DrawableTextures[Accessor];
+	uint32 Index = GetViewportIndex(Accessor);
+	return DrawableTextures[Index];
 }
 
 void FMetalViewport::ReleaseDrawable()
