Deprecate DrawIndexedIndirect, since it's just a more confusing version of DrawIndexedPrimitiveIndirect.

Also, while we still have it, make sure it behaves consistently across all RHIs. D3D11 used to assert that NumInstances is above 1, even though it's not using it, and a few other RHIs inherited that behavior. Vulkan was actually using NumInstances as the number of indirect draws to perform, which would have crashed if the data was set up as D3D11 wanted it.

diff --git a/Engine/Source/Runtime/Apple/MetalRHI/Private/MetalCommands.cpp b/Engine/Source/Runtime/Apple/MetalRHI/Private/MetalCommands.cpp
--- a/Engine/Source/Runtime/Apple/MetalRHI/Private/MetalCommands.cpp
+++ b/Engine/Source/Runtime/Apple/MetalRHI/Private/MetalCommands.cpp
@@ -532,13 +532,11 @@ void FMetalRHICommandContext::RHIDrawIndexedPrimitive(FRHIBuffer* IndexBufferRHI
 	}
 }
 
-void FMetalRHICommandContext::RHIDrawIndexedIndirect(FRHIBuffer* IndexBufferRHI, FRHIBuffer* ArgumentsBufferRHI, int32 DrawArgumentsIndex, uint32 NumInstances)
+void FMetalRHICommandContext::RHIDrawIndexedIndirect(FRHIBuffer* IndexBufferRHI, FRHIBuffer* ArgumentsBufferRHI, int32 DrawArgumentsIndex, uint32 /*NumInstances*/)
 {
 	@autoreleasepool {
 	if (GetMetalDeviceContext().SupportsFeature(EMetalFeaturesIndirectBuffer))
 	{
-		check(NumInstances > 1);
-		
 		SCOPE_CYCLE_COUNTER(STAT_MetalDrawCallTime);
 		
 		uint32 PrimitiveType = Context->GetCurrentState().GetPrimitiveType();
@@ -548,7 +546,7 @@ void FMetalRHICommandContext::RHIDrawIndexedIndirect(FRHIBuffer* IndexBufferRHI,
 		FMetalRHIBuffer* IndexBuffer = ResourceCast(IndexBufferRHI);
 		FMetalRHIBuffer* ArgumentsBuffer = ResourceCast(ArgumentsBufferRHI);
 		
-		Context->DrawIndexedIndirect(IndexBuffer, PrimitiveType, ArgumentsBuffer, DrawArgumentsIndex, NumInstances);
+		Context->DrawIndexedIndirect(IndexBuffer, PrimitiveType, ArgumentsBuffer, DrawArgumentsIndex);
 	}
 	else
 	{
diff --git a/Engine/Source/Runtime/Apple/MetalRHI/Private/MetalContext.cpp b/Engine/Source/Runtime/Apple/MetalRHI/Private/MetalContext.cpp
--- a/Engine/Source/Runtime/Apple/MetalRHI/Private/MetalContext.cpp
+++ b/Engine/Source/Runtime/Apple/MetalRHI/Private/MetalContext.cpp
@@ -1303,7 +1303,7 @@ void FMetalContext::DrawIndexedPrimitive(FMetalBuffer const& IndexBuffer, uint32
 	RenderPass.DrawIndexedPrimitive(IndexBuffer, IndexStride, PrimitiveType, BaseVertexIndex, FirstInstance, NumVertices, StartIndex, NumPrimitives, NumInstances);
 }
 
-void FMetalContext::DrawIndexedIndirect(FMetalRHIBuffer* IndexBuffer, uint32 PrimitiveType, FMetalRHIBuffer* VertexBuffer, int32 DrawArgumentsIndex, uint32 NumInstances)
+void FMetalContext::DrawIndexedIndirect(FMetalRHIBuffer* IndexBuffer, uint32 PrimitiveType, FMetalRHIBuffer* VertexBuffer, int32 DrawArgumentsIndex)
 {
 	// finalize any pending state
 	if(!PrepareToDraw(PrimitiveType))
@@ -1311,7 +1311,7 @@ void FMetalContext::DrawIndexedIndirect(FMetalRHIBuffer* IndexBuffer, uint32 Pri
 		return;
 	}
 	
-	RenderPass.DrawIndexedIndirect(IndexBuffer, PrimitiveType, VertexBuffer, DrawArgumentsIndex, NumInstances);
+	RenderPass.DrawIndexedIndirect(IndexBuffer, PrimitiveType, VertexBuffer, DrawArgumentsIndex);
 }
 
 void FMetalContext::DrawIndexedPrimitiveIndirect(uint32 PrimitiveType,FMetalRHIBuffer* IndexBuffer,FMetalRHIBuffer* VertexBuffer, uint32 ArgumentOffset)
diff --git a/Engine/Source/Runtime/Apple/MetalRHI/Private/MetalContext.h b/Engine/Source/Runtime/Apple/MetalRHI/Private/MetalContext.h
--- a/Engine/Source/Runtime/Apple/MetalRHI/Private/MetalContext.h
+++ b/Engine/Source/Runtime/Apple/MetalRHI/Private/MetalContext.h
@@ -75,7 +75,7 @@ public:
 	void DrawIndexedPrimitive(FMetalBuffer const& IndexBuffer, uint32 IndexStride, mtlpp::IndexType IndexType, uint32 PrimitiveType, int32 BaseVertexIndex, uint32 FirstInstance,
 							  uint32 NumVertices, uint32 StartIndex, uint32 NumPrimitives, uint32 NumInstances);
 	
-	void DrawIndexedIndirect(FMetalRHIBuffer* IndexBufferRHI, uint32 PrimitiveType, FMetalRHIBuffer* VertexBufferRHI, int32 DrawArgumentsIndex, uint32 NumInstances);
+	void DrawIndexedIndirect(FMetalRHIBuffer* IndexBufferRHI, uint32 PrimitiveType, FMetalRHIBuffer* VertexBufferRHI, int32 DrawArgumentsIndex);
 	
 	void DrawIndexedPrimitiveIndirect(uint32 PrimitiveType,FMetalRHIBuffer* IndexBufferRHI,FMetalRHIBuffer* VertexBufferRHI,uint32 ArgumentOffset);
 	
diff --git a/Engine/Source/Runtime/Apple/MetalRHI/Private/MetalRenderPass.cpp b/Engine/Source/Runtime/Apple/MetalRHI/Private/MetalRenderPass.cpp
--- a/Engine/Source/Runtime/Apple/MetalRHI/Private/MetalRenderPass.cpp
+++ b/Engine/Source/Runtime/Apple/MetalRHI/Private/MetalRenderPass.cpp
@@ -395,12 +395,10 @@ void FMetalRenderPass::DrawIndexedPrimitive(FMetalBuffer const& IndexBuffer, uin
 	ConditionalSubmit();
 }
 
-void FMetalRenderPass::DrawIndexedIndirect(FMetalRHIBuffer* IndexBuffer, uint32 PrimitiveType, FMetalRHIBuffer* VertexBuffer, int32 DrawArgumentsIndex, uint32 NumInstances)
+void FMetalRenderPass::DrawIndexedIndirect(FMetalRHIBuffer* IndexBuffer, uint32 PrimitiveType, FMetalRHIBuffer* VertexBuffer, int32 DrawArgumentsIndex)
 {
 	if (GetMetalDeviceContext().SupportsFeature(EMetalFeaturesIndirectBuffer))
 	{
-		check(NumInstances > 1);
-		
 		ConditionalSwitchToRender();
 		check(CurrentEncoder.GetCommandBuffer());
 		check(CurrentEncoder.IsRenderCommandEncoderActive());
diff --git a/Engine/Source/Runtime/Apple/MetalRHI/Private/MetalRenderPass.h b/Engine/Source/Runtime/Apple/MetalRHI/Private/MetalRenderPass.h
--- a/Engine/Source/Runtime/Apple/MetalRHI/Private/MetalRenderPass.h
+++ b/Engine/Source/Runtime/Apple/MetalRHI/Private/MetalRenderPass.h
@@ -37,7 +37,7 @@ public:
     void DrawIndexedPrimitive(FMetalBuffer const& IndexBuffer, uint32 IndexStride, uint32 PrimitiveType, int32 BaseVertexIndex, uint32 FirstInstance,
                          uint32 NumVertices, uint32 StartIndex, uint32 NumPrimitives, uint32 NumInstances);
     
-    void DrawIndexedIndirect(FMetalRHIBuffer* IndexBufferRHI, uint32 PrimitiveType, FMetalRHIBuffer* VertexBufferRHI, int32 DrawArgumentsIndex, uint32 NumInstances);
+    void DrawIndexedIndirect(FMetalRHIBuffer* IndexBufferRHI, uint32 PrimitiveType, FMetalRHIBuffer* VertexBufferRHI, int32 DrawArgumentsIndex);
     
     void DrawIndexedPrimitiveIndirect(uint32 PrimitiveType,FMetalRHIBuffer* IndexBufferRHI,FMetalRHIBuffer* VertexBufferRHI,uint32 ArgumentOffset);
 	
diff --git a/Engine/Source/Runtime/D3D12RHI/Private/D3D12Commands.cpp b/Engine/Source/Runtime/D3D12RHI/Private/D3D12Commands.cpp
--- a/Engine/Source/Runtime/D3D12RHI/Private/D3D12Commands.cpp
+++ b/Engine/Source/Runtime/D3D12RHI/Private/D3D12Commands.cpp
@@ -1589,7 +1589,7 @@ void FD3D12CommandContext::RHIDrawPrimitiveIndirect(FRHIBuffer* ArgumentBufferRH
 	DEBUG_EXECUTE_COMMAND_LIST(this);
 }
 
-void FD3D12CommandContext::RHIDrawIndexedIndirect(FRHIBuffer* IndexBufferRHI, FRHIBuffer* ArgumentsBufferRHI, int32 DrawArgumentsIndex, uint32 NumInstances)
+void FD3D12CommandContext::RHIDrawIndexedIndirect(FRHIBuffer* IndexBufferRHI, FRHIBuffer* ArgumentsBufferRHI, int32 DrawArgumentsIndex, uint32 /*NumInstances*/)
 {
 	const uint32 IndexBufferStride = FD3D12DynamicRHI::ResourceCast(IndexBufferRHI)->GetStride();
 	const uint32 ArgumentsBufferStride = FD3D12DynamicRHI::ResourceCast(ArgumentsBufferRHI)->GetStride();
diff --git a/Engine/Source/Runtime/OpenGLDrv/Private/OpenGLCommands.cpp b/Engine/Source/Runtime/OpenGLDrv/Private/OpenGLCommands.cpp
--- a/Engine/Source/Runtime/OpenGLDrv/Private/OpenGLCommands.cpp
+++ b/Engine/Source/Runtime/OpenGLDrv/Private/OpenGLCommands.cpp
@@ -2497,7 +2497,7 @@ void FOpenGLDynamicRHI::RHIDrawPrimitiveIndirect(FRHIBuffer* ArgumentBufferRHI,
 	}
 }
 
-void FOpenGLDynamicRHI::RHIDrawIndexedIndirect(FRHIBuffer* IndexBufferRHI, FRHIBuffer* ArgumentsBufferRHI, int32 DrawArgumentsIndex, uint32 NumInstances)
+void FOpenGLDynamicRHI::RHIDrawIndexedIndirect(FRHIBuffer* IndexBufferRHI, FRHIBuffer* ArgumentsBufferRHI, int32 DrawArgumentsIndex, uint32 /*NumInstances*/)
 {
 	if (FOpenGL::SupportsDrawIndirect())
 	{
@@ -2508,9 +2508,6 @@ void FOpenGLDynamicRHI::RHIDrawIndexedIndirect(FRHIBuffer* IndexBufferRHI, FRHIB
 
 		check(ArgumentsBufferRHI);
 
-		//Draw indiect has to have a number of instances
-		check(NumInstances > 1);
-
 		FOpenGLContextState& ContextState = GetContextStateForCurrentContext();
 		BindPendingFramebuffer(ContextState);
 		SetPendingBlendStateForActiveRenderTargets(ContextState);
diff --git a/Engine/Source/Runtime/RHI/Public/RHICommandList.h b/Engine/Source/Runtime/RHI/Public/RHICommandList.h
--- a/Engine/Source/Runtime/RHI/Public/RHICommandList.h
+++ b/Engine/Source/Runtime/RHI/Public/RHICommandList.h
@@ -3713,6 +3713,7 @@ PRAGMA_ENABLE_DEPRECATION_WARNINGS
 		ALLOC_COMMAND(FRHICommandDrawPrimitiveIndirect)(ArgumentBuffer, ArgumentOffset);
 	}
 
+	UE_DEPRECATED(5.4, "Use DrawIndexedPrimitiveIndirect.")
 	FORCEINLINE_DEBUGGABLE void DrawIndexedIndirect(FRHIBuffer* IndexBufferRHI, FRHIBuffer* ArgumentsBufferRHI, uint32 DrawArgumentsIndex, uint32 NumInstances)
 	{
 		//check(IsOutsideRenderPass());
diff --git a/Engine/Source/Runtime/VulkanRHI/Private/VulkanCommands.cpp b/Engine/Source/Runtime/VulkanRHI/Private/VulkanCommands.cpp
--- a/Engine/Source/Runtime/VulkanRHI/Private/VulkanCommands.cpp
+++ b/Engine/Source/Runtime/VulkanRHI/Private/VulkanCommands.cpp
@@ -685,7 +685,7 @@ void FVulkanCommandListContext::RHIDrawIndexedPrimitive(FRHIBuffer* IndexBufferR
 	}
 }
 
-void FVulkanCommandListContext::RHIDrawIndexedIndirect(FRHIBuffer* IndexBufferRHI, FRHIBuffer* ArgumentsBufferRHI, int32 DrawArgumentsIndex, uint32 NumInstances)
+void FVulkanCommandListContext::RHIDrawIndexedIndirect(FRHIBuffer* IndexBufferRHI, FRHIBuffer* ArgumentsBufferRHI, int32 DrawArgumentsIndex, uint32 /*NumInstances*/)
 {
 #if VULKAN_ENABLE_AGGRESSIVE_STATS
 	SCOPE_CYCLE_COUNTER(STAT_VulkanDrawCallTime);
@@ -704,7 +704,7 @@ void FVulkanCommandListContext::RHIDrawIndexedIndirect(FRHIBuffer* IndexBufferRH
 	VkDeviceSize ArgumentOffset = DrawArgumentsIndex * sizeof(VkDrawIndexedIndirectCommand);
 
 
-	VulkanRHI::vkCmdDrawIndexedIndirect(CmdBuffer, ArgumentBuffer->GetHandle(), ArgumentBuffer->GetOffset() + ArgumentOffset, NumInstances, sizeof(VkDrawIndexedIndirectCommand));
+	VulkanRHI::vkCmdDrawIndexedIndirect(CmdBuffer, ArgumentBuffer->GetHandle(), ArgumentBuffer->GetOffset() + ArgumentOffset, 1, sizeof(VkDrawIndexedIndirectCommand));
 
 	if (FVulkanPlatform::RegisterGPUWork() && IsImmediate())
 	{
diff --git a/Engine/Source/Runtime/Windows/D3D11RHI/Private/D3D11Commands.cpp b/Engine/Source/Runtime/Windows/D3D11RHI/Private/D3D11Commands.cpp
--- a/Engine/Source/Runtime/Windows/D3D11RHI/Private/D3D11Commands.cpp
+++ b/Engine/Source/Runtime/Windows/D3D11RHI/Private/D3D11Commands.cpp
@@ -1279,14 +1279,7 @@ void FD3D11DynamicRHI::RHIDrawIndexedIndirect(FRHIBuffer* IndexBufferRHI, FRHIBu
 	StateCache.SetIndexBuffer(IndexBuffer->Resource, Format, 0);
 	StateCache.SetPrimitiveTopology(GetD3D11PrimitiveType(PrimitiveType));
 
-	if(NumInstances > 1)
-	{
-		Direct3DDeviceIMContext->DrawIndexedInstancedIndirect(ArgumentsBuffer->Resource, DrawArgumentsIndex * 5 * sizeof(uint32));
-	}
-	else
-	{
-		check(0);
-	}
+	Direct3DDeviceIMContext->DrawIndexedInstancedIndirect(ArgumentsBuffer->Resource, DrawArgumentsIndex * 5 * sizeof(uint32));
 
 	EnableUAVOverlap();
 }
