OpenXRVisionOS non-plugin support files
-The changes required for VisionOS Immersive mode that are not within the OpenXRVisionOS plugin.
-Enables SimpleHMD for VisionOS (stereo rendering without a hmd).

diff --git a/Engine/Platforms/VisionOS/Plugins/Runtime/OpenXR/OpenXR_VisionOS.uplugin b/Engine/Platforms/VisionOS/Plugins/Runtime/OpenXR/OpenXR_VisionOS.uplugin
new file mode 100644
--- /dev/null
+++ b/Engine/Platforms/VisionOS/Plugins/Runtime/OpenXR/OpenXR_VisionOS.uplugin
@@ -0,0 +1,31 @@
+{
+	"FileVersion": 3,
+	"bIsPluginExtension": true,
+	"SupportedTargetPlatforms": [
+		"VisionOS"
+	],
+	"Modules": [
+		{
+			"Name": "OpenXRHMD",
+			"Type": "Runtime",
+			"PlatformAllowList": [
+				"VisionOS"
+			]
+		},
+		{
+			"Name": "OpenXRInput",
+			"Type": "Runtime",
+			"PlatformAllowList": [
+				"VisionOS"
+			]
+		},
+		{
+			"Name": "OpenXRAR",
+			"Type": "Runtime",
+			"PlatformAllowList": [
+				"VisionOS"
+			]
+		}
+	]
+}
+
diff --git a/Engine/Plugins/Experimental/SimpleHMD/SimpleHMD.uplugin b/Engine/Plugins/Experimental/SimpleHMD/SimpleHMD.uplugin
--- a/Engine/Plugins/Experimental/SimpleHMD/SimpleHMD.uplugin
+++ b/Engine/Plugins/Experimental/SimpleHMD/SimpleHMD.uplugin
@@ -19,7 +19,8 @@
 		"Linux",
 		"Android",
 		"Mac",
-		"IOS"
+		"IOS",
+		"VisionOS"
 	],
 	"Modules" : 
 	[
diff --git a/Engine/Plugins/Runtime/OpenXR/OpenXR.uplugin b/Engine/Plugins/Runtime/OpenXR/OpenXR.uplugin
--- a/Engine/Plugins/Runtime/OpenXR/OpenXR.uplugin
+++ b/Engine/Plugins/Runtime/OpenXR/OpenXR.uplugin
@@ -27,7 +27,8 @@
 			"PlatformAllowList": [
 				"Win64",
 				"Linux",
-				"Android"
+				"Android",
+				"Mac"
 			]
 		},
 		{
@@ -37,7 +38,8 @@
 			"PlatformAllowList": [
 				"Win64",
 				"Linux",
-				"Android"
+				"Android",
+				"Mac"
 			]
 		},
 		{
@@ -47,7 +49,8 @@
 			"PlatformAllowList": [
 				"Win64",
 				"Linux",
-				"Android"
+				"Android",
+				"Mac"
 			]
 		},
 		{
@@ -56,7 +59,8 @@
 			"LoadingPhase": "Default",
 			"PlatformAllowList": [
 				"Win64",
-				"Linux"
+				"Linux",
+				"Mac"
 			]
 		}
 	],
diff --git a/Engine/Plugins/Runtime/OpenXR/Source/OpenXRHMD/Private/OpenXRHMD.cpp b/Engine/Plugins/Runtime/OpenXR/Source/OpenXRHMD/Private/OpenXRHMD.cpp
--- a/Engine/Plugins/Runtime/OpenXR/Source/OpenXRHMD/Private/OpenXRHMD.cpp
+++ b/Engine/Plugins/Runtime/OpenXR/Source/OpenXRHMD/Private/OpenXRHMD.cpp
@@ -2833,6 +2833,7 @@ void FOpenXRHMD::OnBeginRendering_RenderThread(FRHICommandListImmediate& RHICmdL
 	{
 		// Locate the views we will actually be rendering for.
 		// This is required to support late-updating the field-of-view.
+        //Note: This LocateViews happens before xrBeginFrame.  Which I don't think is correct.
 		LocateViews(PipelinedFrameStateRendering, false);
 
 		SCOPED_NAMED_EVENT(EnqueueFrame, FColor::Red);
@@ -2941,6 +2942,14 @@ void FOpenXRHMD::OnBeginRendering_GameThread()
 	// can wait for the next frame in the next tick. Without this signal it's possible that two ticks
 	// happen before the next frame is actually rendered.
 	bShouldWait = true;
+    
+    if (bIsReady && bIsRunning)
+    {
+        for (IOpenXRExtensionPlugin* Module : ExtensionPlugins)
+        {
+            Module->OnBeginRendering_GameThread(Session);
+        }
+    }
 
 	ENQUEUE_RENDER_COMMAND(TransferFrameStateToRenderingThread)(
 		[this, GameFrameState = PipelinedFrameStateGame, bBackgroundLayerVisible = IsBackgroundLayerVisible()](FRHICommandListImmediate& RHICmdList) mutable
diff --git a/Engine/Plugins/Runtime/OpenXR/Source/OpenXRHMD/Public/IOpenXRExtensionPlugin.h b/Engine/Plugins/Runtime/OpenXR/Source/OpenXRHMD/Public/IOpenXRExtensionPlugin.h
--- a/Engine/Plugins/Runtime/OpenXR/Source/OpenXRHMD/Public/IOpenXRExtensionPlugin.h
+++ b/Engine/Plugins/Runtime/OpenXR/Source/OpenXRHMD/Public/IOpenXRExtensionPlugin.h
@@ -344,12 +344,17 @@ public:
 	{
 	}
 
-	// OpenXRHMD::OnBeginRendering_GameThread
+	// OpenXRHMD::OnBeginSimulation_GameThread
 	virtual void* OnWaitFrame(XrSession InSession, void* InNext)
 	{
 		return InNext;
 	}
 
+	// OpenXRHMD::OnBeginRendering_GameThread
+	virtual void OnBeginRendering_GameThread(XrSession InSession)
+	{
+	}
+
 	// OpenXRHMD::OnBeginRendering_RenderThread
 	virtual void OnBeginRendering_RenderThread(XrSession InSession)
 	{
diff --git a/Engine/Source/Runtime/Apple/MetalRHI/Private/MetalContext.cpp b/Engine/Source/Runtime/Apple/MetalRHI/Private/MetalContext.cpp
--- a/Engine/Source/Runtime/Apple/MetalRHI/Private/MetalContext.cpp
+++ b/Engine/Source/Runtime/Apple/MetalRHI/Private/MetalContext.cpp
@@ -384,7 +384,8 @@ FMetalDeviceContext::FMetalDeviceContext(mtlpp::Device MetalDevice, uint32 InDev
 		}
 	}
 	
-	if (FParse::Param(FCommandLine::Get(), TEXT("MetalIntermediateBackBuffer")) || FParse::Param(FCommandLine::Get(), TEXT("MetalOffscreenOnly")))
+    const bool bIsVisionOS = PLATFORM_VISIONOS;
+	if (bIsVisionOS || FParse::Param(FCommandLine::Get(), TEXT("MetalIntermediateBackBuffer")) || FParse::Param(FCommandLine::Get(), TEXT("MetalOffscreenOnly")))
 	{
 		GMetalSupportsIntermediateBackBuffer = 1;
 	}
@@ -662,15 +663,19 @@ void FMetalDeviceContext::EndDrawingViewport(FMetalViewport* Viewport, bool bPre
 	if (bPresent && !bOffscreenOnly)
 	{
 		
-#if PLATFORM_MAC
+#if PLATFORM_MAC || PLATFORM_VISIONOS
 		// Handle custom present
+        bool bNeedNativePresent = true;
 		FRHICustomPresent* const CustomPresent = Viewport->GetCustomPresent();
 		if (CustomPresent != nullptr)
 		{
 			int32 SyncInterval = 0;
 			{
 				SCOPE_CYCLE_COUNTER(STAT_MetalCustomPresentTime);
-				CustomPresent->Present(SyncInterval);
+                FMetalRHICommandContext* RHICommandContext = static_cast<FMetalRHICommandContext*>(RHIGetDefaultContext());
+                RHICommandContext->SetCustomPresentViewport(Viewport);
+                bNeedNativePresent = CustomPresent->Present(SyncInterval);
+                RHICommandContext->SetCustomPresentViewport(nullptr);
 			}
 			
 			mtlpp::CommandBuffer CurrentCommandBuffer = GetCurrentCommandBuffer();
@@ -683,10 +688,12 @@ void FMetalDeviceContext::EndDrawingViewport(FMetalViewport* Viewport, bool bPre
 #endif
 		
 		RenderPass.End();
-		
-		SubmitCommandsHint(EMetalSubmitFlagsForce|EMetalSubmitFlagsCreateCommandBuffer);
-		
-		Viewport->Present(GetCommandQueue(), bLockToVsync);
+		if (bNeedNativePresent)
+        {
+            SubmitCommandsHint(EMetalSubmitFlagsForce|EMetalSubmitFlagsCreateCommandBuffer);
+            
+            Viewport->Present(GetCommandQueue(), bLockToVsync);
+        }
 	}
 	
 	bPresented = bPresent;
diff --git a/Engine/Source/Runtime/Apple/MetalRHI/Private/MetalRHIContext.cpp b/Engine/Source/Runtime/Apple/MetalRHI/Private/MetalRHIContext.cpp
--- a/Engine/Source/Runtime/Apple/MetalRHI/Private/MetalRHIContext.cpp
+++ b/Engine/Source/Runtime/Apple/MetalRHI/Private/MetalRHIContext.cpp
@@ -3,6 +3,7 @@
 #include "MetalRHIPrivate.h"
 #include "MetalRHIRenderQuery.h"
 #include "MetalCommandBufferFence.h"
+#include "MetalRHIVisionOSBridge.h"
 
 TGlobalResource<TBoundShaderStateHistory<10000>> FMetalRHICommandContext::BoundShaderStateHistory;
 
@@ -216,3 +217,11 @@ void FMetalRHICommandContext::RHIEndOcclusionQueryBatch()
 	Context->InsertCommandBufferFence(*CommandBufferFence);
 	CommandBufferFence.Reset();
 }
+
+#if PLATFORM_VISIONOS
+void FMetalRHICommandContext::BeginRenderingImmersive(const MetalRHIVisionOS::BeginRenderingImmersiveParams& Params)
+{
+    SwiftFrame = Params.SwiftFrame;
+    UE_LOG(LogMetalVisionOS, Verbose, TEXT("SwiftLayerFrame(0x%x) copying from Params in FMetalRHICommandContext::BeginRenderingImmersive"), SwiftFrame);
+}
+#endif
diff --git a/Engine/Source/Runtime/Apple/MetalRHI/Private/MetalRHIVisionOSBridge.cpp b/Engine/Source/Runtime/Apple/MetalRHI/Private/MetalRHIVisionOSBridge.cpp
new file mode 100644
--- /dev/null
+++ b/Engine/Source/Runtime/Apple/MetalRHI/Private/MetalRHIVisionOSBridge.cpp
@@ -0,0 +1,29 @@
+#include "MetalRHIVisionOSBridge.h"
+
+#if PLATFORM_VISIONOS
+
+#include "MetalRHIPrivate.h"
+#include "MetalDynamicRHI.h"
+
+DEFINE_LOG_CATEGORY(LogMetalVisionOS);
+
+void MetalRHIVisionOS::BeginRenderingImmersive(const MetalRHIVisionOS::BeginRenderingImmersiveParams& Params)
+{
+    UE_LOG(LogMetalVisionOS, Verbose, TEXT("SwiftLayerFrame(0x%x) sending to RHICommandContext in MetalRHIVisionOS::BeginRendering"), Params.SwiftFrame);
+    FMetalRHICommandContext* RHICommandContext = static_cast<FMetalRHICommandContext*>(RHIGetDefaultContext());
+    check(RHICommandContext);
+    RHICommandContext->BeginRenderingImmersive(Params);
+}
+
+void MetalRHIVisionOS::PresentImmersive(const MetalRHIVisionOS::PresentImmersiveParams& Params)
+{
+    FMetalRHICommandContext* RHICommandContext = static_cast<FMetalRHICommandContext*>(RHIGetDefaultContext());
+    check(RHICommandContext);
+    check(RHICommandContext->CustomPresentViewport);
+    FMetalViewport* Viewport = ResourceCast(RHICommandContext->CustomPresentViewport);
+    check(Viewport);
+    FMetalSurface* Surface = GetMetalSurfaceFromRHITexture(Params.Texture);
+
+    Viewport->PresentImmersive(Surface, &Params);
+}
+#endif // PLATFORM_VISIONOS
diff --git a/Engine/Source/Runtime/Apple/MetalRHI/Private/MetalStateCache.cpp b/Engine/Source/Runtime/Apple/MetalRHI/Private/MetalStateCache.cpp
--- a/Engine/Source/Runtime/Apple/MetalRHI/Private/MetalStateCache.cpp
+++ b/Engine/Source/Runtime/Apple/MetalRHI/Private/MetalStateCache.cpp
@@ -406,7 +406,11 @@ void FMetalStateCache::SetRasterizerState(FMetalRasterizerState* InRasterizerSta
 	if(RasterizerState != InRasterizerState)
 	{
 		RasterizerState = InRasterizerState;
+#if PLATFORM_VISIONOS
+		RasterBits |= EMetalRenderFlagFrontFacingWinding|EMetalRenderFlagCullMode|EMetalRenderFlagDepthBias|EMetalRenderFlagTriangleFillMode;
+#else
 		RasterBits |= EMetalRenderFlagFrontFacingWinding|EMetalRenderFlagCullMode|EMetalRenderFlagDepthBias|EMetalRenderFlagTriangleFillMode|EMetalRenderFlagDepthClipMode;
+#endif
 	}
 }
 
diff --git a/Engine/Source/Runtime/Apple/MetalRHI/Private/MetalViewport.cpp b/Engine/Source/Runtime/Apple/MetalRHI/Private/MetalViewport.cpp
--- a/Engine/Source/Runtime/Apple/MetalRHI/Private/MetalViewport.cpp
+++ b/Engine/Source/Runtime/Apple/MetalRHI/Private/MetalViewport.cpp
@@ -15,6 +15,7 @@
 #include "Containers/Set.h"
 #include "MetalProfiler.h"
 #include "RenderUtils.h"
+#include "MetalRHIVisionOSBridge.h"
 
 extern int32 GMetalSupportsIntermediateBackBuffer;
 extern int32 GMetalSeparatePresentThread;
@@ -78,13 +79,18 @@ FMetalViewport::FMetalViewport(void* WindowHandle, uint32 InSizeX, uint32 InSize
 	, bIsFullScreen{bInIsFullscreen}
 #if PLATFORM_MAC
 	, View{nullptr}
+#endif
+#if PLATFORM_MAC || PLATFORM_VISIONOS
 	, CustomPresent{nullptr}
 #endif
 {
 #if PLATFORM_VISIONOS
 	// look to see if we need to hook up to a Swift compositor renderer
 	SwiftLayer = [IOSAppDelegate GetDelegate].SwiftLayer;
-	NewFrame();
+	
+	// Get an initial Frame
+	SwiftLayerFrame = cp_layer_renderer_query_next_frame(SwiftLayer);
+	UE_LOG(LogMetalVisionOS, Verbose, TEXT("SwiftLayerFrame(0x%x) from cp_layer_renderer_query_next_frame in MetalViewport constructor"), SwiftLayerFrame);
 #endif
 
 #if PLATFORM_MAC
@@ -377,29 +383,13 @@ id<CAMetalDrawable> FMetalViewport::GetDrawable(EMetalViewportAccessFlag Accesso
 	return Drawable;
 }
 
-void FMetalViewport::NewFrame()
-{
-#if PLATFORM_VISIONOS
-	if (SwiftLayer)
-	{
-		// end the last frame if there was one
-		if (SwiftLayerFrame != nullptr)
-		{
-			cp_frame_end_submission(SwiftLayerFrame);
-		}
-		
-		// get next frame
-		SwiftLayerFrame = cp_layer_renderer_query_next_frame(SwiftLayer);
-	}
-#endif
-}
-
 FMetalTexture FMetalViewport::GetDrawableTexture(EMetalViewportAccessFlag Accessor)
 {
 #if PLATFORM_VISIONOS
 	if (SwiftLayerFrame != nullptr)
 	{
 		// get the next drawable
+		UE_LOG(LogMetalVisionOS, Verbose, TEXT("SwiftLayerFrame(0x%x) cp_frame_query_drawable in FMetalViewport::GetDrawableTexture"), SwiftLayerFrame);
 		cp_drawable_t SwiftDrawable = cp_frame_query_drawable(SwiftLayerFrame);
 		if (SwiftDrawable == nullptr)
 		{
@@ -509,17 +499,7 @@ void FMetalViewport::Present(FMetalCommandQueue& CommandQueue, bool bLockToVsync
 				FPlatformAtomics::InterlockedDecrement(&FrameAvailable);
 				id<CAMetalDrawable> LocalDrawable = [GetDrawable(EMetalViewportAccessDisplayLink) retain];
 				FMetalTexture DrawableTexture;
-#if PLATFORM_VISIONOS
-				// it seems like this should be usable in all cases, instead of LocalDrawable.texture, but it's causing View to be overwritten with the drawable, somehow!
-				if (SwiftLayer != nullptr)
-				{
-					DrawableTexture = GetDrawableTexture(EMetalViewportAccessDisplayLink);
-				}
-				else
-#endif
-				{
-					DrawableTexture = LocalDrawable.texture;
-				}
+				DrawableTexture = GetDrawableTexture(EMetalViewportAccessDisplayLink);
 				
 				{
 					FScopeLock BlockLock(&Mutex);
@@ -566,7 +546,8 @@ void FMetalViewport::Present(FMetalCommandQueue& CommandQueue, bool bLockToVsync
 							Encoder.EndEncoding();
 							METAL_DEBUG_LAYER(EMetalDebugLevelFastValidation, Debugging.EndEncoder());
 
-							mtlpp::CommandBufferHandler H = [Src, Dst](const mtlpp::CommandBuffer &) {
+							mtlpp::CommandBufferHandler H = [Src, Dst](const mtlpp::CommandBuffer &) 
+							{
 								// void
 							};
 
@@ -585,10 +566,11 @@ void FMetalViewport::Present(FMetalCommandQueue& CommandQueue, bool bLockToVsync
       
 #if PLATFORM_MAC
 						FMetalView* theView = View;
-						mtlpp::CommandBufferHandler C = [LocalDrawable, theView](const mtlpp::CommandBuffer & cmd_buf) {
+						mtlpp::CommandBufferHandler C = [LocalDrawable, theView](const mtlpp::CommandBuffer & cmd_buf) 
 #else
-						mtlpp::CommandBufferHandler C = [LocalDrawable](const mtlpp::CommandBuffer & cmd_buf) {
+						mtlpp::CommandBufferHandler C = [LocalDrawable](const mtlpp::CommandBuffer & cmd_buf) 
 #endif
+						{
 							FMetalGPUProfiler::RecordPresent(cmd_buf);
 							[LocalDrawable release];
 #if PLATFORM_MAC
@@ -668,9 +650,148 @@ void FMetalViewport::Swap()
 		BackBuffer[1] = BB0;
 	}
 	
-	// move on to next frame (only used in SwiftUI/compositor mode)
-	NewFrame();
+#if PLATFORM_VISIONOS
+	if (SwiftLayer && SwiftLayerFrame)
+	{
+		EndFrameImmersive();
+	}
+#endif
+}
+
+#if PLATFORM_VISIONOS
+void FMetalViewport::BeginRenderingImmersive(cp_frame_t SwiftFrame)
+{
+	check(SwiftFrame);
+	SwiftLayerFrame = SwiftFrame;
+	UE_LOG(LogMetalVisionOS, Verbose, TEXT("SwiftLayerFrame(0x%x) set in FMetalViewport::BeginRenderingImmersive"), SwiftLayerFrame);
+}
+
+// This is the present for Immersive visionos, through the OXRVisionOS plugin.
+void FMetalViewport::PresentImmersive(TRefCountPtr<FMetalSurface> CompleteFrame, const MetalRHIVisionOS::PresentImmersiveParams* VisionOSParams)
+{
+	check(SwiftLayer);  // If no SwiftLayer we should not be trying to be immersive.
+	check(SwiftLayerFrame);  // BeginRenderingImmersive should have set this.
+
+	FMetalRHICommandContext* Context = static_cast<FMetalRHICommandContext*>(RHIGetDefaultContext());
+	FMetalDeviceContext& DeviceContext = GetMetalDeviceContext();
+	FMetalCommandQueue& CommandQueue = DeviceContext.GetCommandQueue();
+	
+	FScopeLock Lock(&Mutex);
+
+	LastCompleteFrame = CompleteFrame;
+
+	FPlatformAtomics::InterlockedExchange(&FrameAvailable, 1);
+	
+	{
+		if (FrameAvailable > 0 )
+		{
+			FPlatformAtomics::InterlockedDecrement(&FrameAvailable);
+			id<CAMetalDrawable> LocalDrawable = [GetDrawable(EMetalViewportAccessDisplayLink) retain];
+			FMetalTexture DrawableTexture = GetDrawableTexture(EMetalViewportAccessDisplayLink);
+
+			{
+				FScopeLock BlockLock(&Mutex);
+				
+				if (DrawableTexture)
+				{
+					mtlpp::CommandBuffer CurrentCommandBuffer = CommandQueue.CreateCommandBuffer();
+					check(CurrentCommandBuffer);
+					
+#if ENABLE_METAL_GPUPROFILE
+					FMetalProfiler* Profiler = FMetalProfiler::GetProfiler();
+					FMetalCommandBufferStats* Stats = Profiler->AllocateCommandBuffer(CurrentCommandBuffer, 0);
+#endif
+
+					// TODO Currently we are using intermediate back buffer to connect the OXRVisionOS Swapchain to the drawable.
+					// I think we could use the drawable directly and avoid this copy.
+					check(GMetalSupportsIntermediateBackBuffer);
+					if (GMetalSupportsIntermediateBackBuffer)
+					{
+						TRefCountPtr<FMetalSurface> Texture = LastCompleteFrame;
+						check(IsValidRef(Texture));
+						
+						FMetalTexture Src = Texture->Texture;
+						FMetalTexture Dst = DrawableTexture;
+						
+						NSUInteger Width = FMath::Min(Src.GetWidth(), Dst.GetWidth());
+						NSUInteger Height = FMath::Min(Src.GetHeight(), Dst.GetHeight());
+						
+						mtlpp::BlitCommandEncoder Encoder = CurrentCommandBuffer.BlitCommandEncoder();
+						check(Encoder.GetPtr());
+#if MTLPP_CONFIG_VALIDATE && METAL_DEBUG_OPTIONS
+						FMetalBlitCommandEncoderDebugging Debugging;
+						if (SafeGetRuntimeDebuggingLevel() >= EMetalDebugLevelFastValidation)
+						{
+							FMetalCommandBufferDebugging CmdDebug = FMetalCommandBufferDebugging::Get(CurrentCommandBuffer);
+							Debugging = FMetalBlitCommandEncoderDebugging(Encoder, CmdDebug);
+						}
+#endif
+						METAL_GPUPROFILE(Profiler->EncodeBlit(Stats, __FUNCTION__));
+
+						Encoder.Copy(Src, 0, 0, mtlpp::Origin(0, 0, 0), mtlpp::Size(Width, Height, 1), Dst, 0, 0, mtlpp::Origin(0, 0, 0));
+						METAL_DEBUG_LAYER(EMetalDebugLevelFastValidation, Debugging.Copy(Src, 0, 0, mtlpp::Origin(0, 0, 0), mtlpp::Size(Width, Height, 1), Dst, 0, 0, mtlpp::Origin(0, 0, 0)));
+
+						Encoder.EndEncoding();
+						METAL_DEBUG_LAYER(EMetalDebugLevelFastValidation, Debugging.EndEncoder());
+
+						mtlpp::CommandBufferHandler H = [Src, Dst](const mtlpp::CommandBuffer &) 
+						{
+							// void
+						};
+
+						CurrentCommandBuffer.AddCompletedHandler(H);
+
+						[Drawable release];
+						Drawable = nil;
+					}
+					
+					// This is a bit different than the usual pattern.
+					// This command buffer here is committed directly, instead of going through
+					// FMetalCommandList::Commit. So long as Present() is called within
+					// high level RHI BeginFrame/EndFrame this will not fine.
+					// Otherwise the recording of the Present time will be offset by one in the
+					// FMetalGPUProfiler frame indices.
+  
+					mtlpp::CommandBufferHandler C = [LocalDrawable](const mtlpp::CommandBuffer & cmd_buf) 
+					{
+						FMetalGPUProfiler::RecordPresent(cmd_buf);
+						[LocalDrawable release];
+					};
+					
+					CurrentCommandBuffer.AddCompletedHandler(C);
+
+					if (SwiftLayerFrame != nullptr)
+					{
+						if (VisionOSParams)
+						{
+							UE_LOG(LogMetalVisionOS, Verbose, TEXT("SwiftDrawable(0x%x) cp_drawable_encode_present in FMetalViewport::PresentImmersive"), VisionOSParams->SwiftDrawable);
+							cp_drawable_encode_present(VisionOSParams->SwiftDrawable, CurrentCommandBuffer.GetPtr());
+						}
+					}
+
+					METAL_GPUPROFILE(Stats->End(CurrentCommandBuffer));
+					CommandQueue.CommitCommandBuffer(CurrentCommandBuffer);
+				}
+			}
+		}
+	}
+	
+	EndFrameImmersive();
+}
+
+void FMetalViewport::EndFrameImmersive()
+{
+	check(SwiftLayer);
+	check(SwiftLayerFrame);
+	
+	// end the last frame
+	UE_LOG(LogMetalVisionOS, Verbose, TEXT("SwiftLayerFrame(0x%x) cp_frame_end_submission in FMetalViewport::EndFrameImmersive"), SwiftLayerFrame);
+	cp_frame_end_submission(SwiftLayerFrame);
+	UE_LOG(LogMetalVisionOS, Verbose, TEXT("SwiftLayerFrame nulled in FMetalViewport::EndFrame()"));
+	SwiftLayerFrame = nullptr;
 }
+#endif
+
 
 /*=============================================================================
  *	The following RHI functions must be called from the main thread.
@@ -717,7 +838,14 @@ void FMetalRHIImmediateCommandContext::RHIBeginDrawingViewport(FRHIViewport* Vie
 	@autoreleasepool {
 	FMetalViewport* Viewport = ResourceCast(ViewportRHI);
 	check(Viewport);
-
+	 
+#if PLATFORM_VISIONOS
+		if (SwiftFrame)
+		{
+			Viewport->BeginRenderingImmersive(SwiftFrame);
+		}
+#endif
+		
 	((FMetalDeviceContext*)Context)->BeginDrawingViewport(Viewport);
 
 	// Set the render target and viewport.
diff --git a/Engine/Source/Runtime/Apple/MetalRHI/Public/MetalRHIContext.h b/Engine/Source/Runtime/Apple/MetalRHI/Public/MetalRHIContext.h
--- a/Engine/Source/Runtime/Apple/MetalRHI/Public/MetalRHIContext.h
+++ b/Engine/Source/Runtime/Apple/MetalRHI/Public/MetalRHIContext.h
@@ -13,6 +13,14 @@
 class FMetalDeviceContext;
 struct FMetalCommandBufferFence;
 
+#if PLATFORM_VISIONOS
+namespace MetalRHIVisionOS
+{
+    struct BeginRenderingImmersiveParams;
+    struct PresentImmersiveParams;
+}
+#endif
+
 /** The interface RHI command context. */
 class FMetalRHICommandContext : public IRHICommandContext
 {
@@ -195,6 +203,13 @@ public:
 		ERayTracingBindingType BindingType) final override;
 #endif // METAL_RHI_RAYTRACING
 
+#if PLATFORM_VISIONOS
+    void BeginRenderingImmersive(const MetalRHIVisionOS::BeginRenderingImmersiveParams& Params);
+    cp_frame_t SwiftFrame = nullptr;
+#endif // PLATFORM_VISIONOS
+    void SetCustomPresentViewport(FRHIViewport* Viewport) { CustomPresentViewport = Viewport; }
+    FRHIViewport* CustomPresentViewport = nullptr;
+
 	void BeginRecursiveCommand()
 	{
 		// Nothing to do
diff --git a/Engine/Source/Runtime/Apple/MetalRHI/Public/MetalRHIVisionOSBridge.h b/Engine/Source/Runtime/Apple/MetalRHI/Public/MetalRHIVisionOSBridge.h
new file mode 100644
--- /dev/null
+++ b/Engine/Source/Runtime/Apple/MetalRHI/Public/MetalRHIVisionOSBridge.h
@@ -0,0 +1,26 @@
+// Copyright Epic Games, Inc. All Rights Reserved.
+
+#pragma once
+
+#if PLATFORM_VISIONOS
+#include "RHIFwd.h"
+#import <CompositorServices/CompositorServices.h>
+
+DECLARE_LOG_CATEGORY_EXTERN(LogMetalVisionOS, Warning, All);
+
+namespace MetalRHIVisionOS
+{
+	METALRHI_API struct BeginRenderingImmersiveParams
+    {
+        const cp_frame_t SwiftFrame;
+    };
+    METALRHI_API void BeginRenderingImmersive(const BeginRenderingImmersiveParams& Params);
+
+    METALRHI_API struct PresentImmersiveParams
+    {
+        const FTextureRHIRef& Texture;
+        cp_drawable_t& SwiftDrawable;
+    };
+    METALRHI_API void PresentImmersive(const PresentImmersiveParams& Params);
+}
+#endif // PLATFORM_VISIONOS
diff --git a/Engine/Source/Runtime/Apple/MetalRHI/Public/MetalViewport.h b/Engine/Source/Runtime/Apple/MetalRHI/Public/MetalViewport.h
--- a/Engine/Source/Runtime/Apple/MetalRHI/Public/MetalViewport.h
+++ b/Engine/Source/Runtime/Apple/MetalRHI/Public/MetalViewport.h
@@ -30,6 +30,14 @@ enum EMetalViewportAccessFlag
 
 class FMetalCommandQueue;
 
+#if PLATFORM_VISIONOS
+namespace MetalRHIVisionOS
+{
+    struct BeginRenderingImmersiveParams;
+    struct PresentImmersiveParams;
+}
+#endif
+
 typedef void (^FMetalViewportPresentHandler)(uint32 CGDirectDisplayID, double OutputSeconds, double OutputDuration);
 
 class FMetalViewport : public FRHIViewport
@@ -52,7 +60,9 @@ public:
 	
 #if PLATFORM_MAC
 	NSWindow* GetWindow() const;
-	
+#endif
+    
+#if PLATFORM_MAC || PLATFORM_VISIONOS
 	virtual void SetCustomPresent(FRHICustomPresent* InCustomPresent) override
 	{
 		CustomPresent = InCustomPresent;
@@ -64,11 +74,16 @@ public:
 	void Present(FMetalCommandQueue& CommandQueue, bool bLockToVsync);
 	void Swap();
 	
+#if PLATFORM_VISIONOS
+    void BeginRenderingImmersive(cp_frame_t SwiftFrame);
+    void PresentImmersive(TRefCountPtr<FMetalSurface> CompleteFrame, const MetalRHIVisionOS::PresentImmersiveParams* Params);
+private:
+	void EndFrameImmersive();
+#endif
+	
 private:
 	uint32 GetViewportIndex(EMetalViewportAccessFlag Accessor) const;
 
-	void NewFrame();
-	
 private:
 #if PLATFORM_VISIONOS
 	CP_OBJECT_cp_layer_renderer* SwiftLayer = nullptr;
@@ -89,7 +104,10 @@ private:
 
 #if PLATFORM_MAC
 	FMetalView* View;
-	FRHICustomPresent* CustomPresent;
+#endif
+    
+#if PLATFORM_MAC || PLATFORM_VISIONOS
+    FCustomPresentRHIRef CustomPresent;
 #endif
 };
 
diff --git a/Engine/Source/Runtime/ApplicationCore/Private/IOS/IOSPlatformApplicationMisc.cpp b/Engine/Source/Runtime/ApplicationCore/Private/IOS/IOSPlatformApplicationMisc.cpp
--- a/Engine/Source/Runtime/ApplicationCore/Private/IOS/IOSPlatformApplicationMisc.cpp
+++ b/Engine/Source/Runtime/ApplicationCore/Private/IOS/IOSPlatformApplicationMisc.cpp
@@ -25,7 +25,7 @@ static TAutoConsoleVariable<int32> CVarPhysicalScreenDensity(
 
 EAppReturnType::Type MessageBoxExtImpl( EAppMsgType::Type MsgType, const TCHAR* Text, const TCHAR* Caption )
 {
-#if PLATFORM_TVOS
+#if PLATFORM_TVOS || PLATFORM_VISIONOS
 	return FGenericPlatformMisc::MessageBoxExt(MsgType, Text, Caption);
 #else
 	NSString* CocoaText = (NSString*)FPlatformString::TCHARToCFString(Text);
diff --git a/Engine/Source/Runtime/ApplicationCore/Private/IOS/IOSPlatformFramePacer.cpp b/Engine/Source/Runtime/ApplicationCore/Private/IOS/IOSPlatformFramePacer.cpp
--- a/Engine/Source/Runtime/ApplicationCore/Private/IOS/IOSPlatformFramePacer.cpp
+++ b/Engine/Source/Runtime/ApplicationCore/Private/IOS/IOSPlatformFramePacer.cpp
@@ -132,6 +132,9 @@ uint32 FIOSPlatformRHIFramePacer::Pace = 0;
 
 bool FIOSPlatformRHIFramePacer::IsEnabled()
 {
+#if PLATFORM_VISIONOS
+	return false; // XR does its own frame pacing.
+#else
     static bool bIsRHIFramePacerEnabled = false;
 	static bool bInitialized = false;
 
@@ -170,6 +173,7 @@ bool FIOSPlatformRHIFramePacer::IsEnabled()
 	}
 	
 	return bIsRHIFramePacerEnabled;
+#endif
 }
 
 uint32 FIOSPlatformRHIFramePacer::GetMaxRefreshRate()
diff --git a/Engine/Source/Runtime/Engine/Private/LocalPlayer.cpp b/Engine/Source/Runtime/Engine/Private/LocalPlayer.cpp
--- a/Engine/Source/Runtime/Engine/Private/LocalPlayer.cpp
+++ b/Engine/Source/Runtime/Engine/Private/LocalPlayer.cpp
@@ -772,7 +772,7 @@ bool ULocalPlayer::CalcSceneViewInitOptions(
 		ViewInitOptions.bInCameraCut = PlayerController->PlayerCameraManager->bGameCameraCutThisFrame;
 	}
 
-	if (GEngine->StereoRenderingDevice.IsValid())
+	if (GEngine->IsStereoscopic3D(Viewport))
 	{
 		ViewInitOptions.StereoPass = GEngine->StereoRenderingDevice->GetViewPassForIndex(StereoViewIndex != INDEX_NONE, StereoViewIndex);
 	}
@@ -909,7 +909,7 @@ FSceneView* ULocalPlayer::CalcSceneView( class FSceneViewFamily* ViewFamily,
 			PlayerController->PlayerCameraManager->UpdatePhotographyPostProcessing(View->FinalPostProcessSettings);
 		}
 
-		if (GEngine->StereoRenderingDevice.IsValid())
+		if (GEngine->IsStereoscopic3D(Viewport))
 		{
 			FPostProcessSettings StereoDeviceOverridePostProcessinSettings;
 			float BlendWeight = 1.0f;
@@ -1236,7 +1236,7 @@ bool ULocalPlayer::GetProjectionData(FViewport* Viewport, FSceneViewProjectionDa
 			XRCamera->UseImplicitHMDPosition(bHasActiveCamera);
 		}
 
-		if (GEngine->StereoRenderingDevice.IsValid())
+		if (GEngine->IsStereoscopic3D(Viewport))
 		{
 			GEngine->StereoRenderingDevice->CalculateStereoViewOffset(StereoViewIndex, ViewInfo.Rotation, GetWorld()->GetWorldSettings()->WorldToMeters, StereoViewLocation);
 		}
diff --git a/Engine/Source/Runtime/Launch/Private/IOS/LaunchIOS.cpp b/Engine/Source/Runtime/Launch/Private/IOS/LaunchIOS.cpp
--- a/Engine/Source/Runtime/Launch/Private/IOS/LaunchIOS.cpp
+++ b/Engine/Source/Runtime/Launch/Private/IOS/LaunchIOS.cpp
@@ -598,6 +598,8 @@ int main(int argc, char *argv[])
 void FSwiftAppBootstrap::KickoffWithCompositingLayer(CP_OBJECT_cp_layer_renderer* Layer)
 {
 	IOSAppDelegate* AppDelegate = [IOSAppDelegate GetDelegate];
+    
+    // Might need to cp_layer_renderer_configuration_set_layout here,  or in UESwift... in the future.
 	
 	cp_layer_renderer_properties_t Props = cp_layer_renderer_get_properties(Layer);
 	int NumViews = cp_layer_renderer_properties_get_view_count(Props);
@@ -618,7 +620,7 @@ void FSwiftAppBootstrap::KickoffWithCompositingLayer(CP_OBJECT_cp_layer_renderer
 
 	{
 		cp_frame_t SwiftLayerFrame = cp_layer_renderer_query_next_frame(Layer);
-		cp_drawable_t SwiftDrawable = cp_frame_query_drawable(SwiftLayerFrame); 
+		cp_drawable_t SwiftDrawable = cp_frame_query_drawable(SwiftLayerFrame);
 		{
 			cp_view_t View = cp_drawable_get_view(SwiftDrawable, 0);
 			cp_view_texture_map_t TextureMap = cp_view_get_view_texture_map(View);
