MetalRHIVisionOSBridge unity build fix for IOS

diff --git a/Engine/Source/Runtime/Apple/MetalRHI/Public/MetalRHIVisionOSBridge.h b/Engine/Source/Runtime/Apple/MetalRHI/Public/MetalRHIVisionOSBridge.h
--- a/Engine/Source/Runtime/Apple/MetalRHI/Public/MetalRHIVisionOSBridge.h
+++ b/Engine/Source/Runtime/Apple/MetalRHI/Public/MetalRHIVisionOSBridge.h
@@ -2,6 +2,8 @@
 
 #pragma once
 
+#include "HAL/Platform.h"
+
 #if PLATFORM_VISIONOS
 #include "RHIFwd.h"
 #import <CompositorServices/CompositorServices.h>
