Re-added MTL::Argument path for MacOS lower than 13

diff --git a/Engine/Source/Runtime/Apple/MetalRHI/Private/MetalPipeline.cpp b/Engine/Source/Runtime/Apple/MetalRHI/Private/MetalPipeline.cpp
--- a/Engine/Source/Runtime/Apple/MetalRHI/Private/MetalPipeline.cpp
+++ b/Engine/Source/Runtime/Apple/MetalRHI/Private/MetalPipeline.cpp
@@ -371,87 +371,173 @@ void FMetalShaderPipeline::InitResourceMask()
 
 void FMetalShaderPipeline::InitResourceMask(EMetalShaderFrequency Frequency)
 {
-    NS::Array* Bindings = nullptr;
-    switch(Frequency)
+    if (@available(macOS 13.0, iOS 16.0, *))
     {
-        case EMetalShaderVertex:
+        NS::Array* Bindings = nullptr;
+        switch(Frequency)
         {
-            check(RenderPipelineReflection);
-            Bindings = RenderPipelineReflection->vertexBindings();
-            break;
-        }
-        case EMetalShaderFragment:
-        {
-            check(RenderPipelineReflection);
-            Bindings = RenderPipelineReflection->fragmentBindings();
-            break;
-        }
-        case EMetalShaderCompute:
-        {
-            check(ComputePipelineReflection);
-            Bindings = ComputePipelineReflection->bindings();
-            break;
+            case EMetalShaderVertex:
+            {
+                check(RenderPipelineReflection);
+                Bindings = RenderPipelineReflection->vertexBindings();
+                break;
+            }
+            case EMetalShaderFragment:
+            {
+                check(RenderPipelineReflection);
+                Bindings = RenderPipelineReflection->fragmentBindings();
+                break;
+            }
+            case EMetalShaderCompute:
+            {
+                check(ComputePipelineReflection);
+                Bindings = ComputePipelineReflection->bindings();
+                break;
+            }
+            case EMetalShaderStream:
+            {
+                check(StreamPipelineReflection);
+                Bindings = StreamPipelineReflection->vertexBindings();
+                break;
+            }
+            default:
+                check(false);
+                break;
         }
-        case EMetalShaderStream:
+        
+        for (uint32 i = 0; i < Bindings->count(); i++)
         {
-            check(StreamPipelineReflection);
-            Bindings = StreamPipelineReflection->vertexBindings();
-            break;
+            MTL::Binding* Binding = (MTL::Binding*)Bindings->object(i);
+            check(Binding);
+            
+            if (!Binding->used())
+            {
+                continue;
+            }
+            
+            switch(Binding->type())
+            {
+                case MTL::BindingTypeBuffer:
+                {
+                    MTL::BufferBinding* BufferBinding = (MTL::BufferBinding*)Binding;
+                    checkf(Binding->index() < ML_MaxBuffers, TEXT("Metal buffer index exceeded!"));
+                    if (NSStringToFString(Binding->name()) != TEXT("BufferSizes") && NSStringToFString(Binding->name()) != TEXT("spvBufferSizeConstants"))
+                    {
+                        ResourceMask[Frequency].BufferMask |= (1 << Binding->index());
+                        
+                        if(BufferDataSizes[Frequency].Num() < 31)
+                            BufferDataSizes[Frequency].SetNumZeroed(31);
+                        
+                        BufferDataSizes[Frequency][Binding->index()] = BufferBinding->bufferDataSize();
+                    }
+                    break;
+                }
+                case MTL::BindingTypeThreadgroupMemory:
+                {
+                    break;
+                }
+                case MTL::BindingTypeTexture:
+                {
+                    MTL::TextureBinding* TextureBinding = (MTL::TextureBinding*)Bindings->object(i);
+                    checkf(Binding->index() < ML_MaxTextures, TEXT("Metal texture index exceeded!"));
+                    ResourceMask[Frequency].TextureMask |= (FMetalTextureMask(1) << Binding->index());
+                    TextureTypes[Frequency].Add(Binding->index(), (uint8)TextureBinding->textureType());
+                    break;
+                }
+                case MTL::BindingTypeSampler:
+                {
+                    checkf(Binding->index() < ML_MaxSamplers, TEXT("Metal sampler index exceeded!"));
+                    ResourceMask[Frequency].SamplerMask |= (1 << Binding->index());
+                    break;
+                }
+                default:
+                    check(false);
+                    break;
+            }
         }
-        default:
-            check(false);
-            break;
     }
-    
-    for (uint32 i = 0; i < Bindings->count(); i++)
+    else
     {
-        MTL::Binding* Binding = (MTL::Binding*)Bindings->object(i);
-        check(Binding);
-        
-        if (!Binding->used())
+        NS::Array* Arguments = nullptr;
+        switch(Frequency)
         {
-            continue;
-        }
-        
-        switch(Binding->type())
-        {
-            case MTL::BindingTypeBuffer:
+            case EMetalShaderVertex:
             {
-                MTL::BufferBinding* BufferBinding = (MTL::BufferBinding*)Binding;
-                checkf(Binding->index() < ML_MaxBuffers, TEXT("Metal buffer index exceeded!"));
-                if (NSStringToFString(Binding->name()) != TEXT("BufferSizes") && NSStringToFString(Binding->name()) != TEXT("spvBufferSizeConstants"))
-                {
-                    ResourceMask[Frequency].BufferMask |= (1 << Binding->index());
-                    
-                    if(BufferDataSizes[Frequency].Num() < 31)
-                        BufferDataSizes[Frequency].SetNumZeroed(31);
-                    
-                    BufferDataSizes[Frequency][Binding->index()] = BufferBinding->bufferDataSize();
-                }
+                check(RenderPipelineReflection);
+                Arguments = RenderPipelineReflection->vertexArguments();
                 break;
             }
-            case MTL::BindingTypeThreadgroupMemory:
+            case EMetalShaderFragment:
             {
+                check(RenderPipelineReflection);
+                Arguments = RenderPipelineReflection->fragmentArguments();
                 break;
             }
-            case MTL::BindingTypeTexture:
+            case EMetalShaderCompute:
             {
-                MTL::TextureBinding* TextureBinding = (MTL::TextureBinding*)Bindings->object(i);
-                checkf(Binding->index() < ML_MaxTextures, TEXT("Metal texture index exceeded!"));
-                ResourceMask[Frequency].TextureMask |= (FMetalTextureMask(1) << Binding->index());
-                TextureTypes[Frequency].Add(Binding->index(), (uint8)TextureBinding->textureType());
+                check(ComputePipelineReflection);
+                Arguments = ComputePipelineReflection->arguments();
                 break;
             }
-            case MTL::BindingTypeSampler:
+            case EMetalShaderStream:
             {
-                checkf(Binding->index() < ML_MaxSamplers, TEXT("Metal sampler index exceeded!"));
-                ResourceMask[Frequency].SamplerMask |= (1 << Binding->index());
+                check(StreamPipelineReflection);
+                Arguments = StreamPipelineReflection->vertexArguments();
                 break;
             }
             default:
                 check(false);
                 break;
         }
+        
+        for (uint32 i = 0; i < Arguments->count(); i++)
+        {
+            MTL::Argument* Argument = (MTL::Argument*)Arguments->object(i);
+            check(Argument);
+            
+            if (!Argument->active())
+            {
+                continue;
+            }
+            
+            switch(Argument->type())
+            {
+                case MTL::ArgumentTypeBuffer:
+                {
+                    checkf(Argument->index() < ML_MaxBuffers, TEXT("Metal buffer index exceeded!"));
+                    if (NSStringToFString(Argument->name()) != TEXT("BufferSizes") && NSStringToFString(Argument->name()) != TEXT("spvBufferSizeConstants"))
+                    {
+                        ResourceMask[Frequency].BufferMask |= (1 << Argument->index());
+                        
+                        if(BufferDataSizes[Frequency].Num() < 31)
+                            BufferDataSizes[Frequency].SetNumZeroed(31);
+                        
+                        BufferDataSizes[Frequency][Argument->index()] = Argument->bufferDataSize();
+                    }
+                    break;
+                }
+                case MTL::ArgumentTypeThreadgroupMemory:
+                {
+                    break;
+                }
+                case MTL::ArgumentTypeTexture:
+                {
+                    checkf(Argument->index() < ML_MaxTextures, TEXT("Metal texture index exceeded!"));
+                    ResourceMask[Frequency].TextureMask |= (FMetalTextureMask(1) << Argument->index());
+                    TextureTypes[Frequency].Add(Argument->index(), (uint8)Argument->textureType());
+                    break;
+                }
+                case MTL::ArgumentTypeSampler:
+                {
+                    checkf(Argument->index() < ML_MaxSamplers, TEXT("Metal sampler index exceeded!"));
+                    ResourceMask[Frequency].SamplerMask |= (1 << Argument->index());
+                    break;
+                }
+                default:
+                    check(false);
+                    break;
+            }
+        }
     }
 }
 
