Swapped key/value order in dictionary which was causing a runtime failure on debug shaders

diff --git a/Engine/Source/Runtime/Apple/MetalRHI/Private/Shaders/Types/Templates/MetalBaseShader.h b/Engine/Source/Runtime/Apple/MetalRHI/Private/Shaders/Types/Templates/MetalBaseShader.h
--- a/Engine/Source/Runtime/Apple/MetalRHI/Private/Shaders/Types/Templates/MetalBaseShader.h
+++ b/Engine/Source/Runtime/Apple/MetalRHI/Private/Shaders/Types/Templates/MetalBaseShader.h
@@ -298,12 +298,12 @@ void TMetalBaseShader<BaseResourceType, ShaderType>::Init(TArrayView<const uint8
 #if !PLATFORM_MAC || DEBUG_METAL_SHADERS
             NS::Dictionary* PreprocessorMacros = nullptr;;
 #if !PLATFORM_MAC // Pretty sure that as_type-casts work on macOS, but they don't for half2<->uint on older versions of the iOS runtime compiler.
-			PreprocessorMacros = NS::Dictionary::dictionary(NS::String::string("METAL_RUNTIME_COMPILER", NS::UTF8StringEncoding),
-                                                            NS::String::string("1", NS::UTF8StringEncoding));
+			PreprocessorMacros = NS::Dictionary::dictionary(NS::String::string("1", NS::UTF8StringEncoding),
+                                                            NS::String::string("METAL_RUNTIME_COMPILER", NS::UTF8StringEncoding));
 #endif
 #if DEBUG_METAL_SHADERS
-            PreprocessorMacros = NS::Dictionary::dictionary(NS::String::string("MTLSL_ENABLE_DEBUG_INFO", NS::UTF8StringEncoding),
-                                                            NS::String::string("1", NS::UTF8StringEncoding));
+            PreprocessorMacros = NS::Dictionary::dictionary(NS::String::string("1", NS::UTF8StringEncoding),
+                                                            NS::String::string("MTLSL_ENABLE_DEBUG_INFO", NS::UTF8StringEncoding));
 #endif
             if(PreprocessorMacros)
             {
