Removed NS::Condition and replaced with FEvent

diff --git a/Engine/Source/Runtime/Apple/MetalRHI/Private/MetalUAV.cpp b/Engine/Source/Runtime/Apple/MetalRHI/Private/MetalUAV.cpp
--- a/Engine/Source/Runtime/Apple/MetalRHI/Private/MetalUAV.cpp
+++ b/Engine/Source/Runtime/Apple/MetalRHI/Private/MetalUAV.cpp
@@ -593,73 +593,36 @@ void FMetalUnorderedAccessView::ClearUAV(TRHICommandList_RecursiveHazardous<FMet
     }
 }
 
-FMetalCommandBufferFence::FMetalCommandBufferFence()
+bool FMetalCommandBufferFence::Wait(uint32_t TimeIntervalMs) const
 {
-    Condition = NS::Condition::alloc()->init();
-    check(Condition);
-    bPassed = 0;
-}
-
-FMetalCommandBufferFence::~FMetalCommandBufferFence()
-{
-    Condition->release();
-}
-
-bool FMetalCommandBufferFence::Wait(NS::UInteger TimeInterval) const
-{
-    check(Condition != nullptr);
     check(CmdBuffer);
 
-    Condition->lock();
-    __sync_synchronize();
+    bool bFinished = false;
     
-    uint32_t bFinished = bPassed;
-    
-    if(!bFinished)
+    if(TimeIntervalMs == MAX_uint32)
     {
-        switch (TimeInterval)
-        {
-            case NS::UIntegerMax:
-            {
-                Condition->wait();
-                bFinished = true;
-                break;
-            }
-            case 0:
-            {
-                bFinished = Condition->waitUntilDate(NS::Date::dateWithTimeIntervalSinceNow(0.0));
-                break;
-            }
-            default:
-            {
-                bFinished = Condition->waitUntilDate(NS::Date::dateWithTimeIntervalSinceNow(TimeInterval / 1000.0));
-                break;
-            }
-        }
-        __sync_fetch_and_add((uint32_t*)&bPassed, bFinished);
+        bFinished = Condition->Wait();
+    }
+    else
+    {
+        bFinished = Condition->Wait((uint32_t)TimeIntervalMs);
     }
-    Condition->unlock();
     
-    __sync_synchronize();
-    return (bPassed != 0);
+    return bFinished;
 }
 
 void FMetalCommandBufferFence::Insert(MTLCommandBufferPtr CommandBuffer)
 {
     check(CommandBuffer);
-    check(bPassed == false);
     check(CmdBuffer.get() == nullptr);
-    check(Condition);
     
     CmdBuffer = CommandBuffer;
-
-    MTL::HandlerFunction CommandBufferCompletionHandler = [&, this](MTL::CommandBuffer*) {
-        check(bPassed == false);
-        
-        __sync_fetch_and_add((uint32_t*)&bPassed, 1);
-        Condition->lock();
-        Condition->broadcast();
-        Condition->unlock();
+    
+    Condition->Reset();
+    
+    MTL::HandlerFunction CommandBufferCompletionHandler = [&](MTL::CommandBuffer*)
+    {
+        Condition->Trigger();
     };
     
     CmdBuffer->addCompletedHandler(CommandBufferCompletionHandler);
@@ -667,7 +630,7 @@ void FMetalCommandBufferFence::Insert(MTLCommandBufferPtr CommandBuffer)
 
 void FMetalCommandBufferFence::Signal(const MTL::CommandBuffer* CommandBuffer)
 {
-    Condition->signal();
+    Condition->Trigger();
 }
 
 void FMetalGPUFence::WriteInternal(FMetalCommandBuffer* CommandBuffer)
diff --git a/Engine/Source/Runtime/Apple/MetalRHI/Public/MetalResources.h b/Engine/Source/Runtime/Apple/MetalRHI/Public/MetalResources.h
--- a/Engine/Source/Runtime/Apple/MetalRHI/Public/MetalResources.h
+++ b/Engine/Source/Runtime/Apple/MetalRHI/Public/MetalResources.h
@@ -497,20 +497,14 @@ public:
 class FMetalCommandBufferFence
 {
 public:
-    FMetalCommandBufferFence();
-    ~FMetalCommandBufferFence();
-
-public:
-    //void Init();
-    bool Wait(NS::UInteger TimeInterval) const;
+    bool Wait(uint32_t TimeInterval) const;
     void Insert(MTLCommandBufferPtr CmdBuffer);
     
 private:
     void Signal(const MTL::CommandBuffer* CmdBuffer);
     
-    NS::Condition* Condition = nullptr;
+    FEventRef Condition { EEventMode::ManualReset };
     MTLCommandBufferPtr CmdBuffer;
-    uint32_t bPassed = 0;
 };
 
 class FMetalGPUFence final : public FRHIGPUFence
