Fix for MTL::Texture going out of scope when calling CVMetalTextureGetTexture

diff --git a/Engine/Source/Runtime/Apple/MetalRHI/Private/MetalTexture.cpp b/Engine/Source/Runtime/Apple/MetalRHI/Private/MetalTexture.cpp
--- a/Engine/Source/Runtime/Apple/MetalRHI/Private/MetalTexture.cpp
+++ b/Engine/Source/Runtime/Apple/MetalRHI/Private/MetalTexture.cpp
@@ -537,7 +537,7 @@ FMetalSurface::FMetalSurface(FRHICommandListBase* RHICmdList, FMetalTextureCreat
                 
                 Texture = GetMetalDeviceContext().GetDevice()->newTexture(CreateDesc.Desc, CVPixelBufferGetIOSurface((CVPixelBufferRef)ImageSurfaceRef));
 #else
-				Texture = NS::TransferPtr((__bridge MTL::Texture*)CVMetalTextureGetTexture((CVPixelBufferRef)ImageSurfaceRef));
+				Texture = NS::RetainPtr((__bridge MTL::Texture*)CVMetalTextureGetTexture((CVPixelBufferRef)ImageSurfaceRef));
 #endif
 				METAL_FATAL_ASSERT(Texture, TEXT("Failed to create texture, desc %s"), *NSStringToFString(CreateDesc.Desc->description()));
 
