LLM: Use LLM_IF_ENABLED around all calls to OnLowLevelAlloc and OnLowLevelFree that are guarded by a macro for conditional compilation. LLM_IF_ENABLED is not strictly necessary; LLM or #if ENABLE_LOW_LEVEL_MEM_TRACKER will also both behave correctly. But LLM_IF_ENABLED is a small optimization (avoids a function call) and we want to apply it for consistency.
#rnx

diff --git a/Engine/Source/Runtime/Apple/MetalRHI/Private/MetalBuffer.cpp b/Engine/Source/Runtime/Apple/MetalRHI/Private/MetalBuffer.cpp
--- a/Engine/Source/Runtime/Apple/MetalRHI/Private/MetalBuffer.cpp
+++ b/Engine/Source/Runtime/Apple/MetalRHI/Private/MetalBuffer.cpp
@@ -120,7 +120,7 @@ void FMetalBuffer::Release()
     {
         LLM_PLATFORM_SCOPE_METAL(ELLMTagMetal::Buffers);
         
-        LLM(FLowLevelMemTracker::Get().OnLowLevelFree(ELLMTracker::Platform, this, ELLMAllocType::System));
+		LLM_IF_ENABLED(FLowLevelMemTracker::Get().OnLowLevelFree(ELLMTracker::Platform, this, ELLMAllocType::System));
         
         DEC_MEMORY_STAT_BY(STAT_MetalBufferMemory, GetLength());
         DEC_DWORD_STAT(STAT_MetalBufferCount);
diff --git a/Engine/Source/Runtime/Apple/MetalRHI/Private/MetalLLM.cpp b/Engine/Source/Runtime/Apple/MetalRHI/Private/MetalLLM.cpp
--- a/Engine/Source/Runtime/Apple/MetalRHI/Private/MetalLLM.cpp
+++ b/Engine/Source/Runtime/Apple/MetalRHI/Private/MetalLLM.cpp
@@ -163,7 +163,7 @@ void MetalLLM::LogAllocTexture(MTL::Device* Device, MTL::TextureDescriptor* Desc
 	}
 	INC_DWORD_STAT(STAT_MetalTextureCount);
 	
-	LLM(FLowLevelMemTracker::Get().OnLowLevelAlloc(ELLMTracker::Platform, Ptr, Size, ELLMTag::Untagged, ELLMAllocType::System));
+	LLM_IF_ENABLED(FLowLevelMemTracker::Get().OnLowLevelAlloc(ELLMTracker::Platform, Ptr, Size, ELLMTag::Untagged, ELLMAllocType::System));
 	// Assign a dealloc handler to untrack the memory - but don't track the dispatch block!
 	{
 		LLM_SCOPED_PAUSE_TRACKING(ELLMAllocType::System);
@@ -174,7 +174,7 @@ void MetalLLM::LogAllocTexture(MTL::Device* Device, MTL::TextureDescriptor* Desc
 			[[[FMetalDeallocHandler alloc] initWithBlock:^{
 				LLM_PLATFORM_SCOPE_METAL(ELLMTagMetal::RenderTargets);
 				
-				LLM(FLowLevelMemTracker::Get().OnLowLevelFree(ELLMTracker::Platform, Ptr, ELLMAllocType::System));
+				LLM_IF_ENABLED(FLowLevelMemTracker::Get().OnLowLevelFree(ELLMTracker::Platform, Ptr, ELLMAllocType::System));
 				
 #if PLATFORM_IOS
 				if (!bMemoryless)
@@ -192,7 +192,7 @@ void MetalLLM::LogAllocTexture(MTL::Device* Device, MTL::TextureDescriptor* Desc
 			[[[FMetalDeallocHandler alloc] initWithBlock:^{
 				LLM_PLATFORM_SCOPE_METAL(ELLMTagMetal::Textures);
 			
-				LLM(FLowLevelMemTracker::Get().OnLowLevelFree(ELLMTracker::Platform, Ptr, ELLMAllocType::System));
+				LLM_IF_ENABLED(FLowLevelMemTracker::Get().OnLowLevelFree(ELLMTracker::Platform, Ptr, ELLMAllocType::System));
 			
 #if PLATFORM_IOS
 				if (!bMemoryless)
@@ -215,7 +215,7 @@ void MetalLLM::LogAllocBuffer(MTL::Device* Device, FMetalBufferPtr Buffer)
 	INC_MEMORY_STAT_BY(STAT_MetalBufferMemory, Size);
 	INC_DWORD_STAT(STAT_MetalBufferCount);
 	
-	LLM(FLowLevelMemTracker::Get().OnLowLevelAlloc(ELLMTracker::Platform, Ptr, Size, ELLMTag::Untagged, ELLMAllocType::System));
+	LLM_IF_ENABLED(FLowLevelMemTracker::Get().OnLowLevelAlloc(ELLMTracker::Platform, Ptr, Size, ELLMTag::Untagged, ELLMAllocType::System));
     Buffer->MarkAllocated();
 }
 
@@ -227,7 +227,7 @@ void MetalLLM::LogAllocBufferNative(MTL::Device* Device, MTLBufferPtr Buffer)
     INC_MEMORY_STAT_BY(STAT_MetalBufferMemory, Size);
     INC_DWORD_STAT(STAT_MetalBufferCount);
     
-    LLM(FLowLevelMemTracker::Get().OnLowLevelAlloc(ELLMTracker::Platform, Ptr, Size, ELLMTag::Untagged, ELLMAllocType::System));
+	LLM_IF_ENABLED(FLowLevelMemTracker::Get().OnLowLevelAlloc(ELLMTracker::Platform, Ptr, Size, ELLMTag::Untagged, ELLMAllocType::System));
     // Assign a dealloc handler to untrack the memory - but don't track the dispatch block!
     {
         LLM_SCOPED_PAUSE_TRACKING(ELLMAllocType::System);
@@ -236,7 +236,7 @@ void MetalLLM::LogAllocBufferNative(MTL::Device* Device, MTLBufferPtr Buffer)
         [[[FMetalDeallocHandler alloc] initWithBlock:^{
             LLM_PLATFORM_SCOPE_METAL(ELLMTagMetal::Buffers);
             
-            LLM(FLowLevelMemTracker::Get().OnLowLevelFree(ELLMTracker::Platform, Ptr, ELLMAllocType::System));
+			LLM_IF_ENABLED(FLowLevelMemTracker::Get().OnLowLevelFree(ELLMTracker::Platform, Ptr, ELLMAllocType::System));
             
             DEC_MEMORY_STAT_BY(STAT_MetalBufferMemory, Size);
             DEC_DWORD_STAT(STAT_MetalBufferCount);
@@ -253,7 +253,7 @@ void MetalLLM::LogAllocHeap(MTL::Device* Device, MTL::Heap* Heap)
 	INC_MEMORY_STAT_BY(STAT_MetalHeapMemory, Size);
 	INC_DWORD_STAT(STAT_MetalHeapCount);
 	
-	LLM(FLowLevelMemTracker::Get().OnLowLevelAlloc(ELLMTracker::Platform, Ptr, Size, ELLMTag::Untagged, ELLMAllocType::System));
+	LLM_IF_ENABLED(FLowLevelMemTracker::Get().OnLowLevelAlloc(ELLMTracker::Platform, Ptr, Size, ELLMTag::Untagged, ELLMAllocType::System));
 	// Assign a dealloc handler to untrack the memory - but don't track the dispatch block!
 	{
 		LLM_SCOPED_PAUSE_TRACKING(ELLMAllocType::System);
@@ -263,7 +263,7 @@ void MetalLLM::LogAllocHeap(MTL::Device* Device, MTL::Heap* Heap)
 			LLM_SCOPE_METAL(ELLMTagMetal::Heaps);
 			LLM_PLATFORM_SCOPE_METAL(ELLMTagMetal::Heaps);
 			
-			LLM(FLowLevelMemTracker::Get().OnLowLevelFree(ELLMTracker::Platform, Ptr, ELLMAllocType::System));
+			LLM_IF_ENABLED(FLowLevelMemTracker::Get().OnLowLevelFree(ELLMTracker::Platform, Ptr, ELLMAllocType::System));
 			
 			DEC_MEMORY_STAT_BY(STAT_MetalHeapMemory, Size);
 			DEC_DWORD_STAT(STAT_MetalHeapCount);
diff --git a/Engine/Source/Runtime/Core/Private/Android/AndroidPlatformFile.cpp b/Engine/Source/Runtime/Core/Private/Android/AndroidPlatformFile.cpp
--- a/Engine/Source/Runtime/Core/Private/Android/AndroidPlatformFile.cpp
+++ b/Engine/Source/Runtime/Core/Private/Android/AndroidPlatformFile.cpp
@@ -1064,7 +1064,7 @@ public:
 			UE_LOG(LogAndroidFile, Warning, TEXT("Failed to map memory %s, error is %d"), *Filename, errno);
 			return nullptr;
 		}
-		LLM(FLowLevelMemTracker::Get().OnLowLevelAlloc(ELLMTracker::Platform, AlignedMapPtr, AlignedSize));
+		LLM_IF_ENABLED(FLowLevelMemTracker::Get().OnLowLevelAlloc(ELLMTracker::Platform, AlignedMapPtr, AlignedSize));
 
 		// create a mapping for this range
 		const uint8* MapPtr = AlignedMapPtr + Offset - AlignedOffset;
@@ -1079,7 +1079,7 @@ public:
 		check(NumOutstandingRegions > 0);
 		NumOutstandingRegions--;
 
-		LLM(FLowLevelMemTracker::Get().OnLowLevelFree(ELLMTracker::Platform, Region->AlignedPtr));
+		LLM_IF_ENABLED(FLowLevelMemTracker::Get().OnLowLevelFree(ELLMTracker::Platform, Region->AlignedPtr));
 		const int Res = munmap(const_cast<uint8*>(Region->AlignedPtr), Region->AlignedSize);
 #if LOG_ANDROID_FILE
 		FPlatformMisc::LowLevelOutputDebugStringf(TEXT("Failed to unmap region from %s, errno=%s"), *Filename, UTF8_TO_TCHAR(strerror(errno)));
diff --git a/Engine/Source/Runtime/Core/Private/Android/AndroidPlatformMemory.cpp b/Engine/Source/Runtime/Core/Private/Android/AndroidPlatformMemory.cpp
--- a/Engine/Source/Runtime/Core/Private/Android/AndroidPlatformMemory.cpp
+++ b/Engine/Source/Runtime/Core/Private/Android/AndroidPlatformMemory.cpp
@@ -689,13 +689,13 @@ void* FAndroidPlatformMemory::BinnedAllocFromOS(SIZE_T Size)
 		Ptr = mmap(nullptr, Size, PROT_READ | PROT_WRITE, MAP_PRIVATE | MAP_ANON, -1, 0);
 	}
 
-	LLM(FLowLevelMemTracker::Get().OnLowLevelAlloc(ELLMTracker::Platform, Ptr, Size));
+	LLM_IF_ENABLED(FLowLevelMemTracker::Get().OnLowLevelAlloc(ELLMTracker::Platform, Ptr, Size));
 	return Ptr;
 }
 
 void FAndroidPlatformMemory::BinnedFreeToOS(void* Ptr, SIZE_T Size)
 {
-	LLM(FLowLevelMemTracker::Get().OnLowLevelFree(ELLMTracker::Platform, Ptr));
+	LLM_IF_ENABLED(FLowLevelMemTracker::Get().OnLowLevelFree(ELLMTracker::Platform, Ptr));
 	if (munmap(Ptr, Size) != 0)
 	{
 		const int ErrNo = errno;
diff --git a/Engine/Source/Runtime/Core/Private/Apple/AppleLLM.cpp b/Engine/Source/Runtime/Core/Private/Apple/AppleLLM.cpp
--- a/Engine/Source/Runtime/Core/Private/Apple/AppleLLM.cpp
+++ b/Engine/Source/Runtime/Core/Private/Apple/AppleLLM.cpp
@@ -34,7 +34,7 @@ static id AllocWithZoneInterposer(id Obj, SEL Sel, struct _NSZone * Zone)
 	
 	if (Result)
 	{
-		LLM(FLowLevelMemTracker::Get().OnLowLevelAlloc(ELLMTracker::Default, Result, class_getInstanceSize([Result class]), (ELLMTag)ELLMTagApple::ObjectiveC, ELLMAllocType::System));
+		LLM_IF_ENABLED(FLowLevelMemTracker::Get().OnLowLevelAlloc(ELLMTracker::Default, Result, class_getInstanceSize([Result class]), (ELLMTag)ELLMTagApple::ObjectiveC, ELLMAllocType::System));
 	}
 	
 	return Result;
@@ -44,7 +44,7 @@ static void DeallocInterposer(id Obj, SEL Sel)
 {
 	if (Obj)
 	{
-		LLM(FLowLevelMemTracker::Get().OnLowLevelFree(ELLMTracker::Default, Obj, ELLMAllocType::System));
+		LLM_IF_ENABLED(FLowLevelMemTracker::Get().OnLowLevelFree(ELLMTracker::Default, Obj, ELLMAllocType::System));
 	}
 	DeallocOriginal(Obj, Sel);
 }
@@ -56,7 +56,7 @@ id NSAllocateObjectPtrInterposer(Class aClass, NSUInteger extraBytes, NSZone *zo
 	id Result = NSAllocateObjectPtr(aClass, extraBytes, zone);
 	if (Result)
 	{
-		LLM(FLowLevelMemTracker::Get().OnLowLevelAlloc(ELLMTracker::Default, Result, class_getInstanceSize([Result class]), (ELLMTag)ELLMTagApple::ObjectiveC, ELLMAllocType::System));
+		LLM_IF_ENABLED(FLowLevelMemTracker::Get().OnLowLevelAlloc(ELLMTracker::Default, Result, class_getInstanceSize([Result class]), (ELLMTag)ELLMTagApple::ObjectiveC, ELLMAllocType::System));
 	}
 	
 	return Result;
@@ -68,7 +68,7 @@ id _os_object_alloc_realizedInterposer(Class aClass, size_t size)
 	id Result = _os_object_alloc_realizedPtr(aClass, size);
 	if (Result)
 	{
-		LLM(FLowLevelMemTracker::Get().OnLowLevelAlloc(ELLMTracker::Default, Result, class_getInstanceSize([Result class]), (ELLMTag)ELLMTagApple::ObjectiveC, ELLMAllocType::System));
+		LLM_IF_ENABLED(FLowLevelMemTracker::Get().OnLowLevelAlloc(ELLMTracker::Default, Result, class_getInstanceSize([Result class]), (ELLMTag)ELLMTagApple::ObjectiveC, ELLMAllocType::System));
 	}
 	
 	return Result;
@@ -80,7 +80,7 @@ void NSDeallocateObjectInterposer(id Obj)
 {
 	if (Obj)
 	{
-		LLM(FLowLevelMemTracker::Get().OnLowLevelFree(ELLMTracker::Default, Obj, ELLMAllocType::System));
+		LLM_IF_ENABLED(FLowLevelMemTracker::Get().OnLowLevelFree(ELLMTracker::Default, Obj, ELLMAllocType::System));
 	}
 	NSDeallocateObjectPtr(Obj);
 }
diff --git a/Engine/Source/Runtime/Core/Private/Apple/ApplePlatformMemory.cpp b/Engine/Source/Runtime/Core/Private/Apple/ApplePlatformMemory.cpp
--- a/Engine/Source/Runtime/Core/Private/Apple/ApplePlatformMemory.cpp
+++ b/Engine/Source/Runtime/Core/Private/Apple/ApplePlatformMemory.cpp
@@ -643,7 +643,7 @@ void* _Nullable FApplePlatformMemory::BinnedAllocFromOS(SIZE_T Size)
 		AllocDescriptor->OriginalSizeAsPassed = Size;
 	}
 
-	LLM(FLowLevelMemTracker::Get().OnLowLevelAlloc(ELLMTracker::Platform, Pointer, Size));
+	LLM_IF_ENABLED(FLowLevelMemTracker::Get().OnLowLevelAlloc(ELLMTracker::Platform, Pointer, Size));
 	return Pointer;
 #else
 	void* Ptr = mmap(nullptr, Size, PROT_READ | PROT_WRITE, MAP_PRIVATE | MAP_ANON, -1, 0);
@@ -652,7 +652,7 @@ void* _Nullable FApplePlatformMemory::BinnedAllocFromOS(SIZE_T Size)
 		UE_LOG(LogTemp, Warning, TEXT("mmap failure allocating %d, error code: %d"), Size, errno);
 		Ptr = nullptr;
 	}
-	LLM(FLowLevelMemTracker::Get().OnLowLevelAlloc(ELLMTracker::Platform, Ptr, Size));
+	LLM_IF_ENABLED(FLowLevelMemTracker::Get().OnLowLevelAlloc(ELLMTracker::Platform, Ptr, Size));
 	return Ptr;
 #endif // USE_MALLOC_BINNED2
 }
@@ -661,7 +661,7 @@ void FApplePlatformMemory::BinnedFreeToOS(void* Ptr, SIZE_T Size)
 {
 	// Binned2 requires allocations to be BinnedPageSize-aligned. Simple mmap() does not guarantee this for recommended BinnedPageSize (64KB).
 #if USE_MALLOC_BINNED2
-	LLM(FLowLevelMemTracker::Get().OnLowLevelFree(ELLMTracker::Platform, Ptr));
+	LLM_IF_ENABLED(FLowLevelMemTracker::Get().OnLowLevelFree(ELLMTracker::Platform, Ptr));
 	// guard against someone not passing size in whole pages
 	static SIZE_T OSPageSize = FPlatformMemory::GetConstants().PageSize;
 	SIZE_T SizeInWholePages = (Size % OSPageSize) ? (Size + OSPageSize - (Size % OSPageSize)) : Size;
@@ -717,7 +717,7 @@ void FApplePlatformMemory::BinnedFreeToOS(void* Ptr, SIZE_T Size)
 		}
 	}
 #else
-	LLM(FLowLevelMemTracker::Get().OnLowLevelFree(ELLMTracker::Platform, Ptr));
+	LLM_IF_ENABLED(FLowLevelMemTracker::Get().OnLowLevelFree(ELLMTracker::Platform, Ptr));
 	if (munmap(Ptr, Size) != 0)
 	{
 		const int ErrNo = errno;
diff --git a/Engine/Source/Runtime/Core/Private/GenericPlatform/GenericPlatformMallocCrash.cpp b/Engine/Source/Runtime/Core/Private/GenericPlatform/GenericPlatformMallocCrash.cpp
--- a/Engine/Source/Runtime/Core/Private/GenericPlatform/GenericPlatformMallocCrash.cpp
+++ b/Engine/Source/Runtime/Core/Private/GenericPlatform/GenericPlatformMallocCrash.cpp
@@ -242,9 +242,9 @@ FGenericPlatformMallocCrash::FGenericPlatformMallocCrash( FMalloc* MainMalloc )
 	SmallMemoryPool = (uint8*)FPlatformMemory::BinnedAllocFromOS((SIZE_T)SmallMemoryPoolSize);
 	BookkeepingPool = (uint8*)FPlatformMemory::BinnedAllocFromOS((SIZE_T)BookkeepingPoolSize);
 
-	LLM(FLowLevelMemTracker::Get().OnLowLevelAlloc(ELLMTracker::Default, LargeMemoryPool, LargeMemoryPoolSize));
-	LLM(FLowLevelMemTracker::Get().OnLowLevelAlloc(ELLMTracker::Default, SmallMemoryPool, SmallMemoryPoolSize));
-	LLM(FLowLevelMemTracker::Get().OnLowLevelAlloc(ELLMTracker::Default, BookkeepingPool, BookkeepingPoolSize));
+	LLM_IF_ENABLED(FLowLevelMemTracker::Get().OnLowLevelAlloc(ELLMTracker::Default, LargeMemoryPool, LargeMemoryPoolSize));
+	LLM_IF_ENABLED(FLowLevelMemTracker::Get().OnLowLevelAlloc(ELLMTracker::Default, SmallMemoryPool, SmallMemoryPoolSize));
+	LLM_IF_ENABLED(FLowLevelMemTracker::Get().OnLowLevelAlloc(ELLMTracker::Default, BookkeepingPool, BookkeepingPoolSize));
 
 	MemoryTrace_Alloc((uint64)LargeMemoryPool, LargeMemoryPoolSize, alignof(uint8), EMemoryTraceRootHeap::SystemMemory);
 	MemoryTrace_Alloc((uint64)SmallMemoryPool, SmallMemoryPoolSize, alignof(uint8), EMemoryTraceRootHeap::SystemMemory);
diff --git a/Engine/Source/Runtime/Core/Private/GenericPlatform/GenericPlatformMemory.cpp b/Engine/Source/Runtime/Core/Private/GenericPlatform/GenericPlatformMemory.cpp
--- a/Engine/Source/Runtime/Core/Private/GenericPlatform/GenericPlatformMemory.cpp
+++ b/Engine/Source/Runtime/Core/Private/GenericPlatform/GenericPlatformMemory.cpp
@@ -184,7 +184,7 @@ void FGenericPlatformMemory::SetupMemoryPools()
 		BackupOOMMemoryPool = FPlatformMemory::BinnedAllocFromOS(FPlatformMemory::GetBackMemoryPoolSize());
 
 		MemoryTrace_Alloc((uint64)BackupOOMMemoryPool, FPlatformMemory::GetBackMemoryPoolSize(), alignof(void*), EMemoryTraceRootHeap::SystemMemory);
-		LLM(FLowLevelMemTracker::Get().OnLowLevelAlloc(ELLMTracker::Default, BackupOOMMemoryPool, FPlatformMemory::GetBackMemoryPoolSize()));
+		LLM_IF_ENABLED(FLowLevelMemTracker::Get().OnLowLevelAlloc(ELLMTracker::Default, BackupOOMMemoryPool, FPlatformMemory::GetBackMemoryPoolSize()));
 	}
 }
 
@@ -222,7 +222,7 @@ void FGenericPlatformMemory::OnOutOfMemory(uint64 Size, uint32 Alignment)
 			FPlatformMemory::BinnedFreeToOS(BackupOOMMemoryPool, FPlatformMemory::GetBackMemoryPoolSize());
 			UE_LOG(LogMemory, Warning, TEXT("Freeing %d bytes from backup pool to handle out of memory."), FPlatformMemory::GetBackMemoryPoolSize());
         
-			LLM(FLowLevelMemTracker::Get().OnLowLevelFree(ELLMTracker::Default, BackupOOMMemoryPool));
+			LLM_IF_ENABLED(FLowLevelMemTracker::Get().OnLowLevelFree(ELLMTracker::Default, BackupOOMMemoryPool));
 			MemoryTrace_Free((uint64)BackupOOMMemoryPool, EMemoryTraceRootHeap::SystemMemory);
 		}
 
diff --git a/Engine/Source/Runtime/Core/Private/HAL/Allocators/PooledVirtualMemoryAllocator.cpp b/Engine/Source/Runtime/Core/Private/HAL/Allocators/PooledVirtualMemoryAllocator.cpp
--- a/Engine/Source/Runtime/Core/Private/HAL/Allocators/PooledVirtualMemoryAllocator.cpp
+++ b/Engine/Source/Runtime/Core/Private/HAL/Allocators/PooledVirtualMemoryAllocator.cpp
@@ -85,7 +85,7 @@ void* FPooledVirtualMemoryAllocator::Allocate(SIZE_T Size, uint32 AllocationHint
 			if (void* Ptr = Desc.Pool->Allocate(Size))
 			{
 				// LLM wants to be informed of the allocations of physical RAM, this is the closest we can get.
-				LLM(FLowLevelMemTracker::Get().OnLowLevelAlloc(ELLMTracker::Platform, Ptr, Size));
+				LLM_IF_ENABLED(FLowLevelMemTracker::Get().OnLowLevelAlloc(ELLMTracker::Platform, Ptr, Size));
 				return Ptr;
 			}
 		}
@@ -112,7 +112,7 @@ void* FPooledVirtualMemoryAllocator::Allocate(SIZE_T Size, uint32 AllocationHint
 		void* Ptr = Desc.Pool->Allocate(Size);
 
 		// LLM wants to be informed of the allocations of physical RAM, this is the closest we can get.
-		LLM(FLowLevelMemTracker::Get().OnLowLevelAlloc(ELLMTracker::Platform, Ptr, Size));
+		LLM_IF_ENABLED(FLowLevelMemTracker::Get().OnLowLevelAlloc(ELLMTracker::Platform, Ptr, Size));
 		return Ptr;
 	}
 };
@@ -154,7 +154,7 @@ void FPooledVirtualMemoryAllocator::Free(void* Ptr, SIZE_T Size, FCriticalSectio
 			{
 				// LLVM wants to be informed of the allocations of physical RAM.
 				// This is the closest we can get.
-				LLM(FLowLevelMemTracker::Get().OnLowLevelFree(ELLMTracker::Platform, Ptr));
+				LLM_IF_ENABLED(FLowLevelMemTracker::Get().OnLowLevelFree(ELLMTracker::Platform, Ptr));
 				Desc.Pool->Free(Ptr, Size);
 
 				// check if the pool is empty and delete if so
diff --git a/Engine/Source/Runtime/Core/Private/HAL/MallocBinned.cpp b/Engine/Source/Runtime/Core/Private/HAL/MallocBinned.cpp
--- a/Engine/Source/Runtime/Core/Private/HAL/MallocBinned.cpp
+++ b/Engine/Source/Runtime/Core/Private/HAL/MallocBinned.cpp
@@ -884,7 +884,7 @@ struct FMallocBinned::Private
 	static FORCEINLINE void SmallOSFree(FMallocBinned& Allocator, void* Ptr, SIZE_T Size)
 	{
 #if PLATFORM_IOS
-		LLM(FLowLevelMemTracker::Get().OnLowLevelFree(ELLMTracker::Platform, Ptr));
+		LLM_IF_ENABLED(FLowLevelMemTracker::Get().OnLowLevelFree(ELLMTracker::Platform, Ptr));
 		::free(Ptr);
 		Allocator.bNanoMallocAvailable = true;
 #else
@@ -965,7 +965,7 @@ struct FMallocBinned::Private
 		{
 			UE_LOG(LogTemp, Warning, TEXT("malloc failure allocating %d, error code: %d"), NewSize, errno);
 		}
-		LLM(FLowLevelMemTracker::Get().OnLowLevelAlloc(ELLMTracker::Platform, Ptr, NewSize));
+		LLM_IF_ENABLED(FLowLevelMemTracker::Get().OnLowLevelAlloc(ELLMTracker::Platform, Ptr, NewSize));
 		return Ptr;
 #else
 		(void)OutActualSize;
@@ -1436,9 +1436,9 @@ void* FMallocBinned::Realloc( void* Ptr, SIZE_T NewSize, uint32 Alignment )
             if ((NewSize <= Private::SMALL_BLOCK_POOL_SIZE) && (Alignment <= Private::DEFAULT_BINNED_ALLOCATOR_ALIGNMENT))
             {
                 LLM_PLATFORM_SCOPE(ELLMTag::FMalloc);
-                LLM(FLowLevelMemTracker::Get().OnLowLevelFree(ELLMTracker::Platform, Ptr));
+				LLM_IF_ENABLED(FLowLevelMemTracker::Get().OnLowLevelFree(ELLMTracker::Platform, Ptr));
                 NewPtr = ::realloc(Ptr, NewSize);
-                LLM(FLowLevelMemTracker::Get().OnLowLevelAlloc(ELLMTracker::Platform, NewPtr, NewSize));
+				LLM_IF_ENABLED(FLowLevelMemTracker::Get().OnLowLevelAlloc(ELLMTracker::Platform, NewPtr, NewSize));
                 bFallback = !FPlatformMemory::PtrIsFromNanoMalloc(NewPtr);
             }
 
diff --git a/Engine/Source/Runtime/Core/Private/HAL/MallocBinned3.cpp b/Engine/Source/Runtime/Core/Private/HAL/MallocBinned3.cpp
--- a/Engine/Source/Runtime/Core/Private/HAL/MallocBinned3.cpp
+++ b/Engine/Source/Runtime/Core/Private/HAL/MallocBinned3.cpp
@@ -912,11 +912,11 @@ void FMallocBinned3::Commit(uint32 InPoolIndex, void *Ptr, SIZE_T Size)
 #else
 	PoolBaseVMBlock[InPoolIndex].CommitByPtr(Ptr, Size);
 #endif
-	LLM(FLowLevelMemTracker::Get().OnLowLevelAlloc(ELLMTracker::Platform, Ptr, Size));
+	LLM_IF_ENABLED(FLowLevelMemTracker::Get().OnLowLevelAlloc(ELLMTracker::Platform, Ptr, Size));
 }
 void FMallocBinned3::Decommit(uint32 InPoolIndex, void *Ptr, SIZE_T Size)
 {
-	LLM(FLowLevelMemTracker::Get().OnLowLevelFree(ELLMTracker::Platform, Ptr));
+	LLM_IF_ENABLED(FLowLevelMemTracker::Get().OnLowLevelFree(ELLMTracker::Platform, Ptr));
 
 #if BINNED3_ALLOCATOR_STATS
 	Binned3Decommits++;
@@ -935,7 +935,7 @@ void* FMallocBinned3::AllocateMetaDataMemory(SIZE_T Size)
 	FPlatformMemory::FPlatformVirtualMemoryBlock Block = FPlatformMemory::FPlatformVirtualMemoryBlock::AllocateVirtual(VirtualAlignedSize);
 	size_t CommitAlignedSize = Align(Size, FPlatformMemory::FPlatformVirtualMemoryBlock::GetCommitAlignment());
 	Block.Commit(0, CommitAlignedSize);
-	LLM(FLowLevelMemTracker::Get().OnLowLevelAlloc(ELLMTracker::Platform, Block.GetVirtualPointer(), CommitAlignedSize));
+	LLM_IF_ENABLED(FLowLevelMemTracker::Get().OnLowLevelAlloc(ELLMTracker::Platform, Block.GetVirtualPointer(), CommitAlignedSize));
 	return Block.GetVirtualPointer();
 }
 
@@ -943,7 +943,7 @@ void FMallocBinned3::FreeMetaDataMemory(void *Ptr, SIZE_T InSize)
 {
 	if (Ptr)
 	{
-		LLM(FLowLevelMemTracker::Get().OnLowLevelFree(ELLMTracker::Platform, Ptr));
+		LLM_IF_ENABLED(FLowLevelMemTracker::Get().OnLowLevelFree(ELLMTracker::Platform, Ptr));
 
 		size_t VirtualAlignedSize = Align(InSize, FPlatformMemory::FPlatformVirtualMemoryBlock::GetVirtualSizeAlignment());
 		FPlatformMemory::FPlatformVirtualMemoryBlock Block(Ptr, VirtualAlignedSize / FPlatformMemory::FPlatformVirtualMemoryBlock::GetVirtualSizeAlignment());
@@ -1095,7 +1095,7 @@ void* FMallocBinned3::MallocExternal(SIZE_T Size, uint32 Alignment)
 	FPlatformMemory::FPlatformVirtualMemoryBlock Block = FPlatformMemory::FPlatformVirtualMemoryBlock::AllocateVirtual(AlignedSize, Alignment);
 	Block.Commit(0, AlignedSize);
 	void* Result = Block.GetVirtualPointer();
-	LLM(FLowLevelMemTracker::Get().OnLowLevelAlloc(ELLMTracker::Platform, Result, AlignedSize));
+	LLM_IF_ENABLED(FLowLevelMemTracker::Get().OnLowLevelAlloc(ELLMTracker::Platform, Result, AlignedSize));
 #endif
 
 #if BINNED3_TIME_LARGE_BLOCKS
@@ -1284,7 +1284,7 @@ void FMallocBinned3::FreeExternal(void* Ptr)
 		double StartTime = FPlatformTime::Seconds();
 #endif
 		{
-			LLM(FLowLevelMemTracker::Get().OnLowLevelFree(ELLMTracker::Platform, Ptr));
+			LLM_IF_ENABLED(FLowLevelMemTracker::Get().OnLowLevelFree(ELLMTracker::Platform, Ptr));
 #if USE_CACHED_PAGE_ALLOCATOR_FOR_LARGE_ALLOCS
 			GetCachedOSPageAllocator().Free(Ptr, VMPages * FPlatformMemory::FPlatformVirtualMemoryBlock::GetCommitAlignment());
 #else
diff --git a/Engine/Source/Runtime/Core/Private/IOS/IOSPlatformFile.cpp b/Engine/Source/Runtime/Core/Private/IOS/IOSPlatformFile.cpp
--- a/Engine/Source/Runtime/Core/Private/IOS/IOSPlatformFile.cpp
+++ b/Engine/Source/Runtime/Core/Private/IOS/IOSPlatformFile.cpp
@@ -484,7 +484,7 @@ public:
 			UE_LOG(LogIOS, Warning, TEXT("Failed to map memory %s, error is %d"), *Filename, errno);
 			return nullptr;
 		}
-		LLM(FLowLevelMemTracker::Get().OnLowLevelAlloc(ELLMTracker::Platform, AlignedMapPtr, AlignedSize));
+		LLM_IF_ENABLED(FLowLevelMemTracker::Get().OnLowLevelAlloc(ELLMTracker::Platform, AlignedMapPtr, AlignedSize));
 
 		// create a mapping for this range
 		const uint8* MapPtr = AlignedMapPtr + Offset - AlignedOffset;
@@ -499,7 +499,7 @@ public:
 		check(NumOutstandingRegions > 0);
 		NumOutstandingRegions--;
 		
-		LLM(FLowLevelMemTracker::Get().OnLowLevelFree(ELLMTracker::Platform, (void*)Region->AlignedPtr));
+		LLM_IF_ENABLED(FLowLevelMemTracker::Get().OnLowLevelFree(ELLMTracker::Platform, (void*)Region->AlignedPtr));
 		int Res = munmap((void*)Region->AlignedPtr, Region->AlignedSize);
 		checkf(Res == 0, TEXT("Failed to unmap, error is %d, errno is %d [params: %x, %d]"), Res, errno, MappedPtr, GetFileSize());
 	}
diff --git a/Engine/Source/Runtime/Core/Private/IOS/IOSPlatformMemory.cpp b/Engine/Source/Runtime/Core/Private/IOS/IOSPlatformMemory.cpp
--- a/Engine/Source/Runtime/Core/Private/IOS/IOSPlatformMemory.cpp
+++ b/Engine/Source/Runtime/Core/Private/IOS/IOSPlatformMemory.cpp
@@ -31,7 +31,7 @@ void FIOSPlatformMemory::OnOutOfMemory(uint64 Size, uint32 Alignment)
 			FPlatformMemory::BinnedFreeToOS(BackupOOMMemoryPool, FPlatformMemory::GetBackMemoryPoolSize());
 			UE_LOG(LogMemory, Warning, TEXT("Freeing %d bytes from backup pool to handle out of memory."), FPlatformMemory::GetBackMemoryPoolSize());
         
-			LLM(FLowLevelMemTracker::Get().OnLowLevelFree(ELLMTracker::Default, BackupOOMMemoryPool));
+			LLM_IF_ENABLED(FLowLevelMemTracker::Get().OnLowLevelFree(ELLMTracker::Default, BackupOOMMemoryPool));
 		}
     
 		UE_LOG(LogMemory, Warning, TEXT("MemoryStats:")\
diff --git a/Engine/Source/Runtime/Core/Private/Microsoft/MicrosoftRunnableThread.h b/Engine/Source/Runtime/Core/Private/Microsoft/MicrosoftRunnableThread.h
--- a/Engine/Source/Runtime/Core/Private/Microsoft/MicrosoftRunnableThread.h
+++ b/Engine/Source/Runtime/Core/Private/Microsoft/MicrosoftRunnableThread.h
@@ -106,8 +106,8 @@ public:
 #if UE_MEMORY_TRACE_ENABLED || ENABLE_LOW_LEVEL_MEM_TRACKER
 		const uint64 FakeAddress = uint64(this) | (1ull << 47);
 		MemoryTrace_Free(FakeAddress);
-		LLM(FLowLevelMemTracker::Get().OnLowLevelFree(ELLMTracker::Default, (const void*)FakeAddress));
-		LLM(FLowLevelMemTracker::Get().OnLowLevelFree(ELLMTracker::Platform, (const void*)FakeAddress));
+		LLM_IF_ENABLED(FLowLevelMemTracker::Get().OnLowLevelFree(ELLMTracker::Default, (const void*)FakeAddress));
+		LLM_IF_ENABLED(FLowLevelMemTracker::Get().OnLowLevelFree(ELLMTracker::Platform, (const void*)FakeAddress));
 #endif // UE_MEMORY_TRACE_ENABLED || ENABLE_LOW_LEVEL_MEM_TRACKER
 
 		return bDidExitOK;
@@ -158,8 +158,8 @@ protected:
 			MemoryTrace_Alloc(FakeAddress, ActualStackSize, DefaultAlignment);
 			MemoryTrace_MarkAllocAsHeap(FakeAddress, EMemoryTraceRootHeap::SystemMemory);
 			MemoryTrace_Alloc(FakeAddress, ActualStackSize, DefaultAlignment);
-			LLM(FLowLevelMemTracker::Get().OnLowLevelAlloc(ELLMTracker::Default, (const void*)FakeAddress, ActualStackSize));
-			LLM(FLowLevelMemTracker::Get().OnLowLevelAlloc(ELLMTracker::Platform, (const void*)FakeAddress, ActualStackSize));
+			LLM_IF_ENABLED(FLowLevelMemTracker::Get().OnLowLevelAlloc(ELLMTracker::Default, (const void*)FakeAddress, ActualStackSize));
+			LLM_IF_ENABLED(FLowLevelMemTracker::Get().OnLowLevelAlloc(ELLMTracker::Platform, (const void*)FakeAddress, ActualStackSize));
 #endif // UE_MEMORY_TRACE_ENABLED || ENABLE_LOW_LEVEL_MEM_TRACKER
 
 			// Create the thread as suspended, so we can ensure ThreadId is initialized and the thread manager knows about the thread before it runs.
@@ -172,8 +172,8 @@ protected:
 #if UE_MEMORY_TRACE_ENABLED || ENABLE_LOW_LEVEL_MEM_TRACKER
 			const uint64 FakeAddress = uint64(this) | (1ull << 47);
 			MemoryTrace_Free(FakeAddress);
-			LLM(FLowLevelMemTracker::Get().OnLowLevelFree(ELLMTracker::Default, (const void*)FakeAddress));
-			LLM(FLowLevelMemTracker::Get().OnLowLevelFree(ELLMTracker::Platform, (const void*)FakeAddress));
+			LLM_IF_ENABLED(FLowLevelMemTracker::Get().OnLowLevelFree(ELLMTracker::Default, (const void*)FakeAddress));
+			LLM_IF_ENABLED(FLowLevelMemTracker::Get().OnLowLevelFree(ELLMTracker::Platform, (const void*)FakeAddress));
 #endif // UE_MEMORY_TRACE_ENABLED || ENABLE_LOW_LEVEL_MEM_TRACKER
 
 			Runnable = nullptr;
diff --git a/Engine/Source/Runtime/Core/Private/Unix/UnixPlatformFile.cpp b/Engine/Source/Runtime/Core/Private/Unix/UnixPlatformFile.cpp
--- a/Engine/Source/Runtime/Core/Private/Unix/UnixPlatformFile.cpp
+++ b/Engine/Source/Runtime/Core/Private/Unix/UnixPlatformFile.cpp
@@ -795,7 +795,7 @@ public:
 			UE_LOG(LogUnixPlatformFile, Warning, TEXT("Failed to map memory %s, error is %d"), *Filename, errno);
 			return nullptr;
 		}
-		LLM(FLowLevelMemTracker::Get().OnLowLevelAlloc(ELLMTracker::Platform, AlignedMapPtr, AlignedSize));
+		LLM_IF_ENABLED(FLowLevelMemTracker::Get().OnLowLevelAlloc(ELLMTracker::Platform, AlignedMapPtr, AlignedSize));
 
 		// create a mapping for this range
 		const uint8* MapPtr = AlignedMapPtr + Offset - AlignedOffset;
@@ -810,7 +810,7 @@ public:
 		check(NumOutstandingRegions > 0);
 		NumOutstandingRegions--;
 
-		LLM(FLowLevelMemTracker::Get().OnLowLevelFree(ELLMTracker::Platform, Region->AlignedPtr));
+		LLM_IF_ENABLED(FLowLevelMemTracker::Get().OnLowLevelFree(ELLMTracker::Platform, Region->AlignedPtr));
 		const int Res = munmap(const_cast<uint8*>(Region->AlignedPtr), Region->AlignedSize);
 		checkf(Res == 0, TEXT("Failed to unmap, error is %d, errno is %d [params: %x, %d]"), Res, errno, MappedPtr, GetFileSize());
 	}
diff --git a/Engine/Source/Runtime/Core/Private/Unix/UnixPlatformMemory.cpp b/Engine/Source/Runtime/Core/Private/Unix/UnixPlatformMemory.cpp
--- a/Engine/Source/Runtime/Core/Private/Unix/UnixPlatformMemory.cpp
+++ b/Engine/Source/Runtime/Core/Private/Unix/UnixPlatformMemory.cpp
@@ -540,7 +540,7 @@ void* FUnixPlatformMemory::BinnedAllocFromOS(SIZE_T Size)
 		AllocDescriptor->OriginalSizeAsPassed = Size;
 	}
 
-	LLM(FLowLevelMemTracker::Get().OnLowLevelAlloc(ELLMTracker::Platform, Pointer, Size));
+	LLM_IF_ENABLED(FLowLevelMemTracker::Get().OnLowLevelAlloc(ELLMTracker::Platform, Pointer, Size));
 	UE::FForkPageProtector::Get().AddMemoryRegion(Pointer, Size);
 
 	return Pointer;
@@ -548,7 +548,7 @@ void* FUnixPlatformMemory::BinnedAllocFromOS(SIZE_T Size)
 
 void FUnixPlatformMemory::BinnedFreeToOS(void* Ptr, SIZE_T Size)
 {
-	LLM(FLowLevelMemTracker::Get().OnLowLevelFree(ELLMTracker::Platform, Ptr));
+	LLM_IF_ENABLED(FLowLevelMemTracker::Get().OnLowLevelFree(ELLMTracker::Platform, Ptr));
 	// guard against someone not passing size in whole pages
 	static SIZE_T OSPageSize = FPlatformMemory::GetConstants().PageSize;
 	SIZE_T SizeInWholePages = (Size % OSPageSize) ? (Size + OSPageSize - (Size % OSPageSize)) : Size;
diff --git a/Engine/Source/Runtime/Core/Private/Windows/WindowsPlatformMemory.cpp b/Engine/Source/Runtime/Core/Private/Windows/WindowsPlatformMemory.cpp
--- a/Engine/Source/Runtime/Core/Private/Windows/WindowsPlatformMemory.cpp
+++ b/Engine/Source/Runtime/Core/Private/Windows/WindowsPlatformMemory.cpp
@@ -456,13 +456,13 @@ void* FWindowsPlatformMemory::BinnedAllocFromOS( SIZE_T Size )
 	}
 #endif
 	void* Ptr = VirtualAlloc( NULL, Size, MEM_COMMIT | MEM_RESERVE, PAGE_READWRITE );
-	LLM(FLowLevelMemTracker::Get().OnLowLevelAlloc(ELLMTracker::Platform, Ptr, Size));
+	LLM_IF_ENABLED(FLowLevelMemTracker::Get().OnLowLevelAlloc(ELLMTracker::Platform, Ptr, Size));
 	return Ptr;
 }
 
 void FWindowsPlatformMemory::BinnedFreeToOS( void* Ptr, SIZE_T Size )
 {
-	LLM(FLowLevelMemTracker::Get().OnLowLevelFree(ELLMTracker::Platform, Ptr));
+	LLM_IF_ENABLED(FLowLevelMemTracker::Get().OnLowLevelFree(ELLMTracker::Platform, Ptr));
 
 	CA_SUPPRESS(6001)
 	// Windows maintains the size of allocation internally, so Size is unused
diff --git a/Engine/Source/Runtime/Core/Public/Misc/GenericGrowableAllocator.h b/Engine/Source/Runtime/Core/Public/Misc/GenericGrowableAllocator.h
--- a/Engine/Source/Runtime/Core/Public/Misc/GenericGrowableAllocator.h
+++ b/Engine/Source/Runtime/Core/Public/Misc/GenericGrowableAllocator.h
@@ -267,7 +267,7 @@ public:
 #if !UE_BUILD_SHIPPING
 			Allocation->OwnerType = OwnerType;
 #endif
-			LLM(FLowLevelMemTracker::Get().OnLowLevelAlloc(ELLMTracker::Default, GetAddressForTracking(Offset), Size));
+			LLM_IF_ENABLED(FLowLevelMemTracker::Get().OnLowLevelAlloc(ELLMTracker::Default, GetAddressForTracking(Offset), Size));
 			MemoryTrace_Alloc(uint64(GetAddressForTracking(Offset)), Size, Alignment, EMemoryTraceRootHeap::SystemMemory);
 
 			// let the implementation fill in any more
@@ -290,7 +290,7 @@ public:
 		uint64 AllocationSize = Padding + Size;
 		uint32 Offset = Memory->Offset;
 
-		LLM(FLowLevelMemTracker::Get().OnLowLevelFree(ELLMTracker::Default, GetAddressForTracking(Offset)));
+		LLM_IF_ENABLED(FLowLevelMemTracker::Get().OnLowLevelFree(ELLMTracker::Default, GetAddressForTracking(Offset)));
 		MemoryTrace_Free(uint64(GetAddressForTracking(Offset)), EMemoryTraceRootHeap::SystemMemory);
 
 		// we are now done with the Allocation object
diff --git a/Engine/Source/Runtime/D3D12RHI/Private/D3D12Allocation.cpp b/Engine/Source/Runtime/D3D12RHI/Private/D3D12Allocation.cpp
--- a/Engine/Source/Runtime/D3D12RHI/Private/D3D12Allocation.cpp
+++ b/Engine/Source/Runtime/D3D12RHI/Private/D3D12Allocation.cpp
@@ -453,7 +453,7 @@ void FD3D12BuddyAllocator::Allocate(uint32 SizeInBytes, uint32 Alignment, FD3D12
 
 	// track the allocation
 #if !PLATFORM_WINDOWS
-	LLM(FLowLevelMemTracker::Get().OnLowLevelAlloc(ELLMTracker::Default, ResourceLocation.GetAddressForLLMTracking(), SizeInBytes));
+	LLM_IF_ENABLED(FLowLevelMemTracker::Get().OnLowLevelAlloc(ELLMTracker::Default, ResourceLocation.GetAddressForLLMTracking(), SizeInBytes));
 	// Note: Disabling this LLM hook for Windows is due to a work-around in the way that d3d12 buffers are tracked
 	// by LLM. LLM tracks buffer data in the UpdateBufferStats function because that is the easiest place to ensure that LLM
 	// can be updated whenever a buffer is created or released. Unfortunately, some buffers allocate from this allocator
@@ -520,7 +520,7 @@ void FD3D12BuddyAllocator::Deallocate(FD3D12ResourceLocation& ResourceLocation)
 	// which means that the memory would be counted twice. Because of this the tracking had to be disabled here.
 	// This does mean that non-buffer memory that goes through this allocator won't be tracked, so this does need a better solution.
 	// see UpdateBufferStats for a more detailed explanation.
-	LLM(FLowLevelMemTracker::Get().OnLowLevelFree(ELLMTracker::Default, ResourceLocation.GetAddressForLLMTracking()));
+	LLM_IF_ENABLED(FLowLevelMemTracker::Get().OnLowLevelFree(ELLMTracker::Default, ResourceLocation.GetAddressForLLMTracking()));
 #if UE_MEMORY_TRACE_ENABLED
 	MemoryTrace_Free(ResourceLocation.GetGPUVirtualAddress(), EMemoryTraceRootHeap::VideoMemory);
 #endif
diff --git a/Engine/Source/Runtime/D3D12RHI/Private/D3D12Texture.cpp b/Engine/Source/Runtime/D3D12RHI/Private/D3D12Texture.cpp
--- a/Engine/Source/Runtime/D3D12RHI/Private/D3D12Texture.cpp
+++ b/Engine/Source/Runtime/D3D12RHI/Private/D3D12Texture.cpp
@@ -334,8 +334,8 @@ void FD3D12TextureStats::UpdateD3D12TextureStats(FD3D12Texture& Texture, const F
 #if PLATFORM_WINDOWS
 		// On Windows there is no way to hook into the low level d3d allocations and frees.
 		// This means that we must manually add the tracking here.
-		LLM(FLowLevelMemTracker::Get().OnLowLevelAlloc(ELLMTracker::Platform, Texture.GetResource(), TextureSize, ELLMTag::GraphicsPlatform));
-		LLM(FLowLevelMemTracker::Get().OnLowLevelAlloc(ELLMTracker::Default, Texture.GetResource(), TextureSize, ELLMTag::Textures));
+		LLM_IF_ENABLED(FLowLevelMemTracker::Get().OnLowLevelAlloc(ELLMTracker::Platform, Texture.GetResource(), TextureSize, ELLMTag::GraphicsPlatform));
+		LLM_IF_ENABLED(FLowLevelMemTracker::Get().OnLowLevelAlloc(ELLMTracker::Default, Texture.GetResource(), TextureSize, ELLMTag::Textures));
 		{
 			LLM(UE_MEMSCOPE_DEFAULT(ELLMTag::Textures));
 
@@ -355,8 +355,8 @@ void FD3D12TextureStats::UpdateD3D12TextureStats(FD3D12Texture& Texture, const F
 	else
 	{
 #if PLATFORM_WINDOWS
-		LLM(FLowLevelMemTracker::Get().OnLowLevelFree(ELLMTracker::Platform, Texture.GetResource()));
-		LLM(FLowLevelMemTracker::Get().OnLowLevelFree(ELLMTracker::Default, Texture.GetResource()));
+		LLM_IF_ENABLED(FLowLevelMemTracker::Get().OnLowLevelFree(ELLMTracker::Platform, Texture.GetResource()));
+		LLM_IF_ENABLED(FLowLevelMemTracker::Get().OnLowLevelFree(ELLMTracker::Default, Texture.GetResource()));
 
 #if UE_MEMORY_TRACE_ENABLED
 		// Skip back buffers that aren't traced on alloc and don't have valid GPUVirtualAddress
diff --git a/Engine/Source/Runtime/D3D12RHI/Private/D3D12TransientResourceAllocator.cpp b/Engine/Source/Runtime/D3D12RHI/Private/D3D12TransientResourceAllocator.cpp
--- a/Engine/Source/Runtime/D3D12RHI/Private/D3D12TransientResourceAllocator.cpp
+++ b/Engine/Source/Runtime/D3D12RHI/Private/D3D12TransientResourceAllocator.cpp
@@ -78,7 +78,7 @@ FD3D12TransientHeap::FD3D12TransientHeap(const FInitializer& Initializer, FD3D12
 #if PLATFORM_WINDOWS
 		// On Windows there is no way to hook into the low level d3d allocations and frees.
 		// This means that we must manually add the tracking here.
-		LLM(FLowLevelMemTracker::Get().OnLowLevelAlloc(ELLMTracker::Platform, D3DHeap, Desc.SizeInBytes, ELLMTag::GraphicsPlatform));
+		LLM_IF_ENABLED(FLowLevelMemTracker::Get().OnLowLevelAlloc(ELLMTracker::Platform, D3DHeap, Desc.SizeInBytes, ELLMTag::GraphicsPlatform));
 		MemoryTrace_Alloc((uint64)D3DHeap, Desc.SizeInBytes, 0, EMemoryTraceRootHeap::VideoMemory);
 		// Boost priority to make sure it's not paged out
 		TRefCountPtr<ID3D12Device5> D3DDevice5;
@@ -119,7 +119,7 @@ FD3D12TransientHeap::~FD3D12TransientHeap()
 		DEC_MEMORY_STAT_BY(STAT_D3D12TransientHeaps, Desc.SizeInBytes);
 		DEC_MEMORY_STAT_BY(STAT_D3D12MemoryCurrentTotal, Desc.SizeInBytes);
 #if PLATFORM_WINDOWS
-		LLM(FLowLevelMemTracker::Get().OnLowLevelFree(ELLMTracker::Platform, Heap->GetHeap()));
+		LLM_IF_ENABLED(FLowLevelMemTracker::Get().OnLowLevelFree(ELLMTracker::Platform, Heap->GetHeap()));
 		MemoryTrace_Free((uint64)Heap->GetHeap(), EMemoryTraceRootHeap::VideoMemory);
 #endif
 
diff --git a/Engine/Source/Runtime/OpenGLDrv/Private/OpenGLUniformBuffer.cpp b/Engine/Source/Runtime/OpenGLDrv/Private/OpenGLUniformBuffer.cpp
--- a/Engine/Source/Runtime/OpenGLDrv/Private/OpenGLUniformBuffer.cpp
+++ b/Engine/Source/Runtime/OpenGLDrv/Private/OpenGLUniformBuffer.cpp
@@ -433,7 +433,7 @@ void FOpenGLUniformBuffer::SetGLUniformBufferParams(GLuint InResource, uint32 In
 	bStreamDraw = bInStreamDraw;
 	bOwnsResource = true;
 
-	LLM(FLowLevelMemTracker::Get().OnLowLevelAlloc(ELLMTracker::Default, ((uint8*)this)+1, InAllocatedSize)); //+1 because ptr must be unique for LLM
+	LLM_IF_ENABLED(FLowLevelMemTracker::Get().OnLowLevelAlloc(ELLMTracker::Default, ((uint8*)this)+1, InAllocatedSize)); //+1 because ptr must be unique for LLM
 }
 
 FOpenGLUniformBuffer::~FOpenGLUniformBuffer()
@@ -479,7 +479,7 @@ FOpenGLUniformBuffer::~FOpenGLUniformBuffer()
 			Resource = 0; 
 		}
 
-		LLM(FLowLevelMemTracker::Get().OnLowLevelFree(ELLMTracker::Default, ((uint8*)this)+1)); //+1 because ptr must be unique for LLM
+		LLM_IF_ENABLED(FLowLevelMemTracker::Get().OnLowLevelFree(ELLMTracker::Default, ((uint8*)this)+1)); //+1 because ptr must be unique for LLM
 	}
 }
 
diff --git a/Engine/Source/Runtime/RHI/Private/GPUDefragAllocator.cpp b/Engine/Source/Runtime/RHI/Private/GPUDefragAllocator.cpp
--- a/Engine/Source/Runtime/RHI/Private/GPUDefragAllocator.cpp
+++ b/Engine/Source/Runtime/RHI/Private/GPUDefragAllocator.cpp
@@ -222,7 +222,7 @@ void* FGPUDefragAllocator::Allocate(int64 AllocationSize, int32 Alignment, TStat
 	check(IsAligned(AllocatedChunk->Base, Alignment));
 
 	MemoryTrace_Alloc((uint64)AllocatedChunk->Base, AllocationSize, Alignment);
-	LLM(FLowLevelMemTracker::Get().OnLowLevelAlloc(ELLMTracker::Default, AllocatedChunk->Base, AllocationSize));
+	LLM_IF_ENABLED(FLowLevelMemTracker::Get().OnLowLevelAlloc(ELLMTracker::Default, AllocatedChunk->Base, AllocationSize));
 
 	return AllocatedChunk->Base;
 }
diff --git a/Engine/Source/Runtime/RHI/Public/GPUDefragAllocator.h b/Engine/Source/Runtime/RHI/Public/GPUDefragAllocator.h
--- a/Engine/Source/Runtime/RHI/Public/GPUDefragAllocator.h
+++ b/Engine/Source/Runtime/RHI/Public/GPUDefragAllocator.h
@@ -795,7 +795,7 @@ protected:
 	{
 		MemoryTrace_ReallocFree((uint64)Source);
 		MemoryTrace_ReallocAlloc((uint64)Dest->Base, Size, 4);
-		LLM(FLowLevelMemTracker::Get().OnLowLevelAllocMoved(ELLMTracker::Default, Dest->Base, Source));
+		LLM_IF_ENABLED(FLowLevelMemTracker::Get().OnLowLevelAllocMoved(ELLMTracker::Default, Dest->Base, Source));
 
 		uint8* DestAddr = Dest->Base + DestOffset;
 		int64 MemDistance = (int64)(DestAddr) - (int64)(Source);
diff --git a/Engine/Source/Runtime/VulkanRHI/Private/VulkanMemory.cpp b/Engine/Source/Runtime/VulkanRHI/Private/VulkanMemory.cpp
--- a/Engine/Source/Runtime/VulkanRHI/Private/VulkanMemory.cpp
+++ b/Engine/Source/Runtime/VulkanRHI/Private/VulkanMemory.cpp
@@ -959,11 +959,11 @@ namespace VulkanRHI
 
 #if VULKAN_USE_LLM
 		LLM_PLATFORM_SCOPE_VULKAN(ELLMTagVulkan::VulkanDriverMemoryGPU);
-		LLM(FLowLevelMemTracker::Get().OnLowLevelAlloc(ELLMTracker::Platform, (void*)NewAllocation->Handle, AllocationSize, ELLMTag::GraphicsPlatform, ELLMAllocType::System));
+		LLM_IF_ENABLED(FLowLevelMemTracker::Get().OnLowLevelAlloc(ELLMTracker::Platform, (void*)NewAllocation->Handle, AllocationSize, ELLMTag::GraphicsPlatform, ELLMAllocType::System));
 		LLM_TRACK_VULKAN_SPARE_MEMORY_GPU((int64)AllocationSize);
 #else
-		LLM(FLowLevelMemTracker::Get().OnLowLevelAlloc(ELLMTracker::Platform, (void*)NewAllocation->Handle, AllocationSize, ELLMTag::GraphicsPlatform));
-		LLM(FLowLevelMemTracker::Get().OnLowLevelAlloc(ELLMTracker::Default, (void*)NewAllocation->Handle, AllocationSize, ELLMTag::Untagged));
+		LLM_IF_ENABLED(FLowLevelMemTracker::Get().OnLowLevelAlloc(ELLMTracker::Platform, (void*)NewAllocation->Handle, AllocationSize, ELLMTag::GraphicsPlatform));
+		LLM_IF_ENABLED(FLowLevelMemTracker::Get().OnLowLevelAlloc(ELLMTracker::Default, (void*)NewAllocation->Handle, AllocationSize, ELLMTag::Untagged));
 #endif
 
 		INC_DWORD_STAT(STAT_VulkanNumPhysicalMemAllocations);
@@ -1042,11 +1042,11 @@ namespace VulkanRHI
 		VulkanRHI::vkFreeMemory(DeviceHandle, Allocation->Handle, VULKAN_CPU_ALLOCATOR);
 
 #if VULKAN_USE_LLM
-		LLM(FLowLevelMemTracker::Get().OnLowLevelFree(ELLMTracker::Platform, (void*)Allocation->Handle, ELLMAllocType::System));
+		LLM_IF_ENABLED(FLowLevelMemTracker::Get().OnLowLevelFree(ELLMTracker::Platform, (void*)Allocation->Handle, ELLMAllocType::System));
 		LLM_TRACK_VULKAN_SPARE_MEMORY_GPU(-(int64)Allocation->Size);
 #else
-		LLM(FLowLevelMemTracker::Get().OnLowLevelFree(ELLMTracker::Platform, (void*)Allocation->Handle));
-		LLM(FLowLevelMemTracker::Get().OnLowLevelFree(ELLMTracker::Default, (void*)Allocation->Handle));
+		LLM_IF_ENABLED(FLowLevelMemTracker::Get().OnLowLevelFree(ELLMTracker::Platform, (void*)Allocation->Handle));
+		LLM_IF_ENABLED(FLowLevelMemTracker::Get().OnLowLevelFree(ELLMTracker::Default, (void*)Allocation->Handle));
 #endif
 
 		--NumAllocations;
diff --git a/Engine/Source/Runtime/Windows/D3D11RHI/Private/D3D11Texture.cpp b/Engine/Source/Runtime/Windows/D3D11RHI/Private/D3D11Texture.cpp
--- a/Engine/Source/Runtime/Windows/D3D11RHI/Private/D3D11Texture.cpp
+++ b/Engine/Source/Runtime/Windows/D3D11RHI/Private/D3D11Texture.cpp
@@ -69,8 +69,8 @@ void UpdateD3D11TextureStats(FD3D11Texture& Texture, bool bAllocating)
 	// This means that we must manually add the tracking here.
 	if (bAllocating)
 	{
-		LLM(FLowLevelMemTracker::Get().OnLowLevelAlloc(ELLMTracker::Platform, Texture.GetResource(), TextureSize, ELLMTag::GraphicsPlatform));
-		LLM(FLowLevelMemTracker::Get().OnLowLevelAlloc(ELLMTracker::Default , Texture.GetResource(), TextureSize, ELLMTag::Textures));
+		LLM_IF_ENABLED(FLowLevelMemTracker::Get().OnLowLevelAlloc(ELLMTracker::Platform, Texture.GetResource(), TextureSize, ELLMTag::GraphicsPlatform));
+		LLM_IF_ENABLED(FLowLevelMemTracker::Get().OnLowLevelAlloc(ELLMTracker::Default , Texture.GetResource(), TextureSize, ELLMTag::Textures));
 		{
 			LLM(UE_MEMSCOPE_DEFAULT(ELLMTag::Textures));
 			MemoryTrace_Alloc((uint64)Texture.GetResource(), TextureSize, 1024, EMemoryTraceRootHeap::VideoMemory);
@@ -78,8 +78,8 @@ void UpdateD3D11TextureStats(FD3D11Texture& Texture, bool bAllocating)
 	}
 	else
 	{
-		LLM(FLowLevelMemTracker::Get().OnLowLevelFree(ELLMTracker::Platform, Texture.GetResource()));
-		LLM(FLowLevelMemTracker::Get().OnLowLevelFree(ELLMTracker::Default , Texture.GetResource()));
+		LLM_IF_ENABLED(FLowLevelMemTracker::Get().OnLowLevelFree(ELLMTracker::Platform, Texture.GetResource()));
+		LLM_IF_ENABLED(FLowLevelMemTracker::Get().OnLowLevelFree(ELLMTracker::Default , Texture.GetResource()));
 		MemoryTrace_Free((uint64)Texture.GetResource(), EMemoryTraceRootHeap::VideoMemory);
 	}
 }
