Fix for Metal Cube map UAV's multiplying array size by 6 when using texture type

diff --git a/Engine/Source/Runtime/Apple/MetalRHI/Private/MetalUAV.cpp b/Engine/Source/Runtime/Apple/MetalRHI/Private/MetalUAV.cpp
--- a/Engine/Source/Runtime/Apple/MetalRHI/Private/MetalUAV.cpp
+++ b/Engine/Source/Runtime/Apple/MetalRHI/Private/MetalUAV.cpp
@@ -245,9 +245,7 @@ void FMetalShaderResourceView::UpdateView()
             uint32_t ArrayStart = Info.ArrayRange.First;
             uint32_t ArraySize = Info.ArrayRange.Num;
             
-            // Check the incoming texture type for whether this a cube or cube array
-            if (TextureType == MTL::TextureTypeCube ||
-                TextureType == MTL::TextureTypeCubeArray)
+			if (Info.Dimension == FRHIViewDesc::EDimension::TextureCube || Info.Dimension == FRHIViewDesc::EDimension::TextureCubeArray)
             {
                 ArrayStart = Info.ArrayRange.First * 6;
                 ArraySize = Info.ArrayRange.Num * 6;
@@ -421,12 +419,11 @@ void FMetalUnorderedAccessView::UpdateView()
             uint32_t ArraySize = Info.ArrayRange.Num;
             
             // Check the incoming texture type for whether this a cube or cube array
-            if (TextureType == MTL::TextureTypeCube ||
-                TextureType == MTL::TextureTypeCubeArray)
-            {
-                ArrayStart = Info.ArrayRange.First * 6;
-                ArraySize = Info.ArrayRange.Num * 6;
-            }
+			if (Info.Dimension == FRHIViewDesc::EDimension::TextureCube || Info.Dimension == FRHIViewDesc::EDimension::TextureCubeArray)
+			{
+				ArrayStart = Info.ArrayRange.First * 6;
+				ArraySize = Info.ArrayRange.Num * 6;
+			}
             
             TextureType = UAVDimensionToMetalTextureType(Info.Dimension);
             
