Fixed texture readback error due to missing fence wait

diff --git a/Engine/Source/Runtime/Apple/MetalRHI/Private/MetalRenderTarget.cpp b/Engine/Source/Runtime/Apple/MetalRHI/Private/MetalRenderTarget.cpp
--- a/Engine/Source/Runtime/Apple/MetalRHI/Private/MetalRenderTarget.cpp
+++ b/Engine/Source/Runtime/Apple/MetalRHI/Private/MetalRenderTarget.cpp
@@ -264,6 +264,11 @@ void FMetalDynamicRHI::RHIMapStagingSurface(FRHITexture* TextureRHI, FRHIGPUFenc
 {
     MTL_SCOPED_AUTORELEASE_POOL;
     
+	if (FenceRHI && !FenceRHI->Poll())
+	{
+		ResourceCast(FenceRHI)->WaitCPU();
+	}
+	
 	FMetalSurface* Surface = GetMetalSurfaceFromRHITexture(TextureRHI);
     
 	uint32 Stride = 0;
diff --git a/Engine/Source/Runtime/Apple/MetalRHI/Private/MetalUAV.cpp b/Engine/Source/Runtime/Apple/MetalRHI/Private/MetalUAV.cpp
--- a/Engine/Source/Runtime/Apple/MetalRHI/Private/MetalUAV.cpp
+++ b/Engine/Source/Runtime/Apple/MetalRHI/Private/MetalUAV.cpp
@@ -694,3 +694,11 @@ bool FMetalGPUFence::Poll() const
 		return false;
 	}
 }
+
+void FMetalGPUFence::WaitCPU() const
+{
+	if (Fence)
+	{
+		Fence->Wait(MAX_uint32);
+	}
+}
diff --git a/Engine/Source/Runtime/Apple/MetalRHI/Public/MetalResources.h b/Engine/Source/Runtime/Apple/MetalRHI/Public/MetalResources.h
--- a/Engine/Source/Runtime/Apple/MetalRHI/Public/MetalResources.h
+++ b/Engine/Source/Runtime/Apple/MetalRHI/Public/MetalResources.h
@@ -522,6 +522,8 @@ public:
 	void WriteInternal(FMetalCommandBuffer* CmdBuffer);
 
 	virtual bool Poll() const override final;
+	
+	void WaitCPU() const;
 
 private:
     TSharedPtr<FMetalCommandBufferFence, ESPMode::ThreadSafe> Fence = nullptr;
