Missing files from submit

diff --git a/Engine/Source/Runtime/Apple/MetalRHI/Private/MetalBindlessDescriptors.cpp b/Engine/Source/Runtime/Apple/MetalRHI/Private/MetalBindlessDescriptors.cpp
new file mode 100644
--- /dev/null
+++ b/Engine/Source/Runtime/Apple/MetalRHI/Private/MetalBindlessDescriptors.cpp
@@ -0,0 +1,254 @@
+// Copyright Epic Games, Inc. All Rights Reserved.
+#include "MetalRHIPrivate.h"
+
+#include "MetalBindlessDescriptors.h"
+
+int32 GBindlessResourceDescriptorHeapSize = 1000 * 1000;
+static FAutoConsoleVariableRef CVarBindlessResourceDescriptorHeapSize(
+    TEXT("Metal.Bindless.ResourceDescriptorHeapSize"),
+    GBindlessResourceDescriptorHeapSize,
+    TEXT("Bindless resource descriptor heap size"),
+    ECVF_ReadOnly
+);
+
+int32 GBindlessSamplerDescriptorHeapSize = 64 << 10; // TODO: We should be able to reduce the size of the sampler heap if we fix static sampler creation.
+static FAutoConsoleVariableRef CVarBindlessSamplerDescriptorHeapSize(
+    TEXT("Metal.Bindless.SamplerDescriptorHeapSize"),
+    GBindlessSamplerDescriptorHeapSize,
+    TEXT("Bindless sampler descriptor heap size"),
+    ECVF_ReadOnly
+);
+
+FMetalDescriptorHeap::FMetalDescriptorHeap(const ERHIDescriptorHeapType DescriptorType)
+    : DeferredDeletionListIndex(0)
+    , ResourceHeap(nil)
+    , Type(DescriptorType)
+{
+
+}
+
+void FMetalDescriptorHeap::Init(const int32 HeapSize)
+{
+    ResourceHeap = GetMetalDeviceContext().CreatePooledBuffer(FMetalPooledBufferArgs(GetMetalDeviceContext().GetDevice(), HeapSize, BUF_Dynamic, MTL::StorageModeShared));
+    Descriptors = reinterpret_cast<IRDescriptorTableEntry*>(ResourceHeap->Contents());
+}
+
+void FMetalDescriptorHeap::Reset()
+{
+    DeferredDeletionListIndex = (GetMetalDeviceContext().GetFrameNumberRHIThread() % NumPendingFrame);
+
+    {
+        FScopeLock ScopeLock(&FreeListCS);
+
+        while (!DeferredDeletionList[DeferredDeletionListIndex].IsEmpty())
+        {
+            FRHIDescriptorHandle DeletionIndex;
+            DeferredDeletionList[DeferredDeletionListIndex].Dequeue(DeletionIndex);
+            FreeList.Enqueue(DeletionIndex.GetIndex());
+        }
+    }
+}
+
+void FMetalDescriptorHeap::FreeDescriptor(FRHIDescriptorHandle DescriptorHandle)
+{
+    DeferredDeletionList[DeferredDeletionListIndex].Enqueue(DescriptorHandle);
+}
+
+uint32 FMetalDescriptorHeap::GetFreeResourceIndex()
+{
+    {
+        FScopeLock ScopeLock(&FreeListCS);
+        if (!FreeList.IsEmpty())
+        {
+            uint32 FreeIndex;
+            FreeList.Dequeue(FreeIndex);
+            return FreeIndex;
+        }
+    }
+
+    NSUInteger MaxDescriptorCount = ResourceHeap->GetLength() / sizeof(IRDescriptorTableEntry);
+    checkf((PeakDescriptorCount + 1) < MaxDescriptorCount, TEXT("Reached Heap Max Capacity (%u/%u)"), PeakDescriptorCount + 1, MaxDescriptorCount);
+
+    const uint32 ResourceIndex = PeakDescriptorCount++;
+    return ResourceIndex;
+}
+
+FRHIDescriptorHandle FMetalDescriptorHeap::ReserveDescriptor()
+{
+    const uint32 ResourceIndex = GetFreeResourceIndex();
+    return FRHIDescriptorHandle(Type, ResourceIndex);
+}
+
+void FMetalDescriptorHeap::UpdateDescriptor(FRHIDescriptorHandle DescriptorHandle, IRDescriptorTableEntry DescriptorData)
+{
+    checkf(DescriptorHandle.IsValid(), TEXT("Attemping to update invalid descriptor handle!"));
+
+    uint32 DescriptorIndex = DescriptorHandle.GetIndex();
+    Descriptors[DescriptorIndex] = DescriptorData;
+}
+
+void FMetalDescriptorHeap::BindHeap(FMetalCommandEncoder* Encoder, MTL::FunctionType FunctionType, const uint32 BindIndex)
+{
+    uint32 DescriptorCount = PeakDescriptorCount.load();
+    const uint64 HeapSize = DescriptorCount * sizeof(IRDescriptorTableEntry);
+
+    Encoder->SetShaderBuffer(FunctionType, ResourceHeap, 0, HeapSize, BindIndex, MTL::ResourceUsageRead);
+}
+
+FMetalBindlessDescriptorManager::FMetalBindlessDescriptorManager()
+    : StandardResources(ERHIDescriptorHeapType::Standard)
+    , SamplerResources(ERHIDescriptorHeapType::Sampler)
+{
+	
+}
+
+FMetalBindlessDescriptorManager::~FMetalBindlessDescriptorManager()
+{
+
+}
+
+void FMetalBindlessDescriptorManager::Init()
+{
+    StandardResources.Init(GBindlessResourceDescriptorHeapSize);
+    SamplerResources.Init(GBindlessSamplerDescriptorHeapSize);
+	
+	bIsSupported = true;
+}
+
+void FMetalBindlessDescriptorManager::Reset()
+{
+	if(!bIsSupported)
+	{
+		return;
+	}
+	
+    for (uint32 Frequency = 0; Frequency < EMetalShaderStages::Num; Frequency++)
+    {
+        ResidentResources[Frequency].SetNum(0);
+    }
+
+    StandardResources.Reset();
+    SamplerResources.Reset();
+}
+
+FRHIDescriptorHandle FMetalBindlessDescriptorManager::ReserveDescriptor(ERHIDescriptorHeapType InType)
+{
+    switch (InType)
+    {
+    case ERHIDescriptorHeapType::Standard:
+        return StandardResources.ReserveDescriptor();
+    case ERHIDescriptorHeapType::Sampler:
+        return SamplerResources.ReserveDescriptor();
+    default:
+        checkNoEntry();
+    };
+
+    return FRHIDescriptorHandle();
+}
+
+void FMetalBindlessDescriptorManager::FreeDescriptor(FRHIDescriptorHandle DescriptorHandle)
+{
+    check(DescriptorHandle.IsValid());
+    switch (DescriptorHandle.GetType())
+    {
+    case ERHIDescriptorHeapType::Standard:
+        StandardResources.FreeDescriptor(DescriptorHandle);
+        break;
+    case ERHIDescriptorHeapType::Sampler:
+        SamplerResources.FreeDescriptor(DescriptorHandle);
+        break;
+    default:
+        checkNoEntry();
+    };
+}
+
+void FMetalBindlessDescriptorManager::BindSampler(FRHIDescriptorHandle DescriptorHandle, MTL::SamplerState* Sampler)
+{
+    IRDescriptorTableEntry DescriptorData = {0};
+    IRDescriptorTableSetSampler(&DescriptorData, Sampler, 0.0f);
+
+    SamplerResources.UpdateDescriptor(DescriptorHandle, DescriptorData);
+}
+
+void FMetalBindlessDescriptorManager::BindResource(FRHIDescriptorHandle DescriptorHandle, FMetalResourceViewBase* Resource)
+{
+    IRDescriptorTableEntry DescriptorData = {0};
+
+    switch (Resource->GetMetalType())
+    {
+    case FMetalResourceViewBase::EMetalType::TextureView:
+        {
+            auto const& View = Resource->GetTextureView();
+
+            IRDescriptorTableSetTexture(&DescriptorData, View.get(), 0.0f, 0u);
+        }
+        break;
+    case FMetalResourceViewBase::EMetalType::BufferView:
+        {
+            auto const& View = Resource->GetBufferView();
+
+            IRDescriptorTableSetBuffer(&DescriptorData, View.Buffer->GetGPUAddress() + View.Offset, View.Size);
+        }
+        break;
+    case FMetalResourceViewBase::EMetalType::TextureBufferBacked:
+        {
+            auto const& View = Resource->GetTextureBufferBacked();
+
+            IRBufferView BufferView;
+            BufferView.buffer = View.Buffer->GetMTLBuffer().get();
+            BufferView.bufferOffset = View.Buffer->GetOffset() + View.Offset;
+            BufferView.bufferSize = View.Size;
+            BufferView.typedBuffer = true;
+            BufferView.textureBufferView = View.Texture.get();
+
+            uint32 Stride = GPixelFormats[View.Format].BlockBytes;
+            uint32 FirstElement = View.Offset / Stride;
+            uint32 NumElement = View.Size / Stride;
+
+            uint64 BufferVA              = View.Buffer->GetGPUAddress() + View.Offset;
+            uint64_t ExtraElement        = (BufferVA % 16) / Stride;
+
+            BufferView.textureViewOffsetInElements = ExtraElement;
+
+            IRDescriptorTableSetBufferView(&DescriptorData, &BufferView);
+        }
+        break;
+#if METAL_RHI_RAYTRACING
+    case FMetalResourceViewBase::EMetalType::AccelerationStructure:
+        {
+            MTL::AccelerationStructure const& AccelerationStructure = Resource->GetAccelerationStructure();
+
+            IRDescriptorTableSetAccelerationStructure(&DescriptorData, [AccelerationStructure.GetPtr() gpuResourceID]._impl);
+        }
+        break;
+#endif
+    default:
+        checkNoEntry();
+        return;
+    };
+
+    StandardResources.UpdateDescriptor(DescriptorHandle, DescriptorData);
+}
+
+void FMetalBindlessDescriptorManager::BindTexture(FRHIDescriptorHandle DescriptorHandle, MTL::Texture* Texture)
+{
+    IRDescriptorTableEntry DescriptorData = {0};
+    IRDescriptorTableSetTexture(&DescriptorData, Texture, 0.0f, 0u);
+    StandardResources.UpdateDescriptor(DescriptorHandle, DescriptorData);
+}
+
+void FMetalBindlessDescriptorManager::BindDescriptorHeapsToEncoder(FMetalCommandEncoder* Encoder, MTL::FunctionType FunctionType, EMetalShaderStages Frequency)
+{
+    for (auto& Resources : ResidentResources[Frequency])
+    {
+        Encoder->UseResource(Resources.Key, Resources.Value);
+    }
+
+    StandardResources.BindHeap(Encoder, FunctionType, kIRStandardHeapBindPoint);
+    SamplerResources.BindHeap(Encoder, FunctionType, kIRSamplerHeapBindPoint);
+}
+
+void FMetalBindlessDescriptorManager::MakeResident(FRHIDescriptorHandle DescriptorHandle, MTL::Resource* Resource, MTL::ResourceUsage Usage, EMetalShaderStages Frequency)
+{
+    ResidentResources[Frequency].AddUnique(TTuple<MTL::Resource*, MTL::ResourceUsage>(Resource, Usage));
+}
diff --git a/Engine/Source/Runtime/Apple/MetalRHI/Private/MetalBindlessDescriptors.h b/Engine/Source/Runtime/Apple/MetalRHI/Private/MetalBindlessDescriptors.h
new file mode 100644
--- /dev/null
+++ b/Engine/Source/Runtime/Apple/MetalRHI/Private/MetalBindlessDescriptors.h
@@ -0,0 +1,71 @@
+// Copyright Epic Games, Inc. All Rights Reserved.
+
+#pragma once
+
+#include <Metal/Metal.h>
+#include "MetalResources.h"
+#include "MetalShaderResources.h"
+
+struct FMetalDescriptorHeap
+{
+                                    FMetalDescriptorHeap(const ERHIDescriptorHeapType DescriptorType);
+
+    void                            Init(const int32 HeapSize);
+
+    void                            Reset();
+
+    FRHIDescriptorHandle            ReserveDescriptor();
+    void                            FreeDescriptor(FRHIDescriptorHandle DescriptorHandle);
+    uint32                          GetFreeResourceIndex();
+
+    void                            UpdateDescriptor(FRHIDescriptorHandle DescriptorHandle, struct IRDescriptorTableEntry DescriptorData);
+    void                            BindHeap(FMetalCommandEncoder* Encoder, MTL::FunctionType FunctionType, const uint32 BindIndex);
+
+    //
+    FCriticalSection                FreeListCS;
+
+    TQueue<uint32>                  FreeList;
+
+    uint32                          DeferredDeletionListIndex;
+
+    static constexpr int32 NumPendingFrame = 3;
+
+    TQueue<FRHIDescriptorHandle>    DeferredDeletionList[NumPendingFrame];
+
+    std::atomic<uint32>             PeakDescriptorCount;
+
+    struct IRDescriptorTableEntry*  Descriptors;
+
+    FMetalBufferPtr  				ResourceHeap;
+
+    const ERHIDescriptorHeapType    Type;
+};
+
+class FMetalBindlessDescriptorManager
+{
+public:
+                            FMetalBindlessDescriptorManager();
+                            ~FMetalBindlessDescriptorManager();
+
+    void                    Init();
+    void                    Reset();
+
+    FRHIDescriptorHandle     ReserveDescriptor(ERHIDescriptorHeapType InType);
+    void                    FreeDescriptor(FRHIDescriptorHandle DescriptorHandle);
+
+    void                    BindSampler(FRHIDescriptorHandle DescriptorHandle, MTL::SamplerState* Sampler);
+    void                    BindResource(FRHIDescriptorHandle DescriptorHandle, FMetalResourceViewBase* Resource);
+    void                    BindTexture(FRHIDescriptorHandle DescriptorHandle, MTL::Texture* Texture);
+
+    void                    BindDescriptorHeapsToEncoder(FMetalCommandEncoder* Encoder, MTL::FunctionType FunctionType, EMetalShaderStages Frequency);
+
+    void                    MakeResident(FRHIDescriptorHandle DescriptorHandle, MTL::Resource* Resource, MTL::ResourceUsage Usage, EMetalShaderStages Frequency);
+
+	bool					IsSupported() {return bIsSupported;}
+private:
+	bool 					bIsSupported = false;
+    FMetalDescriptorHeap    StandardResources;
+    FMetalDescriptorHeap    SamplerResources;
+
+    TArray<TTuple<MTL::Resource*, MTL::ResourceUsage>> ResidentResources[EMetalShaderStages::Num];
+};
diff --git a/Engine/Source/Runtime/Apple/MetalRHI/Private/Shaders/Types/MetalAmplificationShader.cpp b/Engine/Source/Runtime/Apple/MetalRHI/Private/Shaders/Types/MetalAmplificationShader.cpp
new file mode 100644
--- /dev/null
+++ b/Engine/Source/Runtime/Apple/MetalRHI/Private/Shaders/Types/MetalAmplificationShader.cpp
@@ -0,0 +1,34 @@
+// Copyright Epic Games, Inc. All Rights Reserved.
+
+/*=============================================================================
+	MetalAmplificationShader.cpp: Metal RHI Amplification Shader Class Implementation.
+=============================================================================*/
+
+
+#include "MetalRHIPrivate.h"
+#include "Templates/MetalBaseShader.h"
+#include "MetalAmplificationShader.h"
+
+
+//------------------------------------------------------------------------------
+
+#pragma mark - Metal RHI Amplification Shader Class
+
+#if PLATFORM_SUPPORTS_MESH_SHADERS
+FMetalAmplificationShader::FMetalAmplificationShader(TArrayView<const uint8> InCode)
+{
+	FMetalCodeHeader Header;
+	Init(InCode, Header, MTLLibraryPtr());
+}
+
+FMetalAmplificationShader::FMetalAmplificationShader(TArrayView<const uint8> InCode, MTLLibraryPtr InLibrary)
+{
+	FMetalCodeHeader Header;
+	Init(InCode, Header, InLibrary);
+}
+
+MTLFunctionPtr FMetalAmplificationShader::GetFunction()
+{
+	return GetCompiledFunction();
+}
+#endif
diff --git a/Engine/Source/Runtime/Apple/MetalRHI/Private/Shaders/Types/MetalAmplificationShader.h b/Engine/Source/Runtime/Apple/MetalRHI/Private/Shaders/Types/MetalAmplificationShader.h
new file mode 100644
--- /dev/null
+++ b/Engine/Source/Runtime/Apple/MetalRHI/Private/Shaders/Types/MetalAmplificationShader.h
@@ -0,0 +1,24 @@
+// Copyright Epic Games, Inc. All Rights Reserved.
+
+/*=============================================================================
+	MetalAmplificationShader.h: Metal RHI Amplification Shader Class Definition.
+=============================================================================*/
+
+#pragma once
+
+
+//------------------------------------------------------------------------------
+
+#pragma mark - Metal RHI Amplification Shader Class
+
+
+#if PLATFORM_SUPPORTS_MESH_SHADERS
+class FMetalAmplificationShader : public TMetalBaseShader<FRHIAmplificationShader, SF_Amplification>
+{
+public:
+    FMetalAmplificationShader(TArrayView<const uint8> InCode);
+    FMetalAmplificationShader(TArrayView<const uint8> InCode, MTLLibraryPtr InLibrary);
+
+    MTLFunctionPtr GetFunction();
+};
+#endif
diff --git a/Engine/Source/Runtime/Apple/MetalRHI/Private/Shaders/Types/MetalGeometryShader.cpp b/Engine/Source/Runtime/Apple/MetalRHI/Private/Shaders/Types/MetalGeometryShader.cpp
new file mode 100644
--- /dev/null
+++ b/Engine/Source/Runtime/Apple/MetalRHI/Private/Shaders/Types/MetalGeometryShader.cpp
@@ -0,0 +1,34 @@
+// Copyright Epic Games, Inc. All Rights Reserved.
+
+/*=============================================================================
+	MetalGeometryShader.cpp: Metal RHI Geometry Shader Class Implementation.
+=============================================================================*/
+
+
+#include "MetalRHIPrivate.h"
+#include "Templates/MetalBaseShader.h"
+#include "MetalGeometryShader.h"
+
+
+//------------------------------------------------------------------------------
+
+#pragma mark - Metal RHI Geometry Shader Class
+
+#if METAL_USE_METAL_SHADER_CONVERTER
+FMetalGeometryShader::FMetalGeometryShader(TArrayView<const uint8> InCode)
+{
+	FMetalCodeHeader Header;
+	Init(InCode, Header, MTLLibraryPtr());
+}
+
+FMetalGeometryShader::FMetalGeometryShader(TArrayView<const uint8> InCode, MTLLibraryPtr InLibrary)
+{
+	FMetalCodeHeader Header;
+	Init(InCode, Header, InLibrary);
+}
+
+MTLFunctionPtr FMetalGeometryShader::GetFunction()
+{
+	return GetCompiledFunction();
+}
+#endif
diff --git a/Engine/Source/Runtime/Apple/MetalRHI/Private/Shaders/Types/MetalMeshShader.cpp b/Engine/Source/Runtime/Apple/MetalRHI/Private/Shaders/Types/MetalMeshShader.cpp
new file mode 100644
--- /dev/null
+++ b/Engine/Source/Runtime/Apple/MetalRHI/Private/Shaders/Types/MetalMeshShader.cpp
@@ -0,0 +1,34 @@
+// Copyright Epic Games, Inc. All Rights Reserved.
+
+/*=============================================================================
+	MetalMeshShader.cpp: Metal RHI Mesh Shader Class Implementation.
+=============================================================================*/
+
+
+#include "MetalRHIPrivate.h"
+#include "Templates/MetalBaseShader.h"
+#include "MetalMeshShader.h"
+
+
+//------------------------------------------------------------------------------
+
+#pragma mark - Metal RHI Mesh Shader Class
+
+#if PLATFORM_SUPPORTS_MESH_SHADERS
+FMetalMeshShader::FMetalMeshShader(TArrayView<const uint8> InCode)
+{
+	FMetalCodeHeader Header;
+	Init(InCode, Header, MTLLibraryPtr());
+}
+
+FMetalMeshShader::FMetalMeshShader(TArrayView<const uint8> InCode, MTLLibraryPtr InLibrary)
+{
+	FMetalCodeHeader Header;
+	Init(InCode, Header, InLibrary);
+}
+
+MTLFunctionPtr FMetalMeshShader::GetFunction()
+{
+	return GetCompiledFunction();
+}
+#endif
diff --git a/Engine/Source/Runtime/Apple/MetalRHI/Private/Shaders/Types/MetalMeshShader.h b/Engine/Source/Runtime/Apple/MetalRHI/Private/Shaders/Types/MetalMeshShader.h
new file mode 100644
--- /dev/null
+++ b/Engine/Source/Runtime/Apple/MetalRHI/Private/Shaders/Types/MetalMeshShader.h
@@ -0,0 +1,24 @@
+// Copyright Epic Games, Inc. All Rights Reserved.
+
+/*=============================================================================
+	MetalMeshShader.h: Metal RHI Mesh Shader Class Definition.
+=============================================================================*/
+
+#pragma once
+
+
+//------------------------------------------------------------------------------
+
+#pragma mark - Metal RHI Mesh Shader Class
+
+
+#if PLATFORM_SUPPORTS_MESH_SHADERS
+class FMetalMeshShader : public TMetalBaseShader<FRHIMeshShader, SF_Mesh>
+{
+public:
+    FMetalMeshShader(TArrayView<const uint8> InCode);
+    FMetalMeshShader(TArrayView<const uint8> InCode, MTLLibraryPtr InLibrary);
+
+    MTLFunctionPtr GetFunction();
+};
+#endif
