Fixed compile errors on iOS due to missing ir_converter.h which wont be supported until later

diff --git a/Engine/Source/Runtime/Apple/MetalRHI/Private/MetalBindlessDescriptors.cpp b/Engine/Source/Runtime/Apple/MetalRHI/Private/MetalBindlessDescriptors.cpp
--- a/Engine/Source/Runtime/Apple/MetalRHI/Private/MetalBindlessDescriptors.cpp
+++ b/Engine/Source/Runtime/Apple/MetalRHI/Private/MetalBindlessDescriptors.cpp
@@ -1,4 +1,6 @@
 // Copyright Epic Games, Inc. All Rights Reserved.
+#if PLATFORM_SUPPORTS_BINDLESS_RENDERING
+
 #include "MetalRHIPrivate.h"
 
 #include "MetalBindlessDescriptors.h"
@@ -252,3 +254,5 @@ void FMetalBindlessDescriptorManager::MakeResident(FRHIDescriptorHandle Descript
 {
     ResidentResources[Frequency].AddUnique(TTuple<MTL::Resource*, MTL::ResourceUsage>(Resource, Usage));
 }
+
+#endif //PLATFORM_SUPPORTS_BINDLESS_RENDERING
diff --git a/Engine/Source/Runtime/Apple/MetalRHI/Private/MetalBindlessDescriptors.h b/Engine/Source/Runtime/Apple/MetalRHI/Private/MetalBindlessDescriptors.h
--- a/Engine/Source/Runtime/Apple/MetalRHI/Private/MetalBindlessDescriptors.h
+++ b/Engine/Source/Runtime/Apple/MetalRHI/Private/MetalBindlessDescriptors.h
@@ -2,6 +2,8 @@
 
 #pragma once
 
+#if PLATFORM_SUPPORTS_BINDLESS_RENDERING
+
 #include <Metal/Metal.h>
 #include "MetalResources.h"
 #include "MetalShaderResources.h"
@@ -69,3 +71,5 @@ private:
 
     TArray<TTuple<MTL::Resource*, MTL::ResourceUsage>> ResidentResources[EMetalShaderStages::Num];
 };
+
+#endif
diff --git a/Engine/Source/Runtime/Apple/MetalRHI/Private/MetalDynamicRHI_Shaders.cpp b/Engine/Source/Runtime/Apple/MetalRHI/Private/MetalDynamicRHI_Shaders.cpp
--- a/Engine/Source/Runtime/Apple/MetalRHI/Private/MetalDynamicRHI_Shaders.cpp
+++ b/Engine/Source/Runtime/Apple/MetalRHI/Private/MetalDynamicRHI_Shaders.cpp
@@ -32,13 +32,19 @@ FPixelShaderRHIRef FMetalDynamicRHI::RHICreatePixelShader(TArrayView<const uint8
 	return Shader;
 }
 
+
 FGeometryShaderRHIRef FMetalDynamicRHI::RHICreateGeometryShader(TArrayView<const uint8> Code, const FSHAHash& Hash)
 {
     MTL_SCOPED_AUTORELEASE_POOL;
+#if PLATFORM_SUPPORTS_GEOMETRY_SHADERS
     FMetalGeometryShader* Shader = new FMetalGeometryShader(Code);
     return Shader;
+#else
+	return nullptr;
+#endif
 }
 
+
 FComputeShaderRHIRef FMetalDynamicRHI::RHICreateComputeShader(TArrayView<const uint8> Code, const FSHAHash& Hash)
 {
     MTL_SCOPED_AUTORELEASE_POOL;
diff --git a/Engine/Source/Runtime/Apple/MetalRHI/Private/MetalRHI.cpp b/Engine/Source/Runtime/Apple/MetalRHI/Private/MetalRHI.cpp
--- a/Engine/Source/Runtime/Apple/MetalRHI/Private/MetalRHI.cpp
+++ b/Engine/Source/Runtime/Apple/MetalRHI/Private/MetalRHI.cpp
@@ -1150,12 +1150,14 @@ FMetalDynamicRHI::FMetalDynamicRHI(ERHIFeatureLevel::Type RequestedFeatureLevel)
 	CompilerInstance = IRCompilerCreate();
 #endif
 
+#if PLATFORM_SUPPORTS_BINDLESS_RENDERING
 	if(GRHIBindlessSupport != ERHIBindlessSupport::Unsupported)
 	{
 		FMetalBindlessDescriptorManager* BindlessDescriptorManager = ImmediateContext.Context->GetBindlessDescriptorManager();
 		BindlessDescriptorManager->Init();
 	}
-
+#endif
+	
 #if ENABLE_METAL_GPUPROFILE
     if (ImmediateContext.Profiler)
 		ImmediateContext.Profiler->EndFrame();
diff --git a/Engine/Source/Runtime/Apple/MetalRHI/Private/MetalStateCache.cpp b/Engine/Source/Runtime/Apple/MetalRHI/Private/MetalStateCache.cpp
--- a/Engine/Source/Runtime/Apple/MetalRHI/Private/MetalStateCache.cpp
+++ b/Engine/Source/Runtime/Apple/MetalRHI/Private/MetalStateCache.cpp
@@ -1664,12 +1664,6 @@ void FMetalStateCache::IRMakeTextureResident(EMetalShaderStages const Frequency,
 	}
 }
 
-void FMetalStateCache::RegisterMetalHeap(MTL::Heap* Heap)
-{
-	FScopeLock ScopeLock(&ActiveHeapsLock);
-	ActiveHeaps.Add(Heap);
-}
-
 void FMetalStateCache::IRForwardBindlessParameters(EMetalShaderStages const Frequency, TConstArrayView<FRHIShaderParameterResource> InBindlessParameters)
 {
     FMetalBindlessDescriptorManager* BindlessDescriptorManager = GetMetalDeviceContext().GetBindlessDescriptorManager();
@@ -1742,6 +1736,12 @@ void FMetalStateCache::IRBindUniformBuffer(EMetalShaderStages const Frequency, i
 }
 #endif
 
+void FMetalStateCache::RegisterMetalHeap(MTL::Heap* Heap)
+{
+	FScopeLock ScopeLock(&ActiveHeapsLock);
+	ActiveHeaps.Add(Heap);
+}
+
 void FMetalStateCache::SetShaderTexture(EMetalShaderStages const Frequency, MTL::Texture* Texture, NS::UInteger const Index, MTL::ResourceUsage const Usage)
 {
 	check(Frequency < EMetalShaderStages::Num);
@@ -2094,8 +2094,8 @@ void FMetalStateCache::CommitRenderResources(FMetalCommandEncoder* Raster)
 			{
 				IRBindResourcesToEncoder<FMetalVertexShader*, EMetalShaderStages::Vertex, MTL::FunctionTypeVertex>(GraphicsPSO->VertexShader, Raster);
 			}
-#endif
 		}
+#endif
     }
     
 #if PLATFORM_SUPPORTS_GEOMETRY_SHADERS
@@ -2108,8 +2108,8 @@ void FMetalStateCache::CommitRenderResources(FMetalCommandEncoder* Raster)
 		{
 			IRBindResourcesToEncoder<FMetalGeometryShader*, EMetalShaderStages::Geometry, MTL::FunctionTypeMesh>(GraphicsPSO->GeometryShader, Raster);
 		}
+	}
 #endif
-    }
 #endif
         
     if (IsValidRef(GraphicsPSO->PixelShader))
@@ -2599,8 +2599,10 @@ void FMetalStateCache::SetRenderPipelineState(FMetalCommandEncoder& CommandEncod
 	
 	
 #if METAL_DEBUG_OPTIONS
+#if PLATFORM_SUPPORTS_BINDLESS_RENDERING
 	FMetalBindlessDescriptorManager* BindlessDescriptorManager = GetMetalDeviceContext().GetBindlessDescriptorManager();
 	if(!BindlessDescriptorManager->IsSupported())
+#endif
 	{
 		Validate();
 		
diff --git a/Engine/Source/Runtime/Core/Public/IOS/IOSPlatform.h b/Engine/Source/Runtime/Core/Public/IOS/IOSPlatform.h
--- a/Engine/Source/Runtime/Core/Public/IOS/IOSPlatform.h
+++ b/Engine/Source/Runtime/Core/Public/IOS/IOSPlatform.h
@@ -74,6 +74,7 @@ typedef FIOSPlatformTypes FPlatformTypes;
 #define PLATFORM_NEEDS_RHIRESOURCELIST					0
 #define PLATFORM_SUPPORTS_GEOMETRY_SHADERS				0
 #define PLATFORM_SUPPORTS_VIRTUAL_TEXTURE_STREAMING		1
+#define PLATFORM_SUPPORTS_BINDLESS_RENDERING			0
 
 // Ensure we can use this builtin - seems to be present on Clang 9, GCC 11 and MSVC 19.26,
 // but gives spurious "non-void function 'BitCast' should return a value" errors on some
