Fixed compile errors on iOS due to missing ir_converter.h which wont be supported until later

diff --git a/Engine/Source/Runtime/Apple/MetalRHI/Private/MetalRHIPrivate.h b/Engine/Source/Runtime/Apple/MetalRHI/Private/MetalRHIPrivate.h
--- a/Engine/Source/Runtime/Apple/MetalRHI/Private/MetalRHIPrivate.h
+++ b/Engine/Source/Runtime/Apple/MetalRHI/Private/MetalRHIPrivate.h
@@ -51,7 +51,7 @@ const uint32 MetalBufferBytesSize = BufferOffsetAlignment * 32;
 #endif
 
 #define METAL_RHI_RAYTRACING (0)
-#define METAL_USE_METAL_SHADER_CONVERTER (1)
+#define METAL_USE_METAL_SHADER_CONVERTER PLATFORM_SUPPORTS_BINDLESS_RENDERING
 
 // Metal Shader Converter
 #if METAL_USE_METAL_SHADER_CONVERTER
