Fixing includes in MetalRHI so that all headers can be included independently.

diff --git a/Engine/Source/Runtime/Apple/MetalRHI/Private/MetalBindlessDescriptors.cpp b/Engine/Source/Runtime/Apple/MetalRHI/Private/MetalBindlessDescriptors.cpp
--- a/Engine/Source/Runtime/Apple/MetalRHI/Private/MetalBindlessDescriptors.cpp
+++ b/Engine/Source/Runtime/Apple/MetalRHI/Private/MetalBindlessDescriptors.cpp
@@ -1,10 +1,11 @@
 // Copyright Epic Games, Inc. All Rights Reserved.
 
-#include "MetalRHIPrivate.h"
 #include "MetalBindlessDescriptors.h"
 
 #if PLATFORM_SUPPORTS_BINDLESS_RENDERING
 
+#include "MetalContext.h"
+
 int32 GBindlessResourceDescriptorHeapSize = 1000 * 1000;
 static FAutoConsoleVariableRef CVarBindlessResourceDescriptorHeapSize(
     TEXT("Metal.Bindless.ResourceDescriptorHeapSize"),
diff --git a/Engine/Source/Runtime/Apple/MetalRHI/Private/MetalBindlessDescriptors.h b/Engine/Source/Runtime/Apple/MetalRHI/Private/MetalBindlessDescriptors.h
--- a/Engine/Source/Runtime/Apple/MetalRHI/Private/MetalBindlessDescriptors.h
+++ b/Engine/Source/Runtime/Apple/MetalRHI/Private/MetalBindlessDescriptors.h
@@ -2,11 +2,13 @@
 
 #pragma once
 
-#include <Metal/Metal.h>
+#include "MetalRHIPrivate.h"
+
+#if PLATFORM_SUPPORTS_BINDLESS_RENDERING
 #include "MetalResources.h"
 #include "MetalShaderResources.h"
 
-#if PLATFORM_SUPPORTS_BINDLESS_RENDERING
+class FMetalCommandEncoder;
 
 struct FMetalDescriptorHeap
 {
diff --git a/Engine/Source/Runtime/Apple/MetalRHI/Private/MetalBuffer.cpp b/Engine/Source/Runtime/Apple/MetalRHI/Private/MetalBuffer.cpp
--- a/Engine/Source/Runtime/Apple/MetalRHI/Private/MetalBuffer.cpp
+++ b/Engine/Source/Runtime/Apple/MetalRHI/Private/MetalBuffer.cpp
@@ -1,13 +1,15 @@
 // Copyright Epic Games, Inc. All Rights Reserved.
 
-#include "MetalRHIPrivate.h"
 #include "MetalBuffer.h"
-#include "Templates/AlignmentTemplates.h"
-#include "MetalLLM.h"
-#include <objc/runtime.h>
 #include "MetalCommandBuffer.h"
+#include "MetalCommandQueue.h"
+#include "MetalContext.h"
+#include "MetalLLM.h"
 #include "MetalProfiler.h"
 #include "MetalRenderPass.h"
+#include "MetalRHIPrivate.h"
+#include "Templates/AlignmentTemplates.h"
+#include <objc/runtime.h>
 
 DECLARE_MEMORY_STAT(TEXT("Used Device Buffer Memory"), STAT_MetalDeviceBufferMemory, STATGROUP_MetalRHI);
 DECLARE_MEMORY_STAT(TEXT("Used Pooled Buffer Memory"), STAT_MetalPooledBufferMemory, STATGROUP_MetalRHI);
diff --git a/Engine/Source/Runtime/Apple/MetalRHI/Private/MetalBuffer.h b/Engine/Source/Runtime/Apple/MetalRHI/Private/MetalBuffer.h
--- a/Engine/Source/Runtime/Apple/MetalRHI/Private/MetalBuffer.h
+++ b/Engine/Source/Runtime/Apple/MetalRHI/Private/MetalBuffer.h
@@ -4,9 +4,15 @@
 #pragma once
 
 #include "MetalRHIPrivate.h"
+#include "MetalResources.h"
 #include "Containers/LockFreeList.h"
 #include "ResourcePool.h"
 
+class FMetalCommandQueue;
+
+class FMetalBuffer;
+typedef TSharedPtr<FMetalBuffer> FMetalBufferPtr;
+
 struct FMetalPooledBufferArgs
 {
     FMetalPooledBufferArgs() : Device(nullptr), Size(0), Flags(BUF_None), Storage(MTL::StorageModeShared), CpuCacheMode(MTL::CPUCacheModeDefaultCache) {}
diff --git a/Engine/Source/Runtime/Apple/MetalRHI/Private/MetalCaptureManager.h b/Engine/Source/Runtime/Apple/MetalRHI/Private/MetalCaptureManager.h
--- a/Engine/Source/Runtime/Apple/MetalRHI/Private/MetalCaptureManager.h
+++ b/Engine/Source/Runtime/Apple/MetalRHI/Private/MetalCaptureManager.h
@@ -2,7 +2,7 @@
 
 #pragma once
 
-#include <Metal/Metal.h>
+#include "MetalRHIPrivate.h"
 
 class FMetalCommandQueue;
 
diff --git a/Engine/Source/Runtime/Apple/MetalRHI/Private/MetalCommandBuffer.h b/Engine/Source/Runtime/Apple/MetalRHI/Private/MetalCommandBuffer.h
--- a/Engine/Source/Runtime/Apple/MetalRHI/Private/MetalCommandBuffer.h
+++ b/Engine/Source/Runtime/Apple/MetalRHI/Private/MetalCommandBuffer.h
@@ -2,7 +2,7 @@
 
 #pragma once
 
-#include <Metal/Metal.h>
+#include "MetalThirdParty.h"
 #include "MetalResources.h"
 #include "MetalShaderResources.h"
 
diff --git a/Engine/Source/Runtime/Apple/MetalRHI/Private/MetalCommandEncoder.cpp b/Engine/Source/Runtime/Apple/MetalRHI/Private/MetalCommandEncoder.cpp
--- a/Engine/Source/Runtime/Apple/MetalRHI/Private/MetalCommandEncoder.cpp
+++ b/Engine/Source/Runtime/Apple/MetalRHI/Private/MetalCommandEncoder.cpp
@@ -4,11 +4,13 @@
 	MetalCommandEncoder.cpp: Metal command encoder wrapper.
 =============================================================================*/
 
+#include "MetalCommandEncoder.h"
+#include "MetalCommandBuffer.h"
+#include "MetalCommandList.h"
+#include "MetalContext.h"
 #include "MetalRHIPrivate.h"
 #include "MetalShaderTypes.h"
 #include "MetalGraphicsPipelineState.h"
-#include "MetalCommandEncoder.h"
-#include "MetalCommandBuffer.h"
 #include "MetalProfiler.h"
 #include "MetalShaderResources.h"
 
diff --git a/Engine/Source/Runtime/Apple/MetalRHI/Private/MetalCommandEncoder.h b/Engine/Source/Runtime/Apple/MetalRHI/Private/MetalCommandEncoder.h
--- a/Engine/Source/Runtime/Apple/MetalRHI/Private/MetalCommandEncoder.h
+++ b/Engine/Source/Runtime/Apple/MetalRHI/Private/MetalCommandEncoder.h
@@ -2,15 +2,20 @@
 
 #pragma once
 
-#include <Metal/Metal.h>
+#include "MetalRHIPrivate.h"
 #include "MetalBuffer.h"
 #include "MetalFence.h"
+#include "MetalPipeline.h"
 #include "MetalProfiler.h"
 
 class FMetalCommandList;
 class FMetalCommandQueue;
 class FMetalGraphicsPipelineState;
 class FMetalCommandBufferFence;
+struct FMetalCommandBufferStats;
+
+class FMetalShaderPipeline;
+typedef TSharedPtr<FMetalShaderPipeline, ESPMode::ThreadSafe> FMetalShaderPipelinePtr;
 
 /**
  * Enumeration for submission hints to avoid unclear bool values.
diff --git a/Engine/Source/Runtime/Apple/MetalRHI/Private/MetalCommandList.h b/Engine/Source/Runtime/Apple/MetalRHI/Private/MetalCommandList.h
--- a/Engine/Source/Runtime/Apple/MetalRHI/Private/MetalCommandList.h
+++ b/Engine/Source/Runtime/Apple/MetalRHI/Private/MetalCommandList.h
@@ -2,9 +2,10 @@
 
 #pragma once
 
-#include <Metal/Metal.h>
+#include "MetalThirdParty.h"
 #include "MetalProfiler.h"
 
+class FMetalCommandBuffer;
 class FMetalCommandQueue;
 
 /**
diff --git a/Engine/Source/Runtime/Apple/MetalRHI/Private/MetalCommandQueue.cpp b/Engine/Source/Runtime/Apple/MetalRHI/Private/MetalCommandQueue.cpp
--- a/Engine/Source/Runtime/Apple/MetalRHI/Private/MetalCommandQueue.cpp
+++ b/Engine/Source/Runtime/Apple/MetalRHI/Private/MetalCommandQueue.cpp
@@ -4,16 +4,16 @@
 	MetalCommandQueue.cpp: Metal command queue wrapper..
 =============================================================================*/
 
-#include "MetalRHIPrivate.h"
-
 #include "MetalCommandQueue.h"
 #include "MetalCommandBuffer.h"
 #include "MetalCommandList.h"
+#include "MetalFence.h"
 #include "MetalProfiler.h"
+#include "MetalRHIPrivate.h"
 #include "Misc/ConfigCacheIni.h"
 
 #if !UE_BUILD_SHIPPING
-#import <Metal/Metal.h>
+#import "MetalThirdParty.h"
 #endif
 
 #pragma mark - Private C++ Statics -
diff --git a/Engine/Source/Runtime/Apple/MetalRHI/Private/MetalCommandQueue.h b/Engine/Source/Runtime/Apple/MetalRHI/Private/MetalCommandQueue.h
--- a/Engine/Source/Runtime/Apple/MetalRHI/Private/MetalCommandQueue.h
+++ b/Engine/Source/Runtime/Apple/MetalRHI/Private/MetalCommandQueue.h
@@ -2,10 +2,15 @@
 
 #pragma once
 
-#include <Metal/Metal.h>
+#include "MetalRHIPrivate.h"
+#include "RHI.h"
+#include "RHIFeatureLevel.h"
 #include "Containers/LockFreeList.h"
 
+class FMetalCommandBuffer;
+class FMetalCommandBufferFence;
 class FMetalCommandList;
+class FMetalFence;
 
 /**
  * Enumeration of features which are present only on some OS/device combinations.
diff --git a/Engine/Source/Runtime/Apple/MetalRHI/Private/MetalCommands.cpp b/Engine/Source/Runtime/Apple/MetalRHI/Private/MetalCommands.cpp
--- a/Engine/Source/Runtime/Apple/MetalRHI/Private/MetalCommands.cpp
+++ b/Engine/Source/Runtime/Apple/MetalRHI/Private/MetalCommands.cpp
@@ -8,6 +8,7 @@
 #include "MetalShaderTypes.h"
 #include "MetalGraphicsPipelineState.h"
 #include "MetalComputePipelineState.h"
+#include "MetalRHIContext.h"
 #include "GlobalShader.h"
 #include "OneColorShader.h"
 #include "RHICommandList.h"
diff --git a/Engine/Source/Runtime/Apple/MetalRHI/Private/MetalContext.cpp b/Engine/Source/Runtime/Apple/MetalRHI/Private/MetalContext.cpp
--- a/Engine/Source/Runtime/Apple/MetalRHI/Private/MetalContext.cpp
+++ b/Engine/Source/Runtime/Apple/MetalRHI/Private/MetalContext.cpp
@@ -1,5 +1,7 @@
 // Copyright Epic Games, Inc. All Rights Reserved.
 
+#include "MetalContext.h"
+#include "MetalRHIContext.h"
 #include "MetalRHIPrivate.h"
 #include "MetalRHIRenderQuery.h"
 #include "MetalVertexDeclaration.h"
@@ -13,7 +15,6 @@
 #include "HAL/PlatformFramePacer.h"
 #include "Runtime/HeadMountedDisplay/Public/IHeadMountedDisplayModule.h"
 
-#include "MetalContext.h"
 #include "MetalProfiler.h"
 #include "MetalCommandBuffer.h"
 
diff --git a/Engine/Source/Runtime/Apple/MetalRHI/Private/MetalContext.h b/Engine/Source/Runtime/Apple/MetalRHI/Private/MetalContext.h
--- a/Engine/Source/Runtime/Apple/MetalRHI/Private/MetalContext.h
+++ b/Engine/Source/Runtime/Apple/MetalRHI/Private/MetalContext.h
@@ -10,6 +10,7 @@
 #include "MetalBuffer.h"
 #include "MetalCaptureManager.h"
 #include "MetalFrameAllocator.h"
+#include "MetalStateCache.h"
 #if PLATFORM_IOS
 #include "IOS/IOSView.h"
 #endif
diff --git a/Engine/Source/Runtime/Apple/MetalRHI/Private/MetalDynamicRHI.cpp b/Engine/Source/Runtime/Apple/MetalRHI/Private/MetalDynamicRHI.cpp
--- a/Engine/Source/Runtime/Apple/MetalRHI/Private/MetalDynamicRHI.cpp
+++ b/Engine/Source/Runtime/Apple/MetalRHI/Private/MetalDynamicRHI.cpp
@@ -4,7 +4,7 @@
 	MetalDynamicRHI.cpp: Metal Dynamic RHI Class Implementation.
 =============================================================================*/
 
-
+#include "MetalDynamicRHI.h"
 #include "MetalRHIPrivate.h"
 #include "MetalRHIRenderQuery.h"
 #include "MetalRHIStagingBuffer.h"
diff --git a/Engine/Source/Runtime/Apple/MetalRHI/Private/MetalDynamicRHI.h b/Engine/Source/Runtime/Apple/MetalRHI/Private/MetalDynamicRHI.h
--- a/Engine/Source/Runtime/Apple/MetalRHI/Private/MetalDynamicRHI.h
+++ b/Engine/Source/Runtime/Apple/MetalRHI/Private/MetalDynamicRHI.h
@@ -9,10 +9,8 @@
 #include "CoreMinimal.h"
 #include "RHI.h"
 
-DECLARE_LOG_CATEGORY_EXTERN(LogMetal, Display, All);
-
 // Metal RHI public headers.
-#include <Metal/Metal.h>
+#include "MetalThirdParty.h"
 #include "MetalState.h"
 #include "MetalResources.h"
 #include "MetalRHIContext.h"
diff --git a/Engine/Source/Runtime/Apple/MetalRHI/Private/MetalDynamicRHIModule.cpp b/Engine/Source/Runtime/Apple/MetalRHI/Private/MetalDynamicRHIModule.cpp
--- a/Engine/Source/Runtime/Apple/MetalRHI/Private/MetalDynamicRHIModule.cpp
+++ b/Engine/Source/Runtime/Apple/MetalRHI/Private/MetalDynamicRHIModule.cpp
@@ -4,10 +4,10 @@
 	MetalDynamicRHIModule.cpp: Metal Dynamic RHI Module Class Implementation.
 ==============================================================================*/
 
-
+#include "MetalDynamicRHIModule.h"
+#include "MetalDynamicRHI.h"
 #include "MetalLLM.h"
 #include "DynamicRHI.h"
-#include "MetalDynamicRHIModule.h"
 #include "Modules/ModuleManager.h"
 
 
diff --git a/Engine/Source/Runtime/Apple/MetalRHI/Private/MetalDynamicRHIModule.h b/Engine/Source/Runtime/Apple/MetalRHI/Private/MetalDynamicRHIModule.h
--- a/Engine/Source/Runtime/Apple/MetalRHI/Private/MetalDynamicRHIModule.h
+++ b/Engine/Source/Runtime/Apple/MetalRHI/Private/MetalDynamicRHIModule.h
@@ -6,6 +6,7 @@
 
 #pragma once
 
+#include "DynamicRHI.h"
 
 //------------------------------------------------------------------------------
 
diff --git a/Engine/Source/Runtime/Apple/MetalRHI/Private/MetalDynamicRHI_Shaders.cpp b/Engine/Source/Runtime/Apple/MetalRHI/Private/MetalDynamicRHI_Shaders.cpp
--- a/Engine/Source/Runtime/Apple/MetalRHI/Private/MetalDynamicRHI_Shaders.cpp
+++ b/Engine/Source/Runtime/Apple/MetalRHI/Private/MetalDynamicRHI_Shaders.cpp
@@ -4,7 +4,7 @@
 	MetalDynamicRHI_Shaders.cpp: Metal Dynamic RHI Class Shader Methods.
 =============================================================================*/
 
-
+#include "MetalDynamicRHI.h"
 #include "MetalRHIPrivate.h"
 #include "MetalShaderTypes.h"
 #include "Shaders/MetalShaderLibrary.h"
diff --git a/Engine/Source/Runtime/Apple/MetalRHI/Private/MetalFence.h b/Engine/Source/Runtime/Apple/MetalRHI/Private/MetalFence.h
--- a/Engine/Source/Runtime/Apple/MetalRHI/Private/MetalFence.h
+++ b/Engine/Source/Runtime/Apple/MetalRHI/Private/MetalFence.h
@@ -3,7 +3,8 @@
 
 #pragma once
 
-#import <Metal/Metal.h>
+#include "MetalRHIPrivate.h"
+#include "Containers/LockFreeList.h"
 
 #pragma clang diagnostic ignored "-Wnullability-completeness"
 
diff --git a/Engine/Source/Runtime/Apple/MetalRHI/Private/MetalFrameAllocator.cpp b/Engine/Source/Runtime/Apple/MetalRHI/Private/MetalFrameAllocator.cpp
--- a/Engine/Source/Runtime/Apple/MetalRHI/Private/MetalFrameAllocator.cpp
+++ b/Engine/Source/Runtime/Apple/MetalRHI/Private/MetalFrameAllocator.cpp
@@ -1,10 +1,11 @@
 // Copyright Epic Games, Inc. All Rights Reserved.
 
-#include "MetalRHIPrivate.h"
 #include "MetalFrameAllocator.h"
+#include "MetalCommandBuffer.h"
+#include "MetalContext.h"
 #include "MetalLLM.h"
 #include "MetalProfiler.h"
-#include "MetalCommandBuffer.h"
+#include "MetalRHIPrivate.h"
 
 #pragma mark Constructor/Destructor
 
diff --git a/Engine/Source/Runtime/Apple/MetalRHI/Private/MetalGraphicsPipelineState.h b/Engine/Source/Runtime/Apple/MetalRHI/Private/MetalGraphicsPipelineState.h
--- a/Engine/Source/Runtime/Apple/MetalRHI/Private/MetalGraphicsPipelineState.h
+++ b/Engine/Source/Runtime/Apple/MetalRHI/Private/MetalGraphicsPipelineState.h
@@ -6,6 +6,14 @@
 
 #pragma once
 
+#include "RHIResources.h"
+#include "MetalState.h"
+#include "Shaders/Types/MetalAmplificationShader.h"
+#include "Shaders/Types/MetalGeometryShader.h"
+#include "Shaders/Types/MetalMeshShader.h"
+#include "Shaders/Types/MetalPixelShader.h"
+#include "Shaders/Types/MetalVertexShader.h"
+
 class FMetalGraphicsPipelineState : public FRHIGraphicsPipelineState
 {
 	friend class FMetalDynamicRHI;
diff --git a/Engine/Source/Runtime/Apple/MetalRHI/Private/MetalHashedVertexDescriptor.h b/Engine/Source/Runtime/Apple/MetalRHI/Private/MetalHashedVertexDescriptor.h
--- a/Engine/Source/Runtime/Apple/MetalRHI/Private/MetalHashedVertexDescriptor.h
+++ b/Engine/Source/Runtime/Apple/MetalRHI/Private/MetalHashedVertexDescriptor.h
@@ -6,6 +6,7 @@
 
 #pragma once
 
+#include "MetalRHIPrivate.h"
 
 //------------------------------------------------------------------------------
 
diff --git a/Engine/Source/Runtime/Apple/MetalRHI/Private/MetalLLM.cpp b/Engine/Source/Runtime/Apple/MetalRHI/Private/MetalLLM.cpp
--- a/Engine/Source/Runtime/Apple/MetalRHI/Private/MetalLLM.cpp
+++ b/Engine/Source/Runtime/Apple/MetalRHI/Private/MetalLLM.cpp
@@ -1,7 +1,9 @@
 // Copyright Epic Games, Inc. All Rights Reserved.
 
 #include "MetalLLM.h"
+#include "MetalCommandQueue.h"
 #include "MetalProfiler.h"
+#include "MetalResources.h"
 #include "RenderUtils.h"
 #include "HAL/LowLevelMemStats.h"
 
diff --git a/Engine/Source/Runtime/Apple/MetalRHI/Private/MetalLLM.h b/Engine/Source/Runtime/Apple/MetalRHI/Private/MetalLLM.h
--- a/Engine/Source/Runtime/Apple/MetalRHI/Private/MetalLLM.h
+++ b/Engine/Source/Runtime/Apple/MetalRHI/Private/MetalLLM.h
@@ -6,6 +6,9 @@
 #include "HAL/LowLevelMemTracker.h"
 #include "Apple/AppleLLM.h"
 
+class FMetalBuffer;
+typedef TSharedPtr<FMetalBuffer> FMetalBufferPtr;
+
 @interface FMetalDeallocHandler : FApplePlatformObject<NSObject>
 {
 	dispatch_block_t Block;
diff --git a/Engine/Source/Runtime/Apple/MetalRHI/Private/MetalPipeline.cpp b/Engine/Source/Runtime/Apple/MetalRHI/Private/MetalPipeline.cpp
--- a/Engine/Source/Runtime/Apple/MetalRHI/Private/MetalPipeline.cpp
+++ b/Engine/Source/Runtime/Apple/MetalRHI/Private/MetalPipeline.cpp
@@ -4,6 +4,7 @@
 	MetalPipeline.cpp: Metal shader pipeline RHI implementation.
 =============================================================================*/
 
+#include "MetalPipeline.h"
 #include "MetalRHIPrivate.h"
 #include "MetalVertexDeclaration.h"
 #include "MetalShaderTypes.h"
diff --git a/Engine/Source/Runtime/Apple/MetalRHI/Private/MetalPipeline.h b/Engine/Source/Runtime/Apple/MetalRHI/Private/MetalPipeline.h
--- a/Engine/Source/Runtime/Apple/MetalRHI/Private/MetalPipeline.h
+++ b/Engine/Source/Runtime/Apple/MetalRHI/Private/MetalPipeline.h
@@ -3,19 +3,17 @@
 #pragma once
 
 #include "MetalRHIPrivate.h"
-
 #include "ShaderPipelineCache.h"
+#include "MetalShaderResources.h"
 
 /**
  * The sampler, buffer and texture resource limits as defined here:
  * https://developer.apple.com/library/ios/documentation/Miscellaneous/Conceptual/MetalProgrammingGuide/Render-Ctx/Render-Ctx.html
  */
 #if PLATFORM_IOS
-#define METAL_MAX_BUFFERS 31
 #define METAL_MAX_TEXTURES 31
 typedef uint32 FMetalTextureMask;
 #elif PLATFORM_MAC
-#define METAL_MAX_BUFFERS 31
 #define METAL_MAX_TEXTURES 128
 typedef __uint128_t FMetalTextureMask;
 #else
diff --git a/Engine/Source/Runtime/Apple/MetalRHI/Private/MetalProfiler.cpp b/Engine/Source/Runtime/Apple/MetalRHI/Private/MetalProfiler.cpp
--- a/Engine/Source/Runtime/Apple/MetalRHI/Private/MetalProfiler.cpp
+++ b/Engine/Source/Runtime/Apple/MetalRHI/Private/MetalProfiler.cpp
@@ -1,7 +1,8 @@
 // Copyright Epic Games, Inc. All Rights Reserved.
 
-#include "MetalRHIPrivate.h"
 #include "MetalProfiler.h"
+#include "MetalContext.h"
+#include "MetalRHIPrivate.h"
 #include "EngineGlobals.h"
 #include "StaticBoundShaderState.h"
 #include "MetalCommandBuffer.h"
diff --git a/Engine/Source/Runtime/Apple/MetalRHI/Private/MetalProfiler.h b/Engine/Source/Runtime/Apple/MetalRHI/Private/MetalProfiler.h
--- a/Engine/Source/Runtime/Apple/MetalRHI/Private/MetalProfiler.h
+++ b/Engine/Source/Runtime/Apple/MetalRHI/Private/MetalProfiler.h
@@ -3,9 +3,12 @@
 #pragma once
 
 #include "MetalRHIPrivate.h"
-#include "MetalCommandQueue.h"
 #include "GPUProfiler.h"
 
+DECLARE_DELEGATE_OneParam(FMetalCommandBufferCompletionHandler, MTL::CommandBuffer*);
+
+class FMetalContext;
+
 // Stats
 DECLARE_CYCLE_STAT_EXTERN(TEXT("MakeDrawable time"),STAT_MetalMakeDrawableTime,STATGROUP_MetalRHI, );
 DECLARE_CYCLE_STAT_EXTERN(TEXT("Draw call time"),STAT_MetalDrawCallTime,STATGROUP_MetalRHI, );
diff --git a/Engine/Source/Runtime/Apple/MetalRHI/Private/MetalRHI.cpp b/Engine/Source/Runtime/Apple/MetalRHI/Private/MetalRHI.cpp
--- a/Engine/Source/Runtime/Apple/MetalRHI/Private/MetalRHI.cpp
+++ b/Engine/Source/Runtime/Apple/MetalRHI/Private/MetalRHI.cpp
@@ -4,6 +4,9 @@
 	MetalRHI.cpp: Metal device RHI implementation.
 =============================================================================*/
 
+#include "MetalRHI.h"
+#include "MetalContext.h"
+#include "MetalDynamicRHI.h"
 #include "MetalRHIPrivate.h"
 #include "Misc/MessageDialog.h"
 #include "Modules/ModuleManager.h"
diff --git a/Engine/Source/Runtime/Apple/MetalRHI/Private/MetalRHIContext.cpp b/Engine/Source/Runtime/Apple/MetalRHI/Private/MetalRHIContext.cpp
--- a/Engine/Source/Runtime/Apple/MetalRHI/Private/MetalRHIContext.cpp
+++ b/Engine/Source/Runtime/Apple/MetalRHI/Private/MetalRHIContext.cpp
@@ -1,5 +1,7 @@
 // Copyright Epic Games, Inc. All Rights Reserved.
 
+#include "MetalRHIContext.h"
+#include "MetalContext.h"
 #include "MetalRHIPrivate.h"
 #include "MetalRHIRenderQuery.h"
 #include "MetalRHIVisionOSBridge.h"
diff --git a/Engine/Source/Runtime/Apple/MetalRHI/Private/MetalRHIPrivate.h b/Engine/Source/Runtime/Apple/MetalRHI/Private/MetalRHIPrivate.h
--- a/Engine/Source/Runtime/Apple/MetalRHI/Private/MetalRHIPrivate.h
+++ b/Engine/Source/Runtime/Apple/MetalRHI/Private/MetalRHIPrivate.h
@@ -6,27 +6,18 @@
 
 #pragma once
 
-#include "CoreMinimal.h"
-#include "Misc/ScopeLock.h"
-#include "Misc/CommandLine.h"
+#include "MetalRHI.h"
+#include "MetalThirdParty.h"
 #include "PixelFormat.h"
+#include "RHICommandList.h"
+#include "RHIGlobals.h"
 
-// Metal C++ wrapper
-THIRD_PARTY_INCLUDES_START
-#include "MetalInclude.h"
-THIRD_PARTY_INCLUDES_END
-
-DECLARE_DELEGATE_OneParam(FMetalCommandBufferCompletionHandler, MTL::CommandBuffer*);
+class FMetalDeviceContext;
+class FMetalSurface;
 
-inline uint32 GetTypeHash(const MTLTexturePtr& TexturePtr)
-{
-    return GetTypeHash(TexturePtr.get());
-}
+DECLARE_LOG_CATEGORY_EXTERN(LogMetal, Display, All);
 
-inline uint32 GetTypeHash(const MTLBufferPtr& BufferPtr)
-{
-    return GetTypeHash(BufferPtr.get());
-}
+DECLARE_DELEGATE_OneParam(FMetalCommandBufferCompletionHandler, MTL::CommandBuffer*);
 
 // Whether the Metal RHI is initialized sufficiently to handle resources
 extern bool GIsMetalInitialized;
@@ -50,27 +41,22 @@ const uint32 MetalBufferBytesSize = BufferOffsetAlignment * 2;
 const uint32 MetalBufferBytesSize = BufferOffsetAlignment * 32;
 #endif
 
-#define METAL_RHI_RAYTRACING (0)
+
 #define METAL_USE_METAL_SHADER_CONVERTER PLATFORM_SUPPORTS_BINDLESS_RENDERING
 
 // Metal Shader Converter
 #if METAL_USE_METAL_SHADER_CONVERTER
 THIRD_PARTY_INCLUDES_START
-#include "metal_irconverter.h"
-#define IR_RUNTIME_METALCPP 1
-#define IR_PRIVATE_IMPLEMENTATION 1
-#include "metal_irconverter_runtime.h"
+    #include "metal_irconverter.h"
+    #define IR_RUNTIME_METALCPP 1
+    #define IR_PRIVATE_IMPLEMENTATION 1
+    #include "metal_irconverter_runtime.h"
 THIRD_PARTY_INCLUDES_END
 
-constexpr uint64_t kIRStandardHeapBindPoint 			   = 0;
+constexpr uint64_t kIRStandardHeapBindPoint                = 0;
 
 #endif
 
-#include "MetalInclude.h"
-#include "MetalRHI.h"
-#include "MetalDynamicRHI.h"
-#include "RHI.h"
-
 #define BUFFER_CACHE_MODE MTL::ResourceCPUCacheModeDefaultCache
 
 #if PLATFORM_MAC
@@ -106,6 +92,7 @@ enum EMTLTextureType
 #define METAL_SUPPORTS_INDIRECT_ARGUMENT_BUFFERS 1
 #define METAL_SUPPORTS_CAPTURE_MANAGER 1
 #define METAL_SUPPORTS_TILE_SHADERS 1
+
 // In addition to compile-time SDK checks we also need a way to check if these are available on runtime
 extern bool GMetalSupportsCaptureManager;
 
@@ -121,24 +108,24 @@ extern FMetalBufferFormat GMetalBufferFormats[PF_MAX];
 
 #define METAL_DEBUG_OPTIONS !(UE_BUILD_SHIPPING || UE_BUILD_TEST)
 #if METAL_DEBUG_OPTIONS
-#define METAL_DEBUG_OPTION(Code) Code
+    #define METAL_DEBUG_OPTION(Code) Code
 #else
-#define METAL_DEBUG_OPTION(Code)
+    #define METAL_DEBUG_OPTION(Code)
 #endif
 
 extern bool GMetalCommandBufferDebuggingEnabled;
 
 /** Set to 1 to enable GPU events in Xcode frame debugger */
 #ifndef ENABLE_METAL_GPUEVENTS_IN_TEST
-	#define ENABLE_METAL_GPUEVENTS_IN_TEST 0
+    #define ENABLE_METAL_GPUEVENTS_IN_TEST 0
 #endif
-#define ENABLE_METAL_GPUEVENTS	(UE_BUILD_DEBUG || UE_BUILD_DEVELOPMENT || (UE_BUILD_TEST && ENABLE_METAL_GPUEVENTS_IN_TEST))
-#define ENABLE_METAL_GPUPROFILE	(ENABLE_METAL_GPUEVENTS && 1)
+#define ENABLE_METAL_GPUEVENTS    (UE_BUILD_DEBUG || UE_BUILD_DEVELOPMENT || (UE_BUILD_TEST && ENABLE_METAL_GPUEVENTS_IN_TEST))
+#define ENABLE_METAL_GPUPROFILE    (ENABLE_METAL_GPUEVENTS && 1)
 
 #if ENABLE_METAL_GPUPROFILE
-#define METAL_GPUPROFILE(Code) Code
+    #define METAL_GPUPROFILE(Code) Code
 #else
-#define METAL_GPUPROFILE(Code) 
+    #define METAL_GPUPROFILE(Code)
 #endif
 
 #define UNREAL_TO_METAL_BUFFER_INDEX(Index) ((MaxMetalStreams - 1) - Index)
@@ -166,9 +153,6 @@ void METALRHI_API SafeReleaseMetalObject(NS::Object* Object);
 // Safely release a metal texture, correctly handling the case where the RHI has been destructed first
 void SafeReleaseMetalTexture(MTLTexturePtr Object);
 
-// Safely release a metal buffer, correctly handling the case where the RHI has been destructed first
-void SafeReleaseMetalBuffer(FMetalBufferPtr Buffer);
-
 // Safely release a fence, correctly handling cases where fences aren't supported or the debug implementation is used.
 void SafeReleaseMetalFence(class FMetalFence* Object);
 
@@ -231,6 +215,11 @@ MTL::PixelFormat UEToMetalFormat(EPixelFormat UEFormat, bool bSRGB);
 
 uint8 GetMetalPixelFormatKey(MTL::PixelFormat Format);
 
+template<class T>
+struct TMetalResourceTraits
+{
+};
+
 template<typename TRHIType>
 static FORCEINLINE typename TMetalResourceTraits<TRHIType>::TConcreteType* ResourceCast(TRHIType* Resource)
 {
@@ -324,6 +313,3 @@ FORCEINLINE NS::String* FStringToNSString(const FString& InputString)
 {
     return ((NS::String*)InputString.GetCFString())->autorelease();
 }
-
-#include "MetalStateCache.h"
-#include "MetalContext.h"
diff --git a/Engine/Source/Runtime/Apple/MetalRHI/Private/MetalRHIRenderQuery.cpp b/Engine/Source/Runtime/Apple/MetalRHI/Private/MetalRHIRenderQuery.cpp
--- a/Engine/Source/Runtime/Apple/MetalRHI/Private/MetalRHIRenderQuery.cpp
+++ b/Engine/Source/Runtime/Apple/MetalRHI/Private/MetalRHIRenderQuery.cpp
@@ -4,8 +4,9 @@
 	MetalRHIRenderQuery.cpp: Metal RHI Render Query Implementation.
 =============================================================================*/
 
-#include "MetalRHIPrivate.h"
 #include "MetalRHIRenderQuery.h"
+#include "MetalContext.h"
+#include "MetalRHIPrivate.h"
 #include "MetalProfiler.h"
 #include "MetalLLM.h"
 #include "MetalCommandBuffer.h"
diff --git a/Engine/Source/Runtime/Apple/MetalRHI/Private/MetalRHIRenderQuery.h b/Engine/Source/Runtime/Apple/MetalRHI/Private/MetalRHIRenderQuery.h
--- a/Engine/Source/Runtime/Apple/MetalRHI/Private/MetalRHIRenderQuery.h
+++ b/Engine/Source/Runtime/Apple/MetalRHI/Private/MetalRHIRenderQuery.h
@@ -6,6 +6,8 @@
 
 #pragma once
 
+#include "MetalRHIPrivate.h"
+#include "RHIResources.h"
 
 //------------------------------------------------------------------------------
 
@@ -16,7 +18,8 @@ class FMetalContext;
 class FMetalQueryBufferPool;
 class FMetalQueryResult;
 class FMetalCommandBufferFence;
-
+class FMetalBuffer;
+typedef TSharedPtr<FMetalBuffer> FMetalBufferPtr;
 
 //------------------------------------------------------------------------------
 
diff --git a/Engine/Source/Runtime/Apple/MetalRHI/Private/MetalRHIStagingBuffer.cpp b/Engine/Source/Runtime/Apple/MetalRHI/Private/MetalRHIStagingBuffer.cpp
--- a/Engine/Source/Runtime/Apple/MetalRHI/Private/MetalRHIStagingBuffer.cpp
+++ b/Engine/Source/Runtime/Apple/MetalRHI/Private/MetalRHIStagingBuffer.cpp
@@ -5,8 +5,9 @@
 =============================================================================*/
 
 
-#include "MetalRHIPrivate.h"
 #include "MetalRHIStagingBuffer.h"
+#include "MetalResources.h"
+#include "MetalRHIPrivate.h"
 
 
 //------------------------------------------------------------------------------
diff --git a/Engine/Source/Runtime/Apple/MetalRHI/Private/MetalRHIStagingBuffer.h b/Engine/Source/Runtime/Apple/MetalRHI/Private/MetalRHIStagingBuffer.h
--- a/Engine/Source/Runtime/Apple/MetalRHI/Private/MetalRHIStagingBuffer.h
+++ b/Engine/Source/Runtime/Apple/MetalRHI/Private/MetalRHIStagingBuffer.h
@@ -6,6 +6,11 @@
 
 #pragma once
 
+#include "RHIResources.h"
+
+class FMetalBuffer;
+typedef TSharedPtr<FMetalBuffer> FMetalBufferPtr;
+
 class FMetalRHIStagingBuffer final : public FRHIStagingBuffer
 {
 	friend class FMetalRHICommandContext;
diff --git a/Engine/Source/Runtime/Apple/MetalRHI/Private/MetalRayTracing.h b/Engine/Source/Runtime/Apple/MetalRHI/Private/MetalRayTracing.h
--- a/Engine/Source/Runtime/Apple/MetalRHI/Private/MetalRayTracing.h
+++ b/Engine/Source/Runtime/Apple/MetalRHI/Private/MetalRayTracing.h
@@ -5,14 +5,11 @@
 ==============================================================================*/
 
 #pragma once
+
 #include "MetalRHIPrivate.h"
 
 #if METAL_RHI_RAYTRACING
 
-THIRD_PARTY_INCLUDES_START
-#include "MetalInclude.h"
-THIRD_PARTY_INCLUDES_END
-
 struct FMetalRayTracingGeometryParameters
 {
 	FMetalRHIBuffer* IndexBuffer;
diff --git a/Engine/Source/Runtime/Apple/MetalRHI/Private/MetalRenderPass.cpp b/Engine/Source/Runtime/Apple/MetalRHI/Private/MetalRenderPass.cpp
--- a/Engine/Source/Runtime/Apple/MetalRHI/Private/MetalRenderPass.cpp
+++ b/Engine/Source/Runtime/Apple/MetalRHI/Private/MetalRenderPass.cpp
@@ -4,12 +4,13 @@
 	MetalRenderPass.cpp: Metal command pass wrapper.
 =============================================================================*/
 
+#include "MetalRenderPass.h"
 
+#include "MetalContext.h"
 #include "MetalRHIPrivate.h"
 #include "MetalShaderTypes.h"
 #include "MetalGraphicsPipelineState.h"
 #include "MetalVertexDeclaration.h"
-#include "MetalRenderPass.h"
 #include "MetalCommandBuffer.h"
 #include "MetalProfiler.h"
 #include "MetalFrameAllocator.h"
diff --git a/Engine/Source/Runtime/Apple/MetalRHI/Private/MetalRenderPass.h b/Engine/Source/Runtime/Apple/MetalRHI/Private/MetalRenderPass.h
--- a/Engine/Source/Runtime/Apple/MetalRHI/Private/MetalRenderPass.h
+++ b/Engine/Source/Runtime/Apple/MetalRHI/Private/MetalRenderPass.h
@@ -10,6 +10,7 @@
 
 class FMetalCommandList;
 class FMetalCommandQueue;
+class FMetalStateCache;
 
 class FMetalRenderPass
 {
diff --git a/Engine/Source/Runtime/Apple/MetalRHI/Private/MetalRenderTarget.cpp b/Engine/Source/Runtime/Apple/MetalRHI/Private/MetalRenderTarget.cpp
--- a/Engine/Source/Runtime/Apple/MetalRHI/Private/MetalRenderTarget.cpp
+++ b/Engine/Source/Runtime/Apple/MetalRHI/Private/MetalRenderTarget.cpp
@@ -6,8 +6,11 @@
 
 #include "MetalRHIPrivate.h"
 #include "ScreenRendering.h"
-#include "MetalProfiler.h"
 #include "MetalCommandBuffer.h"
+#include "MetalCommandQueue.h"
+#include "MetalContext.h"
+#include "MetalDynamicRHI.h"
+#include "MetalProfiler.h"
 #include "ResolveShader.h"
 #include "PipelineStateCache.h"
 #include "Math/PackedVector.h"
diff --git a/Engine/Source/Runtime/Apple/MetalRHI/Private/MetalState.cpp b/Engine/Source/Runtime/Apple/MetalRHI/Private/MetalState.cpp
--- a/Engine/Source/Runtime/Apple/MetalRHI/Private/MetalState.cpp
+++ b/Engine/Source/Runtime/Apple/MetalRHI/Private/MetalState.cpp
@@ -4,6 +4,9 @@
 	MetalState.cpp: Metal state implementation.
 =============================================================================*/
 
+#include "MetalState.h"
+#include "MetalContext.h"
+#include "MetalDynamicRHI.h"
 #include "MetalRHIPrivate.h"
 #include "MetalProfiler.h"
 #include "RHIUtilities.h"
diff --git a/Engine/Source/Runtime/Apple/MetalRHI/Private/MetalStateCache.h b/Engine/Source/Runtime/Apple/MetalRHI/Private/MetalStateCache.h
--- a/Engine/Source/Runtime/Apple/MetalRHI/Private/MetalStateCache.h
+++ b/Engine/Source/Runtime/Apple/MetalRHI/Private/MetalStateCache.h
@@ -3,11 +3,15 @@
 #pragma once
 
 #include "MetalRHIPrivate.h"
+#include "MetalRHIPrivate.h"
+#include "MetalPipeline.h"
+#include "MetalResources.h"
+#include "MetalState.h"
 #include "MetalUniformBuffer.h"
 #include "Shaders/MetalShaderParameterCache.h"
-#include "MetalPipeline.h"
 
 class FMetalGraphicsPipelineState;
+class FMetalComputeShader;
 class FMetalQueryBuffer;
 
 enum EMetalPipelineFlags
diff --git a/Engine/Source/Runtime/Apple/MetalRHI/Private/MetalTexture.cpp b/Engine/Source/Runtime/Apple/MetalRHI/Private/MetalTexture.cpp
--- a/Engine/Source/Runtime/Apple/MetalRHI/Private/MetalTexture.cpp
+++ b/Engine/Source/Runtime/Apple/MetalRHI/Private/MetalTexture.cpp
@@ -5,6 +5,8 @@
  =============================================================================*/
 
 #include "MetalRHIPrivate.h"
+#include "MetalContext.h"
+#include "MetalDynamicRHI.h"
 #include "MetalProfiler.h" // for STAT_MetalTexturePageOffTime
 #include "MetalCommandBuffer.h"
 #include "RenderUtils.h"
diff --git a/Engine/Source/Runtime/Apple/MetalRHI/Private/MetalTransitionData.cpp b/Engine/Source/Runtime/Apple/MetalRHI/Private/MetalTransitionData.cpp
--- a/Engine/Source/Runtime/Apple/MetalRHI/Private/MetalTransitionData.cpp
+++ b/Engine/Source/Runtime/Apple/MetalRHI/Private/MetalTransitionData.cpp
@@ -4,9 +4,10 @@
 	MetalTransitionData.cpp: Metal RHI Resource Transition Implementation.
 ==============================================================================*/
 
-#include "MetalRHIPrivate.h"
 #include "MetalTransitionData.h"
-
+#include "MetalContext.h"
+#include "MetalRHIContext.h"
+#include "MetalRHIPrivate.h"
 
 //------------------------------------------------------------------------------
 
diff --git a/Engine/Source/Runtime/Apple/MetalRHI/Private/MetalTransitionData.h b/Engine/Source/Runtime/Apple/MetalRHI/Private/MetalTransitionData.h
--- a/Engine/Source/Runtime/Apple/MetalRHI/Private/MetalTransitionData.h
+++ b/Engine/Source/Runtime/Apple/MetalRHI/Private/MetalTransitionData.h
@@ -6,6 +6,8 @@
 
 #pragma once
 
+#include "RHITransition.h"
+#include "MetalFence.h"
 
 //------------------------------------------------------------------------------
 
diff --git a/Engine/Source/Runtime/Apple/MetalRHI/Private/MetalUAV.cpp b/Engine/Source/Runtime/Apple/MetalRHI/Private/MetalUAV.cpp
--- a/Engine/Source/Runtime/Apple/MetalRHI/Private/MetalUAV.cpp
+++ b/Engine/Source/Runtime/Apple/MetalRHI/Private/MetalUAV.cpp
@@ -2,6 +2,9 @@
 
 
 #include "MetalRHIPrivate.h"
+#include "MetalContext.h"
+#include "MetalDynamicRHI.h"
+#include "MetalRHIContext.h"
 #include "MetalRHIStagingBuffer.h"
 #include "MetalCommandBuffer.h"
 #include "RenderUtils.h"
diff --git a/Engine/Source/Runtime/Apple/MetalRHI/Private/MetalUniformBuffer.cpp b/Engine/Source/Runtime/Apple/MetalRHI/Private/MetalUniformBuffer.cpp
--- a/Engine/Source/Runtime/Apple/MetalRHI/Private/MetalUniformBuffer.cpp
+++ b/Engine/Source/Runtime/Apple/MetalRHI/Private/MetalUniformBuffer.cpp
@@ -4,9 +4,10 @@
 	MetalConstantBuffer.cpp: Metal Constant buffer implementation.
 =============================================================================*/
 
+#include "MetalUniformBuffer.h"
 #include "MetalRHIPrivate.h"
+#include "MetalContext.h"
 #include "MetalFrameAllocator.h"
-#include "MetalUniformBuffer.h"
 #include "ShaderParameterStruct.h"
 #include "RHIUniformBufferDataShared.h"
 
diff --git a/Engine/Source/Runtime/Apple/MetalRHI/Private/MetalUniformBufferCommands.cpp b/Engine/Source/Runtime/Apple/MetalRHI/Private/MetalUniformBufferCommands.cpp
--- a/Engine/Source/Runtime/Apple/MetalRHI/Private/MetalUniformBufferCommands.cpp
+++ b/Engine/Source/Runtime/Apple/MetalRHI/Private/MetalUniformBufferCommands.cpp
@@ -1,6 +1,7 @@
 // Copyright Epic Games, Inc. All Rights Reserved.
 
 #include "MetalRHIPrivate.h"
+#include "MetalDynamicRHI.h"
 #include "MetalShaderTypes.h"
 #include "MetalFrameAllocator.h"
 
diff --git a/Engine/Source/Runtime/Apple/MetalRHI/Private/MetalVertexBuffer.cpp b/Engine/Source/Runtime/Apple/MetalRHI/Private/MetalVertexBuffer.cpp
--- a/Engine/Source/Runtime/Apple/MetalRHI/Private/MetalVertexBuffer.cpp
+++ b/Engine/Source/Runtime/Apple/MetalRHI/Private/MetalVertexBuffer.cpp
@@ -8,6 +8,8 @@
 #include "MetalProfiler.h"
 #include "MetalCommandBuffer.h"
 #include "MetalCommandQueue.h"
+#include "MetalContext.h"
+#include "MetalDynamicRHI.h"
 #include "Containers/ResourceArray.h"
 #include "RenderUtils.h"
 #include "MetalLLM.h"
diff --git a/Engine/Source/Runtime/Apple/MetalRHI/Private/MetalVertexDeclaration.h b/Engine/Source/Runtime/Apple/MetalRHI/Private/MetalVertexDeclaration.h
--- a/Engine/Source/Runtime/Apple/MetalRHI/Private/MetalVertexDeclaration.h
+++ b/Engine/Source/Runtime/Apple/MetalRHI/Private/MetalVertexDeclaration.h
@@ -6,9 +6,10 @@
 
 #pragma once
 
-
 #include "MetalHashedVertexDescriptor.h"
-
+#include "MetalRHIPrivate.h"
+#include "RHI.h"
+#include "RHIResources.h"
 
 //------------------------------------------------------------------------------
 
diff --git a/Engine/Source/Runtime/Apple/MetalRHI/Private/MetalViewport.cpp b/Engine/Source/Runtime/Apple/MetalRHI/Private/MetalViewport.cpp
--- a/Engine/Source/Runtime/Apple/MetalRHI/Private/MetalViewport.cpp
+++ b/Engine/Source/Runtime/Apple/MetalRHI/Private/MetalViewport.cpp
@@ -4,7 +4,13 @@
 	MetalViewport.cpp: Metal viewport RHI implementation.
 =============================================================================*/
 
+#include "MetalViewport.h"
+#include "MetalContext.h"
+#include "MetalDynamicRHI.h"
 #include "MetalRHIPrivate.h"
+#include "MetalCommandBuffer.h"
+#include "MetalProfiler.h"
+#include "MetalRHIVisionOSBridge.h"
 
 #import <QuartzCore/CAMetalLayer.h>
 
@@ -16,10 +22,7 @@
 #endif
 #include "RenderCommandFence.h"
 #include "Containers/Set.h"
-#include "MetalCommandBuffer.h"
-#include "MetalProfiler.h"
 #include "RenderUtils.h"
-#include "MetalRHIVisionOSBridge.h"
 
 extern int32 GMetalSupportsIntermediateBackBuffer;
 extern int32 GMetalSeparatePresentThread;
diff --git a/Engine/Source/Runtime/Apple/MetalRHI/Private/Shaders/Debugging/MetalShaderDebugCache.cpp b/Engine/Source/Runtime/Apple/MetalRHI/Private/Shaders/Debugging/MetalShaderDebugCache.cpp
--- a/Engine/Source/Runtime/Apple/MetalRHI/Private/Shaders/Debugging/MetalShaderDebugCache.cpp
+++ b/Engine/Source/Runtime/Apple/MetalRHI/Private/Shaders/Debugging/MetalShaderDebugCache.cpp
@@ -4,18 +4,14 @@
 	MetalShaderDebugCache.cpp: Metal RHI Shader Debug Cache.
 =============================================================================*/
 
+#include "MetalShaderDebugCache.h"
 #include "CoreMinimal.h"
 
-THIRD_PARTY_INCLUDES_START
-#include "MetalInclude.h"
-THIRD_PARTY_INCLUDES_END
-
 #include "HAL/FileManager.h"
 #include "HAL/PlatformFileManager.h"
 #include "Misc/Paths.h"
 #include "Misc/ScopeLock.h"
 
-#include "MetalShaderDebugCache.h"
 #include "MetalShaderDebugZipFile.h"
 
 #if !UE_BUILD_SHIPPING
diff --git a/Engine/Source/Runtime/Apple/MetalRHI/Private/Shaders/Debugging/MetalShaderDebugCache.h b/Engine/Source/Runtime/Apple/MetalRHI/Private/Shaders/Debugging/MetalShaderDebugCache.h
--- a/Engine/Source/Runtime/Apple/MetalRHI/Private/Shaders/Debugging/MetalShaderDebugCache.h
+++ b/Engine/Source/Runtime/Apple/MetalRHI/Private/Shaders/Debugging/MetalShaderDebugCache.h
@@ -6,8 +6,12 @@
 
 #pragma once
 
+#include "HAL/Platform.h"
+
 #if !UE_BUILD_SHIPPING
 
+#include "MetalRHIPrivate.h"
+
 struct FMetalShaderDebugCache
 {
 	static FMetalShaderDebugCache& Get()
diff --git a/Engine/Source/Runtime/Apple/MetalRHI/Private/Shaders/Debugging/MetalShaderDebugZipFile.cpp b/Engine/Source/Runtime/Apple/MetalRHI/Private/Shaders/Debugging/MetalShaderDebugZipFile.cpp
--- a/Engine/Source/Runtime/Apple/MetalRHI/Private/Shaders/Debugging/MetalShaderDebugZipFile.cpp
+++ b/Engine/Source/Runtime/Apple/MetalRHI/Private/Shaders/Debugging/MetalShaderDebugZipFile.cpp
@@ -4,20 +4,14 @@
 	MetalShaderDebugZipFile.cpp: Metal shader RHI implementation.
 =============================================================================*/
 
-#include "CoreMinimal.h"
-
-THIRD_PARTY_INCLUDES_START
-#include "MetalInclude.h"
-THIRD_PARTY_INCLUDES_END
-
 
+#include "MetalShaderDebugZipFile.h"
+#include "CoreMinimal.h"
 #include "HAL/FileManager.h"
 #include "HAL/PlatformFileManager.h"
 #include "Misc/Paths.h"
 #include "Misc/ScopeLock.h"
 
-#include "MetalShaderDebugZipFile.h"
-
 #if !UE_BUILD_SHIPPING
 
 FMetalShaderDebugZipFile::FMetalShaderDebugZipFile(FString LibPath)
diff --git a/Engine/Source/Runtime/Apple/MetalRHI/Private/Shaders/Debugging/MetalShaderDebugZipFile.h b/Engine/Source/Runtime/Apple/MetalRHI/Private/Shaders/Debugging/MetalShaderDebugZipFile.h
--- a/Engine/Source/Runtime/Apple/MetalRHI/Private/Shaders/Debugging/MetalShaderDebugZipFile.h
+++ b/Engine/Source/Runtime/Apple/MetalRHI/Private/Shaders/Debugging/MetalShaderDebugZipFile.h
@@ -6,8 +6,13 @@
 
 #pragma once
 
+#include "HAL/Platform.h"
+
 #if !UE_BUILD_SHIPPING
 
+#include "MetalRHIPrivate.h"
+#include "Containers/UnrealString.h"
+
 class FMetalShaderDebugZipFile
 {
 	struct FFileEntry
diff --git a/Engine/Source/Runtime/Apple/MetalRHI/Private/Shaders/MetalCompiledShaderCache.cpp b/Engine/Source/Runtime/Apple/MetalRHI/Private/Shaders/MetalCompiledShaderCache.cpp
--- a/Engine/Source/Runtime/Apple/MetalRHI/Private/Shaders/MetalCompiledShaderCache.cpp
+++ b/Engine/Source/Runtime/Apple/MetalRHI/Private/Shaders/MetalCompiledShaderCache.cpp
@@ -4,12 +4,7 @@
 	MetalCompiledShaderCache.cpp: Metal RHI Compiled Shader Cache.
 =============================================================================*/
 
-#include "CoreMinimal.h"
-
-THIRD_PARTY_INCLUDES_START
-#include "MetalInclude.h"
-THIRD_PARTY_INCLUDES_END
-
+#include "MetalCompiledShaderCache.h"
 #include "MetalCompiledShaderKey.h"
 #include "MetalCompiledShaderCache.h"
 
diff --git a/Engine/Source/Runtime/Apple/MetalRHI/Private/Shaders/MetalCompiledShaderCache.h b/Engine/Source/Runtime/Apple/MetalRHI/Private/Shaders/MetalCompiledShaderCache.h
--- a/Engine/Source/Runtime/Apple/MetalRHI/Private/Shaders/MetalCompiledShaderCache.h
+++ b/Engine/Source/Runtime/Apple/MetalRHI/Private/Shaders/MetalCompiledShaderCache.h
@@ -6,6 +6,7 @@
 
 #pragma once
 
+#include "MetalCompiledShaderKey.h"
 #include "Misc/ScopeRWLock.h"
 #include "MetalRHIPrivate.h"
 
diff --git a/Engine/Source/Runtime/Apple/MetalRHI/Private/Shaders/MetalCompiledShaderKey.h b/Engine/Source/Runtime/Apple/MetalRHI/Private/Shaders/MetalCompiledShaderKey.h
--- a/Engine/Source/Runtime/Apple/MetalRHI/Private/Shaders/MetalCompiledShaderKey.h
+++ b/Engine/Source/Runtime/Apple/MetalRHI/Private/Shaders/MetalCompiledShaderKey.h
@@ -6,6 +6,9 @@
 
 #pragma once
 
+#include "HAL/Platform.h"
+#include "Templates/TypeHash.h"
+
 struct FMetalCompiledShaderKey
 {
 	FMetalCompiledShaderKey(uint32 InCodeSize, uint32 InCodeCRC, uint32 InConstants)
diff --git a/Engine/Source/Runtime/Apple/MetalRHI/Private/Shaders/MetalShaderLibrary.h b/Engine/Source/Runtime/Apple/MetalRHI/Private/Shaders/MetalShaderLibrary.h
--- a/Engine/Source/Runtime/Apple/MetalRHI/Private/Shaders/MetalShaderLibrary.h
+++ b/Engine/Source/Runtime/Apple/MetalRHI/Private/Shaders/MetalShaderLibrary.h
@@ -6,6 +6,11 @@
 
 #pragma once
 
+#include "MetalRHIPrivate.h"
+#include "MetalShaderResources.h"
+#include "RHIShaderLibrary.h"
+#include "ShaderCodeArchive.h"
+
 class FMetalShaderLibrary final : public FRHIShaderLibrary
 {
 public:
diff --git a/Engine/Source/Runtime/Apple/MetalRHI/Private/Shaders/MetalShaderParameterCache.cpp b/Engine/Source/Runtime/Apple/MetalRHI/Private/Shaders/MetalShaderParameterCache.cpp
--- a/Engine/Source/Runtime/Apple/MetalRHI/Private/Shaders/MetalShaderParameterCache.cpp
+++ b/Engine/Source/Runtime/Apple/MetalRHI/Private/Shaders/MetalShaderParameterCache.cpp
@@ -5,8 +5,11 @@
 =============================================================================*/
 
 
-#include "MetalRHIPrivate.h"
 #include "MetalShaderParameterCache.h"
+#include "MetalCommandEncoder.h"
+#include "MetalResources.h"
+#include "MetalRHIPrivate.h"
+#include "MetalStateCache.h"
 
 
 //------------------------------------------------------------------------------
diff --git a/Engine/Source/Runtime/Apple/MetalRHI/Private/Shaders/MetalShaderParameterCache.h b/Engine/Source/Runtime/Apple/MetalRHI/Private/Shaders/MetalShaderParameterCache.h
--- a/Engine/Source/Runtime/Apple/MetalRHI/Private/Shaders/MetalShaderParameterCache.h
+++ b/Engine/Source/Runtime/Apple/MetalRHI/Private/Shaders/MetalShaderParameterCache.h
@@ -6,6 +6,12 @@
 
 #pragma once
 
+#include "CrossCompilerCommon.h"
+#include "MetalRHIPrivate.h"
+
+class FMetalBufferData;
+struct FMetalShaderBindings;
+
 class FMetalShaderParameterCache
 {
 public:
diff --git a/Engine/Source/Runtime/Apple/MetalRHI/Private/Shaders/Types/MetalAmplificationShader.cpp b/Engine/Source/Runtime/Apple/MetalRHI/Private/Shaders/Types/MetalAmplificationShader.cpp
--- a/Engine/Source/Runtime/Apple/MetalRHI/Private/Shaders/Types/MetalAmplificationShader.cpp
+++ b/Engine/Source/Runtime/Apple/MetalRHI/Private/Shaders/Types/MetalAmplificationShader.cpp
@@ -5,8 +5,6 @@
 =============================================================================*/
 
 
-#include "MetalRHIPrivate.h"
-#include "Templates/MetalBaseShader.h"
 #include "MetalAmplificationShader.h"
 
 
diff --git a/Engine/Source/Runtime/Apple/MetalRHI/Private/Shaders/Types/MetalAmplificationShader.h b/Engine/Source/Runtime/Apple/MetalRHI/Private/Shaders/Types/MetalAmplificationShader.h
--- a/Engine/Source/Runtime/Apple/MetalRHI/Private/Shaders/Types/MetalAmplificationShader.h
+++ b/Engine/Source/Runtime/Apple/MetalRHI/Private/Shaders/Types/MetalAmplificationShader.h
@@ -6,6 +6,8 @@
 
 #pragma once
 
+#include "MetalRHIPrivate.h"
+#include "Shaders/Types/Templates/MetalBaseShader.h"
 
 //------------------------------------------------------------------------------
 
diff --git a/Engine/Source/Runtime/Apple/MetalRHI/Private/Shaders/Types/MetalComputeShader.cpp b/Engine/Source/Runtime/Apple/MetalRHI/Private/Shaders/Types/MetalComputeShader.cpp
--- a/Engine/Source/Runtime/Apple/MetalRHI/Private/Shaders/Types/MetalComputeShader.cpp
+++ b/Engine/Source/Runtime/Apple/MetalRHI/Private/Shaders/Types/MetalComputeShader.cpp
@@ -4,9 +4,6 @@
 	MetalComputeShader.cpp: Metal RHI Compute Shader Class Implementation.
 =============================================================================*/
 
-
-#include "MetalRHIPrivate.h"
-#include "Templates/MetalBaseShader.h"
 #include "MetalComputeShader.h"
 #include "MetalCommandBuffer.h"
 
diff --git a/Engine/Source/Runtime/Apple/MetalRHI/Private/Shaders/Types/MetalComputeShader.h b/Engine/Source/Runtime/Apple/MetalRHI/Private/Shaders/Types/MetalComputeShader.h
--- a/Engine/Source/Runtime/Apple/MetalRHI/Private/Shaders/Types/MetalComputeShader.h
+++ b/Engine/Source/Runtime/Apple/MetalRHI/Private/Shaders/Types/MetalComputeShader.h
@@ -6,6 +6,9 @@
 
 #pragma once
 
+#include "MetalRHIPrivate.h"
+#include "Shaders/Types/Templates/MetalBaseShader.h"
+
 
 //------------------------------------------------------------------------------
 
diff --git a/Engine/Source/Runtime/Apple/MetalRHI/Private/Shaders/Types/MetalGeometryShader.cpp b/Engine/Source/Runtime/Apple/MetalRHI/Private/Shaders/Types/MetalGeometryShader.cpp
--- a/Engine/Source/Runtime/Apple/MetalRHI/Private/Shaders/Types/MetalGeometryShader.cpp
+++ b/Engine/Source/Runtime/Apple/MetalRHI/Private/Shaders/Types/MetalGeometryShader.cpp
@@ -4,12 +4,8 @@
 	MetalGeometryShader.cpp: Metal RHI Geometry Shader Class Implementation.
 =============================================================================*/
 
-
-#include "MetalRHIPrivate.h"
-#include "Templates/MetalBaseShader.h"
 #include "MetalGeometryShader.h"
 
-
 //------------------------------------------------------------------------------
 
 #pragma mark - Metal RHI Geometry Shader Class
diff --git a/Engine/Source/Runtime/Apple/MetalRHI/Private/Shaders/Types/MetalGeometryShader.h b/Engine/Source/Runtime/Apple/MetalRHI/Private/Shaders/Types/MetalGeometryShader.h
--- a/Engine/Source/Runtime/Apple/MetalRHI/Private/Shaders/Types/MetalGeometryShader.h
+++ b/Engine/Source/Runtime/Apple/MetalRHI/Private/Shaders/Types/MetalGeometryShader.h
@@ -6,6 +6,9 @@
 
 #pragma once
 
+#include "MetalRHIPrivate.h"
+#include "Shaders/Types/Templates/MetalBaseShader.h"
+
 
 //------------------------------------------------------------------------------
 
diff --git a/Engine/Source/Runtime/Apple/MetalRHI/Private/Shaders/Types/MetalMeshShader.cpp b/Engine/Source/Runtime/Apple/MetalRHI/Private/Shaders/Types/MetalMeshShader.cpp
--- a/Engine/Source/Runtime/Apple/MetalRHI/Private/Shaders/Types/MetalMeshShader.cpp
+++ b/Engine/Source/Runtime/Apple/MetalRHI/Private/Shaders/Types/MetalMeshShader.cpp
@@ -4,12 +4,8 @@
 	MetalMeshShader.cpp: Metal RHI Mesh Shader Class Implementation.
 =============================================================================*/
 
-
-#include "MetalRHIPrivate.h"
-#include "Templates/MetalBaseShader.h"
 #include "MetalMeshShader.h"
 
-
 //------------------------------------------------------------------------------
 
 #pragma mark - Metal RHI Mesh Shader Class
diff --git a/Engine/Source/Runtime/Apple/MetalRHI/Private/Shaders/Types/MetalMeshShader.h b/Engine/Source/Runtime/Apple/MetalRHI/Private/Shaders/Types/MetalMeshShader.h
--- a/Engine/Source/Runtime/Apple/MetalRHI/Private/Shaders/Types/MetalMeshShader.h
+++ b/Engine/Source/Runtime/Apple/MetalRHI/Private/Shaders/Types/MetalMeshShader.h
@@ -6,6 +6,9 @@
 
 #pragma once
 
+#include "MetalRHIPrivate.h"
+#include "Shaders/Types/Templates/MetalBaseShader.h"
+
 
 //------------------------------------------------------------------------------
 
diff --git a/Engine/Source/Runtime/Apple/MetalRHI/Private/Shaders/Types/MetalPixelShader.cpp b/Engine/Source/Runtime/Apple/MetalRHI/Private/Shaders/Types/MetalPixelShader.cpp
--- a/Engine/Source/Runtime/Apple/MetalRHI/Private/Shaders/Types/MetalPixelShader.cpp
+++ b/Engine/Source/Runtime/Apple/MetalRHI/Private/Shaders/Types/MetalPixelShader.cpp
@@ -4,12 +4,8 @@
 	MetalPixelShader.cpp: Metal RHI Pixel Shader Class Implementation.
 =============================================================================*/
 
-
-#include "MetalRHIPrivate.h"
-#include "Templates/MetalBaseShader.h"
 #include "MetalPixelShader.h"
 
-
 //------------------------------------------------------------------------------
 
 #pragma mark - Metal RHI Pixel Shader Class
diff --git a/Engine/Source/Runtime/Apple/MetalRHI/Private/Shaders/Types/MetalPixelShader.h b/Engine/Source/Runtime/Apple/MetalRHI/Private/Shaders/Types/MetalPixelShader.h
--- a/Engine/Source/Runtime/Apple/MetalRHI/Private/Shaders/Types/MetalPixelShader.h
+++ b/Engine/Source/Runtime/Apple/MetalRHI/Private/Shaders/Types/MetalPixelShader.h
@@ -6,6 +6,9 @@
 
 #pragma once
 
+#include "MetalRHIPrivate.h"
+#include "Shaders/Types/Templates/MetalBaseShader.h"
+
 
 //------------------------------------------------------------------------------
 
diff --git a/Engine/Source/Runtime/Apple/MetalRHI/Private/Shaders/Types/MetalVertexShader.cpp b/Engine/Source/Runtime/Apple/MetalRHI/Private/Shaders/Types/MetalVertexShader.cpp
--- a/Engine/Source/Runtime/Apple/MetalRHI/Private/Shaders/Types/MetalVertexShader.cpp
+++ b/Engine/Source/Runtime/Apple/MetalRHI/Private/Shaders/Types/MetalVertexShader.cpp
@@ -4,12 +4,8 @@
 	MetalVertexShader.cpp: Metal RHI Vertex Shader Class Implementation.
 =============================================================================*/
 
-
-#include "MetalRHIPrivate.h"
-#include "Templates/MetalBaseShader.h"
 #include "MetalVertexShader.h"
 
-
 //------------------------------------------------------------------------------
 
 #pragma mark - Metal RHI Vertex Shader Class
diff --git a/Engine/Source/Runtime/Apple/MetalRHI/Private/Shaders/Types/MetalVertexShader.h b/Engine/Source/Runtime/Apple/MetalRHI/Private/Shaders/Types/MetalVertexShader.h
--- a/Engine/Source/Runtime/Apple/MetalRHI/Private/Shaders/Types/MetalVertexShader.h
+++ b/Engine/Source/Runtime/Apple/MetalRHI/Private/Shaders/Types/MetalVertexShader.h
@@ -6,6 +6,9 @@
 
 #pragma once
 
+#include "MetalRHIPrivate.h"
+#include "Shaders/Types/Templates/MetalBaseShader.h"
+
 
 //------------------------------------------------------------------------------
 
diff --git a/Engine/Source/Runtime/Apple/MetalRHI/Private/Shaders/Types/Templates/MetalBaseShader.h b/Engine/Source/Runtime/Apple/MetalRHI/Private/Shaders/Types/Templates/MetalBaseShader.h
--- a/Engine/Source/Runtime/Apple/MetalRHI/Private/Shaders/Types/Templates/MetalBaseShader.h
+++ b/Engine/Source/Runtime/Apple/MetalRHI/Private/Shaders/Types/Templates/MetalBaseShader.h
@@ -9,7 +9,11 @@
 #include "Misc/Paths.h"
 #include "HAL/FileManager.h"
 #include "HAL/PlatformFileManager.h"
+#include "MetalCommandQueue.h"
+#include "MetalContext.h"
+#include "MetalShaderResources.h"
 #include "Misc/FileHelper.h"
+#include "Serialization/MemoryReader.h"
 #include "Shaders/Debugging/MetalShaderDebugCache.h"
 #include "Shaders/MetalCompiledShaderKey.h"
 #include "Shaders/MetalCompiledShaderCache.h"
diff --git a/Engine/Source/Runtime/Apple/MetalRHI/Public/MetalFrameAllocator.h b/Engine/Source/Runtime/Apple/MetalRHI/Public/MetalFrameAllocator.h
--- a/Engine/Source/Runtime/Apple/MetalRHI/Public/MetalFrameAllocator.h
+++ b/Engine/Source/Runtime/Apple/MetalRHI/Public/MetalFrameAllocator.h
@@ -2,10 +2,12 @@
 
 #pragma once
 
-#import <Metal/Metal.h>
+#include "MetalRHI.h"
+#import "MetalThirdParty.h"
+#include "Stats/Stats.h"
 #import <os/lock.h>
 
-#include "MetalRHI.h"
+class FMetalCommandBuffer;
 
 #define METAL_FRAME_ALLOCATOR_VALIDATION !(UE_BUILD_SHIPPING || UE_BUILD_TEST)
 
diff --git a/Engine/Source/Runtime/Apple/MetalRHI/Public/MetalRHI.h b/Engine/Source/Runtime/Apple/MetalRHI/Public/MetalRHI.h
--- a/Engine/Source/Runtime/Apple/MetalRHI/Public/MetalRHI.h
+++ b/Engine/Source/Runtime/Apple/MetalRHI/Public/MetalRHI.h
@@ -7,3 +7,5 @@
 #pragma once 
 
 // TODO: Put public Metal RHI definitions here.
+#define METAL_RHI_RAYTRACING (0)
+
diff --git a/Engine/Source/Runtime/Apple/MetalRHI/Public/MetalRHIContext.h b/Engine/Source/Runtime/Apple/MetalRHI/Public/MetalRHIContext.h
--- a/Engine/Source/Runtime/Apple/MetalRHI/Public/MetalRHIContext.h
+++ b/Engine/Source/Runtime/Apple/MetalRHI/Public/MetalRHIContext.h
@@ -3,7 +3,7 @@
 #pragma once
 
 // Metal RHI public headers.
-#include <Metal/Metal.h>
+#include "MetalThirdParty.h"
 #include "MetalState.h"
 #include "MetalResources.h"
 #include "MetalViewport.h"
diff --git a/Engine/Source/Runtime/Apple/MetalRHI/Public/MetalResources.h b/Engine/Source/Runtime/Apple/MetalRHI/Public/MetalResources.h
--- a/Engine/Source/Runtime/Apple/MetalRHI/Public/MetalResources.h
+++ b/Engine/Source/Runtime/Apple/MetalRHI/Public/MetalResources.h
@@ -6,10 +6,11 @@
 
 #pragma once
 
+#include "MetalRHIPrivate.h"
 #include "BoundShaderStateCache.h"
 #include "MetalShaderResources.h"
 #include "ShaderCodeArchive.h"
-#include "MetalRHIPrivate.h"
+#include "Templates/TypeHash.h"
 
 #define UE_METAL_RHI_SUPPORT_CLEAR_UAV_WITH_BLIT_ENCODER 1
 
@@ -41,6 +42,11 @@ class FMetalSubBufferHeap;
 class FMetalSubBufferLinear;
 class FMetalSubBufferMagazine;
 
+inline uint32 GetTypeHash(const MTLBufferPtr& BufferPtr)
+{
+    return GetTypeHash(BufferPtr.get());
+}
+
 class FMetalBuffer
 {
 public:
@@ -114,6 +120,9 @@ private:
 
 typedef TSharedPtr<FMetalBuffer> FMetalBufferPtr;
 
+// Safely release a metal buffer, correctly handling the case where the RHI has been destructed first
+void SafeReleaseMetalBuffer(FMetalBufferPtr Buffer);
+
 struct FMetalTextureCreateDesc : public FRHITextureCreateDesc
 {
 	FMetalTextureCreateDesc(FRHITextureCreateDesc const& CreateDesc);
@@ -151,6 +160,11 @@ private:
 	FMetalResourceViewBase* LinkedViews = nullptr;
 };
 
+inline uint32 GetTypeHash(const MTLTexturePtr& TexturePtr)
+{
+    return GetTypeHash(TexturePtr.get());
+}
+
 // Metal RHI texture resource
 class METALRHI_API FMetalSurface : public FRHITexture, public FMetalViewableResource
 {
@@ -589,10 +603,7 @@ class FMetalMeshShader;
 class FMetalAmplificationShader;
 #endif
 
-template<class T>
-struct TMetalResourceTraits
-{
-};
+
 template<>
 struct TMetalResourceTraits<FRHIShaderLibrary>
 {
@@ -661,26 +672,6 @@ struct TMetalResourceTraits<FRHIUnorderedAccessView>
 	typedef FMetalUnorderedAccessView TConcreteType;
 };
 
-template<>
-struct TMetalResourceTraits<FRHISamplerState>
-{
-	typedef FMetalSamplerState TConcreteType;
-};
-template<>
-struct TMetalResourceTraits<FRHIRasterizerState>
-{
-	typedef FMetalRasterizerState TConcreteType;
-};
-template<>
-struct TMetalResourceTraits<FRHIDepthStencilState>
-{
-	typedef FMetalDepthStencilState TConcreteType;
-};
-template<>
-struct TMetalResourceTraits<FRHIBlendState>
-{
-	typedef FMetalBlendState TConcreteType;
-};
 template<>
 struct TMetalResourceTraits<FRHIGraphicsPipelineState>
 {
diff --git a/Engine/Source/Runtime/Apple/MetalRHI/Public/MetalState.h b/Engine/Source/Runtime/Apple/MetalRHI/Public/MetalState.h
--- a/Engine/Source/Runtime/Apple/MetalRHI/Public/MetalState.h
+++ b/Engine/Source/Runtime/Apple/MetalRHI/Public/MetalState.h
@@ -6,11 +6,9 @@
 
 #pragma once
 
-#include "CoreMinimal.h"
+#include "MetalRHIPrivate.h"
 #include "RHI.h"
-THIRD_PARTY_INCLUDES_START
-#include "MetalInclude.h"
-THIRD_PARTY_INCLUDES_END
+#include "RHIResources.h"
 
 class FMetalSamplerState : public FRHISamplerState
 {
@@ -94,3 +92,24 @@ private:
 	static uint8 NextKey;
 	static FCriticalSection Mutex;
 };
+
+template<>
+struct TMetalResourceTraits<FRHISamplerState>
+{
+    typedef FMetalSamplerState TConcreteType;
+};
+template<>
+struct TMetalResourceTraits<FRHIRasterizerState>
+{
+    typedef FMetalRasterizerState TConcreteType;
+};
+template<>
+struct TMetalResourceTraits<FRHIDepthStencilState>
+{
+    typedef FMetalDepthStencilState TConcreteType;
+};
+template<>
+struct TMetalResourceTraits<FRHIBlendState>
+{
+    typedef FMetalBlendState TConcreteType;
+};
diff --git a/Engine/Source/Runtime/Apple/MetalRHI/Public/MetalThirdParty.h b/Engine/Source/Runtime/Apple/MetalRHI/Public/MetalThirdParty.h
new file mode 100644
--- /dev/null
+++ b/Engine/Source/Runtime/Apple/MetalRHI/Public/MetalThirdParty.h
@@ -0,0 +1,12 @@
+// Copyright Epic Games, Inc. All Rights Reserved.
+
+#pragma once
+
+// Need CoreMinimal to avoid collisions
+#include "CoreMinimal.h"
+
+// Metal C++ wrapper
+THIRD_PARTY_INCLUDES_START
+    #include <Metal/Metal.h>
+    #include "MetalInclude.h"
+THIRD_PARTY_INCLUDES_END
diff --git a/Engine/Source/Runtime/Apple/MetalRHI/Public/MetalUniformBuffer.h b/Engine/Source/Runtime/Apple/MetalRHI/Public/MetalUniformBuffer.h
--- a/Engine/Source/Runtime/Apple/MetalRHI/Public/MetalUniformBuffer.h
+++ b/Engine/Source/Runtime/Apple/MetalRHI/Public/MetalUniformBuffer.h
@@ -6,7 +6,7 @@
 //  Implements handles to linearly allocated per-frame constant buffers for shared memory systems.
 //
 
-#import <Metal/Metal.h>
+#import "MetalThirdParty.h"
 
 #include "RHIResources.h"
 
diff --git a/Engine/Source/Runtime/Apple/MetalRHI/Public/MetalViewport.h b/Engine/Source/Runtime/Apple/MetalRHI/Public/MetalViewport.h
--- a/Engine/Source/Runtime/Apple/MetalRHI/Public/MetalViewport.h
+++ b/Engine/Source/Runtime/Apple/MetalRHI/Public/MetalViewport.h
@@ -6,20 +6,23 @@
 
 #pragma once
 
+#include "MetalRHIPrivate.h"
+#include "PixelFormat.h"
+#include "RHIResources.h"
+
 #if PLATFORM_MAC
 #include "Mac/CocoaTextView.h"
 @interface FMetalView : FCocoaTextView
 @end
 #endif
 #include "HAL/PlatformFramePacer.h"
-THIRD_PARTY_INCLUDES_START
-#include "MetalInclude.h"
-THIRD_PARTY_INCLUDES_END
 
 #if PLATFORM_VISIONOS
 #import <CompositorServices/CompositorServices.h>
 #endif
 
+class FMetalSurface;
+
 enum EMetalViewportAccessFlag
 {
 	EMetalViewportAccessRHI,
