Disable heaps on AMD\Intel platforms due to lack of shared memory support and heaps being a requirement only for Apple Sillicon

diff --git a/Engine/Config/BaseEngine.ini b/Engine/Config/BaseEngine.ini
--- a/Engine/Config/BaseEngine.ini
+++ b/Engine/Config/BaseEngine.ini
@@ -3154,7 +3154,7 @@ bEnableRayTracing=true
 
 [/Script/MacTargetPlatform.MacTargetSettings]
 MetalLanguageVersion=7
-+TargetedRHIs=SF_METAL_SM6
++TargetedRHIs=SF_METAL_SM5
 UseFastIntrinsics=False
 EnableMathOptimisations=True
 TargetArchitecture=MacTargetArchitectureUniversal
diff --git a/Engine/Source/Runtime/Apple/MetalRHI/Private/MetalCommandQueue.cpp b/Engine/Source/Runtime/Apple/MetalRHI/Private/MetalCommandQueue.cpp
--- a/Engine/Source/Runtime/Apple/MetalRHI/Private/MetalCommandQueue.cpp
+++ b/Engine/Source/Runtime/Apple/MetalRHI/Private/MetalCommandQueue.cpp
@@ -228,7 +228,8 @@ FMetalCommandQueue::FMetalCommandQueue(MTL::Device* InDevice, uint32 const MaxNu
 		}
 	}
     
-	if (!FParse::Param(FCommandLine::Get(),TEXT("nometalheap")))
+	// Temporarily only support heaps for devices with unified memory
+	if (Device->hasUnifiedMemory() && !FParse::Param(FCommandLine::Get(),TEXT("nometalheap")))
 	{
 		Features |= EMetalFeaturesHeaps;
 	}
