Fix memory leaks in Metal:
- MTLFunctionConstantValuesInternal is not getting deallocated
- RHIBindDebugName does not free CFString allocated by CFStringCreateWithBytes in GetCFString

diff --git a/Engine/Source/Runtime/Apple/MetalRHI/Private/MetalTexture.cpp b/Engine/Source/Runtime/Apple/MetalRHI/Private/MetalTexture.cpp
--- a/Engine/Source/Runtime/Apple/MetalRHI/Private/MetalTexture.cpp
+++ b/Engine/Source/Runtime/Apple/MetalRHI/Private/MetalTexture.cpp
@@ -1900,11 +1900,10 @@ void FMetalDynamicRHI::RHIUnlockTextureCubeFace(FRHITexture* TextureCubeRHI,uint
 
 void FMetalDynamicRHI::RHIBindDebugLabelName(FRHICommandListBase& RHICmdList, FRHITexture* TextureRHI, const TCHAR* Name)
 {
-    MTL_SCOPED_AUTORELEASE_POOL;
-    
-    FMetalSurface* Surf = GetMetalSurfaceFromRHITexture(TextureRHI);
-    
-    NS::String* LabelString = (NS::String*)FString(Name).GetCFString();
+	MTL_SCOPED_AUTORELEASE_POOL;
+    FMetalSurface* Surf = GetMetalSurfaceFromRHITexture(TextureRHI);   
+	
+    NS::String* LabelString = FStringToNSString(Name);
     if(Surf->Texture)
     {
         Surf->Texture->setLabel(LabelString);
diff --git a/Engine/Source/Runtime/Apple/MetalRHI/Private/Shaders/Types/Templates/MetalBaseShader.h b/Engine/Source/Runtime/Apple/MetalRHI/Private/Shaders/Types/Templates/MetalBaseShader.h
--- a/Engine/Source/Runtime/Apple/MetalRHI/Private/Shaders/Types/Templates/MetalBaseShader.h
+++ b/Engine/Source/Runtime/Apple/MetalRHI/Private/Shaders/Types/Templates/MetalBaseShader.h
@@ -13,6 +13,7 @@
 #include "MetalContext.h"
 #include "MetalShaderResources.h"
 #include "Misc/FileHelper.h"
+#include "Misc/ScopeExit.h"
 #include "Serialization/MemoryReader.h"
 #include "Shaders/Debugging/MetalShaderDebugCache.h"
 #include "Shaders/MetalCompiledShaderKey.h"
@@ -457,6 +458,7 @@ MTLFunctionPtr TMetalBaseShader<BaseResourceType, ShaderType>::GetCompiledFuncti
             NS::String* Name = (LibraryFunctionIndex != -1)
             ? (NS::String*)Library->functionNames()->object(LibraryFunctionIndex) : FStringToNSString(FString::Printf(TEXT("Main_%0.8x_%0.8x"), SourceLen, SourceCRC));
 			MTL::FunctionConstantValues* ConstantValues = nullptr;
+			ON_SCOPE_EXIT { if (ConstantValues){ ConstantValues->release(); } };
 			if (bHasFunctionConstants)
 			{
 				ConstantValues = MTL::FunctionConstantValues::alloc()->init();
