Fixed memory leak in MetalRHI with MTLVertexDescriptor

diff --git a/Engine/Source/Runtime/Apple/MetalRHI/Private/MetalPipeline.cpp b/Engine/Source/Runtime/Apple/MetalRHI/Private/MetalPipeline.cpp
--- a/Engine/Source/Runtime/Apple/MetalRHI/Private/MetalPipeline.cpp
+++ b/Engine/Source/Runtime/Apple/MetalRHI/Private/MetalPipeline.cpp
@@ -582,7 +582,7 @@ static MTLVertexDescriptorPtr GetMaskedVertexDescriptor(MTLVertexDescriptorPtr I
 	{
 		if (!InOutMask.IsFieldEnabled((int32)Attr) && InputDesc->attributes()->object(Attr) != nullptr)
 		{
-			MTLVertexDescriptorPtr Desc = NS::RetainPtr(InputDesc->copy());
+			MTLVertexDescriptorPtr Desc = NS::TransferPtr(InputDesc->copy());
 			CrossCompiler::FShaderBindingInOutMask BuffersUsed;
 			for (int32 MetalStreamIndex = 0; MetalStreamIndex < MaxMetalStreams; ++MetalStreamIndex)
 			{
