Fixed CubeMap array RT's in Metal RHI using incorrect slice

diff --git a/Engine/Source/Runtime/Apple/MetalRHI/Private/MetalStateCache.cpp b/Engine/Source/Runtime/Apple/MetalRHI/Private/MetalStateCache.cpp
--- a/Engine/Source/Runtime/Apple/MetalRHI/Private/MetalStateCache.cpp
+++ b/Engine/Source/Runtime/Apple/MetalRHI/Private/MetalStateCache.cpp
@@ -603,10 +603,6 @@ bool FMetalStateCache::SetRenderPassInfo(FRHIRenderPassInfo const& InRenderTarge
 	
 				// user code generally passes -1 as a default, but we need 0
 				uint32 ArraySliceIndex = RenderTargetView.ArraySlice == 0xFFFFFFFF ? 0 : RenderTargetView.ArraySlice;
-				if (Surface.GetDesc().IsTextureCube())
-				{
-					ArraySliceIndex = GetMetalCubeFace((ECubeFace)ArraySliceIndex);
-				}
 				
 				switch(Surface.GetDesc().Dimension)
 				{
