Metal Bindless
- Removed code that was binding InResourceParameters in MetalCommands
- Fixed texture residency in heaps
- Fixed crash when bindless was enabled and building for unsupported platform

diff --git a/Engine/Source/Runtime/Apple/MetalRHI/Private/MetalCommands.cpp b/Engine/Source/Runtime/Apple/MetalRHI/Private/MetalCommands.cpp
--- a/Engine/Source/Runtime/Apple/MetalRHI/Private/MetalCommands.cpp
+++ b/Engine/Source/Runtime/Apple/MetalRHI/Private/MetalCommands.cpp
@@ -402,20 +402,6 @@ void FMetalRHICommandContext::RHISetShaderParameters(FRHIGraphicsShader* Shader,
 #if PLATFORM_SUPPORTS_BINDLESS_RENDERING
 	if(IsMetalBindlessEnabled())
 	{
-		for (const FRHIShaderParameterResource& Parameter : InResourceParameters)
-		{
-			if(Parameter.Type == FRHIShaderParameterResource::EType::UniformBuffer)
-			{
-				continue;
-			}
-			const FRHIDescriptorHandle Handle = UE::RHICore::GetBindlessParameterHandle(Parameter);
-			if (Handle.IsValid())
-			{
-				const uint32 BindlessIndex = Handle.GetIndex();
-				RHISetShaderParameter(Shader, 0, Parameter.Index, 4, &BindlessIndex);
-			}
-		}
-		
 		EMetalShaderStages Stage = GetShaderStage(Shader);
 		Context->GetCurrentState().IRForwardBindlessParameters(Stage, InResourceParameters);
 		Context->GetCurrentState().IRForwardBindlessParameters(Stage, InBindlessParameters);
@@ -437,20 +423,6 @@ void FMetalRHICommandContext::RHISetShaderParameters(FRHIComputeShader* Shader,
 #if PLATFORM_SUPPORTS_BINDLESS_RENDERING
 	if(IsMetalBindlessEnabled())
 	{
-		for (const FRHIShaderParameterResource& Parameter : InResourceParameters)
-		{
-			if(Parameter.Type == FRHIShaderParameterResource::EType::UniformBuffer)
-			{
-				continue;
-			}
-			const FRHIDescriptorHandle Handle = UE::RHICore::GetBindlessParameterHandle(Parameter);
-			if (Handle.IsValid())
-			{
-				const uint32 BindlessIndex = Handle.GetIndex();
-				RHISetShaderParameter(Shader, 0, Parameter.Index, 4, &BindlessIndex);
-			}
-		}
-		
 		Context->GetCurrentState().IRForwardBindlessParameters(EMetalShaderStages::Compute, InResourceParameters);
 		Context->GetCurrentState().IRForwardBindlessParameters(EMetalShaderStages::Compute, InBindlessParameters);
 	}
diff --git a/Engine/Source/Runtime/Apple/MetalRHI/Private/MetalDynamicRHI.cpp b/Engine/Source/Runtime/Apple/MetalRHI/Private/MetalDynamicRHI.cpp
--- a/Engine/Source/Runtime/Apple/MetalRHI/Private/MetalDynamicRHI.cpp
+++ b/Engine/Source/Runtime/Apple/MetalRHI/Private/MetalDynamicRHI.cpp
@@ -13,10 +13,6 @@
 #include "MetalGraphicsPipelineState.h"
 #include "MetalComputePipelineState.h"
 #include "MetalTransitionData.h"
- 
-#if METAL_USE_METAL_SHADER_CONVERTER
-#include "metal_irconverter.h"
-#endif
 
 //------------------------------------------------------------------------------
 
diff --git a/Engine/Source/Runtime/Apple/MetalRHI/Private/MetalRHIPrivate.h b/Engine/Source/Runtime/Apple/MetalRHI/Private/MetalRHIPrivate.h
--- a/Engine/Source/Runtime/Apple/MetalRHI/Private/MetalRHIPrivate.h
+++ b/Engine/Source/Runtime/Apple/MetalRHI/Private/MetalRHIPrivate.h
@@ -163,7 +163,8 @@ void SafeReleaseFunction(TFunction<void()> ReleaseFunction);
 
 FORCEINLINE bool IsMetalBindlessEnabled()
 {
-	return GRHIBindlessSupport != ERHIBindlessSupport::Unsupported;
+	return GRHIBindlessSupport != ERHIBindlessSupport::Unsupported &&
+			GMaxRHIFeatureLevel >= ERHIFeatureLevel::SM6;
 }
 
 // Access the underlying surface object from any kind of texture
diff --git a/Engine/Source/Runtime/Apple/MetalRHI/Private/MetalRenderPass.cpp b/Engine/Source/Runtime/Apple/MetalRHI/Private/MetalRenderPass.cpp
--- a/Engine/Source/Runtime/Apple/MetalRHI/Private/MetalRenderPass.cpp
+++ b/Engine/Source/Runtime/Apple/MetalRHI/Private/MetalRenderPass.cpp
@@ -19,7 +19,7 @@
 
 #pragma mark - Private Console Variables -
 
-static int32 GMetalCommandBufferCommitThreshold = 256;
+static int32 GMetalCommandBufferCommitThreshold = 0;
 static FAutoConsoleVariableRef CVarMetalCommandBufferCommitThreshold(
 	TEXT("rhi.Metal.CommandBufferCommitThreshold"),
 	GMetalCommandBufferCommitThreshold,
diff --git a/Engine/Source/Runtime/Apple/MetalRHI/Private/MetalState.cpp b/Engine/Source/Runtime/Apple/MetalRHI/Private/MetalState.cpp
--- a/Engine/Source/Runtime/Apple/MetalRHI/Private/MetalState.cpp
+++ b/Engine/Source/Runtime/Apple/MetalRHI/Private/MetalState.cpp
@@ -310,7 +310,7 @@ FMetalSamplerState::FMetalSamplerState(FMetalDeviceContext* Context, const FSamp
     FMetalBindlessDescriptorManager* BindlessDescriptorManager = Context->GetBindlessDescriptorManager();
     check(BindlessDescriptorManager);
 
-	if(BindlessDescriptorManager->IsSupported())
+	if(IsMetalBindlessEnabled())
 	{
 		BindlessHandle = BindlessDescriptorManager->ReserveDescriptor(ERHIDescriptorHeapType::Sampler);
 		BindlessDescriptorManager->BindSampler(BindlessHandle, State);
@@ -324,7 +324,7 @@ FMetalSamplerState::~FMetalSamplerState()
     FMetalBindlessDescriptorManager* BindlessDescriptorManager = GetMetalDeviceContext().GetBindlessDescriptorManager();
     check(BindlessDescriptorManager);
 
-	if(BindlessDescriptorManager->IsSupported())
+	if(IsMetalBindlessEnabled())
 	{
 		BindlessDescriptorManager->FreeDescriptor(BindlessHandle);
 	}
diff --git a/Engine/Source/Runtime/Apple/MetalRHI/Private/MetalStateCache.cpp b/Engine/Source/Runtime/Apple/MetalRHI/Private/MetalStateCache.cpp
--- a/Engine/Source/Runtime/Apple/MetalRHI/Private/MetalStateCache.cpp
+++ b/Engine/Source/Runtime/Apple/MetalRHI/Private/MetalStateCache.cpp
@@ -1176,7 +1176,7 @@ void FMetalStateCache::SetVertexStream(uint32 const Index, FMetalBufferPtr Buffe
 #if METAL_USE_METAL_SHADER_CONVERTER
 	FMetalBindlessDescriptorManager* BindlessDescriptorManager = GetMetalDeviceContext().GetBindlessDescriptorManager();
 	
-	if(BindlessDescriptorManager->IsSupported())
+	if(IsMetalBindlessEnabled())
 	{
 		// Update GPU VA (assuming the offset has changed since last time).
 		if (Buffer)
@@ -1535,12 +1535,18 @@ void FMetalStateCache::SetShaderBuffer(EMetalShaderStages const Frequency, MTL::
 }
 #endif // METAL_RHI_RAYTRACING
 
+static bool CanMakeTextureResidentViaHeaps(const MTL::Texture* Texture)
+{
+    return !(Texture->usage() & MTL::TextureUsageRenderTarget)
+        && !(Texture->usage() & MTL::TextureUsageShaderWrite);
+}
+
 #if METAL_USE_METAL_SHADER_CONVERTER
 void FMetalStateCache::IRMakeSRVResident(EMetalShaderStages const Frequency, FMetalShaderResourceView* SRV)
 {
     FMetalBindlessDescriptorManager* BindlessDescriptorManager = GetMetalDeviceContext().GetBindlessDescriptorManager();
 
-	if(!BindlessDescriptorManager->IsSupported())
+	if(!IsMetalBindlessEnabled())
 	{
 		return;
 	}
@@ -1555,7 +1561,7 @@ void FMetalStateCache::IRMakeSRVResident(EMetalShaderStages const Frequency, FMe
         {
             auto const& View = SRV->GetTextureView();
 			
-			if (!View->heap())
+            if (!View->heap() || !CanMakeTextureResidentViaHeaps(View.get()))
 			{
 				BindlessDescriptorManager->MakeResident(SRV->GetBindlessHandle(), View.get(), MTL::ResourceUsage(MTL::ResourceUsageRead | MTL::ResourceUsageSample), Frequency);
 			}
@@ -1590,7 +1596,7 @@ void FMetalStateCache::IRMakeSRVResident(EMetalShaderStages const Frequency, FMe
 void FMetalStateCache::IRMakeUAVResident(EMetalShaderStages const Frequency, FMetalUnorderedAccessView* UAV)
 {
     FMetalBindlessDescriptorManager* BindlessDescriptorManager = GetMetalDeviceContext().GetBindlessDescriptorManager();
-	if(!BindlessDescriptorManager->IsSupported())
+	if(!IsMetalBindlessEnabled())
 	{
 		return;
 	}
@@ -1605,10 +1611,7 @@ void FMetalStateCache::IRMakeUAVResident(EMetalShaderStages const Frequency, FMe
         {
             auto const& View = UAV->GetTextureView();
             
-			if (!View->heap())
-			{
-				BindlessDescriptorManager->MakeResident(UAV->GetBindlessHandle(), View.get(), MTL::ResourceUsage(MTL::ResourceUsageRead | MTL::ResourceUsageWrite), Frequency);
-			}
+			BindlessDescriptorManager->MakeResident(UAV->GetBindlessHandle(), View.get(), MTL::ResourceUsage(MTL::ResourceUsageRead | MTL::ResourceUsageWrite), Frequency);
 			break;
         }
             
@@ -1616,10 +1619,7 @@ void FMetalStateCache::IRMakeUAVResident(EMetalShaderStages const Frequency, FMe
         {
             auto const& View = UAV->GetBufferView();
             
-			if (!View.Buffer->GetMTLBuffer()->heap())
-			{
-				BindlessDescriptorManager->MakeResident(UAV->GetBindlessHandle(), View.Buffer->GetMTLBuffer().get(), MTL::ResourceUsage(MTL::ResourceUsageRead | MTL::ResourceUsageWrite), Frequency);
-			}
+			BindlessDescriptorManager->MakeResident(UAV->GetBindlessHandle(), View.Buffer->GetMTLBuffer().get(), MTL::ResourceUsage(MTL::ResourceUsageRead | MTL::ResourceUsageWrite), Frequency);
 			
 			break;
         }
@@ -1628,14 +1628,9 @@ void FMetalStateCache::IRMakeUAVResident(EMetalShaderStages const Frequency, FMe
         {
             auto const& View = UAV->GetTextureBufferBacked();
             
-			if (!View.Buffer->GetMTLBuffer()->heap())
-			{
-				BindlessDescriptorManager->MakeResident(UAV->GetBindlessHandle(), View.Buffer->GetMTLBuffer().get(), MTL::ResourceUsage(MTL::ResourceUsageRead | MTL::ResourceUsageWrite), Frequency);
-			}
-			if (!View.Texture->heap())
-			{
-				BindlessDescriptorManager->MakeResident(UAV->GetBindlessHandle(), View.Texture.get(), MTL::ResourceUsage(MTL::ResourceUsageRead | MTL::ResourceUsageWrite), Frequency);
-			}
+			BindlessDescriptorManager->MakeResident(UAV->GetBindlessHandle(), View.Buffer->GetMTLBuffer().get(), MTL::ResourceUsage(MTL::ResourceUsageRead | MTL::ResourceUsageWrite), Frequency);
+			BindlessDescriptorManager->MakeResident(UAV->GetBindlessHandle(), View.Texture.get(), MTL::ResourceUsage(MTL::ResourceUsageRead | MTL::ResourceUsageWrite), Frequency);
+
 			break;
         }
             
@@ -1654,7 +1649,7 @@ void FMetalStateCache::IRMakeTextureResident(EMetalShaderStages const Frequency,
 {
     FMetalBindlessDescriptorManager* BindlessDescriptorManager = GetMetalDeviceContext().GetBindlessDescriptorManager();
     
-	if (!Texture->heap())
+    if (!Texture->heap() || !CanMakeTextureResidentViaHeaps(Texture))
 	{
 		BindlessDescriptorManager->MakeResident(FRHIDescriptorHandle(), Texture, MTL::ResourceUsage(MTL::ResourceUsageRead | MTL::ResourceUsageSample), Frequency);
 	}
@@ -2574,7 +2569,7 @@ void FMetalStateCache::SetRenderPipelineState(FMetalCommandEncoder& CommandEncod
 #if METAL_USE_METAL_SHADER_CONVERTER
 		FMetalBindlessDescriptorManager* BindlessDescriptorManager = GetMetalDeviceContext().GetBindlessDescriptorManager();
 		
-		if(BindlessDescriptorManager->IsSupported())
+		if(IsMetalBindlessEnabled())
 		{
 #if PLATFORM_SUPPORTS_MESH_SHADERS
 			if (GraphicsPSO->VertexDeclaration != nullptr)
@@ -2596,7 +2591,7 @@ void FMetalStateCache::SetRenderPipelineState(FMetalCommandEncoder& CommandEncod
 #if METAL_DEBUG_OPTIONS
 #if PLATFORM_SUPPORTS_BINDLESS_RENDERING
 	FMetalBindlessDescriptorManager* BindlessDescriptorManager = GetMetalDeviceContext().GetBindlessDescriptorManager();
-	if(!BindlessDescriptorManager->IsSupported())
+	if(!IsMetalBindlessEnabled())
 #endif
 	{
 		Validate();
diff --git a/Engine/Source/Runtime/Apple/MetalRHI/Private/MetalTexture.cpp b/Engine/Source/Runtime/Apple/MetalRHI/Private/MetalTexture.cpp
--- a/Engine/Source/Runtime/Apple/MetalRHI/Private/MetalTexture.cpp
+++ b/Engine/Source/Runtime/Apple/MetalRHI/Private/MetalTexture.cpp
@@ -831,7 +831,7 @@ FMetalSurface::FMetalSurface(FRHICommandListBase* RHICmdList, FMetalTextureCreat
     FMetalBindlessDescriptorManager* BindlessDescriptorManager = GetMetalDeviceContext().GetBindlessDescriptorManager();
     check(BindlessDescriptorManager);
 	
-	if(BindlessDescriptorManager->IsSupported())
+	if(IsMetalBindlessEnabled())
 	{
 		BindlessHandle = BindlessDescriptorManager->ReserveDescriptor(ERHIDescriptorHeapType::Standard);
 		
@@ -911,7 +911,7 @@ FMetalSurface::~FMetalSurface()
     FMetalBindlessDescriptorManager* BindlessDescriptorManager = GetMetalDeviceContext().GetBindlessDescriptorManager();
     check(BindlessDescriptorManager);
 
-	if(BindlessDescriptorManager->IsSupported())
+	if(IsMetalBindlessEnabled())
 	{
 		if (!(GetDesc().Flags & TexCreate_Presentable))
 		{
@@ -2090,7 +2090,7 @@ FTextureReferenceRHIRef FMetalDynamicRHI::RHICreateTextureReference(FRHICommandL
 	FMetalBindlessDescriptorManager* BindlessDescriptorManager = GetMetalDeviceContext().GetBindlessDescriptorManager();
 	check(BindlessDescriptorManager);
 	
-	if(BindlessDescriptorManager->IsSupported())
+	if(IsMetalBindlessEnabled())
 	{
 		// If the referenced texture is configured for bindless, make sure we also create an SRV to use for bindless.
 		if (ReferencedTexture && ReferencedTexture->GetDefaultBindlessHandle().IsValid())
@@ -2115,12 +2115,17 @@ void FMetalDynamicRHI::RHIUpdateTextureReference(FRHICommandListBase& RHICmdList
 		FMetalTextureReference* MetalTextureReference = ResourceCast(TextureRef);
 		FMetalShaderResourceView* MetalTextureRefSRV = MetalTextureReference->BindlessView;
 		
-		FMetalSurface* NewSurface = GetMetalSurfaceFromRHITexture(NewTexture);
-		
-		FMetalBindlessDescriptorManager* BindlessDescriptorManager = GetMetalDeviceContext().GetBindlessDescriptorManager();
-		check(BindlessDescriptorManager);
-	
-		BindlessDescriptorManager->BindTexture(MetalTextureRefSRV->GetBindlessHandle(), NewSurface->Texture.get());
+        FRHIDescriptorHandle DestHandle = MetalTextureRefSRV->GetBindlessHandle();
+        
+        if(DestHandle.IsValid())
+        {
+            FMetalSurface* NewSurface = GetMetalSurfaceFromRHITexture(NewTexture);
+            
+            FMetalBindlessDescriptorManager* BindlessDescriptorManager = GetMetalDeviceContext().GetBindlessDescriptorManager();
+            check(BindlessDescriptorManager);
+            
+            BindlessDescriptorManager->BindTexture(DestHandle, NewSurface->Texture.get());
+        }
 	}
 #endif // PLATFORM_SUPPORTS_BINDLESS_RENDERING
 
diff --git a/Engine/Source/Runtime/Apple/MetalRHI/Private/MetalUAV.cpp b/Engine/Source/Runtime/Apple/MetalRHI/Private/MetalUAV.cpp
--- a/Engine/Source/Runtime/Apple/MetalRHI/Private/MetalUAV.cpp
+++ b/Engine/Source/Runtime/Apple/MetalRHI/Private/MetalUAV.cpp
@@ -86,7 +86,7 @@ FMetalShaderResourceView::FMetalShaderResourceView(FRHICommandListBase& RHICmdLi
 	FMetalBindlessDescriptorManager* BindlessDescriptorManager = GetMetalDeviceContext().GetBindlessDescriptorManager();
     check(BindlessDescriptorManager);
 
-	if(BindlessDescriptorManager->IsSupported())
+	if(IsMetalBindlessEnabled())
 	{
 		BindlessHandle = BindlessDescriptorManager->ReserveDescriptor(ERHIDescriptorHeapType::Standard);
 	}
@@ -105,7 +105,7 @@ FMetalShaderResourceView::~FMetalShaderResourceView()
         FMetalBindlessDescriptorManager* BindlessDescriptorManager = GetMetalDeviceContext().GetBindlessDescriptorManager();
         check(BindlessDescriptorManager);
 
-		if(BindlessDescriptorManager->IsSupported())
+		if(IsMetalBindlessEnabled())
 		{
 			BindlessDescriptorManager->FreeDescriptor(BindlessHandle);
 		}
@@ -341,7 +341,7 @@ void FMetalShaderResourceView::UpdateView()
 	FMetalBindlessDescriptorManager* BindlessDescriptorManager = GetMetalDeviceContext().GetBindlessDescriptorManager();
     check(BindlessDescriptorManager);
 
-	if(BindlessDescriptorManager->IsSupported())
+	if(IsMetalBindlessEnabled())
 	{
 		BindlessDescriptorManager->BindResource(BindlessHandle, this);
 	}
@@ -355,7 +355,7 @@ FMetalUnorderedAccessView::FMetalUnorderedAccessView(FRHICommandListBase& RHICmd
     FMetalBindlessDescriptorManager* BindlessDescriptorManager = GetMetalDeviceContext().GetBindlessDescriptorManager();
 	check(BindlessDescriptorManager);
 
-	if(BindlessDescriptorManager->IsSupported())
+	if(IsMetalBindlessEnabled())
 	{
 		BindlessHandle = BindlessDescriptorManager->ReserveDescriptor(ERHIDescriptorHeapType::Standard);
 	}
@@ -374,7 +374,7 @@ FMetalUnorderedAccessView::~FMetalUnorderedAccessView()
 	FMetalBindlessDescriptorManager* BindlessDescriptorManager = GetMetalDeviceContext().GetBindlessDescriptorManager();
 	check(BindlessDescriptorManager);
 
-	if(BindlessDescriptorManager->IsSupported())
+	if(IsMetalBindlessEnabled())
 	{
 		BindlessDescriptorManager->FreeDescriptor(BindlessHandle);
 	}
@@ -579,7 +579,7 @@ void FMetalUnorderedAccessView::UpdateView()
     FMetalBindlessDescriptorManager* BindlessDescriptorManager = GetMetalDeviceContext().GetBindlessDescriptorManager();
 	check(BindlessDescriptorManager);
 
-	if(BindlessDescriptorManager->IsSupported())
+	if(IsMetalBindlessEnabled())
 	{
 		BindlessDescriptorManager->BindResource(BindlessHandle, this);
 	}
diff --git a/Engine/Source/Runtime/Apple/MetalRHI/Private/MetalVertexDeclaration.cpp b/Engine/Source/Runtime/Apple/MetalRHI/Private/MetalVertexDeclaration.cpp
--- a/Engine/Source/Runtime/Apple/MetalRHI/Private/MetalVertexDeclaration.cpp
+++ b/Engine/Source/Runtime/Apple/MetalRHI/Private/MetalVertexDeclaration.cpp
@@ -10,10 +10,6 @@
 #include "MetalHashedVertexDescriptor.h"
 #include "MetalProfiler.h"
 
-#if METAL_USE_METAL_SHADER_CONVERTER
-#include "metal_irconverter.h"
-#endif
-
 //------------------------------------------------------------------------------
 
 #pragma mark - Metal Vertex Declaration Globals
