MetalRHI cleanup
- Moving shader specific data to FMetalShaderData so that it can be accessed without each shader type
- Implementing custom Metal version of RHISetShaderParametersShared

diff --git a/Engine/Source/Runtime/Apple/MetalRHI/Private/MetalCommands.cpp b/Engine/Source/Runtime/Apple/MetalRHI/Private/MetalCommands.cpp
--- a/Engine/Source/Runtime/Apple/MetalRHI/Private/MetalCommands.cpp
+++ b/Engine/Source/Runtime/Apple/MetalRHI/Private/MetalCommands.cpp
@@ -96,13 +96,86 @@ void FMetalRHICommandContext::RHISetStreamSource(uint32 StreamIndex, FRHIBuffer*
     Context->GetCurrentState().SetVertexStream(StreamIndex, VertexBuffer ? TheBuffer : nullptr, VertexBuffer ? VertexBuffer->Data : nullptr, Offset, VertexBuffer ? VertexBuffer->GetSize() : 0);
 }
 
+static void SetUniformBufferInternal(FMetalStateCache& StateCache, FMetalShaderData* ShaderData, EMetalShaderStages Stage, uint32 BufferIndex, FRHIUniformBuffer* UBRHI)
+{
+    StateCache.BindUniformBuffer(Stage, BufferIndex, UBRHI);
+        
+    FMetalShaderBindings& Bindings = ShaderData->Bindings;
+    if ((Bindings.ConstantBuffers) & (1 << BufferIndex))
+    {
+        FMetalUniformBuffer* UB = ResourceCast(UBRHI);
+        UB->PrepareToBind();
+#if METAL_USE_METAL_SHADER_CONVERTER
+        if (IsMetalBindlessEnabled())
+        {
+            StateCache.IRBindUniformBuffer(Stage, BufferIndex, UB);
+        }
+        else
+#endif
+        {
+            FMetalBufferPtr Buf = FMetalBufferPtr(new FMetalBuffer(UB->Backing));
+            StateCache.SetShaderBuffer(Stage, Buf, nullptr, UB->Offset, UB->GetSize(), BufferIndex, MTL::ResourceUsageRead);
+        }
+    }
+}
+
+inline FMetalShaderData* GetShaderData(FRHIShader* InShaderRHI, EMetalShaderStages Stage)
+{
+    switch (Stage)
+    {
+    case EMetalShaderStages::Vertex:        return ResourceCast(static_cast<FRHIVertexShader*>(InShaderRHI));
+#if PLATFORM_SUPPORTS_MESH_SHADERS
+    case EMetalShaderStages::Mesh:          return ResourceCast(static_cast<FRHIMeshShader*>(InShaderRHI));
+    case EMetalShaderStages::Amplification: return ResourceCast(static_cast<FRHIAmplificationShader*>(InShaderRHI));
+#endif
+    case EMetalShaderStages::Pixel:         return ResourceCast(static_cast<FRHIPixelShader*>(InShaderRHI));
+#if PLATFORM_SUPPORTS_GEOMETRY_SHADERS
+    case EMetalShaderStages::Geometry:      return ResourceCast(static_cast<FRHIGeometryShader*>(InShaderRHI));
+#endif
+    case EMetalShaderStages::Compute:       return ResourceCast(static_cast<FRHIComputeShader*>(InShaderRHI));
+            
+    default:
+        checkf(0, TEXT("FRHIShader Type %d is invalid or unsupported!"), (int32)InShaderRHI->GetFrequency());
+        NOT_SUPPORTED("RHIShaderStage");
+    }
+    return nullptr;
+}
+
+static void BindUniformBuffer(FMetalStateCache& StateCache, FRHIShader* Shader, EMetalShaderStages Stage, uint32 BufferIndex, FRHIUniformBuffer* InBuffer)
+{
+    if (FMetalShaderData* ShaderData = GetShaderData(Shader, Stage))
+    {
+        SetUniformBufferInternal(StateCache, ShaderData, Stage, BufferIndex, InBuffer);
+    }
+}
+
+static void ApplyStaticUniformBuffersOnContext(FMetalRHICommandContext& Context, FRHIShader* Shader, FMetalShaderData* ShaderData)
+{
+    if (Shader)
+    {
+        MTL_SCOPED_AUTORELEASE_POOL;
+
+        FMetalStateCache& StateCache = Context.GetInternalContext().GetCurrentState();
+        const EMetalShaderStages Stage = GetMetalShaderFrequency(Shader->GetFrequency());
+
+        UE::RHICore::ApplyStaticUniformBuffers(
+            Shader,
+            Context.GetStaticUniformBuffers(),
+            [&StateCache, ShaderData, Stage](int32 BufferIndex, FRHIUniformBuffer* Buffer)
+            {
+                SetUniformBufferInternal(StateCache, ShaderData, Stage, BufferIndex, ResourceCast(Buffer));
+            }
+        );
+    }
+}
+
 template <typename TRHIShader>
-void FMetalRHICommandContext::ApplyStaticUniformBuffers(TRHIShader* Shader)
+static void ApplyStaticUniformBuffersOnContext(FMetalRHICommandContext& Context, TRefCountPtr<TRHIShader>& Shader)
 {
-	if (Shader)
-	{
-		UE::RHICore::ApplyStaticUniformBuffers(this, Shader, GlobalUniformBuffers);
-	}
+    if (IsValidRef(Shader))
+    {
+        ApplyStaticUniformBuffersOnContext(Context, Shader, static_cast<FMetalShaderData*>(Shader.GetReference()));
+    }
 }
 
 void FMetalRHICommandContext::RHISetComputePipelineState(FRHIComputePipelineState* ComputePipelineState)
@@ -115,8 +188,7 @@ void FMetalRHICommandContext::RHISetComputePipelineState(FRHIComputePipelineStat
 	// sets this compute shader pipeline as the current (this resets all state, so we need to set all resources after calling this)
 	Context->GetCurrentState().SetComputeShader(ComputePipeline->GetComputeShader());
 
-	ApplyStaticUniformBuffers(ComputePipeline->GetComputeShader());
-    
+    ApplyStaticUniformBuffersOnContext(*this, ComputePipeline->ComputeShader);
 }
 
 void FMetalRHICommandContext::RHIDispatchComputeShader(uint32 ThreadGroupCountX, uint32 ThreadGroupCountY, uint32 ThreadGroupCountZ)
@@ -237,21 +309,14 @@ void FMetalRHICommandContext::RHISetGraphicsPipelineState(FRHIGraphicsPipelineSt
     if (bApplyAdditionalState)
     {
 #if PLATFORM_SUPPORTS_MESH_SHADERS
-        if (IsValidRef(PipelineState->MeshShader))
-        {
-            ApplyStaticUniformBuffers(PipelineState->MeshShader.GetReference());
-            if (IsValidRef(PipelineState->AmplificationShader))
-            {
-                ApplyStaticUniformBuffers(PipelineState->AmplificationShader.GetReference());
-            }
-        }
-        else if (IsValidRef(PipelineState->VertexShader))
+        ApplyStaticUniformBuffersOnContext(*this, PipelineState->MeshShader);
+        ApplyStaticUniformBuffersOnContext(*this, PipelineState->AmplificationShader);
 #endif
-		{
-			ApplyStaticUniformBuffers(PipelineState->VertexShader.GetReference());
-		}
-		
-		ApplyStaticUniformBuffers(PipelineState->PixelShader.GetReference());
+#if PLATFORM_SUPPORTS_GEOMETRY_SHADERS
+        ApplyStaticUniformBuffersOnContext(*this, PipelineState->GeometryShader);
+#endif
+        ApplyStaticUniformBuffersOnContext(*this, PipelineState->VertexShader);
+        ApplyStaticUniformBuffersOnContext(*this, PipelineState->PixelShader);
     }
 }
 
@@ -265,168 +330,190 @@ void FMetalRHICommandContext::RHISetStaticUniformBuffers(const FUniformBufferSta
 	}
 }
 
-void FMetalRHICommandContext::RHISetUAVParameter(FRHIPixelShader* PixelShaderRHI, uint32 UAVIndex, FRHIUnorderedAccessView* UAVRHI)
+struct FMetalShaderBinder
 {
-    MTL_SCOPED_AUTORELEASE_POOL;
-    
-	FMetalUnorderedAccessView* UAV = ResourceCast(UAVRHI);
-	Context->GetCurrentState().SetShaderUnorderedAccessView(EMetalShaderStages::Pixel, UAVIndex, UAV);
-}
-
-void FMetalRHICommandContext::RHISetUAVParameter(FRHIComputeShader* ComputeShaderRHI, uint32 UAVIndex, FRHIUnorderedAccessView* UAVRHI)
-{
-    MTL_SCOPED_AUTORELEASE_POOL;
-    
-	FMetalUnorderedAccessView* UAV = ResourceCast(UAVRHI);
-	Context->GetCurrentState().SetShaderUnorderedAccessView(EMetalShaderStages::Compute, UAVIndex, UAV);
-}
-
-void FMetalRHICommandContext::RHISetUAVParameter(FRHIComputeShader* ComputeShaderRHI,uint32 UAVIndex, FRHIUnorderedAccessView* UAVRHI, uint32 InitialCount)
-{
-    MTL_SCOPED_AUTORELEASE_POOL;
+    FMetalDeviceContext& Context;
+    FMetalStateCache& StateCache;
+    const EMetalShaderStages Stage;
+    FMetalShaderParameterCache& ShaderParameters;
+#if PLATFORM_SUPPORTS_BINDLESS_RENDERING
+    const bool bBindlessResources;
+    const bool bBindlessSamplers;
+#endif
     
-	FMetalUnorderedAccessView* UAV = ResourceCast(UAVRHI);
-    Context->GetCurrentState().SetShaderUnorderedAccessView(EMetalShaderStages::Compute, UAVIndex, UAV);
-}
-
-
-void FMetalRHICommandContext::RHISetShaderTexture(FRHIGraphicsShader* ShaderRHI, uint32 TextureIndex, FRHITexture* NewTextureRHI)
-{
-    MTL_SCOPED_AUTORELEASE_POOL;
+    FMetalShaderBinder(FMetalDeviceContext& InContext, EShaderFrequency ShaderFrequency)
+    : Context(InContext)
+    , StateCache(InContext.GetCurrentState())
+    , Stage(GetMetalShaderFrequency(ShaderFrequency))
+    , ShaderParameters(StateCache.GetShaderParameters(Stage))
+#if PLATFORM_SUPPORTS_BINDLESS_RENDERING
+    , bBindlessResources(IsMetalBindlessEnabled())
+    , bBindlessSamplers(IsMetalBindlessEnabled())
+#endif
+    {
+    }
     
-	FMetalSurface* Surface = GetMetalSurfaceFromRHITexture(NewTextureRHI);
-	EMetalShaderStages Stage = GetShaderStage(ShaderRHI);
-	if (Surface != nullptr)
-	{
-        if (Surface->Texture || !(Surface->GetDesc().Flags & TexCreate_Presentable))
-        {
-            Context->GetCurrentState().SetShaderTexture(Stage, Surface->Texture.get(), TextureIndex, (MTL::ResourceUsage)(MTL::ResourceUsageRead|MTL::ResourceUsageSample));
-        }
-        else
+#if PLATFORM_SUPPORTS_BINDLESS_RENDERING
+    void SetBindlessHandle(const FRHIDescriptorHandle& Handle, uint32 Offset)
+    {
+        if (Handle.IsValid())
         {
-            MTLTexturePtr Tex = Surface->GetCurrentTexture();
-            Context->GetCurrentState().SetShaderTexture(Stage, Tex.get(), TextureIndex, (MTL::ResourceUsage)(MTL::ResourceUsageRead|MTL::ResourceUsageSample));
+            const uint32 BindlessIndex = Handle.GetIndex();
+            ShaderParameters.Set(Offset, 0, 4, &BindlessIndex);
         }
-	}
-	else
-	{
-		Context->GetCurrentState().SetShaderTexture(Stage, nullptr, TextureIndex, MTL::ResourceUsage(0));
-	}
-}
+    }
+#endif
 
-void FMetalRHICommandContext::RHISetShaderTexture(FRHIComputeShader* ComputeShader, uint32 TextureIndex, FRHITexture* NewTextureRHI)
-{
-    MTL_SCOPED_AUTORELEASE_POOL;
-    
-	FMetalSurface* Surface = GetMetalSurfaceFromRHITexture(NewTextureRHI);
-	if (Surface != nullptr)
+    void SetUAV(FRHIUnorderedAccessView* InUnorderedAccessView, uint32 Index, bool bClearResources = false)
     {
-        if (Surface->Texture || !(Surface->GetDesc().Flags & TexCreate_Presentable))
+        FMetalUnorderedAccessView* UAV = ResourceCast(InUnorderedAccessView);
+        StateCache.SetShaderUnorderedAccessView(Stage, Index, UAV);
+    }
+
+    void SetSRV(FRHIShaderResourceView* InShaderResourceView, uint32 Index)
+    {
+        FMetalShaderResourceView* SRV = ResourceCast(InShaderResourceView);
+        StateCache.SetShaderResourceView(Stage, Index, SRV);
+    }
+
+    void SetTexture(FRHITexture* InTexture, uint32 Index)
+    {
+        if (FMetalSurface* Surface = GetMetalSurfaceFromRHITexture(InTexture))
         {
-            Context->GetCurrentState().SetShaderTexture(EMetalShaderStages::Compute, Surface->Texture.get(), TextureIndex, (MTL::ResourceUsage)(MTL::ResourceUsageRead|MTL::ResourceUsageSample));
+            if (Surface->Texture || !EnumHasAnyFlags(Surface->GetDesc().Flags, ETextureCreateFlags::Presentable))
+            {
+                StateCache.SetShaderTexture(Stage, Surface->Texture.get(), Index, (MTL::ResourceUsage)(MTL::ResourceUsageRead|MTL::ResourceUsageSample));
+            }
+            else
+            {
+                MTLTexturePtr Tex = Surface->GetCurrentTexture();
+                StateCache.SetShaderTexture(Stage, Tex.get(), Index, (MTL::ResourceUsage)(MTL::ResourceUsageRead|MTL::ResourceUsageSample));
+            }
         }
         else
         {
-            MTLTexturePtr Tex = Surface->GetCurrentTexture();
-            Context->GetCurrentState().SetShaderTexture(EMetalShaderStages::Compute, Tex.get(), TextureIndex, (MTL::ResourceUsage)(MTL::ResourceUsageRead|MTL::ResourceUsageSample));
+            StateCache.SetShaderTexture(Stage, nullptr, Index, MTL::ResourceUsage(0));
         }
-	}
-	else
-	{
-		Context->GetCurrentState().SetShaderTexture(EMetalShaderStages::Compute, nullptr, TextureIndex, MTL::ResourceUsage(0));
-	}
-}
+    }
 
+    void SetSampler(FRHISamplerState* InSampler, uint32 Index)
+    {
+        FMetalSamplerState* Sampler = ResourceCast(InSampler);
+        StateCache.SetShaderSamplerState(Stage, Sampler, Index);
+    }
+};
 
-void FMetalRHICommandContext::RHISetShaderResourceViewParameter(FRHIGraphicsShader* ShaderRHI, uint32 TextureIndex, FRHIShaderResourceView* SRVRHI)
+static void SetShaderParametersOnContext(
+      FMetalDeviceContext& Context
+    , FRHIShader* Shader
+    , EShaderFrequency ShaderFrequency
+    , TArrayView<const uint8> InParametersData
+    , TArrayView<const FRHIShaderParameter> InParameters
+    , TArrayView<const FRHIShaderParameterResource> InResourceParameters
+    , TArrayView<const FRHIShaderParameterResource> InBindlessParameters)
 {
     MTL_SCOPED_AUTORELEASE_POOL;
-    
-	FMetalShaderResourceView* SRV = ResourceCast(SRVRHI);
-	EMetalShaderStages Stage = GetShaderStage(ShaderRHI);
-	Context->GetCurrentState().SetShaderResourceView(Stage, TextureIndex, SRV);
-}
 
-void FMetalRHICommandContext::RHISetShaderResourceViewParameter(FRHIComputeShader* ComputeShaderRHI, uint32 TextureIndex, FRHIShaderResourceView* SRVRHI)
-{
-    MTL_SCOPED_AUTORELEASE_POOL;
+    FMetalShaderBinder Binder(Context, ShaderFrequency);
     
-	FMetalShaderResourceView* SRV = ResourceCast(SRVRHI);
-	Context->GetCurrentState().SetShaderResourceView(EMetalShaderStages::Compute, TextureIndex, SRV);
-}
+    for (const FRHIShaderParameter& Parameter : InParameters)
+    {
+        Binder.ShaderParameters.Set(Parameter.BufferIndex, Parameter.BaseIndex, Parameter.ByteSize, &InParametersData[Parameter.ByteOffset]);
+    }
 
+#if PLATFORM_SUPPORTS_BINDLESS_RENDERING
+    for (const FRHIShaderParameterResource& Parameter : InBindlessParameters)
+    {
+        const FRHIDescriptorHandle Handle = UE::RHICore::GetBindlessParameterHandle(Parameter);
+        if (Handle.IsValid())
+        {
+            checkf(Handle.IsValid(), TEXT("Metal resource did not provide a valid descriptor handle. Please validate that all Metal types can provide this or that the resource is still valid."));
+            Binder.SetBindlessHandle(Handle, Parameter.Index);
+        }
+    }
+#endif
 
-void FMetalRHICommandContext::RHISetShaderSampler(FRHIGraphicsShader* ShaderRHI, uint32 SamplerIndex, FRHISamplerState* NewStateRHI)
-{
-    MTL_SCOPED_AUTORELEASE_POOL;
-    
-	FMetalSamplerState* NewState = ResourceCast(NewStateRHI);
-	EMetalShaderStages Stage = GetShaderStage(ShaderRHI);
-	Context->GetCurrentState().SetShaderSamplerState(Stage, NewState, SamplerIndex);
-}
+    for (const FRHIShaderParameterResource& Parameter : InResourceParameters)
+    {
+        if (Parameter.Type == FRHIShaderParameterResource::EType::UnorderedAccessView)
+        {
+            if (ShaderFrequency == SF_Pixel || ShaderFrequency == SF_Compute)
+            {
+                Binder.SetUAV(static_cast<FRHIUnorderedAccessView*>(Parameter.Resource), Parameter.Index, true);
+            }
+            else
+            {
+                checkf(false, TEXT("TShaderRHI Can't have compute shader to be set. UAVs are not supported on vertex, tessellation and geometry shaders."));
+            }
+        }
+    }
 
-void FMetalRHICommandContext::RHISetShaderSampler(FRHIComputeShader* ComputeShader, uint32 SamplerIndex, FRHISamplerState* NewStateRHI)
-{
-    MTL_SCOPED_AUTORELEASE_POOL;
+    for (const FRHIShaderParameterResource& Parameter : InResourceParameters)
+    {
+        switch (Parameter.Type)
+        {
+        case FRHIShaderParameterResource::EType::Texture:
+            Binder.SetTexture(static_cast<FRHITexture*>(Parameter.Resource), Parameter.Index);
+            break;
+        case FRHIShaderParameterResource::EType::ResourceView:
+            Binder.SetSRV(static_cast<FRHIShaderResourceView*>(Parameter.Resource), Parameter.Index);
+            break;
+        case FRHIShaderParameterResource::EType::UnorderedAccessView:
+            break;
+        case FRHIShaderParameterResource::EType::Sampler:
+            Binder.SetSampler(static_cast<FRHISamplerState*>(Parameter.Resource), Parameter.Index);
+            break;
+        case FRHIShaderParameterResource::EType::UniformBuffer:
+            BindUniformBuffer(Binder.StateCache, Shader, Binder.Stage, Parameter.Index, static_cast<FRHIUniformBuffer*>(Parameter.Resource));
+            break;
+        default:
+            checkf(false, TEXT("Unhandled resource type?"));
+            break;
+        }
+    }
     
-	FMetalSamplerState* NewState = ResourceCast(NewStateRHI);
-	Context->GetCurrentState().SetShaderSamplerState(EMetalShaderStages::Compute, NewState, SamplerIndex);
-}
-
-void FMetalRHICommandContext::RHISetShaderParameter(FRHIGraphicsShader* ShaderRHI, uint32 BufferIndex, uint32 BaseIndex, uint32 NumBytes, const void* NewValue)
-{
-    MTL_SCOPED_AUTORELEASE_POOL;
     
-	EMetalShaderStages Stage = GetShaderStage(ShaderRHI);
-	Context->GetCurrentState().GetShaderParameters(Stage).Set(BufferIndex, BaseIndex, NumBytes, NewValue);
-}
-
-void FMetalRHICommandContext::RHISetShaderParameter(FRHIComputeShader* ComputeShaderRHI,uint32 BufferIndex, uint32 BaseIndex, uint32 NumBytes, const void* NewValue)
-{
-    MTL_SCOPED_AUTORELEASE_POOL;
-	Context->GetCurrentState().GetShaderParameters(EMetalShaderStages::Compute).Set(BufferIndex, BaseIndex, NumBytes, NewValue);
+#if PLATFORM_SUPPORTS_BINDLESS_RENDERING
+    if(IsMetalBindlessEnabled())
+    {
+        Binder.StateCache.IRForwardBindlessParameters(Binder.Stage, InResourceParameters);
+        Binder.StateCache.IRForwardBindlessParameters(Binder.Stage, InBindlessParameters);
+    }
+#endif
 }
 
 void FMetalRHICommandContext::RHISetShaderParameters(FRHIGraphicsShader* Shader, TConstArrayView<uint8> InParametersData, TConstArrayView<FRHIShaderParameter> InParameters, TConstArrayView<FRHIShaderParameterResource> InResourceParameters, TConstArrayView<FRHIShaderParameterResource> InBindlessParameters)
 {
-	// TODO: CL - Implement metal specific version of RHISetShaderParametersShared
-	UE::RHICore::RHISetShaderParametersShared(
-		*this
-		, Shader
-		, InParametersData
-		, InParameters
-		, InResourceParameters
-		, InBindlessParameters
-	);
-
-#if PLATFORM_SUPPORTS_BINDLESS_RENDERING
-	if(IsMetalBindlessEnabled())
-	{
-		EMetalShaderStages Stage = GetShaderStage(Shader);
-		Context->GetCurrentState().IRForwardBindlessParameters(Stage, InResourceParameters);
-		Context->GetCurrentState().IRForwardBindlessParameters(Stage, InBindlessParameters);
-	}
-#endif
+    const EShaderFrequency ShaderFrequency = Shader->GetFrequency();
+    if (IsValidGraphicsFrequency(ShaderFrequency))
+    {
+        SetShaderParametersOnContext(
+            *Context
+            , Shader
+            , ShaderFrequency
+            , InParametersData
+            , InParameters
+            , InResourceParameters
+            , InBindlessParameters
+        );
+    }
+    else
+    {
+        checkf(0, TEXT("Unsupported FRHIGraphicsShader Type '%s'!"), GetShaderFrequencyString(ShaderFrequency, false));
+    }
 }
 
 void FMetalRHICommandContext::RHISetShaderParameters(FRHIComputeShader* Shader, TConstArrayView<uint8> InParametersData, TConstArrayView<FRHIShaderParameter> InParameters, TConstArrayView<FRHIShaderParameterResource> InResourceParameters, TConstArrayView<FRHIShaderParameterResource> InBindlessParameters)
 {
-	UE::RHICore::RHISetShaderParametersShared(
-		*this
-		, Shader
-		, InParametersData
-		, InParameters
-		, InResourceParameters
-		, InBindlessParameters
-	);
-
-#if PLATFORM_SUPPORTS_BINDLESS_RENDERING
-	if(IsMetalBindlessEnabled())
-	{
-		Context->GetCurrentState().IRForwardBindlessParameters(EMetalShaderStages::Compute, InResourceParameters);
-		Context->GetCurrentState().IRForwardBindlessParameters(EMetalShaderStages::Compute, InBindlessParameters);
-	}
-#endif
+    SetShaderParametersOnContext(
+        *Context
+        , Shader
+        , SF_Compute
+        , InParametersData
+        , InParameters
+        , InResourceParameters
+        , InBindlessParameters
+    );
 }
 
 void FMetalRHICommandContext::RHISetStencilRef(uint32 StencilRef)
diff --git a/Engine/Source/Runtime/Apple/MetalRHI/Private/MetalComputePipelineState.h b/Engine/Source/Runtime/Apple/MetalRHI/Private/MetalComputePipelineState.h
--- a/Engine/Source/Runtime/Apple/MetalRHI/Private/MetalComputePipelineState.h
+++ b/Engine/Source/Runtime/Apple/MetalRHI/Private/MetalComputePipelineState.h
@@ -15,16 +15,10 @@ public:
 		check(InComputeShader);
 	}
 
-	virtual ~FMetalComputePipelineState()
-	{
-		// void
-	}
-
 	FMetalComputeShader* GetComputeShader()
 	{
 		return ComputeShader;
 	}
 
-private:
 	TRefCountPtr<FMetalComputeShader> ComputeShader;
 };
diff --git a/Engine/Source/Runtime/Apple/MetalRHI/Private/MetalUniformBufferCommands.cpp b/Engine/Source/Runtime/Apple/MetalRHI/Private/MetalUniformBufferCommands.cpp
--- a/Engine/Source/Runtime/Apple/MetalRHI/Private/MetalUniformBufferCommands.cpp
+++ b/Engine/Source/Runtime/Apple/MetalRHI/Private/MetalUniformBufferCommands.cpp
@@ -32,72 +32,3 @@ void FMetalDynamicRHI::RHIUpdateUniformBuffer(FRHICommandListBase& RHICmdList, F
 		UniformBuffer->Update(SrcContents);
 	});
 }
-
-template<EMetalShaderStages Stage, typename RHIShaderType>
-static void SetUniformBufferInternal(FMetalContext* Context, RHIShaderType* ShaderRHI, uint32 BufferIndex, FRHIUniformBuffer* UBRHI)
-{
-    MTL_SCOPED_AUTORELEASE_POOL;
-
-    auto Shader = ResourceCast(ShaderRHI);
-    Context->GetCurrentState().BindUniformBuffer(Stage, BufferIndex, UBRHI);
-        
-    auto& Bindings = Shader->Bindings;
-    if((Bindings.ConstantBuffers) & (1 << BufferIndex))
-    {
-        FMetalUniformBuffer* UB = ResourceCast(UBRHI);
-        UB->PrepareToBind();
-#if METAL_USE_METAL_SHADER_CONVERTER
-		if(IsMetalBindlessEnabled())
-		{
-			Context->GetCurrentState().IRBindUniformBuffer(Stage, BufferIndex, UB);
-		}
-		else
-#endif
-		{
-			FMetalBufferPtr Buf = FMetalBufferPtr(new FMetalBuffer(UB->Backing));
-			Context->GetCurrentState().SetShaderBuffer(Stage, Buf, nullptr, UB->Offset, UB->GetSize(), BufferIndex, MTL::ResourceUsageRead);
-		}
-    }
-}
-
-void FMetalRHICommandContext::RHISetShaderUniformBuffer(FRHIGraphicsShader* Shader, uint32 BufferIndex, FRHIUniformBuffer* Buffer)
-{
-	switch (Shader->GetFrequency())
-	{
-		case SF_Vertex:
-			SetUniformBufferInternal<EMetalShaderStages::Vertex, FRHIVertexShader>(Context, static_cast<FRHIVertexShader*>(Shader), BufferIndex, Buffer);
-			break;
-
-		case SF_Geometry:
-#if PLATFORM_SUPPORTS_GEOMETRY_SHADERS
-            SetUniformBufferInternal<EMetalShaderStages::Geometry, FRHIGeometryShader>(Context, static_cast<FRHIGeometryShader*>(Shader), BufferIndex, Buffer);
-#else
-			NOT_SUPPORTED("RHISetShaderUniformBuffer-Geometry");
-#endif
-			break;
-
-		case SF_Pixel:
-			SetUniformBufferInternal<EMetalShaderStages::Pixel, FRHIPixelShader>(Context, static_cast<FRHIPixelShader*>(Shader), BufferIndex, Buffer);
-			break;
-
-#if PLATFORM_SUPPORTS_MESH_SHADERS
-        case SF_Mesh:
-            SetUniformBufferInternal<EMetalShaderStages::Mesh, FRHIMeshShader>(Context, static_cast<FRHIMeshShader*>(Shader), BufferIndex, Buffer);
-            break;
-
-        case SF_Amplification:
-            SetUniformBufferInternal<EMetalShaderStages::Amplification, FRHIAmplificationShader>(Context, static_cast<FRHIAmplificationShader*>(Shader), BufferIndex, Buffer);
-            break;
-#endif
-
-		default:
-			checkf(0, TEXT("FRHIShader Type %d is invalid or unsupported!"), (int32)Shader->GetFrequency());
-			NOT_SUPPORTED("RHIShaderStage");
-			break;
-	}
-}
-
-void FMetalRHICommandContext::RHISetShaderUniformBuffer(FRHIComputeShader* ComputeShaderRHI, uint32 BufferIndex, FRHIUniformBuffer* BufferRHI)
-{
-	SetUniformBufferInternal<EMetalShaderStages::Compute, FRHIComputeShader>(Context, ComputeShaderRHI, BufferIndex, BufferRHI);
-}
diff --git a/Engine/Source/Runtime/Apple/MetalRHI/Private/Shaders/Types/Templates/MetalBaseShader.h b/Engine/Source/Runtime/Apple/MetalRHI/Private/Shaders/Types/Templates/MetalBaseShader.h
--- a/Engine/Source/Runtime/Apple/MetalRHI/Private/Shaders/Types/Templates/MetalBaseShader.h
+++ b/Engine/Source/Runtime/Apple/MetalRHI/Private/Shaders/Types/Templates/MetalBaseShader.h
@@ -41,8 +41,57 @@ extern MTL::LanguageVersion ValidateVersion(uint32 Version);
 
 #pragma mark - Metal RHI Base Shader Class Template
 
+struct FMetalShaderData
+{
+    /** External bindings for this shader. */
+    FMetalShaderBindings Bindings;
+    
+    // List of memory copies from RHIUniformBuffer to packed uniforms
+    TArray<CrossCompiler::FUniformBufferCopyInfo> UniformBuffersCopyInfo;
+
+    /* Argument encoders for shader IABs */
+    TMap<uint32, MTL::ArgumentEncoder*> ArgumentEncoders;
+
+    /* Tier1 Argument buffer bitmasks */
+    TMap<uint32, TBitArray<>> ArgumentBitmasks;
+
+    /* Uniform buffer static slots */
+    TArray<FUniformBufferStaticSlot> StaticSlots;
+
+    /** The binding for the buffer side-table if present */
+    int32 SideTableBinding = -1;
+
+    /** CRC & Len for name disambiguation */
+    uint32 SourceLen = 0;
+    uint32 SourceCRC = 0;
+
+    /** Hash for the shader/material permutation constants */
+    uint32 ConstantValueHash = 0;
+    
+    // this is the compiler shader
+    MTLFunctionPtr Function;
+    // This is the MTLLibrary for the shader so we can dynamically refine the MTLFunction
+    MTLLibraryPtr Library;
+
+    /** The debuggable text source */
+    NS::String* GlslCodeNSString = nullptr;
+
+    /** The compressed text source */
+    TArray<uint8> CompressedSource;
+
+    /** The uncompressed text source size */
+    uint32 CodeSize = 0;
+
+    // Function constant states
+    bool bHasFunctionConstants = false;
+    bool bDeviceFunctionConstants = false;
+
+    /** Index of the function (in the library) pointing to the function requested by the user (when GetCompiledFunction() is called with an explicit index). */
+    uint32 LibraryFunctionIndex = -1;
+};
+
 template<typename BaseResourceType, int32 ShaderType>
-class TMetalBaseShader : public BaseResourceType, public IRefCountedObject
+class TMetalBaseShader : public BaseResourceType, public FMetalShaderData
 {
 public:
 	enum
@@ -70,58 +119,8 @@ public:
 	 */
 	NS::String* GetSourceCode();
 
-	// IRefCountedObject interface:
-	virtual uint32 AddRef() const override final;
-	virtual uint32 Release() const override final;
-	virtual uint32 GetRefCount() const override final;
-
-	/** External bindings for this shader. */
-	FMetalShaderBindings Bindings;
-
-	// List of memory copies from RHIUniformBuffer to packed uniforms
-	TArray<CrossCompiler::FUniformBufferCopyInfo> UniformBuffersCopyInfo;
-
-	/* Argument encoders for shader IABs */
-	TMap<uint32, MTL::ArgumentEncoder*> ArgumentEncoders;
-
-	/* Tier1 Argument buffer bitmasks */
-	TMap<uint32, TBitArray<>> ArgumentBitmasks;
-
-	/** The binding for the buffer side-table if present */
-	int32 SideTableBinding = -1;
-
-	/** CRC & Len for name disambiguation */
-	uint32 SourceLen = 0;
-	uint32 SourceCRC = 0;
-
-	/** Hash for the shader/material permutation constants */
-	uint32 ConstantValueHash = 0;
-
 protected:
 	MTLFunctionPtr GetCompiledFunction(bool const bAsync = false, const int32 FunctionIndex = -1);
-
-	// this is the compiler shader
-	MTLFunctionPtr Function;
-
-private:
-	// This is the MTLLibrary for the shader so we can dynamically refine the MTLFunction
-	MTLLibraryPtr Library;
-
-	/** The debuggable text source */
-	NS::String* GlslCodeNSString = nullptr;
-
-	/** The compressed text source */
-	TArray<uint8> CompressedSource;
-
-	/** The uncompressed text source size */
-	uint32 CodeSize = 0;
-
-	// Function constant states
-	bool bHasFunctionConstants = false;
-	bool bDeviceFunctionConstants = false;
-
-    /** Index of the function (in the library) pointing to the function requested by the user (when GetCompiledFunction() is called with an explicit index). */
-    uint32 LibraryFunctionIndex = -1;
 };
 
 
@@ -410,24 +409,6 @@ inline NS::String* TMetalBaseShader<BaseResourceType, ShaderType>::GetSourceCode
 	return GlslCodeNSString;
 }
 
-template<typename BaseResourceType, int32 ShaderType>
-uint32 TMetalBaseShader<BaseResourceType, ShaderType>::AddRef() const
-{
-	return FRHIResource::AddRef();
-}
-
-template<typename BaseResourceType, int32 ShaderType>
-uint32 TMetalBaseShader<BaseResourceType, ShaderType>::Release() const
-{
-	return FRHIResource::Release();
-}
-
-template<typename BaseResourceType, int32 ShaderType>
-uint32 TMetalBaseShader<BaseResourceType, ShaderType>::GetRefCount() const
-{
-	return FRHIResource::GetRefCount();
-}
-
 template<typename BaseResourceType, int32 ShaderType>
 MTLFunctionPtr TMetalBaseShader<BaseResourceType, ShaderType>::GetCompiledFunction(bool const bAsync, const int32 FunctionIndex)
 {
diff --git a/Engine/Source/Runtime/Apple/MetalRHI/Public/MetalRHIContext.h b/Engine/Source/Runtime/Apple/MetalRHI/Public/MetalRHIContext.h
--- a/Engine/Source/Runtime/Apple/MetalRHI/Public/MetalRHIContext.h
+++ b/Engine/Source/Runtime/Apple/MetalRHI/Public/MetalRHIContext.h
@@ -108,20 +108,6 @@ public:
 	
 	virtual void RHISetStaticUniformBuffers(const FUniformBufferStaticBindings& InUniformBuffers) final override;
 
-	void RHISetShaderTexture(FRHIGraphicsShader* Shader, uint32 TextureIndex, FRHITexture* NewTexture);
-	void RHISetShaderTexture(FRHIComputeShader* PixelShader, uint32 TextureIndex, FRHITexture* NewTexture);
-	void RHISetShaderSampler(FRHIComputeShader* ComputeShader, uint32 SamplerIndex, FRHISamplerState* NewState);
-	void RHISetShaderSampler(FRHIGraphicsShader* Shader, uint32 SamplerIndex, FRHISamplerState* NewState);
-	void RHISetUAVParameter(FRHIPixelShader* PixelShaderRHI, uint32 UAVIndex, FRHIUnorderedAccessView* UAVRHI);
-	void RHISetUAVParameter(FRHIComputeShader* ComputeShader, uint32 UAVIndex, FRHIUnorderedAccessView* UAV);
-	void RHISetUAVParameter(FRHIComputeShader* ComputeShader, uint32 UAVIndex, FRHIUnorderedAccessView* UAV, uint32 InitialCount);
-	void RHISetShaderResourceViewParameter(FRHIGraphicsShader* Shader, uint32 SamplerIndex, FRHIShaderResourceView* SRV);
-	void RHISetShaderResourceViewParameter(FRHIComputeShader* ComputeShader, uint32 SamplerIndex, FRHIShaderResourceView* SRV);
-	void RHISetShaderUniformBuffer(FRHIGraphicsShader* Shader, uint32 BufferIndex, FRHIUniformBuffer* Buffer);
-	void RHISetShaderUniformBuffer(FRHIComputeShader* ComputeShader, uint32 BufferIndex, FRHIUniformBuffer* Buffer);
-	void RHISetShaderParameter(FRHIGraphicsShader* Shader, uint32 BufferIndex, uint32 BaseIndex, uint32 NumBytes, const void* NewValue);
-	void RHISetShaderParameter(FRHIComputeShader* ComputeShader, uint32 BufferIndex, uint32 BaseIndex, uint32 NumBytes, const void* NewValue);
-	
 	virtual void RHISetShaderParameters(FRHIGraphicsShader* Shader, TConstArrayView<uint8> InParametersData, TConstArrayView<FRHIShaderParameter> InParameters, TConstArrayView<FRHIShaderParameterResource> InResourceParameters, TConstArrayView<FRHIShaderParameterResource> InBindlessParameters) final override;
 	virtual void RHISetShaderParameters(FRHIComputeShader* Shader, TConstArrayView<uint8> InParametersData, TConstArrayView<FRHIShaderParameter> InParameters, TConstArrayView<FRHIShaderParameterResource> InResourceParameters, TConstArrayView<FRHIShaderParameterResource> InBindlessParameters) final override;
 
@@ -204,6 +190,11 @@ public:
 	{
 		// Nothing to do
 	}
+    
+    inline const TArray<FRHIUniformBuffer*>& GetStaticUniformBuffers() const
+    {
+        return GlobalUniformBuffers;
+    }
 protected:
 	static TGlobalResource<TBoundShaderStateHistory<10000>> BoundShaderStateHistory;
 	
@@ -226,9 +217,6 @@ protected:
 	uint32 PendingPrimitiveType = 0;
 	uint32 PendingNumPrimitives = 0;
 
-	template <typename TRHIShader>
-	void ApplyStaticUniformBuffers(TRHIShader* Shader);
-
 	void ResolveTexture(UE::RHICore::FResolveTextureInfo Info);
 
 	TArray<FRHIUniformBuffer*> GlobalUniformBuffers;
