Fix crash with hair when using PF_D24

diff --git a/Engine/Config/Mac/BaseMacEngine.ini b/Engine/Config/Mac/BaseMacEngine.ini
--- a/Engine/Config/Mac/BaseMacEngine.ini
+++ b/Engine/Config/Mac/BaseMacEngine.ini
@@ -35,6 +35,9 @@ fx.NiagaraAllowRuntimeScalabilityChanges=1
 ;Perplatform to PerQualityLevel conversion mapping for platform groups
 QualityLevelMapping="high"
 
+[ConsoleVariables]
+r.HairStrands.Visibility.MSAA.SamplePerPixel=1
+
 [SF_METAL_SM6]
 ;BindlessResources=Enabled
 ;BindlessSamplers=Enabled
diff --git a/Engine/Source/Runtime/Apple/MetalRHI/Private/MetalPipeline.cpp b/Engine/Source/Runtime/Apple/MetalRHI/Private/MetalPipeline.cpp
--- a/Engine/Source/Runtime/Apple/MetalRHI/Private/MetalPipeline.cpp
+++ b/Engine/Source/Runtime/Apple/MetalRHI/Private/MetalPipeline.cpp
@@ -700,10 +700,11 @@ static bool ConfigureRenderPipelineDescriptor(TDescriptorType* RenderPipelineDes
 			break;
 		}
 		case PF_ShadowDepth:
-		{
-			RenderPipelineDesc->setDepthAttachmentPixelFormat((MTL::PixelFormat)GPixelFormats[PF_ShadowDepth].PlatformFormat);
-			break;
-		}
+        case PF_D24:
+        {
+            RenderPipelineDesc->setDepthAttachmentPixelFormat((MTL::PixelFormat)GPixelFormats[Init.DepthStencilTargetFormat].PlatformFormat);
+            break;
+        }
 		default:
 		{
 			break;
