Fixed bug on systems with Integrated and Discrete GPU's, OpenXR was picking the incorrect GPU.
Fixed crash on Intel integrated GPU's

diff --git a/Engine/Plugins/Runtime/OpenXR/Source/OpenXRHMD/Private/OpenXRHMDModule.cpp b/Engine/Plugins/Runtime/OpenXR/Source/OpenXRHMD/Private/OpenXRHMDModule.cpp
--- a/Engine/Plugins/Runtime/OpenXR/Source/OpenXRHMD/Private/OpenXRHMDModule.cpp
+++ b/Engine/Plugins/Runtime/OpenXR/Source/OpenXRHMD/Private/OpenXRHMDModule.cpp
@@ -149,16 +149,23 @@ void FOpenXRHMDModule::ShutdownModule()
 
 uint64 FOpenXRHMDModule::GetGraphicsAdapterLuid()
 {
+    uint64 DefaultValue = 0;
+    
+    // Mac platforms expect the device ID to be returned here
+#if PLATFORM_MAC
+    DefaultValue = (uint64)-1;
+#endif
+    
 	if (FParse::Param(FCommandLine::Get(), TEXT("xrtrackingonly")))
 	{
-		return 0;
+        return DefaultValue;
 	}
 
 	if (!RenderBridge)
 	{
 		if (!InitRenderBridge())
 		{
-			return 0;
+            return DefaultValue;
 		}
 	}
 
@@ -166,7 +173,7 @@ uint64 FOpenXRHMDModule::GetGraphicsAdapterLuid()
 	XrSystemId System = GetSystemId();
 	if (!System)
 	{
-		int64 AdapterLuid = 0;
+		int64 AdapterLuid = (int64)DefaultValue;
 		EngineIni->GetInt64(TEXT("OpenXR.Settings"), TEXT("GraphicsAdapter"), AdapterLuid);
 		return reinterpret_cast<uint64&>(AdapterLuid);
 	}
diff --git a/Engine/Source/Runtime/Apple/MetalRHI/Private/MetalCommandQueue.cpp b/Engine/Source/Runtime/Apple/MetalRHI/Private/MetalCommandQueue.cpp
--- a/Engine/Source/Runtime/Apple/MetalRHI/Private/MetalCommandQueue.cpp
+++ b/Engine/Source/Runtime/Apple/MetalRHI/Private/MetalCommandQueue.cpp
@@ -169,10 +169,10 @@ FMetalCommandQueue::FMetalCommandQueue(MTL::Device* InDevice, uint32 const MaxNu
 	Features = EMetalFeaturesCountingQueries | EMetalFeaturesBaseVertexInstance | EMetalFeaturesIndirectBuffer |
                 EMetalFeaturesLayeredRendering | EMetalFeaturesCubemapArrays | EMetalFeaturesSetBufferOffset;
 
+    FString DeviceName(Device->name()->cString(NS::UTF8StringEncoding));
+    
 	if (Device->supportsFeatureSet(MTL::FeatureSet_macOS_GPUFamily1_v2))
 	{
-		FString DeviceName(Device->name()->cString(NS::UTF8StringEncoding));
-
 		Features |= EMetalFeaturesMSAADepthResolve | EMetalFeaturesMSAAStoreAndResolve;
         
 		// Assume that set*Bytes only works on macOS Sierra and above as no-one has tested it anywhere else.
@@ -229,7 +229,9 @@ FMetalCommandQueue::FMetalCommandQueue(MTL::Device* InDevice, uint32 const MaxNu
 	}
     
 	// Temporarily only support heaps for devices with unified memory
-	if (Device->hasUnifiedMemory() && !FParse::Param(FCommandLine::Get(),TEXT("nometalheap")))
+	if (!DeviceName.Contains(TEXT("Intel")) &&
+          Device->hasUnifiedMemory() &&
+          !FParse::Param(FCommandLine::Get(),TEXT("nometalheap")))
 	{
 		Features |= EMetalFeaturesHeaps;
 	}
