Metal RHI: Fix for Present being called before committing the the command buffer. Present requires work to be submitted before the call to track dependencies.

diff --git a/Engine/Source/Runtime/Apple/MetalRHI/Private/MetalViewport.cpp b/Engine/Source/Runtime/Apple/MetalRHI/Private/MetalViewport.cpp
--- a/Engine/Source/Runtime/Apple/MetalRHI/Private/MetalViewport.cpp
+++ b/Engine/Source/Runtime/Apple/MetalRHI/Private/MetalViewport.cpp
@@ -551,18 +551,23 @@ void FMetalViewport::Present(FMetalCommandQueue& CommandQueue, bool bLockToVsync
 #else // PLATFORM_MAC
 						CurrentCommandBuffer->GetMTLCmdBuffer()->addCompletedHandler(CommandBufferHandler);
 
-						if (MinPresentDuration && GEnablePresentPacing)
-						{
-                            LocalDrawable->presentAfterMinimumDuration(1.0f/(float)FramePace);
-						}
-						else
-						{
-                            LocalDrawable->present();
-						}
+                        {
+                            LocalDrawable->retain();
+                            
+                            // Queue this on the current command buffer to ensure that all work is committed prior to the present, present only knows about dependencies on committed work.
+                            if (MinPresentDuration && GEnablePresentPacing)
+                            {
+                                CurrentCommandBuffer->GetMTLCmdBuffer()->presentDrawableAfterMinimumDuration(LocalDrawable, 1.0f/(float)FramePace);
+                            }
+                            else
+                            {
+                                CurrentCommandBuffer->GetMTLCmdBuffer()->presentDrawable(LocalDrawable);
+                            }
+                        }
 #endif // PLATFORM_MAC
-
-						METAL_GPUPROFILE(Stats->End(CurrentCommandBuffer->GetMTLCmdBuffer()));
-						CommandQueue.CommitCommandBuffer(CurrentCommandBuffer);
+                        METAL_GPUPROFILE(Stats->End(CurrentCommandBuffer->GetMTLCmdBuffer()));
+                        CommandQueue.CommitCommandBuffer(CurrentCommandBuffer);
+						
 					}
 				}
 			}
