Fix bug in bindless with incorrect SetBindlessHandle function

diff --git a/Engine/Source/Runtime/Apple/MetalRHI/Private/MetalCommands.cpp b/Engine/Source/Runtime/Apple/MetalRHI/Private/MetalCommands.cpp
--- a/Engine/Source/Runtime/Apple/MetalRHI/Private/MetalCommands.cpp
+++ b/Engine/Source/Runtime/Apple/MetalRHI/Private/MetalCommands.cpp
@@ -359,7 +359,7 @@ struct FMetalShaderBinder
         if (Handle.IsValid())
         {
             const uint32 BindlessIndex = Handle.GetIndex();
-            ShaderParameters.Set(Offset, 0, 4, &BindlessIndex);
+            Context.GetCurrentState().GetShaderParameters(Stage).Set(0, Offset, 4, &BindlessIndex);
         }
     }
 #endif
