MetalRHI:

Fixed memory leak when deallocating textures created on the heap
Missing texture memory on MetalRHI

diff --git a/Engine/Source/Runtime/Apple/MetalRHI/Private/MetalBuffer.cpp b/Engine/Source/Runtime/Apple/MetalRHI/Private/MetalBuffer.cpp
--- a/Engine/Source/Runtime/Apple/MetalRHI/Private/MetalBuffer.cpp
+++ b/Engine/Source/Runtime/Apple/MetalRHI/Private/MetalBuffer.cpp
@@ -2041,7 +2041,7 @@ void FMetalResourceHeap::ReleaseTexture(FMetalSurface* Surface, MTLTexturePtr Te
 {
 	FScopeLock ScopeLock(&InUseResourcesCS);
 	
-	if (Texture && !Texture->buffer() && !Texture->parentTexture() && !Texture->heap())
+	if (Texture && !Texture->buffer() && !Texture->parentTexture())
 	{
 		auto It = AllocationHandlesLUT.Find(Texture.get());
 		if (It && *It != UINT32_MAX)
diff --git a/Engine/Source/Runtime/Apple/MetalRHI/Private/MetalTexture.cpp b/Engine/Source/Runtime/Apple/MetalRHI/Private/MetalTexture.cpp
--- a/Engine/Source/Runtime/Apple/MetalRHI/Private/MetalTexture.cpp
+++ b/Engine/Source/Runtime/Apple/MetalRHI/Private/MetalTexture.cpp
@@ -1387,7 +1387,7 @@ uint32 FMetalSurface::GetMemorySize()
 		return TotalTextureSize;
 	}
 	
-	if (Texture)
+	if (!Texture)
 	{
 		return 0;
 	}
