MetalRHI: Fix validation error for freeing pointer while it's potentially in use

diff --git a/Engine/Source/Runtime/Apple/MetalRHI/Private/MetalUAV.cpp b/Engine/Source/Runtime/Apple/MetalRHI/Private/MetalUAV.cpp
--- a/Engine/Source/Runtime/Apple/MetalRHI/Private/MetalUAV.cpp
+++ b/Engine/Source/Runtime/Apple/MetalRHI/Private/MetalUAV.cpp
@@ -44,14 +44,8 @@ void FMetalResourceViewBase::Invalidate()
                 
         case EMetalType::TextureBufferBacked:
             FTextureBufferBacked & View = Storage.Get<FTextureBufferBacked>();
-            if (View.bIsBuffer)
-            {
-                SafeReleaseMetalTexture(View.Texture);
-            }
-            else
-            {
-                SafeReleaseMetalBuffer(View.Buffer);
-            }
+            SafeReleaseMetalTexture(View.Texture);
+            SafeReleaseMetalBuffer(View.Buffer);
             break;
 		}
 	}
