MetalRHI: Removed unrequired functions from Immediate Context

diff --git a/Engine/Source/Runtime/Apple/MetalRHI/Private/MetalRHI.cpp b/Engine/Source/Runtime/Apple/MetalRHI/Private/MetalRHI.cpp
--- a/Engine/Source/Runtime/Apple/MetalRHI/Private/MetalRHI.cpp
+++ b/Engine/Source/Runtime/Apple/MetalRHI/Private/MetalRHI.cpp
@@ -1205,56 +1205,36 @@ void FMetalDynamicRHI::Init()
 	GIsRHIInitialized = true;
 }
 
-void FMetalRHIImmediateCommandContext::RHIBeginFrame()
+void FMetalRHICommandContext::RHIBeginFrame()
 {
     MTL_SCOPED_AUTORELEASE_POOL;
 #if ENABLE_METAL_GPUPROFILE
-	Profiler->BeginFrame();
+    Profiler->BeginFrame();
 #endif
-	((FMetalDeviceContext*)Context)->BeginFrame();
-}
-
-void FMetalRHICommandContext::RHIBeginFrame()
-{
-	check(false);
+    ((FMetalDeviceContext*)Context)->BeginFrame();
 }
 
-void FMetalRHIImmediateCommandContext::RHIEndFrame()
+void FMetalRHICommandContext::RHIEndFrame()
 {
     MTL_SCOPED_AUTORELEASE_POOL;
 #if ENABLE_METAL_GPUPROFILE
-	Profiler->EndFrame();
+    Profiler->EndFrame();
 #endif
-	((FMetalDeviceContext*)Context)->EndFrame();
-}
-
-void FMetalRHICommandContext::RHIEndFrame()
-{
-	check(false);
+    ((FMetalDeviceContext*)Context)->EndFrame();
 }
 
-void FMetalRHIImmediateCommandContext::RHIBeginScene()
+void FMetalRHICommandContext::RHIBeginScene()
 {
     MTL_SCOPED_AUTORELEASE_POOL;
 	((FMetalDeviceContext*)Context)->BeginScene();
 }
 
-void FMetalRHICommandContext::RHIBeginScene()
-{
-	check(false);
-}
-
-void FMetalRHIImmediateCommandContext::RHIEndScene()
+void FMetalRHICommandContext::RHIEndScene()
 {
     MTL_SCOPED_AUTORELEASE_POOL;
 	((FMetalDeviceContext*)Context)->EndScene();
 }
 
-void FMetalRHICommandContext::RHIEndScene()
-{
-	check(false);
-}
-
 #if WITH_RHI_BREADCRUMBS
 	void FMetalRHICommandContext::RHIBeginBreadcrumbGPU(FRHIBreadcrumbNode* Breadcrumb)
 	{
diff --git a/Engine/Source/Runtime/Apple/MetalRHI/Private/MetalViewport.cpp b/Engine/Source/Runtime/Apple/MetalRHI/Private/MetalViewport.cpp
--- a/Engine/Source/Runtime/Apple/MetalRHI/Private/MetalViewport.cpp
+++ b/Engine/Source/Runtime/Apple/MetalRHI/Private/MetalViewport.cpp
@@ -777,11 +777,6 @@ void FMetalDynamicRHI::RHITick( float DeltaTime )
  *=============================================================================*/
 
 void FMetalRHICommandContext::RHIBeginDrawingViewport(FRHIViewport* ViewportRHI, FRHITexture* RenderTargetRHI)
-{
-	check(false);
-}
-
-void FMetalRHIImmediateCommandContext::RHIBeginDrawingViewport(FRHIViewport* ViewportRHI, FRHITexture* RenderTargetRHI)
 {
     MTL_SCOPED_AUTORELEASE_POOL;
 	FMetalViewport* Viewport = ResourceCast(ViewportRHI);
@@ -803,11 +798,6 @@ void FMetalRHIImmediateCommandContext::RHIBeginDrawingViewport(FRHIViewport* Vie
 }
 
 void FMetalRHICommandContext::RHIEndDrawingViewport(FRHIViewport* ViewportRHI,bool bPresent,bool bLockToVsync)
-{
-	check(false);
-}
-
-void FMetalRHIImmediateCommandContext::RHIEndDrawingViewport(FRHIViewport* ViewportRHI,bool bPresent,bool bLockToVsync)
 {
     MTL_SCOPED_AUTORELEASE_POOL;
     
diff --git a/Engine/Source/Runtime/Apple/MetalRHI/Public/MetalRHIContext.h b/Engine/Source/Runtime/Apple/MetalRHI/Public/MetalRHIContext.h
--- a/Engine/Source/Runtime/Apple/MetalRHI/Public/MetalRHIContext.h
+++ b/Engine/Source/Runtime/Apple/MetalRHI/Public/MetalRHIContext.h
@@ -155,7 +155,6 @@ public:
 	virtual void RHIEndTransitions(TArrayView<const FRHITransition*> Transitions);
 
 	virtual void RHIBeginRenderPass(const FRHIRenderPassInfo& InInfo, const TCHAR* InName) final override;
-
 	virtual void RHIEndRenderPass() final override;
 	
 	virtual void RHINextSubpass() final override;
@@ -229,14 +228,6 @@ class FMetalRHIImmediateCommandContext : public FMetalRHICommandContext
 {
 public:
 	FMetalRHIImmediateCommandContext(class FMetalProfiler* InProfiler, FMetalDeviceContext* WrapContext);
-
-	// FRHICommandContext API accessible only on the immediate device context
-	virtual void RHIBeginDrawingViewport(FRHIViewport* Viewport, FRHITexture* RenderTargetRHI) final override;
-	virtual void RHIEndDrawingViewport(FRHIViewport* Viewport, bool bPresent, bool bLockToVsync) final override;
-	virtual void RHIBeginFrame() final override;
-	virtual void RHIEndFrame() final override;
-	virtual void RHIBeginScene() final override;
-	virtual void RHIEndScene() final override;
 	
 protected:
 	friend class FMetalDynamicRHI;
