Fixed: QAGame crash on IOS devices
Metal buffer lock command was trying to end and restart current render-pass which is not possible in some cases (memoryless targets)

diff --git a/Engine/Source/Runtime/Apple/MetalRHI/Private/MetalContext.cpp b/Engine/Source/Runtime/Apple/MetalRHI/Private/MetalContext.cpp
--- a/Engine/Source/Runtime/Apple/MetalRHI/Private/MetalContext.cpp
+++ b/Engine/Source/Runtime/Apple/MetalRHI/Private/MetalContext.cpp
@@ -1080,7 +1080,13 @@ void FMetalContext::TransitionResource(FRHITexture* InResource)
 
 void FMetalContext::SubmitCommandsHint(uint32 const Flags)
 {
-	// When the command-buffer is submitted for a reason other than a break of a logical command-buffer (where one high-level command-sequence becomes more than one command-buffer).
+    if (RenderPass.IsWithinRenderPass() && !StateCache.CanRestartRenderPass())
+    {
+		// Make sure we don't try to end render-passes that can't be restarted (eg. render-passes with a memoryless targets)
+		return;
+    }
+    
+    // When the command-buffer is submitted for a reason other than a break of a logical command-buffer (where one high-level command-sequence becomes more than one command-buffer).
 	if (!(Flags & EMetalSubmitFlagsBreakCommandBuffer))
 	{
 		// Release the current query buffer if there are outstanding writes so that it isn't transitioned by a future encoder that will cause a resource access conflict and lifetime error.
diff --git a/Engine/Source/Runtime/Apple/MetalRHI/Private/MetalRenderPass.h b/Engine/Source/Runtime/Apple/MetalRHI/Private/MetalRenderPass.h
--- a/Engine/Source/Runtime/Apple/MetalRHI/Private/MetalRenderPass.h
+++ b/Engine/Source/Runtime/Apple/MetalRHI/Private/MetalRenderPass.h
@@ -150,6 +150,8 @@ public:
 	
 	void InsertTextureBarrier();
 
+	inline bool IsWithinRenderPass() const { return bWithinRenderPass; }
+
 #if METAL_RHI_RAYTRACING
 	// TODO: Crappy workaround for inline raytracing support.
 	inline void SetRayTracingInstanceBufferSRV(TRefCountPtr<FMetalShaderResourceView>& SRV) { InstanceBufferSRV = SRV; }
diff --git a/Engine/Source/Runtime/Engine/Private/SceneTexturesConfig.cpp b/Engine/Source/Runtime/Engine/Private/SceneTexturesConfig.cpp
--- a/Engine/Source/Runtime/Engine/Private/SceneTexturesConfig.cpp
+++ b/Engine/Source/Runtime/Engine/Private/SceneTexturesConfig.cpp
@@ -225,6 +225,10 @@ static void SetupMobileGBufferFlags(FGBufferBindings GBufferBindings[GBL_Num], b
 		// memoryless GBuffer cant be used with compute
 		RemoveFlags |= TexCreate_UAV;
 	}
+	else
+	{
+		RemoveFlags |= TexCreate_Memoryless;
+	}
 
 	for (uint32 Layout = 0; Layout < GBL_Num; ++Layout)
 	{
