Fix for Mac when submitting commands, ignore the CanRestartRenderPass in SubmitCommandsHint

diff --git a/Engine/Source/Runtime/Apple/MetalRHI/Private/MetalContext.cpp b/Engine/Source/Runtime/Apple/MetalRHI/Private/MetalContext.cpp
--- a/Engine/Source/Runtime/Apple/MetalRHI/Private/MetalContext.cpp
+++ b/Engine/Source/Runtime/Apple/MetalRHI/Private/MetalContext.cpp
@@ -1080,11 +1080,13 @@ void FMetalContext::TransitionResource(FRHITexture* InResource)
 
 void FMetalContext::SubmitCommandsHint(uint32 const Flags)
 {
+#if !PLATFORM_MAC
     if (RenderPass.IsWithinRenderPass() && !StateCache.CanRestartRenderPass())
     {
 		// Make sure we don't try to end render-passes that can't be restarted (eg. render-passes with a memoryless targets)
 		return;
     }
+#endif
     
     // When the command-buffer is submitted for a reason other than a break of a logical command-buffer (where one high-level command-sequence becomes more than one command-buffer).
 	if (!(Flags & EMetalSubmitFlagsBreakCommandBuffer))
