MetalRHI:

Fix crash when accessing a Texture2DArray as an SRV using AtomicCompatible

diff --git a/Engine/Source/Runtime/Apple/MetalRHI/Private/MetalUAV.cpp b/Engine/Source/Runtime/Apple/MetalRHI/Private/MetalUAV.cpp
--- a/Engine/Source/Runtime/Apple/MetalRHI/Private/MetalUAV.cpp
+++ b/Engine/Source/Runtime/Apple/MetalRHI/Private/MetalUAV.cpp
@@ -316,6 +316,19 @@ void FMetalShaderResourceView::UpdateView()
 			{
 				ModifyTextureTypeForBindless(TextureType);
 			}
+			else
+			{
+				// We don't support Texture2DArray with atomic compatible so
+				// ensure we are creating a view on a Texture2D with the correct size
+				bool bIsAtomicCompatible = EnumHasAllFlags(Texture->GetDesc().Flags, TexCreate_AtomicCompatible)  ||
+											EnumHasAllFlags(Texture->GetDesc().Flags, ETextureCreateFlags::Atomic64Compatible);
+				
+				if(TextureType == MTL::TextureType2D && bIsAtomicCompatible)
+				{
+					ArrayStart = 0;
+					ArraySize = 1;
+				}
+			}
             
             MTLTexturePtr View = NS::TransferPtr(Texture->Texture->newTextureView(
 				MetalFormat,
