Metal RHI

Fix validation error in RenderPipeline when binding with dummy depth/stencil

diff --git a/Engine/Source/Runtime/Apple/MetalRHI/Private/MetalPipeline.cpp b/Engine/Source/Runtime/Apple/MetalRHI/Private/MetalPipeline.cpp
--- a/Engine/Source/Runtime/Apple/MetalRHI/Private/MetalPipeline.cpp
+++ b/Engine/Source/Runtime/Apple/MetalRHI/Private/MetalPipeline.cpp
@@ -678,9 +678,10 @@ static bool ConfigureRenderPipelineDescriptor(TDescriptorType* RenderPipelineDes
 	
 	switch(Init.DepthStencilTargetFormat)
 	{
+		case PF_X24_G8:
 		case PF_DepthStencil:
 		{
-			MTL::PixelFormat MetalFormat = (MTL::PixelFormat)GPixelFormats[PF_DepthStencil].PlatformFormat;
+			MTL::PixelFormat MetalFormat = (MTL::PixelFormat)GPixelFormats[Init.DepthStencilTargetFormat].PlatformFormat;
 			if(MetalFormat == MTL::PixelFormatDepth32Float)
 			{
 				if (Init.DepthTargetLoadAction != ERenderTargetLoadAction::ENoAction || Init.DepthTargetStoreAction != ERenderTargetStoreAction::ENoAction)
@@ -705,6 +706,19 @@ static bool ConfigureRenderPipelineDescriptor(TDescriptorType* RenderPipelineDes
             RenderPipelineDesc->setDepthAttachmentPixelFormat((MTL::PixelFormat)GPixelFormats[Init.DepthStencilTargetFormat].PlatformFormat);
             break;
         }
+		case PF_Unknown:
+		{
+			if(!Init.RenderTargetsEnabled)
+			{
+				// This hack is added due to this comment
+				// @todo Improve the way we handle binding a dummy depth/stencil so we can get pure UAV raster operations...
+				// Need to remove this ASAP.
+				MTL::PixelFormat MetalFormat = (MTL::PixelFormat)GPixelFormats[PF_DepthStencil].PlatformFormat;
+				RenderPipelineDesc->setDepthAttachmentPixelFormat(MetalFormat);
+				RenderPipelineDesc->setStencilAttachmentPixelFormat(MetalFormat);
+			}
+			break;
+		}
 		default:
 		{
 			break;
