Enable PSO pre-caching support for Metal RHI

diff --git a/Engine/Source/Runtime/Apple/MetalRHI/Private/MetalRHI.cpp b/Engine/Source/Runtime/Apple/MetalRHI/Private/MetalRHI.cpp
--- a/Engine/Source/Runtime/Apple/MetalRHI/Private/MetalRHI.cpp
+++ b/Engine/Source/Runtime/Apple/MetalRHI/Private/MetalRHI.cpp
@@ -630,6 +630,7 @@ FMetalDynamicRHI::FMetalDynamicRHI(ERHIFeatureLevel::Type RequestedFeatureLevel)
 	GSupportsSeparateRenderTargetBlendState = (GMaxRHIFeatureLevel >= ERHIFeatureLevel::SM5);
 
 	GRHISupportsPipelineFileCache = true;
+	GRHISupportsPSOPrecaching = true;
 	
 	// Appears to be no queryable value for max texture_buffer size and its not specified in the docs.  However,
 	// current testing across Apple Silicon macs, AMD and iPhone all quote a max value of 268435456 (1 << 28)
