Adding missing PLATFORM_SUPPORTS_BINDLESS_RENDERING checks

diff --git a/Engine/Source/Runtime/Apple/MetalRHI/Private/MetalCommands.cpp b/Engine/Source/Runtime/Apple/MetalRHI/Private/MetalCommands.cpp
--- a/Engine/Source/Runtime/Apple/MetalRHI/Private/MetalCommands.cpp
+++ b/Engine/Source/Runtime/Apple/MetalRHI/Private/MetalCommands.cpp
@@ -475,9 +475,11 @@ static void SetShaderParametersOnContext(
         case FRHIShaderParameterResource::EType::UniformBuffer:
             BindUniformBuffer(Binder.StateCache, Shader, Binder.Stage, Parameter.Index, static_cast<FRHIUniformBuffer*>(Parameter.Resource));
             break;
+#if PLATFORM_SUPPORTS_BINDLESS_RENDERING
         case FRHIShaderParameterResource::EType::ResourceCollection:
             Binder.SetResourceCollection(static_cast<FRHIResourceCollection*>(Parameter.Resource), Parameter.Index);
             break;
+#endif
         default:
             checkf(false, TEXT("Unhandled resource type?"));
             break;
diff --git a/Engine/Source/Runtime/Apple/MetalRHI/Private/MetalResourceCollection.h b/Engine/Source/Runtime/Apple/MetalRHI/Private/MetalResourceCollection.h
--- a/Engine/Source/Runtime/Apple/MetalRHI/Private/MetalResourceCollection.h
+++ b/Engine/Source/Runtime/Apple/MetalRHI/Private/MetalResourceCollection.h
@@ -5,6 +5,8 @@
 #include "RHICoreResourceCollection.h"
 #include "MetalRHIPrivate.h"
 
+#if PLATFORM_SUPPORTS_BINDLESS_RENDERING
+
 using FMetalResourceCollection = UE::RHICore::FGenericResourceCollection;
 
 template<>
@@ -12,3 +14,5 @@ struct TMetalResourceTraits<FRHIResourceCollection>
 {
 	using TConcreteType = FMetalResourceCollection;
 };
+
+#endif
