MetalRHI

Fixed error when using unoptimized shaders for debugging

diff --git a/Engine/Source/Runtime/Apple/MetalRHI/Private/Shaders/Types/Templates/MetalBaseShader.h b/Engine/Source/Runtime/Apple/MetalRHI/Private/Shaders/Types/Templates/MetalBaseShader.h
--- a/Engine/Source/Runtime/Apple/MetalRHI/Private/Shaders/Types/Templates/MetalBaseShader.h
+++ b/Engine/Source/Runtime/Apple/MetalRHI/Private/Shaders/Types/Templates/MetalBaseShader.h
@@ -191,8 +191,10 @@ void TMetalBaseShader<BaseResourceType, ShaderType>::Init(TArrayView<const uint8
 		else if(bForceTextShaders)
 		{
             GlslCodeNSString = FMetalShaderDebugCache::Get().GetShaderCode(SourceLen, SourceCRC);
-            check(GlslCodeNSString);
-            GlslCodeNSString->retain();
+			if(GlslCodeNSString)
+			{
+				GlslCodeNSString->retain();
+			}
 		}
 #endif
 		if (bForceTextShaders && CodeSize && CompressedSource.Num())
