Fix for incorrectly freeing buffer which is owned by the UAV

diff --git a/Engine/Source/Runtime/Apple/MetalRHI/Private/MetalUAV.cpp b/Engine/Source/Runtime/Apple/MetalRHI/Private/MetalUAV.cpp
--- a/Engine/Source/Runtime/Apple/MetalRHI/Private/MetalUAV.cpp
+++ b/Engine/Source/Runtime/Apple/MetalRHI/Private/MetalUAV.cpp
@@ -44,8 +44,16 @@ void FMetalResourceViewBase::Invalidate()
                 
         case EMetalType::TextureBufferBacked:
             FTextureBufferBacked & View = Storage.Get<FTextureBufferBacked>();
-            SafeReleaseMetalTexture(View.Texture);
-            SafeReleaseMetalBuffer(View.Buffer);
+			// If it is a buffer we don't own the resource
+			if (View.bIsBuffer)
+			{
+				SafeReleaseMetalTexture(View.Texture);
+			}
+			else
+			{
+				SafeReleaseMetalBuffer(View.Buffer);
+				SafeReleaseMetalTexture(View.Texture);
+			}
             break;
 		}
 	}
