Metal: Disable by default PSO file cache when PSO pre-caching is enabled
PSO file cache can be still enabled with r.Metal.EnablePSOFileCacheWhenPrecachingActive=1

diff --git a/Engine/Source/Runtime/Apple/MetalRHI/Private/MetalRHI.cpp b/Engine/Source/Runtime/Apple/MetalRHI/Private/MetalRHI.cpp
--- a/Engine/Source/Runtime/Apple/MetalRHI/Private/MetalRHI.cpp
+++ b/Engine/Source/Runtime/Apple/MetalRHI/Private/MetalRHI.cpp
@@ -46,6 +46,15 @@ static TAutoConsoleVariable<int32> CVarUseIOSRHIThread(
 													ECVF_Default | ECVF_RenderThreadSafe
 													);
 
+// If precaching is active we should not need the file cache.
+// however, precaching and filecache are compatible with each other, there maybe some scenarios in which both could be used.
+static TAutoConsoleVariable<bool> CVarEnableMetalPSOFileCacheWhenPrecachingActive(
+	TEXT("r.Metal.EnablePSOFileCacheWhenPrecachingActive"),
+	false,
+	TEXT("false: If precaching is available (r.PSOPrecaching=1, then disable the PSO filecache. (default)\n")
+	TEXT("true: Allow both PSO file cache and precaching."),
+	ECVF_RenderThreadSafe | ECVF_ReadOnly);
+
 
 static void ValidateTargetedRHIFeatureLevelExists(EShaderPlatform Platform)
 {
@@ -630,8 +639,8 @@ FMetalDynamicRHI::FMetalDynamicRHI(ERHIFeatureLevel::Type RequestedFeatureLevel)
 	GSupportsWideMRT = true;
 	GSupportsSeparateRenderTargetBlendState = (GMaxRHIFeatureLevel >= ERHIFeatureLevel::SM5);
 
-	GRHISupportsPipelineFileCache = true;
 	GRHISupportsPSOPrecaching = true;
+	GRHISupportsPipelineFileCache = !GRHISupportsPSOPrecaching || CVarEnableMetalPSOFileCacheWhenPrecachingActive.GetValueOnAnyThread();
 	
 	// Appears to be no queryable value for max texture_buffer size and its not specified in the docs.  However,
 	// current testing across Apple Silicon macs, AMD and iPhone all quote a max value of 268435456 (1 << 28)
diff --git a/Engine/Source/Runtime/RHI/Private/PipelineFileCache.cpp b/Engine/Source/Runtime/RHI/Private/PipelineFileCache.cpp
--- a/Engine/Source/Runtime/RHI/Private/PipelineFileCache.cpp
+++ b/Engine/Source/Runtime/RHI/Private/PipelineFileCache.cpp
@@ -3120,6 +3120,11 @@ void FPipelineFileCacheManager::Initialize(uint32 InGameVersion)
 
 bool FPipelineFileCacheManager::ShouldEnableFileCache()
 {
+	if (!GRHISupportsPipelineFileCache)
+	{
+		return false;
+	}
+
 #if PLATFORM_IOS
 	if (CVarAlwaysGeneratePOSSOFileCache.GetValueOnAnyThread() == 0)
 	{
@@ -3133,7 +3138,8 @@ bool FPipelineFileCacheManager::ShouldEnableFileCache()
 		}
 	}
 #endif
-	return GRHISupportsPipelineFileCache;
+
+	return true;
 }
 
 void FPipelineFileCacheManager::PreCompileComplete()
@@ -3173,6 +3179,7 @@ void FPipelineFileCacheManager::ClearOSPipelineCache()
 				FTimespan DataTime(0, 0, FileInfo.st_atime);
 				if (ExecutableTime > DataTime)
 				{
+					UE_LOG(LogTemp, Display, TEXT("Clearing functions.data"));
 					unlink(TCHAR_TO_UTF8(*Result));
 				}
 			}
@@ -3182,6 +3189,7 @@ void FPipelineFileCacheManager::ClearOSPipelineCache()
 				FTimespan MapsTime(0, 0, FileInfo.st_atime);
 				if (ExecutableTime > MapsTime)
 				{
+					UE_LOG(LogTemp, Display, TEXT("Clearing functions.maps"));
 					unlink(TCHAR_TO_UTF8(*Result));
 				}
 			}
