MetalRHI: Fix bug where buffer bound with bytes was ignored.

diff --git a/Engine/Source/Runtime/Apple/MetalRHI/Private/MetalStateCache.cpp b/Engine/Source/Runtime/Apple/MetalRHI/Private/MetalStateCache.cpp
--- a/Engine/Source/Runtime/Apple/MetalRHI/Private/MetalStateCache.cpp
+++ b/Engine/Source/Runtime/Apple/MetalRHI/Private/MetalStateCache.cpp
@@ -1180,7 +1180,7 @@ void FMetalStateCache::SetVertexStream(uint32 const Index, FMetalBufferPtr Buffe
 	if(IsMetalBindlessEnabled())
 	{
 		// Update GPU VA (assuming the offset has changed since last time).
-		if (Buffer)
+		if (Buffer || Bytes)
 		{
 			if (Bytes != nil)
 			{
