Metal: reduce shader metadata size by ~20% for most cases

diff --git a/Engine/Source/Developer/Apple/MetalShaderFormat/Private/MetalShaderCompiler.cpp b/Engine/Source/Developer/Apple/MetalShaderFormat/Private/MetalShaderCompiler.cpp
--- a/Engine/Source/Developer/Apple/MetalShaderFormat/Private/MetalShaderCompiler.cpp
+++ b/Engine/Source/Developer/Apple/MetalShaderFormat/Private/MetalShaderCompiler.cpp
@@ -299,13 +299,10 @@ void BuildMetalShaderOutput(
 	Header.CompileFlags |= (ShaderInput.Environment.CompilerFlags.Contains(CFLAG_BoundsChecking) ? (1 << CFLAG_BoundsChecking) : 0);
 	Header.CompileFlags |= (ShaderInput.Environment.CompilerFlags.Contains(CFLAG_Archive) ? (1 << CFLAG_Archive) : 0);
 
-	Header.CompilerVersion = FMetalCompilerToolchain::Get()->GetCompilerVersion((EShaderPlatform)ShaderInput.Target.Platform).Version;
-	Header.CompilerBuild = FMetalCompilerToolchain::Get()->GetTargetVersion((EShaderPlatform)ShaderInput.Target.Platform).Version;
 	Header.Version = Version;
 	Header.SideTable = -1;
 	Header.SourceLen = SourceCRCLen;
 	Header.SourceCRC = SourceCRC;
-	Header.Bindings.bDiscards = false;
 	Header.Bindings.ConstantBuffers = ConstantBuffers;
     
 	FShaderParameterMap& ParameterMap = ShaderOutput.ParameterMap;
@@ -366,7 +363,7 @@ void BuildMetalShaderOutput(
 		// For fragment shaders that discard but don't output anything we need at least a depth-stencil surface, so we need a way to validate this at runtime.
 		if (FCStringAnsi::Strstr(USFSource, "discard_fragment()") != nullptr)
 		{
-			Header.Bindings.bDiscards = true;
+			EnumAddFlags(Header.Bindings.Flags, EMetalBindingsFlags::PixelDiscard);
 		}
 	}
 
@@ -459,9 +456,6 @@ void BuildMetalShaderOutput(
 		Header.Bindings.PackedGlobalArrays.Add(Info);
 	}
 
-	// Setup Packed Uniform Buffers info
-	Header.Bindings.PackedUniformBuffers.Reserve(PackedUniformBuffersSize.Num());
-	
 	// In this mode there should only be 0 or 1 packed UB that contains all the aligned & named global uniform parameters
 	check(PackedUniformBuffersSize.Num() <= 1);
 	for (auto Iterator = PackedUniformBuffersSize.CreateIterator(); Iterator; ++Iterator)
@@ -523,7 +517,9 @@ void BuildMetalShaderOutput(
 #if UE_METAL_USE_METAL_SHADER_CONVERTER
     if (bUseMetalShaderConverter)
     {
-        // Only needed for VS Input (to generate the stage-in function used to convert inputs).
+		EnumAddFlags(Header.Bindings.Flags, EMetalBindingsFlags::UseMetalShaderConverter);
+				
+		// Only needed for VS Input (to generate the stage-in function used to convert inputs).
         if (Frequency == SF_Vertex)
         {
             Header.Bindings.IRConverterReflectionJSON = ANSI_TO_TCHAR(ShaderReflectionJSON);
@@ -536,8 +532,15 @@ void BuildMetalShaderOutput(
         }
 
         Header.Bindings.RSNumCBVs = NumCBVs;
-        Header.Bindings.bDiscards = bUsesDiscard;
         Header.Bindings.OutputSizeVS = OutputSizeVS;
+		if (bUsesDiscard)
+		{
+			EnumAddFlags(Header.Bindings.Flags, EMetalBindingsFlags::PixelDiscard);
+		}
+		else
+		{
+			EnumRemoveFlags(Header.Bindings.Flags, EMetalBindingsFlags::PixelDiscard);
+		}
 
 #if PLATFORM_SUPPORTS_GEOMETRY_SHADERS
         Header.Bindings.MaxInputPrimitivesPerMeshThreadgroupGS = MaxInputPrimitivesPerMeshThreadgroupGS;
@@ -568,7 +571,7 @@ void BuildMetalShaderOutput(
 		Header.RayTracing.InstanceIndexBuffer = CCHeader.RayTracingInstanceIndexBuffer;
 	}
 
-	Header.bDeviceFunctionConstants = (FCStringAnsi::Strstr(USFSource, "#define __METAL_DEVICE_CONSTANT_INDEX__ 1") != nullptr);
+	Header.bDeviceFunctionConstants = (FCStringAnsi::Strstr(USFSource, "#define __METAL_DEVICE_CONSTANT_INDEX__ 1") != nullptr) ? 1 : 0;
 	Header.SideTable = CCHeader.SideTable;
 	Header.Bindings.ArgumentBufferMasks = CCHeader.ArgumentBuffers;
 	Header.Bindings.ArgumentBuffers = 0;
diff --git a/Engine/Source/Developer/Apple/MetalShaderFormat/Private/MetalShaderFormat.cpp b/Engine/Source/Developer/Apple/MetalShaderFormat/Private/MetalShaderFormat.cpp
--- a/Engine/Source/Developer/Apple/MetalShaderFormat/Private/MetalShaderFormat.cpp
+++ b/Engine/Source/Developer/Apple/MetalShaderFormat/Private/MetalShaderFormat.cpp
@@ -38,7 +38,7 @@ extern uint64 AppendShader_Metal(class FString const& ArchivePath, const FSHAHas
 extern bool FinalizeLibrary_Metal(class FName const& Format, class FString const& ArchivePath, class FString const& LibraryPath, TSet<uint64> const& Shaders, class FString const& DebugOutputDir);
 
 /** Version for shader format, this becomes part of the DDC key. */
-static const FGuid UE_SHADER_METAL_VER = FGuid("282D7BF5-C5F8-47B7-B065-E165D60FC17B");
+static const FGuid UE_SHADER_METAL_VER = FGuid("B0DC25EF-C34D-437A-94E5-5E5146AF1B9A");
 
 class FMetalShaderFormat : public IShaderFormat
 {
diff --git a/Engine/Source/Runtime/Apple/MetalRHI/Private/MetalPipeline.cpp b/Engine/Source/Runtime/Apple/MetalRHI/Private/MetalPipeline.cpp
--- a/Engine/Source/Runtime/Apple/MetalRHI/Private/MetalPipeline.cpp
+++ b/Engine/Source/Runtime/Apple/MetalRHI/Private/MetalPipeline.cpp
@@ -620,7 +620,7 @@ static bool ConfigureRenderPipelineDescriptor(TDescriptorType* RenderPipelineDes
 	check(NumActiveTargets <= MaxSimultaneousRenderTargets);
 	if (PixelShader)
 	{
-		if (PixelShader->Bindings.InOutMask.Bitmask == 0 && PixelShader->Bindings.NumUAVs == 0 && PixelShader->Bindings.bDiscards == false)
+		if (PixelShader->Bindings.InOutMask.Bitmask == 0 && PixelShader->Bindings.NumUAVs == 0 && !EnumHasAnyFlags(PixelShader->Bindings.Flags, EMetalBindingsFlags::PixelDiscard))
 		{
 			UE_LOG(LogMetal, Error, TEXT("Pixel shader has no outputs which is not permitted. No Discards, In-Out Mask: %x\nNumber UAVs: %d\nSource Code:\n%s"), PixelShader->Bindings.InOutMask.Bitmask, PixelShader->Bindings.NumUAVs, *NSStringToFString(PixelShader->GetSourceCode()));
 			return false;
diff --git a/Engine/Source/Runtime/Apple/MetalRHI/Private/Shaders/Types/Templates/MetalBaseShader.h b/Engine/Source/Runtime/Apple/MetalRHI/Private/Shaders/Types/Templates/MetalBaseShader.h
--- a/Engine/Source/Runtime/Apple/MetalRHI/Private/Shaders/Types/Templates/MetalBaseShader.h
+++ b/Engine/Source/Runtime/Apple/MetalRHI/Private/Shaders/Types/Templates/MetalBaseShader.h
@@ -46,9 +46,6 @@ struct FMetalShaderData
     /** External bindings for this shader. */
     FMetalShaderBindings Bindings;
     
-    // List of memory copies from RHIUniformBuffer to packed uniforms
-    TArray<CrossCompiler::FUniformBufferCopyInfo> UniformBuffersCopyInfo;
-
     /* Argument encoders for shader IABs */
     TMap<uint32, MTL::ArgumentEncoder*> ArgumentEncoders;
 
@@ -154,7 +151,7 @@ void TMetalBaseShader<BaseResourceType, ShaderType>::Init(TArrayView<const uint8
 	// If this triggers than a level above us has failed to provide valid shader data and the cook is probably bogus
 	UE_CLOG(Header.SourceLen == 0 || Header.SourceCRC == 0, LogMetal, Fatal, TEXT("Invalid Shader Bytecode provided."));
 
-	bDeviceFunctionConstants = Header.bDeviceFunctionConstants;
+	bDeviceFunctionConstants = (Header.bDeviceFunctionConstants == 0 ? false : true);
 
 	// remember where the header ended and code (precompiled or source) begins
 	int32 CodeOffset = Ar.Tell();
@@ -210,7 +207,7 @@ void TMetalBaseShader<BaseResourceType, ShaderType>::Init(TArrayView<const uint8
 		GlslCodeNSString->retain();
 	}
 
-	bHasFunctionConstants = (Header.bDeviceFunctionConstants);
+	bHasFunctionConstants = (Header.bDeviceFunctionConstants == 0 ? false : true);
 
 	ConstantValueHash = 0;
 
@@ -375,7 +372,6 @@ void TMetalBaseShader<BaseResourceType, ShaderType>::Init(TArrayView<const uint8
 
 		GetCompiledFunction(true);
 	}
-	UniformBuffersCopyInfo = Header.UniformBuffersCopyInfo;
 	SideTableBinding = Header.SideTable;
 
 	UE::RHICore::InitStaticUniformBufferSlots(this);
diff --git a/Engine/Source/Runtime/Apple/MetalRHI/Public/MetalShaderResources.h b/Engine/Source/Runtime/Apple/MetalRHI/Public/MetalShaderResources.h
--- a/Engine/Source/Runtime/Apple/MetalRHI/Public/MetalShaderResources.h
+++ b/Engine/Source/Runtime/Apple/MetalRHI/Public/MetalShaderResources.h
@@ -99,9 +99,14 @@ enum class EMetalBufferFormat : uint8
 	Max						=54
 };
 
+enum class EMetalBindingsFlags : uint8
+{
+	PixelDiscard = 1 << 0,
+	UseMetalShaderConverter = 1 << 1,
+};
+
 struct FMetalShaderBindings
 {
-	TArray<TArray<CrossCompiler::FPackedArrayInfo>>	PackedUniformBuffers;
 	TArray<CrossCompiler::FPackedArrayInfo>			PackedGlobalArrays;
 	TMap<uint8, TArray<uint8>>						ArgumentBufferMasks;
 	CrossCompiler::FShaderBindingInOutMask			InOutMask;
@@ -115,28 +120,33 @@ struct FMetalShaderBindings
 	uint8	NumSamplers = 0;
 	uint8	NumUniformBuffers = 0;
 	uint8	NumUAVs = 0;
-	bool	bDiscards = false;
-
+	EMetalBindingsFlags	Flags{};
+	
 	inline FArchive& Serialize(FArchive& Ar, FShaderResourceTable& SRT);
 };
 
 inline FArchive& FMetalShaderBindings::Serialize(FArchive& Ar, FShaderResourceTable& SRT)
 {
-	Ar << PackedUniformBuffers;
 	Ar << PackedGlobalArrays;
 	Ar << SRT;
-	Ar << ArgumentBufferMasks;
 	Ar << ConstantBuffers;
-	Ar << ArgumentBuffers;
 	Ar << InOutMask;
+	Ar << ArgumentBuffers;
+	if (ArgumentBuffers)
+	{
+		Ar << ArgumentBufferMasks;
+	}
 	Ar << NumSamplers;
 	Ar << NumUniformBuffers;
 	Ar << NumUAVs;
-	Ar << bDiscards;
-    Ar << IRConverterReflectionJSON;
-    Ar << RSNumCBVs;
-    Ar << OutputSizeVS;
-    Ar << MaxInputPrimitivesPerMeshThreadgroupGS;
+	Ar << Flags;
+	if (EnumHasAnyFlags(Flags, EMetalBindingsFlags::UseMetalShaderConverter))
+	{
+		Ar << IRConverterReflectionJSON;
+		Ar << RSNumCBVs;
+		Ar << OutputSizeVS;
+		Ar << MaxInputPrimitivesPerMeshThreadgroupGS;
+	}
 	return Ar;
 }
 
@@ -218,35 +228,30 @@ struct FMetalAttribute
 struct FMetalCodeHeader
 {
 	FMetalShaderBindings Bindings;
-	TArray<CrossCompiler::FUniformBufferCopyInfo> UniformBuffersCopyInfo;
 
-	uint64 CompilerBuild;
-	uint32 CompilerVersion;
 	uint32 SourceLen;
 	uint32 SourceCRC;
+	uint32 Version;
 	uint32 NumThreadsX;
 	uint32 NumThreadsY;
 	uint32 NumThreadsZ;
 	uint32 CompileFlags;
+	FMetalRayTracingHeader RayTracing;
 	uint8 Frequency;
-	uint32 Version;
 	int8 SideTable;
-	bool bDeviceFunctionConstants;
-	FMetalRayTracingHeader RayTracing;
-
+	uint8 bDeviceFunctionConstants;
+	
 	FMetalCodeHeader()
-	: CompilerBuild(0)
-	, CompilerVersion(0)
-	, SourceLen(0)
+	: SourceLen(0)
 	, SourceCRC(0)
+	, Version(0)
 	, NumThreadsX(0)
 	, NumThreadsY(0)
 	, NumThreadsZ(0)
 	, CompileFlags(0)
 	, Frequency(0)
-	, Version(0)
 	, SideTable(-1)
-	, bDeviceFunctionConstants(false)
+	, bDeviceFunctionConstants(0)
 	{
 	}
 
@@ -257,40 +262,22 @@ struct FMetalCodeHeader
 inline FArchive& FMetalCodeHeader::Serialize(FArchive& Ar, FShaderResourceTable& SRT)
 {
 	Bindings.Serialize(Ar, SRT);
-
-	int32 NumInfos = UniformBuffersCopyInfo.Num();
-	Ar << NumInfos;
-	if (Ar.IsSaving())
-	{
-		for (int32 Index = 0; Index < NumInfos; ++Index)
-		{
-			Ar << UniformBuffersCopyInfo[Index];
-		}
-	}
-	else if (Ar.IsLoading())
-	{
-		UniformBuffersCopyInfo.Empty(NumInfos);
-		for (int32 Index = 0; Index < NumInfos; ++Index)
-		{
-			CrossCompiler::FUniformBufferCopyInfo Info;
-			Ar << Info;
-			UniformBuffersCopyInfo.Add(Info);
-		}
-	}
 	
-	Ar << CompilerBuild;
-	Ar << CompilerVersion;
 	Ar << SourceLen;
 	Ar << SourceCRC;
-	Ar << NumThreadsX;
-	Ar << NumThreadsY;
-	Ar << NumThreadsZ;
-	Ar << CompileFlags;
-	Ar << Frequency;
 	Ar << Version;
+	Ar << Frequency;
+	if (Frequency == SF_Compute || IsRayTracingShaderFrequency((EShaderFrequency)Frequency))
+	{
+		Ar << NumThreadsX;
+		Ar << NumThreadsY;
+		Ar << NumThreadsZ;
+		Ar << RayTracing;
+	}
+	Ar << CompileFlags;
 	Ar << SideTable;
 	Ar << bDeviceFunctionConstants;
-	Ar << RayTracing;
+	
     return Ar;
 }
 
