Metal: End current render pass on RHIEndRenderPass

diff --git a/Engine/Source/Runtime/Apple/MetalRHI/Private/MetalContext.cpp b/Engine/Source/Runtime/Apple/MetalRHI/Private/MetalContext.cpp
--- a/Engine/Source/Runtime/Apple/MetalRHI/Private/MetalContext.cpp
+++ b/Engine/Source/Runtime/Apple/MetalRHI/Private/MetalContext.cpp
@@ -1307,6 +1307,11 @@ void FMetalContext::SetRenderPassInfo(const FRHIRenderPassInfo& RenderTargetsInf
 	}
 }
 
+void FMetalContext::EndRenderPass()
+{
+    RenderPass.EndRenderPass();
+}
+
 FMetalBufferPtr FMetalContext::AllocateFromRingBuffer(uint32 Size, uint32 Alignment)
 {
 	return RenderPass.GetRingBuffer().NewBuffer(Size, Alignment);
diff --git a/Engine/Source/Runtime/Apple/MetalRHI/Private/MetalContext.h b/Engine/Source/Runtime/Apple/MetalRHI/Private/MetalContext.h
--- a/Engine/Source/Runtime/Apple/MetalRHI/Private/MetalContext.h
+++ b/Engine/Source/Runtime/Apple/MetalRHI/Private/MetalContext.h
@@ -50,6 +50,7 @@ public:
 	 * Set the color, depth and stencil render targets, and then make the new command buffer/encoder
 	 */
 	void SetRenderPassInfo(const FRHIRenderPassInfo& RenderTargetsInfo, bool const bRestart = false);
+    void EndRenderPass();
 	
 	/**
 	 * Allocate from a dynamic ring buffer - by default align to the allowed alignment for offset field when setting buffers
diff --git a/Engine/Source/Runtime/Apple/MetalRHI/Private/MetalRHIContext.cpp b/Engine/Source/Runtime/Apple/MetalRHI/Private/MetalRHIContext.cpp
--- a/Engine/Source/Runtime/Apple/MetalRHI/Private/MetalRHIContext.cpp
+++ b/Engine/Source/Runtime/Apple/MetalRHI/Private/MetalRHIContext.cpp
@@ -130,6 +130,8 @@ void FMetalRHICommandContext::RHIEndRenderPass()
 	{
 		RHIEndOcclusionQueryBatch();
 	}
+    
+    Context->EndRenderPass();
 
 	UE::RHICore::ResolveRenderPassTargets(RenderPassInfo, [this](UE::RHICore::FResolveTextureInfo Info)
 	{
