PSO Pre-caching: Fixed case where mobile deferred was requesting a wrong DepthAux pixel format.
Metal: Removed irrelevant properties from PSO pre-caching hash computation (eg. load/store actions)

diff --git a/Engine/Source/Runtime/Apple/MetalRHI/Private/MetalDynamicRHI.cpp b/Engine/Source/Runtime/Apple/MetalRHI/Private/MetalDynamicRHI.cpp
--- a/Engine/Source/Runtime/Apple/MetalRHI/Private/MetalDynamicRHI.cpp
+++ b/Engine/Source/Runtime/Apple/MetalRHI/Private/MetalDynamicRHI.cpp
@@ -173,3 +173,98 @@ bool FMetalDynamicRHI::RHIGetRenderQueryResult(FRHIRenderQuery* QueryRHI, uint64
 	FMetalRHIRenderQuery* Query = ResourceCast(QueryRHI);
 	return Query->GetResult(OutNumPixels, bWait, GPUIndex);
 }
+
+uint64 FMetalDynamicRHI::RHIComputePrecachePSOHash(const FGraphicsPipelineStateInitializer& Initializer)
+{
+	// When compute precache PSO hash we assume a valid state precache PSO hash is already provided
+	checkf(Initializer.StatePrecachePSOHash != 0, TEXT("Initializer should have a valid state precache PSO hash set when computing the full initializer PSO hash"));
+
+	// All members which are not part of the state objects and influence the PSO on Metal
+	struct FNonStateHashKey
+	{
+		uint64							StatePrecachePSOHash;
+
+		EPrimitiveType					PrimitiveType;
+		uint32							RenderTargetsEnabled;
+		FGraphicsPipelineStateInitializer::TRenderTargetFormats RenderTargetFormats;
+		EPixelFormat					DepthStencilTargetFormat;
+		uint16							NumSamples;
+		EConservativeRasterization		ConservativeRasterization;
+		bool							bDepthBounds;
+		uint8							MultiViewCount;
+		bool							bHasFragmentDensityAttachment;
+		EVRSShadingRate					ShadingRate;
+	} HashKey;
+
+	FMemory::Memzero(&HashKey, sizeof(FNonStateHashKey));
+
+	HashKey.StatePrecachePSOHash			= Initializer.StatePrecachePSOHash;
+
+	HashKey.PrimitiveType					= Initializer.PrimitiveType;
+	HashKey.RenderTargetsEnabled			= Initializer.RenderTargetsEnabled;
+	HashKey.RenderTargetFormats				= Initializer.RenderTargetFormats;
+	HashKey.DepthStencilTargetFormat		= Initializer.DepthStencilTargetFormat;
+	HashKey.NumSamples						= Initializer.NumSamples;
+	HashKey.ConservativeRasterization		= Initializer.ConservativeRasterization;
+	HashKey.bDepthBounds					= Initializer.bDepthBounds;
+	HashKey.MultiViewCount					= Initializer.MultiViewCount;
+	HashKey.bHasFragmentDensityAttachment	= Initializer.bHasFragmentDensityAttachment;
+	HashKey.ShadingRate						= Initializer.ShadingRate;
+
+	return CityHash64((const char*)&HashKey, sizeof(FNonStateHashKey));
+}
+
+bool FMetalDynamicRHI::RHIMatchPrecachePSOInitializers(const FGraphicsPipelineStateInitializer& LHS, const FGraphicsPipelineStateInitializer& RHS)
+{
+	// first check non pointer objects
+	if (LHS.ImmutableSamplerState != RHS.ImmutableSamplerState ||
+		LHS.PrimitiveType != RHS.PrimitiveType ||
+		LHS.bDepthBounds != RHS.bDepthBounds ||
+		LHS.MultiViewCount != RHS.MultiViewCount ||
+		LHS.ShadingRate != RHS.ShadingRate ||
+		LHS.bHasFragmentDensityAttachment != RHS.bHasFragmentDensityAttachment ||
+		LHS.RenderTargetsEnabled != RHS.RenderTargetsEnabled ||
+		LHS.RenderTargetFormats != RHS.RenderTargetFormats ||
+		LHS.DepthStencilTargetFormat != RHS.DepthStencilTargetFormat ||
+		LHS.NumSamples != RHS.NumSamples ||
+		LHS.ConservativeRasterization != RHS.ConservativeRasterization)
+	{
+		return false;
+	}
+
+	// check the RHI shaders (pointer check for shaders should be fine)
+	if (LHS.BoundShaderState.GetVertexShader() != RHS.BoundShaderState.GetVertexShader() ||
+		LHS.BoundShaderState.GetPixelShader() != RHS.BoundShaderState.GetPixelShader() ||
+		LHS.BoundShaderState.GetMeshShader() != RHS.BoundShaderState.GetMeshShader() ||
+		LHS.BoundShaderState.GetAmplificationShader() != RHS.BoundShaderState.GetAmplificationShader() ||
+		LHS.BoundShaderState.GetGeometryShader() != RHS.BoundShaderState.GetGeometryShader())
+	{
+		return false;
+	}
+
+	// Compare the VertexDecl
+	FMetalHashedVertexDescriptor LHSVertexElements;
+	if (LHS.BoundShaderState.VertexDeclarationRHI)
+	{
+		LHSVertexElements = ((FMetalVertexDeclaration*)LHS.BoundShaderState.VertexDeclarationRHI)->Layout;
+	}
+	FMetalHashedVertexDescriptor RHSVertexElements;
+	if (RHS.BoundShaderState.VertexDeclarationRHI)
+	{
+		RHSVertexElements = ((FMetalVertexDeclaration*)RHS.BoundShaderState.VertexDeclarationRHI)->Layout;
+	}
+	if (!(LHSVertexElements == RHSVertexElements))
+	{
+		return false;
+	}
+
+	// Check actual state content (each initializer can have it's own state and not going through a factory)
+	if (!MatchRHIState<FRHIBlendState, FBlendStateInitializerRHI>(LHS.BlendState, RHS.BlendState) ||
+		!MatchRHIState<FRHIRasterizerState, FRasterizerStateInitializerRHI>(LHS.RasterizerState, RHS.RasterizerState) ||
+		!MatchRHIState<FRHIDepthStencilState, FDepthStencilStateInitializerRHI>(LHS.DepthStencilState, RHS.DepthStencilState))
+	{
+		return false;
+	}
+
+	return true;
+}
diff --git a/Engine/Source/Runtime/Apple/MetalRHI/Private/MetalDynamicRHI.h b/Engine/Source/Runtime/Apple/MetalRHI/Private/MetalDynamicRHI.h
--- a/Engine/Source/Runtime/Apple/MetalRHI/Private/MetalDynamicRHI.h
+++ b/Engine/Source/Runtime/Apple/MetalRHI/Private/MetalDynamicRHI.h
@@ -161,6 +161,10 @@ public:
 	virtual FTextureReferenceRHIRef RHICreateTextureReference(FRHICommandListBase& RHICmdList, FRHITexture* InReferencedTexture) final override;
 	virtual void RHIUpdateTextureReference(FRHICommandListBase& RHICmdList, FRHITextureReference* TextureRef, FRHITexture* NewTexture) final override;
 
+
+	virtual uint64 RHIComputePrecachePSOHash(const FGraphicsPipelineStateInitializer& Initializer) final override;
+	virtual bool RHIMatchPrecachePSOInitializers(const FGraphicsPipelineStateInitializer& LHS, const FGraphicsPipelineStateInitializer& RHS) final override;
+
 private:
 	FTextureMemoryStats MemoryStats;
 	FMetalRHIImmediateCommandContext ImmediateContext;
diff --git a/Engine/Source/Runtime/Engine/Private/SceneTexturesConfig.cpp b/Engine/Source/Runtime/Engine/Private/SceneTexturesConfig.cpp
--- a/Engine/Source/Runtime/Engine/Private/SceneTexturesConfig.cpp
+++ b/Engine/Source/Runtime/Engine/Private/SceneTexturesConfig.cpp
@@ -256,6 +256,16 @@ static void SetupMobileGBufferFlags(FGBufferBindings GBufferBindings[GBL_Num], b
 	}
 }
 
+static bool MobileRequiresPreciseSceneDepthAux(EShaderPlatform ShaderPlatform)
+{
+	static const auto CVar = IConsoleManager::Get().FindTConsoleVariableDataInt(TEXT("r.Mobile.SceneDepthAux"));
+	if (IsMobileDeferredShadingEnabled(ShaderPlatform) || CVar->GetValueOnAnyThread() == 2)
+	{
+		return true;
+	}
+	return false;
+}
+
 void FSceneTexturesConfig::Init(const FSceneTexturesConfigInitSettings& InitSettings)
 {
 	FeatureLevel			= InitSettings.FeatureLevel;
@@ -327,6 +337,7 @@ void FSceneTexturesConfig::Init(const FSceneTexturesConfigInitSettings& InitSett
 	if (ShadingPath == EShadingPath::Mobile)
 	{
 		bRequiresDepthAux = MobileRequiresSceneDepthAux(ShaderPlatform);
+		bPreciseDepthAux = bPreciseDepthAux || MobileRequiresPreciseSceneDepthAux(ShaderPlatform);
 	}
 }
 
diff --git a/Engine/Source/Runtime/Renderer/Private/SceneTextures.cpp b/Engine/Source/Runtime/Renderer/Private/SceneTextures.cpp
--- a/Engine/Source/Runtime/Renderer/Private/SceneTextures.cpp
+++ b/Engine/Source/Runtime/Renderer/Private/SceneTextures.cpp
@@ -99,27 +99,6 @@ EPixelFormat FSceneTextures::GetGBufferFFormatAndCreateFlags(ETextureCreateFlags
 	return NormalGBufferFormat;
 }
 
-inline EPixelFormat GetMobileSceneDepthAuxPixelFormat(EShaderPlatform ShaderPlatform, bool bPreciseFormat)
-{
-	if (IsMobileDeferredShadingEnabled(ShaderPlatform) || bPreciseFormat)
-	{
-		return PF_R32_FLOAT;
-	}
-
-	static const auto CVar = IConsoleManager::Get().FindTConsoleVariableDataInt(TEXT("r.Mobile.SceneDepthAux"));
-	EPixelFormat Format = PF_R16F;
-	switch (CVar->GetValueOnAnyThread())
-	{
-	case 1:
-		Format =  PF_R16F;
-		break;
-	case 2:
-		Format = PF_R32_FLOAT;
-		break;
-	}
-	return Format;
-}
-
 static IStereoRenderTargetManager* FindStereoRenderTargetManager()
 {
 	if (!GEngine->StereoRenderingDevice.IsValid() || !GEngine->StereoRenderingDevice->IsStereoEnabled())
@@ -615,7 +594,7 @@ void FSceneTextures::InitializeViewFamily(FRDGBuilder& GraphBuilder, FViewFamily
 		#endif
 		}
 
-		EPixelFormat DepthAuxFormat = GetMobileSceneDepthAuxPixelFormat(Config.ShaderPlatform, Config.bPreciseDepthAux);
+		EPixelFormat DepthAuxFormat = Config.bPreciseDepthAux ? PF_R32_FLOAT : PF_R16F;
 		FRDGTextureDesc Desc = Config.bRequireMultiView ? 
 			FRDGTextureDesc::Create2DArray(Config.Extent, DepthAuxFormat, FClearValueBinding(FarDepthColor), TexCreate_RenderTargetable | TexCreate_ShaderResource | TexCreate_InputAttachmentRead | MemorylessFlag, 2) :
 			FRDGTextureDesc::Create2D(Config.Extent, DepthAuxFormat, FClearValueBinding(FarDepthColor), TexCreate_RenderTargetable | TexCreate_ShaderResource | TexCreate_InputAttachmentRead| MemorylessFlag);
