Add rhi.Metal.MaxBucketSize cvar to restrict the MaxBucket Size of the Pooled Buffer

diff --git a/Engine/Source/Runtime/Apple/MetalRHI/Private/MetalBuffer.cpp b/Engine/Source/Runtime/Apple/MetalRHI/Private/MetalBuffer.cpp
--- a/Engine/Source/Runtime/Apple/MetalRHI/Private/MetalBuffer.cpp
+++ b/Engine/Source/Runtime/Apple/MetalRHI/Private/MetalBuffer.cpp
@@ -35,6 +35,13 @@ static FAutoConsoleVariableRef CVarMetalResourcePurgeInPool(
 	GMetalResourcePurgeInPool,
 	TEXT("Use the SetPurgeableState function to allow the OS to reclaim memory from resources while they are unused in the pools. (Default: 0, Off)"));
 
+static int32 GMetalMaxBucketSize = 32 * 1024 * 1024;
+static FAutoConsoleVariableRef CVarMetalMaxBucketSize(
+	TEXT("rhi.Metal.MaxBucketSize"),
+	GMetalMaxBucketSize,
+	TEXT("Set the Max Bucket Size for the Pooled Buffer. (Default: 33547255)"));
+
+
 #if METAL_DEBUG_OPTIONS
 extern int32 GMetalBufferScribble;
 #endif
@@ -1767,7 +1774,8 @@ FMetalBufferPtr FMetalResourceHeap::CreateBuffer(uint32 Size, uint32 Alignment,
 	// Write combined should be on a case by case basis
 	check(CpuMode == MTL::CPUCacheModeDefaultCache);
 	
-	if (BlockSize <= 33554432)
+	check(GMetalMaxBucketSize <= 32 * 1024 * 1024);
+	if (BlockSize <= GMetalMaxBucketSize)
 	{
 		switch (StorageMode)
 		{
