iOS: Fixed drawable surface leak

diff --git a/Engine/Source/Runtime/Apple/MetalRHI/Private/MetalViewport.cpp b/Engine/Source/Runtime/Apple/MetalRHI/Private/MetalViewport.cpp
--- a/Engine/Source/Runtime/Apple/MetalRHI/Private/MetalViewport.cpp
+++ b/Engine/Source/Runtime/Apple/MetalRHI/Private/MetalViewport.cpp
@@ -552,8 +552,6 @@ void FMetalViewport::Present(FMetalCommandQueue& CommandQueue, bool bLockToVsync
 						CurrentCommandBuffer->GetMTLCmdBuffer()->addCompletedHandler(CommandBufferHandler);
 
                         {
-                            LocalDrawable->retain();
-                            
                             // Queue this on the current command buffer to ensure that all work is committed prior to the present, present only knows about dependencies on committed work.
                             if (MinPresentDuration && GEnablePresentPacing)
                             {
