VisionOS missing include
-This include was not needed in 5.4, presumably other includes were tightened up causing this to break.

diff --git a/Engine/Source/Runtime/Apple/MetalRHI/Private/MetalRenderPass.h b/Engine/Source/Runtime/Apple/MetalRHI/Private/MetalRenderPass.h
--- a/Engine/Source/Runtime/Apple/MetalRHI/Private/MetalRenderPass.h
+++ b/Engine/Source/Runtime/Apple/MetalRHI/Private/MetalRenderPass.h
@@ -8,6 +8,10 @@
 #include "MetalBuffer.h"
 #include "MetalCommandEncoder.h"
 
+#if PLATFORM_VISIONOS
+#import <CompositorServices/CompositorServices.h>
+#endif
+
 class FMetalCommandList;
 class FMetalCommandQueue;
 class FMetalStateCache;
