Remove CurrentContextTLSSlot and associated code in Metal RHI
 - Previously used to provide the context pointer to the FMetalProfiler / FMetalGPUProfiler, and set inside RHIBeginScene.
 - Since there is only one context in Metal at the moment, this is redundant as the profiler already has a Context pointer.

diff --git a/Engine/Source/Runtime/Apple/MetalRHI/Private/MetalContext.cpp b/Engine/Source/Runtime/Apple/MetalRHI/Private/MetalContext.cpp
--- a/Engine/Source/Runtime/Apple/MetalRHI/Private/MetalContext.cpp
+++ b/Engine/Source/Runtime/Apple/MetalRHI/Private/MetalContext.cpp
@@ -444,11 +444,7 @@ void FMetalDeviceContext::Init(void)
 }
 
 void FMetalDeviceContext::BeginFrame()
-{
-#if ENABLE_METAL_GPUPROFILE
-	FPlatformTLS::SetTlsValue(CurrentContextTLSSlot, this);
-#endif
-	
+{	
 	// Wait for the frame semaphore on the immediate context.
 	dispatch_semaphore_wait(CommandBufferSemaphore, DISPATCH_TIME_FOREVER);
 
@@ -595,22 +591,8 @@ void FMetalDeviceContext::EndFrame()
 	InitFrame();
 }
 
-void FMetalDeviceContext::BeginScene()
-{
-#if ENABLE_METAL_GPUPROFILE
-	FPlatformTLS::SetTlsValue(CurrentContextTLSSlot, this);
-#endif
-}
-
-void FMetalDeviceContext::EndScene()
-{
-}
-
 void FMetalDeviceContext::BeginDrawingViewport(FMetalViewport* Viewport)
 {
-#if ENABLE_METAL_GPUPROFILE
-	FPlatformTLS::SetTlsValue(CurrentContextTLSSlot, this);
-#endif
 }
 
 bool FMetalDeviceContext::FMetalDelayedFreeList::IsComplete() const
@@ -925,11 +907,6 @@ bool FMetalDeviceContext::ValidateIsInactiveBuffer(MTL::Buffer* Buffer, const NS
 }
 #endif
 
-
-#if ENABLE_METAL_GPUPROFILE
-uint32 FMetalContext::CurrentContextTLSSlot = FPlatformTLS::AllocTlsSlot();
-#endif
-
 FMetalContext::FMetalContext(MTL::Device* InDevice, FMetalCommandQueue& Queue)
 	: Device(InDevice)
 	, CommandQueue(Queue)
@@ -976,37 +953,8 @@ void FMetalContext::InsertCommandBufferFence(TSharedPtr<FMetalCommandBufferFence
 	RenderPass.InsertCommandBufferFence(Fence, Handler);
 }
 
-#if ENABLE_METAL_GPUPROFILE
-FMetalContext* FMetalContext::GetCurrentContext()
-{
-	FMetalContext* Current = (FMetalContext*)FPlatformTLS::GetTlsValue(CurrentContextTLSSlot);
-	
-	if (!Current)
-	{
-		// If we are executing this outside of a pass we'll return the default.
-		// TODO This needs further investigation. We should fix all the cases that call this without
-		// a context set.
-		FMetalRHICommandContext* CmdContext = static_cast<FMetalRHICommandContext*>(RHIGetDefaultContext());
-		check(CmdContext);
-		Current = &CmdContext->GetInternalContext();
-	}
-	
-	check(Current);
-	return Current;
-}
-
-void FMetalContext::MakeCurrent(FMetalContext* Context)
-{
-	FPlatformTLS::SetTlsValue(CurrentContextTLSSlot, Context);
-}
-#endif
-
 void FMetalContext::InitFrame()
-{
-#if ENABLE_METAL_GPUPROFILE
-	FPlatformTLS::SetTlsValue(CurrentContextTLSSlot, this);
-#endif
-	
+{	
 	// Reset cached state in the encoder
 	StateCache.Reset();
 	
@@ -1035,10 +983,6 @@ void FMetalContext::FinishFrame(bool const bImmediateContext)
 	{
 		StateCache.Reset();
 	}
-
-#if ENABLE_METAL_GPUPROFILE
-	FPlatformTLS::SetTlsValue(CurrentContextTLSSlot, nullptr);
-#endif
 }
 
 void FMetalContext::TransitionResource(FRHIUnorderedAccessView* InResource)
diff --git a/Engine/Source/Runtime/Apple/MetalRHI/Private/MetalContext.h b/Engine/Source/Runtime/Apple/MetalRHI/Private/MetalContext.h
--- a/Engine/Source/Runtime/Apple/MetalRHI/Private/MetalContext.h
+++ b/Engine/Source/Runtime/Apple/MetalRHI/Private/MetalContext.h
@@ -111,11 +111,6 @@ public:
 
 	void StartTiming(class FMetalEventNode* EventNode);
 	void EndTiming(class FMetalEventNode* EventNode);
-
-#if ENABLE_METAL_GPUPROFILE
-	static void MakeCurrent(FMetalContext* Context);
-	static FMetalContext* GetCurrentContext();
-#endif
 	
 	void InitFrame();
 	void FinishFrame(bool const bImmediateContext);
@@ -148,11 +143,6 @@ protected:
 	/** A pool of buffers for writing visibility query results. */
 	TSharedPtr<FMetalQueryBufferPool, ESPMode::ThreadSafe> QueryBuffer;
 	
-#if ENABLE_METAL_GPUPROFILE
-	/** the slot to store a per-thread context ref */
-	static uint32 CurrentContextTLSSlot;
-#endif
-	
 	/** Whether the validation layer is enabled */
 	bool bValidationEnabled;
 };
@@ -193,11 +183,6 @@ public:
 	void DrainHeap();
 	void EndFrame();
 	
-	/** RHIBeginScene helper */
-	void BeginScene();
-	/** RHIEndScene helper */
-	void EndScene();
-	
 	void BeginDrawingViewport(FMetalViewport* Viewport);
 	void EndDrawingViewport(IRHICommandContext& RHICmdContext, FMetalViewport* Viewport, bool bPresent, bool bLockToVsync);
 	
diff --git a/Engine/Source/Runtime/Apple/MetalRHI/Private/MetalProfiler.cpp b/Engine/Source/Runtime/Apple/MetalRHI/Private/MetalProfiler.cpp
--- a/Engine/Source/Runtime/Apple/MetalRHI/Private/MetalProfiler.cpp
+++ b/Engine/Source/Runtime/Apple/MetalRHI/Private/MetalProfiler.cpp
@@ -155,7 +155,7 @@ void FMetalEventNodeFrame::LogDisjointQuery()
 FGPUProfilerEventNode* FMetalGPUProfiler::CreateEventNode(const TCHAR* InName, FGPUProfilerEventNode* InParent)
 {
 #if ENABLE_METAL_GPUPROFILE
-	FMetalEventNode* EventNode = new FMetalEventNode(FMetalContext::GetCurrentContext(), InName, InParent, false, false);
+	FMetalEventNode* EventNode = new FMetalEventNode(Context, InName, InParent, false, false);
 	return EventNode;
 #else
 	return nullptr;
diff --git a/Engine/Source/Runtime/Apple/MetalRHI/Private/MetalRHI.cpp b/Engine/Source/Runtime/Apple/MetalRHI/Private/MetalRHI.cpp
--- a/Engine/Source/Runtime/Apple/MetalRHI/Private/MetalRHI.cpp
+++ b/Engine/Source/Runtime/Apple/MetalRHI/Private/MetalRHI.cpp
@@ -1244,14 +1244,10 @@ void FMetalRHICommandContext::RHIEndFrame()
 
 void FMetalRHICommandContext::RHIBeginScene()
 {
-    MTL_SCOPED_AUTORELEASE_POOL;
-	((FMetalDeviceContext*)Context)->BeginScene();
 }
 
 void FMetalRHICommandContext::RHIEndScene()
 {
-    MTL_SCOPED_AUTORELEASE_POOL;
-	((FMetalDeviceContext*)Context)->EndScene();
 }
 
 #if WITH_RHI_BREADCRUMBS
diff --git a/Engine/Source/Runtime/Apple/MetalRHI/Public/MetalRHIContext.h b/Engine/Source/Runtime/Apple/MetalRHI/Public/MetalRHIContext.h
--- a/Engine/Source/Runtime/Apple/MetalRHI/Public/MetalRHIContext.h
+++ b/Engine/Source/Runtime/Apple/MetalRHI/Public/MetalRHIContext.h
@@ -74,17 +74,9 @@ public:
 	// This method is queued with an RHIThread, otherwise it will flush after it is queued; without an RHI thread there is no benefit to queuing this frame advance commands
 	virtual void RHIEndFrame() override;
 	
-	/**
-		* Signals the beginning of scene rendering. The RHI makes certain caching assumptions between
-		* calls to BeginScene/EndScene. Currently the only restriction is that you can't update texture
-		* references.
-		*/
 	// This method is queued with an RHIThread, otherwise it will flush after it is queued; without an RHI thread there is no benefit to queuing this frame advance commands
 	virtual void RHIBeginScene() override;
-	
-	/**
-		* Signals the end of scene rendering. See RHIBeginScene.
-		*/
+
 	// This method is queued with an RHIThread, otherwise it will flush after it is queued; without an RHI thread there is no benefit to queuing this frame advance commands
 	virtual void RHIEndScene() override;
 	
