- Make sure to submit an encoder before finalizing VisionOS frame

diff --git a/Engine/Source/Runtime/Apple/MetalRHI/Private/MetalViewport.cpp b/Engine/Source/Runtime/Apple/MetalRHI/Private/MetalViewport.cpp
--- a/Engine/Source/Runtime/Apple/MetalRHI/Private/MetalViewport.cpp
+++ b/Engine/Source/Runtime/Apple/MetalRHI/Private/MetalViewport.cpp
@@ -687,6 +687,7 @@ void FMetalViewport::PresentImmersive(const MetalRHIVisionOS::PresentImmersivePa
 				}
 
 				RenderPass.EndRenderPass();
+				RenderPass.Submit(EMetalSubmitFlagsCreateCommandBuffer);
 				RenderPass.EncodePresentImmersive(VisionOSParams.SwiftDrawable, VisionOSParams.SwiftFrame);
 			}
 		}
