Add r.Metal.PSOPrecaching.MemoryThreshold to skip PSO precaching if memory falls under a threshold

diff --git a/Engine/Source/Runtime/Apple/MetalRHI/Private/MetalDynamicRHI.cpp b/Engine/Source/Runtime/Apple/MetalRHI/Private/MetalDynamicRHI.cpp
--- a/Engine/Source/Runtime/Apple/MetalRHI/Private/MetalDynamicRHI.cpp
+++ b/Engine/Source/Runtime/Apple/MetalRHI/Private/MetalDynamicRHI.cpp
@@ -12,11 +12,20 @@
 #include "MetalVertexDeclaration.h"
 #include "MetalGraphicsPipelineState.h"
 #include "MetalTransitionData.h"
+#include "HAL/PlatformMemoryHelpers.h"
 
 //------------------------------------------------------------------------------
 
 #pragma mark - Metal Dynamic RHI Vertex Declaration Methods -
 
+int32 GPSOMemoryThreshold = -1;
+static FAutoConsoleVariableRef CVarPSOMemoryThreshold(
+	TEXT("r.Metal.PSOPrecaching.MemoryThreshold"),
+	GPSOMemoryThreshold,
+	TEXT("-1 No Memory Threshold set, PSO precaching is always on if enabled\n")
+	TEXT(">= 0  Memory Threshold set, when memory drops under this limit, PSO preaching is disabled\n"),
+	ECVF_Default
+);
 
 FVertexDeclarationRHIRef FMetalDynamicRHI::RHICreateVertexDeclaration(const FVertexDeclarationElementList& Elements)
 {
@@ -45,6 +54,18 @@ FGraphicsPipelineStateRHIRef FMetalDynamicRHI::RHICreateGraphicsPipelineState(co
 {
     MTL_SCOPED_AUTORELEASE_POOL;
     
+	if (Initializer.bPSOPrecache)
+	{
+		if (GPSOMemoryThreshold >= 0)
+		{
+			float MBFree = float(PlatformMemoryHelpers::GetFrameMemoryStats().AvailablePhysical / 1024 / 1024);
+			if (MBFree < GPSOMemoryThreshold)
+			{
+				return nullptr;
+			}
+		}
+	}
+
     TRefCountPtr<FMetalGraphicsPipelineState> State = new FMetalGraphicsPipelineState(Initializer);
 
 #if METAL_USE_METAL_SHADER_CONVERTER
