Expand rhi.Metal.CacheShaderPipelines to also have the ability to only cache PSO's that were not part of the precache

diff --git a/Engine/Source/Runtime/Apple/MetalRHI/Private/MetalPipeline.cpp b/Engine/Source/Runtime/Apple/MetalRHI/Private/MetalPipeline.cpp
--- a/Engine/Source/Runtime/Apple/MetalRHI/Private/MetalPipeline.cpp
+++ b/Engine/Source/Runtime/Apple/MetalRHI/Private/MetalPipeline.cpp
@@ -24,6 +24,7 @@ static FAutoConsoleVariableRef CVarMetalCacheShaderPipelines(
 	TEXT("rhi.Metal.CacheShaderPipelines"),
 	GMetalCacheShaderPipelines,
 	TEXT("When enabled (1, default) cache all graphics pipeline state objects created in MetalRHI for the life of the program, this trades memory for performance as creating PSOs is expensive in Metal.\n")
+	TEXT("When set to 2 it will cache all graphics pipeline state objects(that are not requested by PSO precaching) created in MetalRHI for the life of the program, this trades memory for performance as creating PSOs is expensive in Metal.\n")
 	TEXT("Disable in the project configuration to allow PSOs to be released to save memory at the expense of reduced performance and increased hitching in-game\n. (On by default (1))"), ECVF_ReadOnly);
 
 static int32 GMetalCacheMinSize = 32;
@@ -269,7 +270,8 @@ public:
 			{
 				Desc = CreateMTLRenderPipeline(bSync, Key, Init, State);
 
-				if (Desc != nullptr)
+				// Don't cache PSOPrecache if cvar is set to exclude them
+				if (Desc != nullptr && (!Init.bPSOPrecache || GMetalCacheShaderPipelines != 2))
 				{
 					PipelineMutex.WriteLock();
 
