MetalRHI: Fix for readback, not submitting command buffer with GPU fence

diff --git a/Engine/Source/Runtime/Apple/MetalRHI/Private/MetalRenderTarget.cpp b/Engine/Source/Runtime/Apple/MetalRHI/Private/MetalRenderTarget.cpp
--- a/Engine/Source/Runtime/Apple/MetalRHI/Private/MetalRenderTarget.cpp
+++ b/Engine/Source/Runtime/Apple/MetalRHI/Private/MetalRenderTarget.cpp
@@ -267,6 +267,9 @@ void FMetalDynamicRHI::RHIMapStagingSurface(FRHITexture* TextureRHI, FRHIGPUFenc
 {
     MTL_SCOPED_AUTORELEASE_POOL;
     
+	FRHICommandListExecutor::GetImmediateCommandList().ImmediateFlush(EImmediateFlushType::FlushRHIThread);
+	ImmediateContext.Context->SubmitCommandsHint();
+	
 	if (FenceRHI && !FenceRHI->Poll())
 	{
 		ResourceCast(FenceRHI)->WaitCPU();
