MetalRHI: Disable Metal heaps as they are increasing the minimum required memory and have potential leaks.

diff --git a/Engine/Source/Runtime/Apple/MetalRHI/Private/MetalCommandQueue.cpp b/Engine/Source/Runtime/Apple/MetalRHI/Private/MetalCommandQueue.cpp
--- a/Engine/Source/Runtime/Apple/MetalRHI/Private/MetalCommandQueue.cpp
+++ b/Engine/Source/Runtime/Apple/MetalRHI/Private/MetalCommandQueue.cpp
@@ -229,9 +229,10 @@ FMetalCommandQueue::FMetalCommandQueue(MTL::Device* InDevice, uint32 const MaxNu
 	}
     
 	// Temporarily only support heaps for devices with unified memory
+	// Disable this by default code while we work on metal heaps
 	if (!DeviceName.Contains(TEXT("Intel")) &&
           Device->hasUnifiedMemory() &&
-          !FParse::Param(FCommandLine::Get(),TEXT("nometalheap")))
+          FParse::Param(FCommandLine::Get(),TEXT("metalheap")))
 	{
 		Features |= EMetalFeaturesHeaps;
 	}
