Support zenserver for loading metallibs by passing binary blobs to MTLLibrary newLibrary()

diff --git a/Engine/Source/Runtime/Apple/MetalRHI/Private/MetalDynamicRHI_Shaders.cpp b/Engine/Source/Runtime/Apple/MetalRHI/Private/MetalDynamicRHI_Shaders.cpp
--- a/Engine/Source/Runtime/Apple/MetalRHI/Private/MetalDynamicRHI_Shaders.cpp
+++ b/Engine/Source/Runtime/Apple/MetalRHI/Private/MetalDynamicRHI_Shaders.cpp
@@ -136,14 +136,50 @@ FRHIShaderLibraryRef FMetalDynamicRHI::RHICreateShaderLibrary(EShaderPlatform Pl
             for (uint32 i = 0; i < Header.NumLibraries; i++)
             {
                 FString MetalLibraryFilePath = (FilePath / LibName) + FString::Printf(TEXT(".%d.metallib"), i);
-                MetalLibraryFilePath = IFileManager::Get().ConvertToAbsolutePathForExternalAppForRead(*MetalLibraryFilePath);
+                FString MetalLibraryAbsoluteFilePath = IFileManager::Get().ConvertToAbsolutePathForExternalAppForRead(*MetalLibraryFilePath);
 
-                METAL_GPUPROFILE(FScopedMetalCPUStats CPUStat(FString::Printf(TEXT("NewLibraryFile: %s"), *MetalLibraryFilePath)));
+                METAL_GPUPROFILE(FScopedMetalCPUStats CPUStat(FString::Printf(TEXT("NewLibraryFile: %s"), *MetalLibraryAbsoluteFilePath)));
                 NS::Error* Error;
-                NS::String* MetalLibraryFilePathNSString = FStringToNSString(MetalLibraryFilePath);
+                NS::String* MetalLibraryFilePathNSString = FStringToNSString(MetalLibraryAbsoluteFilePath);
                 NS::URL *metalLibraryURL = NS::URL::fileURLWithPath(MetalLibraryFilePathNSString);
                 MTLLibraryPtr Library = NS::TransferPtr(GetMetalDeviceContext().GetDevice()->newLibrary(metalLibraryURL, &Error));
 
+                if (Library.get() == nullptr)
+                {
+                    // Metallib not found. Is Zen server being used?
+                    static const bool bRunningWithZenStore = FPlatformFileManager::Get().FindPlatformFile(TEXT("StorageServer")) != nullptr;
+                    if( bRunningWithZenStore )
+                    {
+                        FArchive* Reader = IFileManager::Get().CreateFileReader(*MetalLibraryFilePath);
+                        TArray<uint8> LibraryData;
+                        bool Success = false;
+                        if (Reader)
+                        {
+                            const int64 FileSize = Reader->TotalSize();
+                            LibraryData.Reset( FileSize + 2 );
+                            LibraryData.AddUninitialized( FileSize );
+                            Reader->Serialize(LibraryData.GetData(), LibraryData.Num());
+                            Success = Reader->Close();
+                            delete Reader;
+                        }
+                        if (Success)
+                        {
+                            dispatch_data_t data = dispatch_data_create(LibraryData.GetData(), LibraryData.Num(), nil, DISPATCH_DATA_DESTRUCTOR_DEFAULT);
+                            Library = NS::TransferPtr(GetMetalDeviceContext().GetDevice()->newLibrary(data, &Error));
+                        }
+                        else
+                        {
+                            UE_LOG(LogMetal, Warning, TEXT("Metallib '%s' unable to be read from ZenStore."), *MetalLibraryFilePath);
+                            return nullptr;
+                        }
+                    }
+                    else
+                    {
+                        UE_LOG(LogMetal, Warning, TEXT("Metallib '%s' not found and ZenStore is not being used."), *MetalLibraryFilePath);
+                        return nullptr;
+                    }
+                }
+
                 if (Library.get() != nullptr)
                 {
                     Libraries.Add(Library);
