- MSAA depth working on VisionOS. Still in a state of investigation for best possible method, so there is a cvar to control different modes, not all of which work yet
- Still to investigate (in no particular order):
  * the override store actions
  * bRenderToSceneColor - does it need to be 1 on VisionOS?
  * Can we do both eyes in one pass so we don't need to store everything out just to turn around and load for right eye? Left eye goes black without storing everything out

diff --git a/Engine/Config/BaseDeviceProfiles.ini b/Engine/Config/BaseDeviceProfiles.ini
--- a/Engine/Config/BaseDeviceProfiles.ini
+++ b/Engine/Config/BaseDeviceProfiles.ini
@@ -703,6 +703,7 @@ BaseProfileName=AppleTV4K
 [VisionPro DeviceProfile]
 DeviceType=VisionOS
 BaseProfileName=IOS
++CVars=r.Mobile.XRMSAAMode=2
 
 [VisionProSimulator DeviceProfile]
 DeviceType=VisionOS
diff --git a/Engine/Platforms/VisionOS/Plugins/Runtime/OpenXRVisionOS/Source/OXRVisionOS/Private/OXRVisionOSSwapchain.cpp b/Engine/Platforms/VisionOS/Plugins/Runtime/OpenXRVisionOS/Source/OXRVisionOS/Private/OXRVisionOSSwapchain.cpp
--- a/Engine/Platforms/VisionOS/Plugins/Runtime/OpenXRVisionOS/Source/OXRVisionOS/Private/OXRVisionOSSwapchain.cpp
+++ b/Engine/Platforms/VisionOS/Plugins/Runtime/OpenXRVisionOS/Source/OXRVisionOS/Private/OXRVisionOSSwapchain.cpp
@@ -24,6 +24,18 @@
 // uncomment to make intellisense work better
 //#include "../../../../../../../../Source/ThirdParty/OpenXR/include/openxr/openxr_platform.h"
 
+static int GetMSAASampleCount()
+{
+	static int Mode = IConsoleManager::Get().FindTConsoleVariableDataInt(TEXT("r.Mobile.XRMSAAMode"))->GetValueOnAnyThread();
+	// mode 2 is to use an MSAA target as the swapchain image (depth only for now) and then we copy it's resolved target to the compositor
+	if (Mode == 2)
+	{
+		static int MSAACount = IConsoleManager::Get().FindTConsoleVariableDataInt(TEXT("r.MSAACount"))->GetValueOnAnyThread();
+		return MSAACount;
+	}
+	return 1;
+}
+
 FOXRVisionOSSwapchain::FSwapchainImage::FSwapchainImage(FTextureRHIRef InImage, EImageState InImageState)
 	: Image(InImage)
 	, ImageState(InImageState)
@@ -76,14 +88,16 @@ FOXRVisionOSSwapchain::FOXRVisionOSSwapchain(const XrSwapchainCreateInfo* create
 		ETextureCreateFlags TextureCreateFlags = ETextureCreateFlags::None;
 		TextureCreateFlags |= TexCreate_RenderTargetable;
 		TextureCreateFlags |= TexCreate_ShaderResource;
-		if (isColor) 
+		uint8 NumSamples = 1;
+		if (isColor)
 		{
 			TextureCreateFlags |= TexCreate_UAV;
 			TextureCreateFlags |= TexCreate_SRGB;
 		} 
 		else
 		{
-			if ((EPixelFormat)CreateInfo.format == PF_DepthStencil) 
+			NumSamples = GetMSAASampleCount();
+			if ((EPixelFormat)CreateInfo.format == PF_DepthStencil)
 			{
 				TextureCreateFlags |= TexCreate_DepthStencilTargetable;
 				TextureCreateFlags |= TexCreate_DepthStencilResolveTarget;
@@ -105,12 +119,19 @@ FOXRVisionOSSwapchain::FOXRVisionOSSwapchain(const XrSwapchainCreateInfo* create
 											CreateInfo.width, CreateInfo.height,
 											(EPixelFormat)CreateInfo.format)
 			.SetFlags(TextureCreateFlags)
+			.SetNumSamples(NumSamples)
 			.DetermineInititialState();
 		
 		Images.Reserve(SwapchainLength);
 		for (int32 BufferIndex = 0; BufferIndex < SwapchainLength; ++BufferIndex)
 		{
-			Images.Emplace(RHICreateTexture(BackBufferDesc), EImageState::Released);
+			FTextureRHIRef Tex = RHICreateTexture(BackBufferDesc);
+			Images.Emplace(Tex, EImageState::Released);
+
+			// hacky way to set the debug label for these textures- this can go when everything is working, but it's helpful for debugging in the capture
+			FString Name = FString::Printf(TEXT("%s_SwapChain%d"), (CreateInfo.usageFlags & XR_SWAPCHAIN_USAGE_COLOR_ATTACHMENT_BIT) ? TEXT("Color") : TEXT("Depth"), BufferIndex);
+			FRHICommandListBase* Base = (FRHICommandListBase*)nullptr;
+			GDynamicRHI->RHIBindDebugLabelName(*Base, Tex, *Name);
 		}
 	}
 }
diff --git a/Engine/Source/Runtime/Apple/MetalRHI/Private/MetalPipeline.cpp b/Engine/Source/Runtime/Apple/MetalRHI/Private/MetalPipeline.cpp
--- a/Engine/Source/Runtime/Apple/MetalRHI/Private/MetalPipeline.cpp
+++ b/Engine/Source/Runtime/Apple/MetalRHI/Private/MetalPipeline.cpp
@@ -741,8 +741,8 @@ static bool ConfigureRenderPipelineDescriptor(TDescriptorType* RenderPipelineDes
 		RenderPipelineDesc->setStencilAttachmentPixelFormat((MTL::PixelFormat)GPixelFormats[PF_DepthStencil].PlatformFormat);
 	}
 	
-	static bool bNoMSAA = FParse::Param(FCommandLine::Get(), TEXT("nomsaa"));
-	uint16 NumSamples = !bNoMSAA ? FMath::Max(Init.NumSamples, (uint16)1u) : (uint16)1u;
+	static bool bAllowMSAA = AllowMSAA();
+	uint16 NumSamples = bAllowMSAA ? FMath::Max(Init.NumSamples, (uint16)1u) : (uint16)1u;
 	if constexpr(bIsRenderPipelineDesc)
 	{
 		RenderPipelineDesc->setSampleCount(NumSamples);
diff --git a/Engine/Source/Runtime/Apple/MetalRHI/Private/MetalRHIPrivate.h b/Engine/Source/Runtime/Apple/MetalRHI/Private/MetalRHIPrivate.h
--- a/Engine/Source/Runtime/Apple/MetalRHI/Private/MetalRHIPrivate.h
+++ b/Engine/Source/Runtime/Apple/MetalRHI/Private/MetalRHIPrivate.h
@@ -314,3 +314,7 @@ FORCEINLINE NS::String* FStringToNSString(const FString& InputString)
 {
     return ((NS::String*)InputString.GetCFString())->autorelease();
 }
+
+// helper functions to reduce copy and pasted checks all around
+bool ShouldUseMemoryless(ETextureCreateFlags Flags);
+bool AllowMSAA();
diff --git a/Engine/Source/Runtime/Apple/MetalRHI/Private/MetalStateCache.cpp b/Engine/Source/Runtime/Apple/MetalRHI/Private/MetalStateCache.cpp
--- a/Engine/Source/Runtime/Apple/MetalRHI/Private/MetalStateCache.cpp
+++ b/Engine/Source/Runtime/Apple/MetalRHI/Private/MetalStateCache.cpp
@@ -52,6 +52,23 @@ static MTL::DepthClipMode TranslateDepthClipMode(ERasterizerDepthClipMode DepthC
 	}
 }
 
+static MTL::StoreAction ConditionalOverrideStoreAction(MTL::StoreAction StoreAction, bool bIsDepth, bool bIsMSAA)
+{
+	if (!bIsMSAA)
+	{
+		return StoreAction;
+	}
+	
+	static int Mode = IConsoleManager::Get().FindTConsoleVariableDataInt(TEXT("r.Mobile.XRMSAAMode"))->GetValueOnAnyThread();
+	if (Mode != 0)
+	{
+		// @todo visionos: this needs to be experimented on to find best mode
+		return MTL::StoreActionStoreAndMultisampleResolve;
+	}
+	
+	return StoreAction;
+}
+
 FORCEINLINE MTL::StoreAction GetMetalRTStoreAction(ERenderTargetStoreAction StoreAction)
 {
 	switch(StoreAction)
@@ -62,7 +79,7 @@ FORCEINLINE MTL::StoreAction GetMetalRTStoreAction(ERenderTargetStoreAction Stor
         //because we may render to the same MSAA target twice in two separate passes.  BasePass, then some stuff, then translucency for example and we need to not lose the prior MSAA contents to do this properly.
 		case ERenderTargetStoreAction::EMultisampleResolve:
 		{
-            static bool bNoMSAA = FParse::Param(FCommandLine::Get(), TEXT("nomsaa"));
+            static bool bNoMSAA = !AllowMSAA();
 			static bool bSupportsMSAAStoreResolve = FMetalCommandQueue::SupportsFeature(EMetalFeaturesMSAAStoreAndResolve) && (GMaxRHIFeatureLevel >= ERHIFeatureLevel::SM5);
             if (bNoMSAA)
             {
@@ -653,8 +670,12 @@ bool FMetalStateCache::SetRenderPassInfo(FRHIRenderPassInfo const& InRenderTarge
 #endif
 					// set up an MSAA attachment
 					ColorAttachment->setTexture(Surface.MSAATexture.get());
+
 					NewColorStore[RenderTargetIndex] = GetMetalRTStoreAction(ERenderTargetStoreAction::EMultisampleResolve);
-					ColorAttachment->setStoreAction(!bMemoryless && GRHIDeviceId > 2 ? MTL::StoreActionUnknown : NewColorStore[RenderTargetIndex]);
+					MTL::StoreAction FinalStoreAction = !bMemoryless && GRHIDeviceId > 2 ? MTL::StoreActionUnknown : NewColorStore[RenderTargetIndex];
+					FinalStoreAction = ConditionalOverrideStoreAction(FinalStoreAction, false, !!Surface.MSAATexture);
+					ColorAttachment->setStoreAction(FinalStoreAction);
+					
 					ColorAttachment->setResolveTexture(Surface.MSAAResolveTexture ? Surface.MSAAResolveTexture.get() : Surface.Texture.get());
 					SampleCount = Surface.MSAATexture->sampleCount();
 				}
@@ -909,7 +930,9 @@ bool FMetalStateCache::SetRenderPassInfo(FRHIRenderPassInfo const& InRenderTarge
                 //needed to quiet the metal validation that runs when you end renderpass. (it requires some kind of 'resolve' for an msaa target)
 				//But with deferredstore we don't set the real one until submit time.
 				NewDepthStore = !Surface.MSAATexture || bSupportsMSAADepthResolve ? GetMetalRTStoreAction(HighLevelStoreAction) : MTL::StoreActionDontCare;
-				DepthAttachment->setStoreAction(!bDepthTextureMemoryless && Surface.MSAATexture && GRHIDeviceId > 2 ? MTL::StoreActionUnknown : NewDepthStore);
+				MTL::StoreAction FinalStoreAction = !bDepthTextureMemoryless && Surface.MSAATexture && GRHIDeviceId > 2 ? MTL::StoreActionUnknown : NewDepthStore;
+				FinalStoreAction = ConditionalOverrideStoreAction(FinalStoreAction, true, !!Surface.MSAATexture);
+				DepthAttachment->setStoreAction(FinalStoreAction);
 				DepthAttachment->setClearDepth(DepthClearValue);
 				check(SampleCount > 0);
 
diff --git a/Engine/Source/Runtime/Apple/MetalRHI/Private/MetalTexture.cpp b/Engine/Source/Runtime/Apple/MetalRHI/Private/MetalTexture.cpp
--- a/Engine/Source/Runtime/Apple/MetalRHI/Private/MetalTexture.cpp
+++ b/Engine/Source/Runtime/Apple/MetalRHI/Private/MetalTexture.cpp
@@ -42,6 +42,25 @@ FAutoConsoleVariableRef CVarMetalDisableIOSMemoryless(
 													  TEXT("If true, disabled the use of Memoryless textures on iOS"),
 													  ECVF_ReadOnly|ECVF_RenderThreadSafe);
 
+bool ShouldUseMemoryless(ETextureCreateFlags Flags)
+{
+	constexpr bool bIsMobile = PLATFORM_IOS;
+	const bool bWantsMemoryless = EnumHasAnyFlags(Flags, TexCreate_Memoryless);
+
+	constexpr bool bIsVisionOS = PLATFORM_VISIONOS;
+	const bool bIsDepth = EnumHasAnyFlags(Flags, TexCreate_DepthStencilTargetable | TexCreate_DepthStencilResolveTarget);
+	static int MSAAMode = IConsoleManager::Get().FindTConsoleVariableDataInt(TEXT("r.Mobile.XRMSAAMode"))->GetValueOnAnyThread();
+
+	// allow memoryless on mobile if not disabled
+	return bIsMobile && bWantsMemoryless && !GMetalDisableIOSMemoryless && !(bIsVisionOS && bIsDepth && MSAAMode == 1);
+}
+
+bool AllowMSAA()
+{
+	static const bool bAllowMSAA = !FParse::Param(FCommandLine::Get(), TEXT("nomsaa"));	
+	return bAllowMSAA;
+}
+
 /** Given a pointer to a RHI texture that was created by the Metal RHI, returns a pointer to the FMetalTextureBase it encapsulates. */
 FMetalSurface* GetMetalSurfaceFromRHITexture(FRHITexture* Texture)
 {
@@ -508,8 +527,7 @@ FMetalTextureCreateDesc::FMetalTextureCreateDesc(FRHITextureCreateDesc const& In
 #endif
 		}
 
-#if PLATFORM_IOS
-		if (!GMetalDisableIOSMemoryless && EnumHasAnyFlags(InDesc.Flags, TexCreate_Memoryless))
+		if (ShouldUseMemoryless(InDesc.Flags))
 		{
 			ensure(EnumHasAnyFlags(InDesc.Flags, (TexCreate_RenderTargetable | TexCreate_DepthStencilTargetable)));
 			ensure(!EnumHasAnyFlags(InDesc.Flags, (TexCreate_CPUReadback | TexCreate_CPUWritable)));
@@ -517,7 +535,6 @@ FMetalTextureCreateDesc::FMetalTextureCreateDesc(FRHITextureCreateDesc const& In
 			Desc->setStorageMode(MTL::StorageModeMemoryless);
 			Desc->setResourceOptions((MTL::ResourceOptions)(MTL::ResourceCPUCacheModeDefaultCache | MTL::ResourceStorageModeMemoryless));
 		}
-#endif
 
 		MTL::ResourceOptions HazardTrackingMode = MTL::ResourceHazardTrackingModeUntracked;
 		static bool bSupportsHeaps = GetMetalDeviceContext().SupportsFeature(EMetalFeaturesHeaps);
@@ -566,7 +583,7 @@ FMetalSurface::FMetalSurface(FRHICommandListBase* RHICmdList, FMetalTextureCreat
 		return;
 	}
     
-    bool bIsMSAARequired = CreateDesc.NumSamples > 1 && !FParse::Param(FCommandLine::Get(), TEXT("nomsaa"));
+    bool bIsMSAARequired = CreateDesc.NumSamples > 1 && AllowMSAA();
 
 	FResourceBulkDataInterface* BulkData = CreateDesc.BulkData;
 
@@ -627,9 +644,8 @@ FMetalSurface::FMetalSurface(FRHICommandListBase* RHICmdList, FMetalTextureCreat
 
         FMetalTextureCreateDesc NewCreateDesc = CreateDesc;
         
-#if PLATFORM_IOS
         // If we are attempting to create an MSAA texture the texture cannot be memoryless unless we are creating a depth texture
-        if(bIsMSAARequired && CreateDesc.Format != PF_DepthStencil && !GMetalDisableIOSMemoryless && EnumHasAllFlags(CreateDesc.Flags, TexCreate_Memoryless))
+        if(ShouldUseMemoryless(CreateDesc.Flags) && bIsMSAARequired && CreateDesc.Format != PF_DepthStencil)
         {
             NewCreateDesc.Flags &= ~TexCreate_Memoryless;
             
@@ -644,8 +660,8 @@ FMetalSurface::FMetalSurface(FRHICommandListBase* RHICmdList, FMetalTextureCreat
                 NewCreateDesc.Desc->setResourceOptions((MTL::ResourceOptions)(MTL::ResourceCPUCacheModeDefaultCache | MTL::ResourceStorageModePrivate));
             }
         }
-#endif
-        const bool bAtomicCompatible = EnumHasAllFlags(CreateDesc.Flags, TexCreate_AtomicCompatible) || EnumHasAllFlags(CreateDesc.Flags, ETextureCreateFlags::Atomic64Compatible);
+
+		const bool bAtomicCompatible = EnumHasAllFlags(CreateDesc.Flags, TexCreate_AtomicCompatible) || EnumHasAllFlags(CreateDesc.Flags, ETextureCreateFlags::Atomic64Compatible);
 		
 		bool bIsBindless = IsMetalBindlessEnabled();
 		
@@ -751,14 +767,12 @@ FMetalSurface::FMetalSurface(FRHICommandListBase* RHICmdList, FMetalTextureCreat
 
 		bool bMemoryless = false;
         
-#if PLATFORM_IOS
-		if (!GMetalDisableIOSMemoryless && EnumHasAllFlags(CreateDesc.Flags, TexCreate_Memoryless))
+		if (ShouldUseMemoryless(CreateDesc.Flags))
 		{
 			bMemoryless = true;
 			Desc->setStorageMode(MTL::StorageModeMemoryless);
 			Desc->setResourceOptions(MTL::ResourceStorageModeMemoryless);
 		}
-#endif
 
 		MSAATexture = GetMetalDeviceContext().CreateTexture(this, Desc.get());
 			
diff --git a/Engine/Source/Runtime/Apple/MetalRHI/Private/MetalViewport.cpp b/Engine/Source/Runtime/Apple/MetalRHI/Private/MetalViewport.cpp
--- a/Engine/Source/Runtime/Apple/MetalRHI/Private/MetalViewport.cpp
+++ b/Engine/Source/Runtime/Apple/MetalRHI/Private/MetalViewport.cpp
@@ -662,19 +662,7 @@ void FMetalViewport::PresentImmersive(const MetalRHIVisionOS::PresentImmersivePa
 					}
 
 					{
-						//HACK BEGIN
-						// using LastCompleteFrame, the color result, is very wrong, but the depth texture is not yet plumbed through
-						// on the mobile rendere and providing nothing results in an all black scene while this results in mostly-ok-pixels.
-						// With the deferred renderer depth does already work, so flip this to using depth.
-						static bool bUseRealDepth = false;
-						static bool bCheckedConfig = false;
-						if (bCheckedConfig == false) {
-							GConfig->GetBool(TEXT("/Script/IOSRuntimeSettings.IOSRuntimeSettings"), TEXT("bSupportsMetalMRT"), bUseRealDepth, GEngineIni);
-							bCheckedConfig = true;
-						}
-						TRefCountPtr<FMetalSurface> Texture = bUseRealDepth ? MyLastCompleteDepth : MyLastCompleteFrame;
-						//HACK END
-						
+						TRefCountPtr<FMetalSurface> Texture = MyLastCompleteDepth;
 						check(IsValidRef(Texture));
 						MTLTexturePtr Src = Texture->Texture;
 						MTLTexturePtr& Dst = DrawableDepthTexture;
@@ -685,7 +673,6 @@ void FMetalViewport::PresentImmersive(const MetalRHIVisionOS::PresentImmersivePa
 						RenderPass.CopyFromTextureToTexture(Src.get(), 0, 0, MTL::Origin(0, 0, 0), MTL::Size(Width, Height, 1), Dst.get(), 0, 0, MTL::Origin(0, 0, 0));
 					}
 				}
-
 				RenderPass.EndRenderPass();
 				RenderPass.Submit(EMetalSubmitFlagsCreateCommandBuffer);
 				RenderPass.EncodePresentImmersive(VisionOSParams.SwiftDrawable, VisionOSParams.SwiftFrame);
diff --git a/Engine/Source/Runtime/Renderer/Private/MobileShadingRenderer.cpp b/Engine/Source/Runtime/Renderer/Private/MobileShadingRenderer.cpp
--- a/Engine/Source/Runtime/Renderer/Private/MobileShadingRenderer.cpp
+++ b/Engine/Source/Runtime/Renderer/Private/MobileShadingRenderer.cpp
@@ -109,6 +109,17 @@ static TAutoConsoleVariable<int32> CVarMobileTonemapSubpass(
 	TEXT(" 1 = On"),
 	ECVF_Scalability | ECVF_RenderThreadSafe);
 
+
+static TAutoConsoleVariable<int32> CVarMobileXRMSAAMode(
+	TEXT("r.Mobile.XRMSAAMode"),
+	0,
+	TEXT(" Whether to modify how mobile XR msaa support works\n")
+	TEXT(" 0 = Standard depth pass/swapchain mode [default]\n")
+	TEXT(" 1 = Perform a copy of depth to the depth resolve target")
+	TEXT(" 2 = Make the depth swap chain be MSAA and use it directly as scene depth"),
+	ECVF_ReadOnly | ECVF_RenderThreadSafe);
+
+
 static bool IsMobileTonemapSubpassEnabled(const FStaticShaderPlatform Platform)
 {
 	static auto* MobileTonemapSubpassPathCvar = IConsoleManager::Get().FindTConsoleVariableDataInt(TEXT("r.Mobile.TonemapSubpass"));
@@ -1447,6 +1458,11 @@ void FMobileSceneRenderer::Render(FRDGBuilder& GraphBuilder)
 							{ 
 								AddMobilePostProcessingPasses(GraphBuilder, Scene, Views[ViewIndex], ViewIndex, GetSceneUniforms(), PostProcessingInputs, InstanceCullingManager);
 							}
+
+							if (CVarMobileXRMSAAMode.GetValueOnAnyThread() == 1)
+							{
+								AddDrawTexturePass(GraphBuilder, Views[ViewIndex], SceneTextures.Depth.Target, SceneTextures.Depth.Resolve);
+							}
 						}
 					}
 				}
@@ -1520,7 +1536,7 @@ FRenderTargetBindingSlots FMobileSceneRenderer::InitRenderTargetBindings_Forward
 	}
 	SceneDepth = SceneTextures.Depth.Target;
 	SceneDepthResolve = GRHISupportsDepthStencilResolve && bMobileMSAA && SceneTextures.Depth.IsSeparate() ? SceneTextures.Depth.Resolve : nullptr;
-
+	
 	FRenderTargetBindingSlots BasePassRenderTargets;
 	BasePassRenderTargets[0] = FRenderTargetBinding(SceneColor, SceneColorResolve, ERenderTargetLoadAction::EClear);
 	if (bRequiresSceneDepthAux)
diff --git a/Engine/Source/Runtime/Renderer/Private/SceneTextures.cpp b/Engine/Source/Runtime/Renderer/Private/SceneTextures.cpp
--- a/Engine/Source/Runtime/Renderer/Private/SceneTextures.cpp
+++ b/Engine/Source/Runtime/Renderer/Private/SceneTextures.cpp
@@ -109,14 +109,13 @@ static IStereoRenderTargetManager* FindStereoRenderTargetManager()
 	return GEngine->StereoRenderingDevice->GetRenderTargetManager();
 }
 
-static TRefCountPtr<FRHITexture> FindStereoDepthTexture(uint32 bSupportsXRDepth, FIntPoint TextureExtent, ETextureCreateFlags RequestedCreateFlags)
+static TRefCountPtr<FRHITexture> FindStereoDepthTexture(uint32 bSupportsXRDepth, FIntPoint TextureExtent, ETextureCreateFlags RequestedCreateFlags, uint8 NumSamples)
 {
 	if (bSupportsXRDepth == 1)
 	{
 		if (IStereoRenderTargetManager* StereoRenderTargetManager = FindStereoRenderTargetManager())
 		{
 			TRefCountPtr<FRHITexture> DepthTex, SRTex;
-			constexpr uint32 NumSamples = 1;
 			StereoRenderTargetManager->AllocateDepthTexture(0, TextureExtent.X, TextureExtent.Y, PF_DepthStencil, 1, RequestedCreateFlags, TexCreate_DepthStencilTargetable | TexCreate_ShaderResource | TexCreate_InputAttachmentRead, DepthTex, SRTex, NumSamples);
 			return MoveTemp(SRTex);
 		}
@@ -407,6 +406,12 @@ void InitializeSceneTexturesConfig(FSceneTexturesConfig& Config, const FSceneVie
 	Config.Init(SceneTexturesConfigInitSettings);
 }
 
+static bool UseMSAAStereoDepthTextureDirectly()
+{
+	static int Mode = IConsoleManager::Get().FindTConsoleVariableDataInt(TEXT("r.Mobile.XRMSAAMode"))->GetValueOnAnyThread();
+	return Mode == 2;
+}
+
 void FMinimalSceneTextures::InitializeViewFamily(FRDGBuilder& GraphBuilder, FViewFamilyInfo& ViewFamily)
 {
 	checkf(ViewFamily.SceneTextures->Owner == &ViewFamily, TEXT("Scene Textures should only be initialized by their owning view family -- possible duplicate initialization"));
@@ -422,7 +427,9 @@ void FMinimalSceneTextures::InitializeViewFamily(FRDGBuilder& GraphBuilder, FVie
 
 	// If not using MSAA, we need to make sure to grab the stereo depth texture if appropriate.
 	FTextureRHIRef StereoDepthRHI;
-	if (Config.NumSamples == 1 && (StereoDepthRHI = FindStereoDepthTexture(Config.bSupportsXRTargetManagerDepthAlloc, Config.Extent, ETextureCreateFlags::None)) != nullptr)
+	bool bUseDepthTextureDirectly = Config.NumSamples == 1 || UseMSAAStereoDepthTextureDirectly();
+	if (bUseDepthTextureDirectly && (StereoDepthRHI = FindStereoDepthTexture(Config.bSupportsXRTargetManagerDepthAlloc, Config.Extent, ETextureCreateFlags::None, Config.NumSamples)) != nullptr)
+
 	{
 		SceneTextures.Depth = RegisterExternalTexture(GraphBuilder, StereoDepthRHI, TEXT("SceneDepthZ"));
 		SceneTextures.Stencil = GraphBuilder.CreateSRV(FRDGTextureSRVDesc::CreateWithPixelFormat(SceneTextures.Depth.Target, PF_X24_G8));
@@ -440,7 +447,7 @@ void FMinimalSceneTextures::InitializeViewFamily(FRDGBuilder& GraphBuilder, FVie
 		{
 			Desc.NumSamples = 1;
 
-			if ((StereoDepthRHI = FindStereoDepthTexture(Config.bSupportsXRTargetManagerDepthAlloc, Config.Extent, ETextureCreateFlags::DepthStencilResolveTarget)) != nullptr)
+			if ((StereoDepthRHI = FindStereoDepthTexture(Config.bSupportsXRTargetManagerDepthAlloc, Config.Extent, ETextureCreateFlags::DepthStencilResolveTarget, Desc.NumSamples)) != nullptr)
 			{
 				ensureMsgf(Desc.ArraySize == StereoDepthRHI->GetDesc().ArraySize, TEXT("Resolve texture does not agree in dimensionality with Target (Resolve.ArraySize=%d, Target.ArraySize=%d)"),
 					Desc.ArraySize, StereoDepthRHI->GetDesc().ArraySize);
