Fixed MetalGPUTiming not calibrating timestamp

diff --git a/Engine/Source/Runtime/Apple/MetalRHI/Private/MetalProfiler.cpp b/Engine/Source/Runtime/Apple/MetalRHI/Private/MetalProfiler.cpp
--- a/Engine/Source/Runtime/Apple/MetalRHI/Private/MetalProfiler.cpp
+++ b/Engine/Source/Runtime/Apple/MetalRHI/Private/MetalProfiler.cpp
@@ -68,6 +68,24 @@ void WriteString(FArchive* OutputFile, const char* String)
 	OutputFile->Serialize((void*)String, sizeof(ANSICHAR)*FCStringAnsi::Strlen(String));
 }
 
+void FMetalGPUTiming::PlatformStaticInitialize(void* UserData)
+{
+	// Are the static variables initialized?
+	if ( !GAreGlobalsInitialized )
+	{
+		GIsSupported = true;
+		SetTimingFrequency(1000 * 1000 * 1000);
+		GAreGlobalsInitialized = true;
+		
+		MTL::Device* MTLDevice = ((FMetalContext*)UserData)->GetDevice();
+		
+		MTL::Timestamp CPUTimeStamp, GPUTimestamp;
+		MTLDevice->sampleTimestamps(&CPUTimeStamp, &GPUTimestamp);
+
+		FGPUTiming::SetCalibrationTimestamp({ GPUTimestamp, CPUTimeStamp });
+	}
+}
+
 FMetalEventNode::~FMetalEventNode()
 {
     
diff --git a/Engine/Source/Runtime/Apple/MetalRHI/Private/MetalProfiler.h b/Engine/Source/Runtime/Apple/MetalRHI/Private/MetalProfiler.h
--- a/Engine/Source/Runtime/Apple/MetalRHI/Private/MetalProfiler.h
+++ b/Engine/Source/Runtime/Apple/MetalRHI/Private/MetalProfiler.h
@@ -152,9 +152,9 @@ public:
 	/**
 	 * Constructor.
 	 */
-	FMetalGPUTiming()
+	FMetalGPUTiming(FMetalContext* Context)
 	{
-		StaticInitialize(nullptr, PlatformStaticInitialize);
+		StaticInitialize((void*)Context, PlatformStaticInitialize);
 	}
 	
 	void SetCalibrationTimestamp(uint64 GPU, uint64 CPU)
@@ -167,16 +167,7 @@ private:
 	/**
 	 * Initializes the static variables, if necessary.
 	 */
-	static void PlatformStaticInitialize(void* UserData)
-	{
-		// Are the static variables initialized?
-		if ( !GAreGlobalsInitialized )
-		{
-			GIsSupported = true;
-			SetTimingFrequency(1000 * 1000 * 1000);
-			GAreGlobalsInitialized = true;
-		}
-	}
+	static void PlatformStaticInitialize(void* UserData);
 };
 
 struct IMetalStatsScope
@@ -274,6 +265,7 @@ struct FMetalGPUProfiler : public FGPUProfiler
 	
 	FMetalGPUProfiler(FMetalContext* InContext)
 	:	FGPUProfiler()
+	,	TimingSupport(InContext)
 	,	Context(InContext)
 	,   NumNestedFrames(0)
 	{}
