Change rhi.Metal.CacheShaderPipeline to exclude all internal caching, not just precached when set to 2

diff --git a/Engine/Source/Runtime/Apple/MetalRHI/Private/MetalPipeline.cpp b/Engine/Source/Runtime/Apple/MetalRHI/Private/MetalPipeline.cpp
--- a/Engine/Source/Runtime/Apple/MetalRHI/Private/MetalPipeline.cpp
+++ b/Engine/Source/Runtime/Apple/MetalRHI/Private/MetalPipeline.cpp
@@ -24,7 +24,7 @@ static FAutoConsoleVariableRef CVarMetalCacheShaderPipelines(
 	TEXT("rhi.Metal.CacheShaderPipelines"),
 	GMetalCacheShaderPipelines,
 	TEXT("When enabled (1, default) cache all graphics pipeline state objects created in MetalRHI for the life of the program, this trades memory for performance as creating PSOs is expensive in Metal.\n")
-	TEXT("When set to 2 it will cache all graphics pipeline state objects(that are not requested by PSO precaching) created in MetalRHI for the life of the program, this trades memory for performance as creating PSOs is expensive in Metal.\n")
+	TEXT("When set to 2 it will not cache any graphics pipeline state objects\n")
 	TEXT("Disable in the project configuration to allow PSOs to be released to save memory at the expense of reduced performance and increased hitching in-game\n. (On by default (1))"), ECVF_ReadOnly);
 
 static int32 GMetalCacheMinSize = 32;
@@ -270,8 +270,8 @@ public:
 			{
 				Desc = CreateMTLRenderPipeline(bSync, Key, Init, State);
 
-				// Don't cache PSOPrecache if cvar is set to exclude them
-				if (Desc != nullptr && (!Init.bPSOPrecache || GMetalCacheShaderPipelines != 2))
+				// Don't cache anything if the cvar is set to 2
+				if (Desc != nullptr &&  GMetalCacheShaderPipelines != 2)
 				{
 					PipelineMutex.WriteLock();
 
