Fix when rhi.Metal.CacheShaderPipelines is set to 2

diff --git a/Engine/Source/Runtime/Apple/MetalRHI/Private/MetalPipeline.cpp b/Engine/Source/Runtime/Apple/MetalRHI/Private/MetalPipeline.cpp
--- a/Engine/Source/Runtime/Apple/MetalRHI/Private/MetalPipeline.cpp
+++ b/Engine/Source/Runtime/Apple/MetalRHI/Private/MetalPipeline.cpp
@@ -230,76 +230,83 @@ public:
 		FMetalGraphicsPipelineKey Key;
 		InitMetalGraphicsPipelineKey(Key, Init);
 		
-		// By default there'll be more threads trying to read this than to write it.
-		PipelineMutex.ReadLock();
-
-		// Try to find the entry in the cache.
-        FMetalShaderPipelinePtr Desc = Pipelines.FindRef(Key);
-
-		PipelineMutex.ReadUnlock();
-
-		if (Desc == nullptr)
+		FMetalShaderPipelinePtr Desc = nullptr;
+		// Don't cache anything if the cvar is set to 2
+		if (GMetalCacheShaderPipelines != 2)
 		{
 			// By default there'll be more threads trying to read this than to write it.
-			EventsMutex.ReadLock();
+			PipelineMutex.ReadLock();
 
-			// Try to find a pipeline creation event for this key. If it's found, we already have a thread creating this pipeline and we just have to wait.
-			TSharedPtr<FPThreadEvent, ESPMode::ThreadSafe> Event = PipelineEvents.FindRef(Key);
+			// Try to find the entry in the cache.
+			Desc = Pipelines.FindRef(Key);
 
-			EventsMutex.ReadUnlock();
+			PipelineMutex.ReadUnlock();
 
-			bool bCompile = false;
-			if (!Event.IsValid())
+			if (Desc == nullptr)
 			{
-				// Create an event other threads can use to wait if they request the same pipeline this thread is creating
-				EventsMutex.WriteLock();
+				// By default there'll be more threads trying to read this than to write it.
+				EventsMutex.ReadLock();
+
+				// Try to find a pipeline creation event for this key. If it's found, we already have a thread creating this pipeline and we just have to wait.
+				TSharedPtr<FPThreadEvent, ESPMode::ThreadSafe> Event = PipelineEvents.FindRef(Key);
 
-				Event = PipelineEvents.FindRef(Key);
+				EventsMutex.ReadUnlock();
+
+				bool bCompile = false;
 				if (!Event.IsValid())
 				{
-					Event = PipelineEvents.Add(Key, MakeShareable(new FPThreadEvent()));
-					Event->Create(true);
-					bCompile = true;
-				}
-				check(Event.IsValid());
+					// Create an event other threads can use to wait if they request the same pipeline this thread is creating
+					EventsMutex.WriteLock();
 
-				EventsMutex.WriteUnlock();
-			}
+					Event = PipelineEvents.FindRef(Key);
+					if (!Event.IsValid())
+					{
+						Event = PipelineEvents.Add(Key, MakeShareable(new FPThreadEvent()));
+						Event->Create(true);
+						bCompile = true;
+					}
+					check(Event.IsValid());
 
-			if (bCompile)
-			{
-				Desc = CreateMTLRenderPipeline(bSync, Key, Init, State);
+					EventsMutex.WriteUnlock();
+				}
 
-				// Don't cache anything if the cvar is set to 2
-				if (Desc != nullptr &&  GMetalCacheShaderPipelines != 2)
+				if (bCompile)
 				{
-					PipelineMutex.WriteLock();
+					Desc = CreateMTLRenderPipeline(bSync, Key, Init, State);
 
-					Pipelines.Add(Key, Desc);
-					ReverseLookup.Add(Desc, Key);
+					if (Desc != nullptr)
+					{
+						PipelineMutex.WriteLock();
 
-					PipelineMutex.WriteUnlock();
-				}
+						Pipelines.Add(Key, Desc);
+						ReverseLookup.Add(Desc, Key);
 
-				EventsMutex.WriteLock();
+						PipelineMutex.WriteUnlock();
+					}
 
-				Event->Trigger();
-				PipelineEvents.Remove(Key);
+					EventsMutex.WriteLock();
 
-				EventsMutex.WriteUnlock();
-			}
-			else
-			{
-				check(Event.IsValid());
-				Event->Wait();
+					Event->Trigger();
+					PipelineEvents.Remove(Key);
 
-				PipelineMutex.ReadLock();
-				Desc = Pipelines.FindRef(Key);
-				PipelineMutex.ReadUnlock();
-				check(Desc);
+					EventsMutex.WriteUnlock();
+				}
+				else
+				{
+					check(Event.IsValid());
+					Event->Wait();
+
+					PipelineMutex.ReadLock();
+					Desc = Pipelines.FindRef(Key);
+					PipelineMutex.ReadUnlock();
+					check(Desc);
+				}
 			}
 		}
-		
+		else
+		{
+			Desc = CreateMTLRenderPipeline(bSync, Key, Init, State);
+		}
 		return Desc;
 	}
 	
