Add PF_R16G16_SINT format.

This pixel format is required for (nvidia) optical flow tests & use cases.

diff --git a/Engine/Source/Runtime/Apple/MetalRHI/Private/MetalRHI.cpp b/Engine/Source/Runtime/Apple/MetalRHI/Private/MetalRHI.cpp
--- a/Engine/Source/Runtime/Apple/MetalRHI/Private/MetalRHI.cpp
+++ b/Engine/Source/Runtime/Apple/MetalRHI/Private/MetalRHI.cpp
@@ -773,6 +773,7 @@ FMetalDynamicRHI::FMetalDynamicRHI(ERHIFeatureLevel::Type RequestedFeatureLevel)
 	GMetalBufferFormats[PF_ETC2_RGBA			] = { MTL::PixelFormatInvalid, (uint8)EMetalBufferFormat::Unknown };
 	GMetalBufferFormats[PF_R32G32B32A32_UINT	] = { MTL::PixelFormatRGBA32Uint, (uint8)EMetalBufferFormat::RGBA32Uint };
 	GMetalBufferFormats[PF_R16G16_UINT			] = { MTL::PixelFormatRG16Uint, (uint8)EMetalBufferFormat::RG16Uint };
+	GMetalBufferFormats[PF_R16G16_SINT			] = { MTL::PixelFormatRG16Sint, (uint8)EMetalBufferFormat::RG16Sint };
 	GMetalBufferFormats[PF_R32G32_UINT			] = { MTL::PixelFormatRG32Uint, (uint8)EMetalBufferFormat::RG32Uint };
 	GMetalBufferFormats[PF_ASTC_4x4             ] = { MTL::PixelFormatInvalid, (uint8)EMetalBufferFormat::Unknown };
 	GMetalBufferFormats[PF_ASTC_6x6             ] = { MTL::PixelFormatInvalid, (uint8)EMetalBufferFormat::Unknown };
@@ -826,6 +827,7 @@ FMetalDynamicRHI::FMetalDynamicRHI(ERHIFeatureLevel::Type RequestedFeatureLevel)
 	GPixelFormats[PF_G16				].PlatformFormat	= (uint32)MTL::PixelFormatR16Unorm;
 	GPixelFormats[PF_R32G32B32A32_UINT	].PlatformFormat	= (uint32)MTL::PixelFormatRGBA32Uint;
 	GPixelFormats[PF_R16G16_UINT		].PlatformFormat	= (uint32)MTL::PixelFormatRG16Uint;
+	GPixelFormats[PF_R16G16_SINT		].PlatformFormat	= (uint32)MTL::PixelFormatRG16Sint;
 	GPixelFormats[PF_R32G32_UINT		].PlatformFormat	= (uint32)MTL::PixelFormatRG32Uint;
 
 #if PLATFORM_IOS
diff --git a/Engine/Source/Runtime/Core/Private/Misc/PixelFormat.cpp b/Engine/Source/Runtime/Core/Private/Misc/PixelFormat.cpp
--- a/Engine/Source/Runtime/Core/Private/Misc/PixelFormat.cpp
+++ b/Engine/Source/Runtime/Core/Private/Misc/PixelFormat.cpp
@@ -144,6 +144,8 @@ FPixelFormatInfo    GPixelFormats[PF_MAX] =
 	FPixelFormatInfo(PF_ASTC_8x8_NORM_RG,   TEXT("ASTC_8x8_NORM_RG"),      8,          8,          1,          16,         2,              0),
 	FPixelFormatInfo(PF_ASTC_10x10_NORM_RG, TEXT("ASTC_10x10_NORM_RG"),    10,         10,         1,          16,         2,              0),
 	FPixelFormatInfo(PF_ASTC_12x12_NORM_RG, TEXT("ASTC_12x12_NORM_RG"),    12,         12,         1,          16,         2,              0),
+
+	FPixelFormatInfo(PF_R16G16_SINT,        TEXT("PF_R16G16_SINT"),        1,          1,          1,          4,          4,              1),
 };
 
 
diff --git a/Engine/Source/Runtime/Core/Public/PixelFormat.h b/Engine/Source/Runtime/Core/Public/PixelFormat.h
--- a/Engine/Source/Runtime/Core/Public/PixelFormat.h
+++ b/Engine/Source/Runtime/Core/Public/PixelFormat.h
@@ -106,7 +106,8 @@ enum EPixelFormat : uint8
 	PF_ASTC_8x8_NORM_RG		=89,	
 	PF_ASTC_10x10_NORM_RG	=90,	
 	PF_ASTC_12x12_NORM_RG	=91,	
-	PF_MAX					=92,
+	PF_R16G16_SINT			=92,
+	PF_MAX					=93,
 };
 #define FOREACH_ENUM_EPIXELFORMAT(op) \
 	op(PF_Unknown) \
@@ -200,7 +201,8 @@ enum EPixelFormat : uint8
 	op(PF_ASTC_6x6_NORM_RG) \
 	op(PF_ASTC_8x8_NORM_RG) \
 	op(PF_ASTC_10x10_NORM_RG) \
-	op(PF_ASTC_12x12_NORM_RG)
+	op(PF_ASTC_12x12_NORM_RG) \
+	op(PF_R16G16_SINT)
 
 // Defines which channel is valid for each pixel format
 enum class EPixelFormatChannelFlags : uint8
@@ -344,6 +346,7 @@ FORCEINLINE bool IsInteger(EPixelFormat PixelFormat)
 	case PF_R32G32B32_SINT:
 	case PF_R8_SINT:
 	case PF_R64_UINT:
+	case PF_R16G16_SINT:
 		return true;
 	}
 	return false;
diff --git a/Engine/Source/Runtime/CoreUObject/Public/UObject/NoExportTypes.h b/Engine/Source/Runtime/CoreUObject/Public/UObject/NoExportTypes.h
--- a/Engine/Source/Runtime/CoreUObject/Public/UObject/NoExportTypes.h
+++ b/Engine/Source/Runtime/CoreUObject/Public/UObject/NoExportTypes.h
@@ -302,6 +302,7 @@ enum EPixelFormat : int
 	PF_ASTC_8x8_NORM_RG,
 	PF_ASTC_10x10_NORM_RG,
 	PF_ASTC_12x12_NORM_RG,
+	PF_R16G16_SINT,
 	PF_MAX,
 };
 
diff --git a/Engine/Source/Runtime/D3D12RHI/Private/D3D12RHI.cpp b/Engine/Source/Runtime/D3D12RHI/Private/D3D12RHI.cpp
--- a/Engine/Source/Runtime/D3D12RHI/Private/D3D12RHI.cpp
+++ b/Engine/Source/Runtime/D3D12RHI/Private/D3D12RHI.cpp
@@ -168,6 +168,7 @@ FD3D12DynamicRHI::FD3D12DynamicRHI(const TArray<TSharedPtr<FD3D12Adapter>>& Chos
 	GPixelFormats[PF_R8G8			].PlatformFormat = DXGI_FORMAT_R8G8_UNORM;
 	GPixelFormats[PF_R32G32B32A32_UINT].PlatformFormat = DXGI_FORMAT_R32G32B32A32_UINT;
 	GPixelFormats[PF_R16G16_UINT	].PlatformFormat = DXGI_FORMAT_R16G16_UINT;
+	GPixelFormats[PF_R16G16_SINT	].PlatformFormat = DXGI_FORMAT_R16G16_SINT;
 	GPixelFormats[PF_R32G32_UINT	].PlatformFormat = DXGI_FORMAT_R32G32_UINT;
 
 	GPixelFormats[PF_BC6H			].PlatformFormat = DXGI_FORMAT_BC6H_UF16;
diff --git a/Engine/Source/Runtime/OpenGLDrv/Private/OpenGLDevice.cpp b/Engine/Source/Runtime/OpenGLDrv/Private/OpenGLDevice.cpp
--- a/Engine/Source/Runtime/OpenGLDrv/Private/OpenGLDevice.cpp
+++ b/Engine/Source/Runtime/OpenGLDrv/Private/OpenGLDevice.cpp
@@ -1163,6 +1163,7 @@ static void InitRHICapabilitiesForGL()
 	SetupTextureFormat( PF_R16G16B16A16_SNORM,	FOpenGLTextureFormat( GL_RGBA16,				GL_RGBA16,				GL_RGBA,			GL_SHORT,						false,			false));
 	
 	SetupTextureFormat( PF_R16G16_UINT,			FOpenGLTextureFormat( GL_RG16UI,				GL_RG16UI,				GL_RG_INTEGER,		GL_UNSIGNED_SHORT,				false,			false));
+	SetupTextureFormat( PF_R16G16_SINT,			FOpenGLTextureFormat( GL_RG16I,					GL_RG16I,				GL_RG_INTEGER,		GL_SHORT,						false,			false));
 	SetupTextureFormat( PF_R8,					FOpenGLTextureFormat( GL_R8,					GL_R8,					GL_RED,				GL_UNSIGNED_BYTE,				false,			false));
 
 	SetupTextureFormat( PF_R5G6B5_UNORM,        FOpenGLTextureFormat( GL_RGB565,                GL_RGB565,              GL_RGB,             GL_UNSIGNED_SHORT_5_6_5,        false,          false));
diff --git a/Engine/Source/Runtime/RenderCore/Private/RenderUtils.cpp b/Engine/Source/Runtime/RenderCore/Private/RenderUtils.cpp
--- a/Engine/Source/Runtime/RenderCore/Private/RenderUtils.cpp
+++ b/Engine/Source/Runtime/RenderCore/Private/RenderUtils.cpp
@@ -329,6 +329,7 @@ EPixelFormatChannelFlags GetPixelFormatValidChannels(EPixelFormat InPixelFormat)
 		EPixelFormatChannelFlags::RG,		// PF_ASTC_8x8_NORM_RG
 		EPixelFormatChannelFlags::RG,		// PF_ASTC_10x10_NORM_RG
 		EPixelFormatChannelFlags::RG,		// PF_ASTC_12x12_NORM_RG
+		EPixelFormatChannelFlags::RG,		// PF_R16G16_SINT
 	};
 	static_assert(UE_ARRAY_COUNT(PixelFormatToChannelFlags) == (uint8)PF_MAX, "Missing pixel format");
 	return (InPixelFormat < PF_MAX) ? PixelFormatToChannelFlags[(uint8)InPixelFormat] : EPixelFormatChannelFlags::None;
diff --git a/Engine/Source/Runtime/Renderer/Private/SystemTextures.cpp b/Engine/Source/Runtime/Renderer/Private/SystemTextures.cpp
--- a/Engine/Source/Runtime/Renderer/Private/SystemTextures.cpp
+++ b/Engine/Source/Runtime/Renderer/Private/SystemTextures.cpp
@@ -1257,7 +1257,7 @@ template<typename TInType>
 void InitializeData(const TInType& InData, EPixelFormat InFormat, uint8* OutData, uint32& OutByteCount)
 {
 	// If a new format is added insure that it is either supported here, or at least flagged as not supported
-	static_assert(PF_MAX == 92);
+	static_assert(PF_MAX == 93);
 
 	switch (InFormat)
 	{
@@ -1284,6 +1284,7 @@ void InitializeData(const TInType& InData, EPixelFormat InFormat, uint8* OutData
 		case PF_A16B16G16R16:			{ FormatData<EDefaultInputType::UNorm, TInType,	uint16,   3, 2, 1, 0>	(InData, OutData, OutByteCount); } break;
 		case PF_FloatRGBA:				{ FormatData<EDefaultInputType::Typed, TInType,	FFloat16, 0, 1, 2, 3>	(InData, OutData, OutByteCount); } break;
 		case PF_R16G16_UINT:			{ FormatData<EDefaultInputType::Typed, TInType,	uint16,   0, 1>			(InData, OutData, OutByteCount); } break;
+		case PF_R16G16_SINT:			{ FormatData<EDefaultInputType::Typed, TInType,	int16,    0, 1>			(InData, OutData, OutByteCount); } break;
 		case PF_G16R16:					{ FormatData<EDefaultInputType::UNorm, TInType,	uint16,   1, 0>			(InData, OutData, OutByteCount); } break;
 		case PF_G16R16_SNORM:			{ FormatData<EDefaultInputType::SNorm, TInType,	int16,    1, 0>			(InData, OutData, OutByteCount); } break;
 		case PF_G16R16F:				{ FormatData<EDefaultInputType::Typed, TInType,	FFloat16, 0, 1>			(InData, OutData, OutByteCount); } break;
diff --git a/Engine/Source/Runtime/VulkanRHI/Private/VulkanDevice.cpp b/Engine/Source/Runtime/VulkanRHI/Private/VulkanDevice.cpp
--- a/Engine/Source/Runtime/VulkanRHI/Private/VulkanDevice.cpp
+++ b/Engine/Source/Runtime/VulkanRHI/Private/VulkanDevice.cpp
@@ -568,6 +568,7 @@ void FVulkanDevice::SetupFormats()
 	MapFormatSupport(PF_R8G8B8A8_UINT, { VK_FORMAT_R8G8B8A8_UINT }, ComponentMappingRGBA);
 	MapFormatSupport(PF_R8G8B8A8_SNORM, { VK_FORMAT_R8G8B8A8_SNORM }, ComponentMappingRGBA);
 	MapFormatSupport(PF_R16G16_UINT, { VK_FORMAT_R16G16_UINT }, ComponentMappingRG01);
+	MapFormatSupport(PF_R16G16_SINT, { VK_FORMAT_R16G16_SINT }, ComponentMappingRG01);
 	MapFormatSupport(PF_R16G16B16A16_UINT, { VK_FORMAT_R16G16B16A16_UINT }, ComponentMappingRGBA);
 	MapFormatSupport(PF_R16G16B16A16_SINT, { VK_FORMAT_R16G16B16A16_SINT }, ComponentMappingRGBA);
 	MapFormatSupport(PF_R32G32_UINT, { VK_FORMAT_R32G32_UINT }, ComponentMappingRG01);
diff --git a/Engine/Source/Runtime/Windows/D3D11RHI/Private/D3D11Device.cpp b/Engine/Source/Runtime/Windows/D3D11RHI/Private/D3D11Device.cpp
--- a/Engine/Source/Runtime/Windows/D3D11RHI/Private/D3D11Device.cpp
+++ b/Engine/Source/Runtime/Windows/D3D11RHI/Private/D3D11Device.cpp
@@ -165,6 +165,7 @@ FD3D11DynamicRHI::FD3D11DynamicRHI(IDXGIFactory1* InDXGIFactory1, D3D_FEATURE_LE
 	GPixelFormats[ PF_R8G8			].PlatformFormat	= DXGI_FORMAT_R8G8_UNORM;
 	GPixelFormats[ PF_R32G32B32A32_UINT].PlatformFormat = DXGI_FORMAT_R32G32B32A32_UINT;
 	GPixelFormats[ PF_R16G16_UINT   ].PlatformFormat    = DXGI_FORMAT_R16G16_UINT;
+	GPixelFormats[ PF_R16G16_SINT   ].PlatformFormat    = DXGI_FORMAT_R16G16_SINT;
 	GPixelFormats[ PF_R32G32_UINT   ].PlatformFormat    = DXGI_FORMAT_R32G32_UINT;
 
 	GPixelFormats[ PF_BC6H			].PlatformFormat	= DXGI_FORMAT_BC6H_UF16;
