Fixed missing static assert update since addition of PF_R16G16_SINT.

diff --git a/Engine/Source/Runtime/Apple/MetalRHI/Private/MetalRHI.cpp b/Engine/Source/Runtime/Apple/MetalRHI/Private/MetalRHI.cpp
--- a/Engine/Source/Runtime/Apple/MetalRHI/Private/MetalRHI.cpp
+++ b/Engine/Source/Runtime/Apple/MetalRHI/Private/MetalRHI.cpp
@@ -817,7 +817,7 @@ FMetalDynamicRHI::FMetalDynamicRHI(ERHIFeatureLevel::Type RequestedFeatureLevel)
 	GMetalBufferFormats[PF_ASTC_8x8_NORM_RG		] = { MTL::PixelFormatInvalid, (uint8)EMetalBufferFormat::Unknown };
 	GMetalBufferFormats[PF_ASTC_10x10_NORM_RG	] = { MTL::PixelFormatInvalid, (uint8)EMetalBufferFormat::Unknown };
 	GMetalBufferFormats[PF_ASTC_12x12_NORM_RG	] = { MTL::PixelFormatInvalid, (uint8)EMetalBufferFormat::Unknown };
-	static_assert(PF_MAX == 92, "Please setup GMetalBufferFormats properly for the new pixel format");
+	static_assert(PF_MAX == 93, "Please setup GMetalBufferFormats properly for the new pixel format");
 
 	// Initialize the platform pixel format map.
 	GPixelFormats[PF_Unknown			].PlatformFormat	= (uint32)MTL::PixelFormatInvalid;
