Removed unused BoundShaderStateHistory

diff --git a/Engine/Source/Runtime/Apple/MetalRHI/Private/MetalRHIContext.cpp b/Engine/Source/Runtime/Apple/MetalRHI/Private/MetalRHIContext.cpp
--- a/Engine/Source/Runtime/Apple/MetalRHI/Private/MetalRHIContext.cpp
+++ b/Engine/Source/Runtime/Apple/MetalRHI/Private/MetalRHIContext.cpp
@@ -6,8 +6,6 @@
 #include "MetalRHIRenderQuery.h"
 #include "MetalRHIVisionOSBridge.h"
 
-TGlobalResource<TBoundShaderStateHistory<10000>> FMetalRHICommandContext::BoundShaderStateHistory;
-
 FMetalDeviceContext& GetMetalDeviceContext()
 {
 	FMetalRHICommandContext* Context = static_cast<FMetalRHICommandContext*>(RHIGetDefaultContext());
diff --git a/Engine/Source/Runtime/Apple/MetalRHI/Public/MetalRHIContext.h b/Engine/Source/Runtime/Apple/MetalRHI/Public/MetalRHIContext.h
--- a/Engine/Source/Runtime/Apple/MetalRHI/Public/MetalRHIContext.h
+++ b/Engine/Source/Runtime/Apple/MetalRHI/Public/MetalRHIContext.h
@@ -8,7 +8,6 @@
 #include "MetalResources.h"
 #include "MetalViewport.h"
 #include "RHICore.h"
-#include "BoundShaderStateHistory.h"
 
 class FMetalDeviceContext;
 class FMetalCommandBufferFence;
@@ -186,7 +185,6 @@ public:
         return GlobalUniformBuffers;
     }
 protected:
-	static TGlobalResource<TBoundShaderStateHistory<10000>> BoundShaderStateHistory;
 	
 	/** Context implementation details. */
 	FMetalDeviceContext* Context = nullptr;
