Remove r.Metal.PSOPrecaching.MemoryThreshold

diff --git a/Engine/Source/Runtime/Apple/MetalRHI/Private/MetalDynamicRHI.cpp b/Engine/Source/Runtime/Apple/MetalRHI/Private/MetalDynamicRHI.cpp
--- a/Engine/Source/Runtime/Apple/MetalRHI/Private/MetalDynamicRHI.cpp
+++ b/Engine/Source/Runtime/Apple/MetalRHI/Private/MetalDynamicRHI.cpp
@@ -12,21 +12,11 @@
 #include "MetalVertexDeclaration.h"
 #include "MetalGraphicsPipelineState.h"
 #include "MetalTransitionData.h"
-#include "HAL/PlatformMemoryHelpers.h"
 
 //------------------------------------------------------------------------------
 
 #pragma mark - Metal Dynamic RHI Vertex Declaration Methods -
 
-int32 GPSOMemoryThreshold = -1;
-static FAutoConsoleVariableRef CVarPSOMemoryThreshold(
-	TEXT("r.Metal.PSOPrecaching.MemoryThreshold"),
-	GPSOMemoryThreshold,
-	TEXT("-1 No Memory Threshold set, PSO precaching is always on if enabled\n")
-	TEXT(">= 0  Memory Threshold set, when memory drops under this limit, PSO preaching is disabled\n"),
-	ECVF_Default
-);
-
 FVertexDeclarationRHIRef FMetalDynamicRHI::RHICreateVertexDeclaration(const FVertexDeclarationElementList& Elements)
 {
     MTL_SCOPED_AUTORELEASE_POOL;
@@ -53,18 +43,6 @@ FVertexDeclarationRHIRef FMetalDynamicRHI::RHICreateVertexDeclaration(const FVer
 FGraphicsPipelineStateRHIRef FMetalDynamicRHI::RHICreateGraphicsPipelineState(const FGraphicsPipelineStateInitializer& Initializer)
 {
     MTL_SCOPED_AUTORELEASE_POOL;
-    
-	if (Initializer.bPSOPrecache)
-	{
-		if (GPSOMemoryThreshold >= 0)
-		{
-			float MBFree = float(PlatformMemoryHelpers::GetFrameMemoryStats().AvailablePhysical / 1024 / 1024);
-			if (MBFree < GPSOMemoryThreshold)
-			{
-				return nullptr;
-			}
-		}
-	}
 
     TRefCountPtr<FMetalGraphicsPipelineState> State = new FMetalGraphicsPipelineState(Initializer);
 
