Metal RHI - Attempting to use a Command buffer fence when no fence has been set

diff --git a/Engine/Source/Runtime/Apple/MetalRHI/Private/MetalDynamicRHI.h b/Engine/Source/Runtime/Apple/MetalRHI/Private/MetalDynamicRHI.h
--- a/Engine/Source/Runtime/Apple/MetalRHI/Private/MetalDynamicRHI.h
+++ b/Engine/Source/Runtime/Apple/MetalRHI/Private/MetalDynamicRHI.h
@@ -172,6 +172,9 @@ public:
 
 	virtual uint64 RHIComputePrecachePSOHash(const FGraphicsPipelineStateInitializer& Initializer) final override;
 	virtual bool RHIMatchPrecachePSOInitializers(const FGraphicsPipelineStateInitializer& LHS, const FGraphicsPipelineStateInitializer& RHS) final override;
+	
+	virtual void RHIBeginRenderQuery_TopOfPipe(FRHICommandListBase& RHICmdList, FRHIRenderQuery* RenderQuery) override;
+	virtual void RHIEndRenderQuery_TopOfPipe  (FRHICommandListBase& RHICmdList, FRHIRenderQuery* RenderQuery) override;
 
 private:
 	FMetalDevice* Device;
diff --git a/Engine/Source/Runtime/Apple/MetalRHI/Private/MetalRHIContext.cpp b/Engine/Source/Runtime/Apple/MetalRHI/Private/MetalRHIContext.cpp
--- a/Engine/Source/Runtime/Apple/MetalRHI/Private/MetalRHIContext.cpp
+++ b/Engine/Source/Runtime/Apple/MetalRHI/Private/MetalRHIContext.cpp
@@ -7,6 +7,7 @@
 #include "MetalBindlessDescriptors.h"
 #include "MetalDevice.h"
 #include "MetalCommandBuffer.h"
+#include "MetalDynamicRHI.h"
 
 #if PLATFORM_VISIONOS
 #import <CompositorServices/CompositorServices.h>
@@ -308,6 +309,22 @@ void FMetalRHICommandContext::RHICalibrateTimers(FRHITimestampCalibrationQuery*
     CalibrationQuery->GPUMicroseconds[0] = uint64(GPUTimestamp / 1000.0);
 }
 
+void FMetalDynamicRHI::RHIBeginRenderQuery_TopOfPipe(FRHICommandListBase& RHICmdList, FRHIRenderQuery* RenderQuery)
+{
+	FMetalRHIRenderQuery* Query = ResourceCast(RenderQuery);
+	Query->Begin_TopOfPipe();
+
+	FDynamicRHI::RHIBeginRenderQuery_TopOfPipe(RHICmdList, RenderQuery);
+}
+
+void FMetalDynamicRHI::RHIEndRenderQuery_TopOfPipe(FRHICommandListBase& RHICmdList, FRHIRenderQuery* RenderQuery)
+{
+	FMetalRHIRenderQuery* Query = ResourceCast(RenderQuery);
+	Query->End_TopOfPipe();
+
+	FDynamicRHI::RHIEndRenderQuery_TopOfPipe(RHICmdList, RenderQuery);
+}
+
 void FMetalRHICommandContext::RHIBeginRenderQuery(FRHIRenderQuery* QueryRHI)
 {
     MTL_SCOPED_AUTORELEASE_POOL;
diff --git a/Engine/Source/Runtime/Apple/MetalRHI/Private/MetalRHIRenderQuery.cpp b/Engine/Source/Runtime/Apple/MetalRHI/Private/MetalRHIRenderQuery.cpp
--- a/Engine/Source/Runtime/Apple/MetalRHI/Private/MetalRHIRenderQuery.cpp
+++ b/Engine/Source/Runtime/Apple/MetalRHI/Private/MetalRHIRenderQuery.cpp
@@ -145,6 +145,11 @@ void FMetalQueryBufferPool::ReleaseQueryBuffer(FMetalBufferPtr Buffer)
 
 #pragma mark - Metal RHI Private Query Result Class -
 
+void FMetalQueryResult::Reset()
+{
+	CommandBufferFence.Reset();
+	bCompleted = false;
+}
 
 bool FMetalQueryResult::Wait(uint64 Millis)
 {
@@ -194,6 +199,25 @@ FMetalRHIRenderQuery::~FMetalRHIRenderQuery()
 	}
 }
 
+void FMetalRHIRenderQuery::Begin_TopOfPipe()
+{
+	Buffer.Reset();
+	bAvailable = false;
+}
+
+void FMetalRHIRenderQuery::End_TopOfPipe()
+{
+	if (Type == RQT_AbsoluteTime)
+	{
+		Buffer.Reset();
+		if(QueryWrittenEvent)
+		{
+			QueryWrittenEvent->Reset();
+		}
+	}
+	bAvailable = false;
+}
+
 void FMetalRHIRenderQuery::Begin(FMetalRHICommandContext* Context, TSharedPtr<FMetalCommandBufferFence, ESPMode::ThreadSafe> const& BatchFence)
 {
 	Buffer.CommandBufferFence.Reset();
@@ -269,8 +293,7 @@ void FMetalRHIRenderQuery::End(FMetalRHICommandContext* Context)
 			Buffer.Offset = 0;
 			Buffer.bCompleted = false;
 			Buffer.bBatchFence = false;
-			Buffer.CommandBufferFence = MakeShareable(new FMetalCommandBufferFence);
-			check(Buffer.CommandBufferFence.IsValid());
+			TSharedPtr<FMetalCommandBufferFence, ESPMode::ThreadSafe> CommandBufferFence = MakeShareable(new FMetalCommandBufferFence);
 
 			Result = 0;
 			bAvailable = false;
@@ -316,7 +339,8 @@ void FMetalRHIRenderQuery::End(FMetalRHICommandContext* Context)
 				this->Release();
 			});
 
-            Context->InsertCommandBufferFence(Buffer.CommandBufferFence, Handler);
+            Context->InsertCommandBufferFence(CommandBufferFence, Handler);
+			Buffer.CommandBufferFence = CommandBufferFence;
 			break;
 		}
 		default:
@@ -344,10 +368,8 @@ bool FMetalRHIRenderQuery::GetResult(uint64& OutNumPixels, bool bWait, uint32 GP
 			FRHICommandListImmediate& RHICmdList = FRHICommandListImmediate::Get();
 			
 			// RHI thread *must* be flushed at this point if the internal handles we rely upon are not yet valid.
-			// We *CANNOT* have one event per query as it consumes too many pthread objects.
-			if (!RHICmdList.Bypass() && IsRunningRHIInSeparateThread() && !Buffer.CommandBufferFence.IsValid())
+			if (!Buffer.CommandBufferFence.IsValid())
 			{
-				RHICmdList.RHIThreadFence(true);
 				RHICmdList.ImmediateFlush(EImmediateFlushType::FlushRHIThread);
 			}
 
@@ -369,7 +391,10 @@ bool FMetalRHIRenderQuery::GetResult(uint64& OutNumPixels, bool bWait, uint32 GP
 		}
 		else
 		{
-			bOK = Buffer.Wait(0);
+			if(Buffer.CommandBufferFence.IsValid())
+			{
+				bOK = Buffer.Wait(0);
+			}
 		}
 
 		if (bOK == false)
diff --git a/Engine/Source/Runtime/Apple/MetalRHI/Private/MetalRHIRenderQuery.h b/Engine/Source/Runtime/Apple/MetalRHI/Private/MetalRHIRenderQuery.h
--- a/Engine/Source/Runtime/Apple/MetalRHI/Private/MetalRHIRenderQuery.h
+++ b/Engine/Source/Runtime/Apple/MetalRHI/Private/MetalRHIRenderQuery.h
@@ -84,6 +84,7 @@ public:
 	FMetalQueryResult() = default;
 	~FMetalQueryResult() = default;
 
+	void Reset();
 	bool Wait(uint64 Millis);
 	uint64 GetResult();
 
@@ -106,6 +107,9 @@ public:
 	FMetalRHIRenderQuery(FMetalDevice& MetalDevice, ERenderQueryType InQueryType);
 	virtual ~FMetalRHIRenderQuery();
 
+	void Begin_TopOfPipe();
+	void End_TopOfPipe();
+	
 	/**
 	 * Kick off an occlusion test
 	 */
@@ -134,7 +138,7 @@ private:
 	volatile uint64 Result;
 
 	// Result availability - if not set the first call to acquire it will read the buffer & cache
-	volatile bool bAvailable;
+	std::atomic<bool> bAvailable;
 
 	// Timer event completion signal
 	FEvent* QueryWrittenEvent;
