VisionOS MetalVeiwport build breaks.
-Fixed a few simple refactoring build breaks in PresentImmersive.

diff --git a/Engine/Source/Runtime/Apple/MetalRHI/Private/MetalViewport.cpp b/Engine/Source/Runtime/Apple/MetalRHI/Private/MetalViewport.cpp
--- a/Engine/Source/Runtime/Apple/MetalRHI/Private/MetalViewport.cpp
+++ b/Engine/Source/Runtime/Apple/MetalRHI/Private/MetalViewport.cpp
@@ -682,18 +682,18 @@ void FMetalViewport::PresentImmersive(const MetalRHIVisionOS::PresentImmersivePa
 					}
 				}
 				FMetalCommandBuffer* Buffer = Context.Finalize();
-				Context->GetDevice().GetCommandQueue().CommitCommandBuffer(Buffer);
-				Context->Reset();
+				Context.GetDevice().GetCommandQueue().CommitCommandBuffer(Buffer);
+				Context.ResetContext();
 				
 				cp_frame_t CompositorServicesFrame;
 				
-				FMetalCommandBuffer* CurrentCommandBuffer = Context->GetCurrentCommandBuffer();
+				FMetalCommandBuffer* CurrentCommandBuffer = Context.GetCurrentCommandBuffer();
 				cp_drawable_encode_present(VisionOSParams.SwiftDrawable, (__bridge id<MTLCommandBuffer>)CurrentCommandBuffer->GetMTLCmdBuffer().get());
 				CompositorServicesFrame = VisionOSParams.SwiftFrame;
 				
-				Buffer = Context->Finalize();
+				Buffer = Context.Finalize();
 				Context.GetDevice().GetCommandQueue().CommitCommandBuffer(Buffer);
-				Context->Reset();
+				Context.ResetContext();
 				
 				cp_frame_end_submission(CompositorServicesFrame);
 			}
