Fix for invalid render encoder after MSAA resolve

diff --git a/Engine/Source/Runtime/Apple/MetalRHI/Private/MetalRHIContext.cpp b/Engine/Source/Runtime/Apple/MetalRHI/Private/MetalRHIContext.cpp
--- a/Engine/Source/Runtime/Apple/MetalRHI/Private/MetalRHIContext.cpp
+++ b/Engine/Source/Runtime/Apple/MetalRHI/Private/MetalRHIContext.cpp
@@ -71,6 +71,7 @@ void FMetalRHICommandContext::BeginComputeEncoder()
 {
 	SCOPE_CYCLE_COUNTER(STAT_MetalSwitchToComputeTime);
 	
+	check(!bWithinRenderPass);
 	check(CurrentEncoder.GetCommandBuffer());
 	check(IsInParallelRenderingThread());
 	
@@ -104,7 +105,7 @@ void FMetalRHICommandContext::EndComputeEncoder()
 void FMetalRHICommandContext::BeginBlitEncoder()
 {
 	SCOPE_CYCLE_COUNTER(STAT_MetalSwitchToBlitTime);
-	
+	check(!bWithinRenderPass);
 	check(CurrentEncoder.GetCommandBuffer());
 	
 	if(!CurrentEncoder.IsBlitCommandEncoderActive())
@@ -137,8 +138,9 @@ void FMetalRHICommandContext::RHIBeginRenderPass(const FRHIRenderPassInfo& InInf
     MTL_SCOPED_AUTORELEASE_POOL;
 	
 #if PLATFORM_SUPPORTS_BINDLESS_RENDERING
-	if(!bWithinRenderPass && IsMetalBindlessEnabled())
+	if(IsMetalBindlessEnabled())
 	{
+		check(!bWithinRenderPass);
 		Device.GetBindlessDescriptorManager()->UpdateDescriptorsWithGPU();
 	}
 #endif
@@ -204,22 +206,22 @@ void FMetalRHICommandContext::RHIEndRenderPass()
 		RHIEndOcclusionQueryBatch();
 	}
     
-	UE::RHICore::ResolveRenderPassTargets(RenderPassInfo, [this](UE::RHICore::FResolveTextureInfo Info)
-	{
-		ResolveTexture(Info);
-	});
-	
 	check(bWithinRenderPass);
 	check(CurrentEncoder.IsRenderCommandEncoderActive());
 	
 	StateCache.FlushVisibilityResults(CurrentEncoder);
 	
 	CurrentEncoderFence = CurrentEncoder.EndEncoding();
+	bWithinRenderPass = false;
 	
-	StateCache.SetRenderTargetsActive(false);
+	// Uses a Blit encoder so need to run after end encoding 
+	UE::RHICore::ResolveRenderPassTargets(RenderPassInfo, [this](UE::RHICore::FResolveTextureInfo Info)
+	{
+		ResolveTexture(Info);
+	});
 	
+	StateCache.SetRenderTargetsActive(false);
 	RenderPassDesc = nullptr;
-	bWithinRenderPass = false;
 }
 
 void FMetalRHICommandContext::ResolveTexture(UE::RHICore::FResolveTextureInfo Info)
