Fixed bug with scissor not resetting to viewport width when enabled.

diff --git a/Engine/Source/Runtime/Apple/MetalRHI/Private/MetalStateCache.cpp b/Engine/Source/Runtime/Apple/MetalRHI/Private/MetalStateCache.cpp
--- a/Engine/Source/Runtime/Apple/MetalRHI/Private/MetalStateCache.cpp
+++ b/Engine/Source/Runtime/Apple/MetalRHI/Private/MetalStateCache.cpp
@@ -1034,15 +1034,22 @@ void FMetalStateCache::SetViewport(const MTL::Viewport& InViewport)
 	
 	ActiveViewports = 1;
 	
+	MTL::ScissorRect Rect;
 	if (!bScissorRectEnabled)
 	{
-		MTL::ScissorRect Rect;
 		Rect.x = InViewport.originX;
 		Rect.y = InViewport.originY;
 		Rect.width = InViewport.width;
 		Rect.height = InViewport.height;
-		SetScissorRect(false, Rect);
 	}
+	else
+	{
+		Rect.x = FMath::Min(Scissor[0].x, InViewport.originX);
+		Rect.y = FMath::Min(Scissor[0].y, InViewport.originY);
+		Rect.width = FMath::Min(Scissor[0].width, InViewport.width);
+		Rect.height = FMath::Min(Scissor[0].height, InViewport.height);
+	}
+	SetScissorRect(bScissorRectEnabled, Rect);
 }
 
 void FMetalStateCache::SetViewport(uint32 Index, const MTL::Viewport& InViewport)
