Metal RHI - Fixed RHI Adapter Name

diff --git a/Engine/Source/Runtime/Apple/MetalRHI/Private/MetalRHI.cpp b/Engine/Source/Runtime/Apple/MetalRHI/Private/MetalRHI.cpp
--- a/Engine/Source/Runtime/Apple/MetalRHI/Private/MetalRHI.cpp
+++ b/Engine/Source/Runtime/Apple/MetalRHI/Private/MetalRHI.cpp
@@ -198,8 +198,7 @@ FMetalDynamicRHI::FMetalDynamicRHI(ERHIFeatureLevel::Type RequestedFeatureLevel)
 	// Metal always needs a render target to render with fragment shaders!
 	GRHIRequiresRenderTargetForPixelShaderUAVs = true;
 
-	//@todo-rco: Query name from API
-	GRHIAdapterName = TEXT("Metal");
+	GRHIAdapterName = NSStringToFString(Device->GetDevice()->name());
 	GRHIVendorId = 1; // non-zero to avoid asserts
 
 	bool const bRequestedFeatureLevel = (RequestedFeatureLevel != ERHIFeatureLevel::Num);
