MetalRHI: Added error code to GPU error

diff --git a/Engine/Source/Runtime/Apple/MetalRHI/Private/MetalCommandList.cpp b/Engine/Source/Runtime/Apple/MetalRHI/Private/MetalCommandList.cpp
--- a/Engine/Source/Runtime/Apple/MetalRHI/Private/MetalCommandList.cpp
+++ b/Engine/Source/Runtime/Apple/MetalRHI/Private/MetalCommandList.cpp
@@ -60,7 +60,7 @@ static void ReportMetalCommandBufferFailure(MTL::CommandBuffer* CompletedBuffer,
 	FString RecoveryString = RecoveryDesc ? FString(RecoveryDesc->cString(NS::UTF8StringEncoding)) : FString(TEXT("Unknown"));
 	
 	NS::String* Desc = CompletedBuffer->debugDescription();
-	UE_LOG(LogMetal, Warning, TEXT("%s"), *FString(Desc->cString(NS::UTF8StringEncoding)));
+	UE_LOG(LogMetal, Warning, TEXT("Metal Command Buffer Failure: %s, %s"), ErrorType, *FString(Desc->cString(NS::UTF8StringEncoding)));
 	
 #if PLATFORM_IOS
     if (bDoCheck && !GIsSuspended)
