Metal RHI - Descriptor upload using correct context WIP

diff --git a/Engine/Source/Runtime/Apple/MetalRHI/Private/MetalBindlessDescriptors.cpp b/Engine/Source/Runtime/Apple/MetalRHI/Private/MetalBindlessDescriptors.cpp
--- a/Engine/Source/Runtime/Apple/MetalRHI/Private/MetalBindlessDescriptors.cpp
+++ b/Engine/Source/Runtime/Apple/MetalRHI/Private/MetalBindlessDescriptors.cpp
@@ -11,7 +11,8 @@
 #include "MetalDevice.h"
 #include "MetalCommandEncoder.h"
 
-#define USE_DESCRIPTOR_BUFFER_COPY 1
+#define USE_CPU_DESCRIPTOR_COPY 0
+#define USE_DESCRIPTOR_BUFFER_COPY !USE_CPU_DESCRIPTOR_COPY
 
 int32 GBindlessResourceDescriptorHeapSize = 2048 * 1024;
 static FAutoConsoleVariableRef CVarBindlessResourceDescriptorHeapSize(
@@ -267,7 +268,7 @@ void FMetalBindlessDescriptorManager::BindResource(FRHIDescriptorHandle Descript
         return;
     };
 
-#if !USE_DESCRIPTOR_BUFFER_COPY
+#if !USE_DESCRIPTOR_BUFFER_COPY && !USE_CPU_DESCRIPTOR_COPY
 	if(GIsRHIInitialized)
 	{
 		FScopeLock ScopeLock(&ComputeDescriptorCS);
@@ -284,17 +285,20 @@ void FMetalBindlessDescriptorManager::BindResource(FRHIDescriptorHandle Descript
 #endif
 }
 
-void FMetalBindlessDescriptorManager::UpdateDescriptorsWithGPU()
+void FMetalBindlessDescriptorManager::UpdateDescriptorsWithGPU(FMetalRHICommandContext* Context)
 {
-#ifdef USE_DESCRIPTOR_BUFFER_COPY
-	UpdateDescriptorsWithCopy();
+#if USE_CPU_DESCRIPTOR_COPY
+	return;
+#elif USE_DESCRIPTOR_BUFFER_COPY
+	UpdateDescriptorsWithCopy(Context);
 #else
 	UpdateDescriptorsWithCompute();
 #endif
 }
 
-void FMetalBindlessDescriptorManager::UpdateDescriptorsWithCopy()
+void FMetalBindlessDescriptorManager::UpdateDescriptorsWithCopy(FMetalRHICommandContext* Context)
 {
+#if !USE_CPU_DESCRIPTOR_COPY
 	FScopeLock ScopeLock(&ComputeDescriptorCS);
 	
 	if(!StandardResources.DescriptorsDirty)
@@ -302,7 +306,6 @@ void FMetalBindlessDescriptorManager::UpdateDescriptorsWithCopy()
 		return;
 	}
 	
-	FMetalRHICommandContext* Context = static_cast<FMetalRHICommandContext*>(RHIGetDefaultContext());
 	TRHIComputeCommandList_RecursiveHazardous<FMetalRHICommandContext> RHICmdList(Context);
 	
 	uint32_t IndexOffset = StandardResources.MinDirtyIndex;
@@ -421,6 +424,7 @@ void FMetalBindlessDescriptorManager::UpdateDescriptorsWithCompute()
 	});
 	
 	Device.ReleaseFunction([Evt](){});
+#endif
 }
 
 void FMetalBindlessDescriptorManager::BindTexture(FRHICommandListBase& RHICmdList, FRHIDescriptorHandle DescriptorHandle, MTL::Texture* Texture, EDescriptorUpdateType UpdateType)
@@ -428,7 +432,7 @@ void FMetalBindlessDescriptorManager::BindTexture(FRHICommandListBase& RHICmdLis
     IRDescriptorTableEntry DescriptorData = {0};
     IRDescriptorTableSetTexture(&DescriptorData, Texture, 0.0f, 0u);
     
-#if USE_DESCRIPTOR_BUFFER_COPY
+#if USE_DESCRIPTOR_BUFFER_COPY || USE_CPU_DESCRIPTOR_COPY
 	UpdateType = EDescriptorUpdateType_Immediate;
 #else
 	UpdateType = !GIsRHIInitialized ? EDescriptorUpdateType_Immediate : UpdateType;
diff --git a/Engine/Source/Runtime/Apple/MetalRHI/Private/MetalBindlessDescriptors.h b/Engine/Source/Runtime/Apple/MetalRHI/Private/MetalBindlessDescriptors.h
--- a/Engine/Source/Runtime/Apple/MetalRHI/Private/MetalBindlessDescriptors.h
+++ b/Engine/Source/Runtime/Apple/MetalRHI/Private/MetalBindlessDescriptors.h
@@ -80,11 +80,11 @@ public:
 
 	bool					IsSupported() {return bIsSupported;}
     
-	void 					UpdateDescriptorsWithGPU();
+	void 					UpdateDescriptorsWithGPU(FMetalRHICommandContext* Context);
     
 private:
 	void                    UpdateDescriptorsWithCompute();
-	void 					UpdateDescriptorsWithCopy();
+	void 					UpdateDescriptorsWithCopy(FMetalRHICommandContext* Context);
 	
 	bool 					bIsSupported = false;
 	FMetalDevice& 			Device;
diff --git a/Engine/Source/Runtime/Apple/MetalRHI/Private/MetalRHIContext.cpp b/Engine/Source/Runtime/Apple/MetalRHI/Private/MetalRHIContext.cpp
--- a/Engine/Source/Runtime/Apple/MetalRHI/Private/MetalRHIContext.cpp
+++ b/Engine/Source/Runtime/Apple/MetalRHI/Private/MetalRHIContext.cpp
@@ -141,7 +141,7 @@ void FMetalRHICommandContext::RHIBeginRenderPass(const FRHIRenderPassInfo& InInf
 	if(IsMetalBindlessEnabled())
 	{
 		check(!bWithinRenderPass);
-		Device.GetBindlessDescriptorManager()->UpdateDescriptorsWithGPU();
+		Device.GetBindlessDescriptorManager()->UpdateDescriptorsWithGPU(this);
 	}
 #endif
     
