MetalRHI - Missing files on fixing RHI context

diff --git a/Engine/Source/Runtime/Apple/MetalRHI/Private/MetalCommands.cpp b/Engine/Source/Runtime/Apple/MetalRHI/Private/MetalCommands.cpp
--- a/Engine/Source/Runtime/Apple/MetalRHI/Private/MetalCommands.cpp
+++ b/Engine/Source/Runtime/Apple/MetalRHI/Private/MetalCommands.cpp
@@ -193,7 +193,7 @@ void FMetalRHICommandContext::RHISetComputePipelineState(FRHIComputePipelineStat
 #if PLATFORM_SUPPORTS_BINDLESS_RENDERING
 	if(!bWithinRenderPass && IsMetalBindlessEnabled())
 	{
-		Device.GetBindlessDescriptorManager()->UpdateDescriptorsWithGPU();
+		Device.GetBindlessDescriptorManager()->UpdateDescriptorsWithGPU(this);
 	}
 #endif
     
