Metal RHI

- Removed deprecated device checks
- Fixed issue with Vertex Color selecting the incorrect format

diff --git a/Engine/Source/Runtime/Apple/MetalRHI/Private/MetalCommandQueue.cpp b/Engine/Source/Runtime/Apple/MetalRHI/Private/MetalCommandQueue.cpp
--- a/Engine/Source/Runtime/Apple/MetalRHI/Private/MetalCommandQueue.cpp
+++ b/Engine/Source/Runtime/Apple/MetalRHI/Private/MetalCommandQueue.cpp
@@ -19,11 +19,9 @@
 
 #pragma mark - Private C++ Statics -
 NS::UInteger FMetalCommandQueue::PermittedOptions = 0;
-extern MTL::VertexFormat GMetalFColorVertexFormat;
 
 bool GMetalCommandBufferDebuggingEnabled = 0;
 
-
 #pragma mark - Public C++ Boilerplate -
 
 FMetalCommandQueue::FMetalCommandQueue(FMetalDevice& MetalDevice, uint32 const MaxNumCommandBuffers /* = 0 */)
@@ -56,11 +54,7 @@ FMetalCommandQueue::FMetalCommandQueue(FMetalDevice& MetalDevice, uint32 const M
 	check(CommandQueue);
 	
 #if PLATFORM_IOS
-#if PLATFORM_TVOS
-	GMetalFColorVertexFormat = MTL::VertexFormatUChar4Normalized_BGRA;
-#else
-	GMetalFColorVertexFormat = MTL::VertexFormatUChar4Normalized_BGRA;
-
+#if !PLATFORM_TVOS
 	if (Device.GetDevice()->supportsFeatureSet(MTL::FeatureSet_iOS_GPUFamily4_v1))
 	{
 		// The below implies tile shaders which are necessary to order the draw calls and generate a buffer that shows what PSOs/draws ran on each tile.
diff --git a/Engine/Source/Runtime/Apple/MetalRHI/Private/MetalDevice.cpp b/Engine/Source/Runtime/Apple/MetalRHI/Private/MetalDevice.cpp
--- a/Engine/Source/Runtime/Apple/MetalRHI/Private/MetalDevice.cpp
+++ b/Engine/Source/Runtime/Apple/MetalRHI/Private/MetalDevice.cpp
@@ -536,53 +536,50 @@ void FMetalDevice::EnumerateFeatureSupport()
 
 	FString DeviceName(Device->name()->cString(NS::UTF8StringEncoding));
 	
-	if (Device->supportsFeatureSet(MTL::FeatureSet_macOS_GPUFamily1_v2))
+	Features |= EMetalFeaturesMSAADepthResolve | EMetalFeaturesMSAAStoreAndResolve;
+	
+	// Assume that set*Bytes only works on macOS Sierra and above as no-one has tested it anywhere else.
+	Features |= EMetalFeaturesSetBytes;
+	
+	// On earlier OS versions Intel Broadwell couldn't suballocate properly
+	if (!(DeviceName.Contains(TEXT("Intel")) && (DeviceName.Contains(TEXT("5300")) || DeviceName.Contains(TEXT("6000")) || DeviceName.Contains(TEXT("6100")))))
 	{
-		Features |= EMetalFeaturesMSAADepthResolve | EMetalFeaturesMSAAStoreAndResolve;
-		
-		// Assume that set*Bytes only works on macOS Sierra and above as no-one has tested it anywhere else.
-		Features |= EMetalFeaturesSetBytes;
-		
-		// On earlier OS versions Intel Broadwell couldn't suballocate properly
-		if (!(DeviceName.Contains(TEXT("Intel")) && (DeviceName.Contains(TEXT("5300")) || DeviceName.Contains(TEXT("6000")) || DeviceName.Contains(TEXT("6100")))))
-		{
-			// Using Private Memory & BlitEncoders for Vertex & Index data should be *much* faster.
-			Features |= EMetalFeaturesEfficientBufferBlits;
-			
-			Features |= EMetalFeaturesBufferSubAllocation;
-					
-			// On earlier OS versions Vega didn't like non-zero blit offsets
-			if (!DeviceName.Contains(TEXT("Vega")))
-			{
-				Features |= EMetalFeaturesPrivateBufferSubAllocation;
-			}
-		}
+		// Using Private Memory & BlitEncoders for Vertex & Index data should be *much* faster.
+		Features |= EMetalFeaturesEfficientBufferBlits;
 		
-		if (!FParse::Param(FCommandLine::Get(), TEXT("nometalparallelencoder")))
+		Features |= EMetalFeaturesBufferSubAllocation;
+				
+		// On earlier OS versions Vega didn't like non-zero blit offsets
+		if (!DeviceName.Contains(TEXT("Vega")))
 		{
-			Features |= EMetalFeaturesParallelRenderEncoders;
+			Features |= EMetalFeaturesPrivateBufferSubAllocation;
 		}
-		Features |= EMetalFeaturesTextureBuffers;
-		if (IndirectArgumentTier >= 1)
+	}
+	
+	if (!FParse::Param(FCommandLine::Get(), TEXT("nometalparallelencoder")))
+	{
+		Features |= EMetalFeaturesParallelRenderEncoders;
+	}
+	Features |= EMetalFeaturesTextureBuffers;
+	if (IndirectArgumentTier >= 1)
+	{
+		Features |= EMetalFeaturesIABs;
+			
+		if (IndirectArgumentTier >= 2)
 		{
-			Features |= EMetalFeaturesIABs;
-				
-			if (IndirectArgumentTier >= 2)
-			{
-				Features |= EMetalFeaturesTier2IABs;
-			}
+			Features |= EMetalFeaturesTier2IABs;
 		}
+	}
 
-		// The editor spawns so many viewports and preview icons that we can run out of hardware fences!
-		// Need to figure out a way to safely flush the rendering and reuse the fences when that happens.
+	// The editor spawns so many viewports and preview icons that we can run out of hardware fences!
+	// Need to figure out a way to safely flush the rendering and reuse the fences when that happens.
 #if WITH_EDITORONLY_DATA
-		if (!GIsEditor)
+	if (!GIsEditor)
 #endif
+	{
+		if (FParse::Param(FCommandLine::Get(),TEXT("metalfence")))
 		{
-			if (FParse::Param(FCommandLine::Get(),TEXT("metalfence")))
-			{
-				Features |= EMetalFeaturesFences;
-			}
+			Features |= EMetalFeaturesFences;
 		}
 	}
 	
@@ -595,19 +592,16 @@ void FMetalDevice::EnumerateFeatureSupport()
 		Features |= EMetalFeaturesHeaps;
 	}
 	
-	if(Device->supportsFeatureSet(MTL::FeatureSet_macOS_GPUFamily1_v3))
+	Features |= EMetalFeaturesMultipleViewports | EMetalFeaturesPipelineBufferMutability | EMetalFeaturesGPUCaptureManager;
+	
+	if (FParse::Param(FCommandLine::Get(),TEXT("metalfence")))
 	{
-		Features |= EMetalFeaturesMultipleViewports | EMetalFeaturesPipelineBufferMutability | EMetalFeaturesGPUCaptureManager;
-		
-		if (FParse::Param(FCommandLine::Get(),TEXT("metalfence")))
-		{
-			Features |= EMetalFeaturesFences;
-		}
-		
-		if (FParse::Param(FCommandLine::Get(),TEXT("metaliabs")))
-		{
-			Features |= EMetalFeaturesIABs;
-		}
+		Features |= EMetalFeaturesFences;
+	}
+	
+	if (FParse::Param(FCommandLine::Get(),TEXT("metaliabs")))
+	{
+		Features |= EMetalFeaturesIABs;
 	}
 #endif
 	
diff --git a/Engine/Source/Runtime/Apple/MetalRHI/Private/MetalRHI.cpp b/Engine/Source/Runtime/Apple/MetalRHI/Private/MetalRHI.cpp
--- a/Engine/Source/Runtime/Apple/MetalRHI/Private/MetalRHI.cpp
+++ b/Engine/Source/Runtime/Apple/MetalRHI/Private/MetalRHI.cpp
@@ -542,13 +542,8 @@ FMetalDynamicRHI::FMetalDynamicRHI(ERHIFeatureLevel::Type RequestedFeatureLevel)
 	
 #endif
 		
-#if PLATFORM_MAC
-    if (MTLDevice->supportsFeatureSet(MTL::FeatureSet_macOS_GPUFamily1_v3))
-#endif
-	{
-		GRHISupportsDynamicResolution = true;
-		GRHISupportsFrameCyclesBubblesRemoval = true;
-	}
+	GRHISupportsDynamicResolution = true;
+	GRHISupportsFrameCyclesBubblesRemoval = true;
 
 	GPoolSizeVRAMPercentage = 0;
 	GTexturePoolSize = 0;
@@ -655,7 +650,7 @@ FMetalDynamicRHI::FMetalDynamicRHI(ERHIFeatureLevel::Type RequestedFeatureLevel)
 	GMaxBufferDimensions = 1 << 27;
 
 #if PLATFORM_MAC
-	check(MTLDevice->supportsFeatureSet(MTL::FeatureSet_macOS_GPUFamily1_v1));
+	check(MTLDevice->supportsFamily(MTL::GPUFamilyMac2));
 	GRHISupportsBaseVertexIndex = true;
 	GRHISupportsFirstInstance = true; // Supported on macOS & iOS but not tvOS.
 	GMaxTextureDimensions = 16384;
@@ -663,8 +658,8 @@ FMetalDynamicRHI::FMetalDynamicRHI(ERHIFeatureLevel::Type RequestedFeatureLevel)
 	GMaxTextureArrayLayers = 2048;
 	GMaxShadowDepthBufferSizeX = GMaxTextureDimensions;
 	GMaxShadowDepthBufferSizeY = GMaxTextureDimensions;
-    bSupportsD16 = !FParse::Param(FCommandLine::Get(),TEXT("nometalv2")) && MTLDevice->supportsFeatureSet(MTL::FeatureSet_macOS_GPUFamily1_v2);
-    GRHISupportsHDROutput = MTLDevice->supportsFeatureSet(MTL::FeatureSet_macOS_GPUFamily1_v2);
+    bSupportsD16 = true;
+    GRHISupportsHDROutput = true;
 	GRHIHDRDisplayOutputFormat = (GRHISupportsHDROutput) ? PF_PLATFORM_HDR_0 : PF_B8G8R8A8;
 	// Based on the spec below, the maxTotalThreadsPerThreadgroup is not a fixed number but calculated according to the device current ability, so the available threads could less than the maximum number.
 	// For safety and keep the consistency for all platform, reduce the maximum number to half of the device based.
diff --git a/Engine/Source/Runtime/Apple/MetalRHI/Private/MetalVertexDeclaration.cpp b/Engine/Source/Runtime/Apple/MetalRHI/Private/MetalVertexDeclaration.cpp
--- a/Engine/Source/Runtime/Apple/MetalRHI/Private/MetalVertexDeclaration.cpp
+++ b/Engine/Source/Runtime/Apple/MetalRHI/Private/MetalVertexDeclaration.cpp
@@ -14,10 +14,6 @@
 
 #pragma mark - Metal Vertex Declaration Globals
 
-
-MTL::VertexFormat GMetalFColorVertexFormat = MTL::VertexFormatUChar4Normalized;
-
-
 //------------------------------------------------------------------------------
 
 #pragma mark - Metal Vertex Declaration Support Routines
@@ -63,7 +59,7 @@ static MTL::VertexFormat TranslateElementTypeToMTLType(EVertexElementType Type)
 		case VET_PackedNormal:	return MTL::VertexFormatChar4Normalized;
 		case VET_UByte4:		return MTL::VertexFormatUChar4;
 		case VET_UByte4N:		return MTL::VertexFormatUChar4Normalized;
-		case VET_Color:			return GMetalFColorVertexFormat;
+		case VET_Color:			return MTL::VertexFormatUChar4Normalized_BGRA;
 		case VET_Short2:		return MTL::VertexFormatShort2;
 		case VET_Short4:		return MTL::VertexFormatShort4;
 		case VET_Short2N:		return MTL::VertexFormatShort2Normalized;
