Metal RHI

Change SetViewport to always override scissor state

diff --git a/Engine/Source/Runtime/Apple/MetalRHI/Private/MetalStateCache.cpp b/Engine/Source/Runtime/Apple/MetalRHI/Private/MetalStateCache.cpp
--- a/Engine/Source/Runtime/Apple/MetalRHI/Private/MetalStateCache.cpp
+++ b/Engine/Source/Runtime/Apple/MetalRHI/Private/MetalStateCache.cpp
@@ -1034,21 +1034,14 @@ void FMetalStateCache::SetViewport(const MTL::Viewport& InViewport)
 	
 	ActiveViewports = 1;
 	
+	bScissorRectEnabled = false;
 	MTL::ScissorRect Rect;
-	if (!bScissorRectEnabled)
-	{
-		Rect.x = InViewport.originX;
-		Rect.y = InViewport.originY;
-		Rect.width = InViewport.width;
-		Rect.height = InViewport.height;
-	}
-	else
-	{
-		Rect.x = FMath::Min(Scissor[0].x, InViewport.originX);
-		Rect.y = FMath::Min(Scissor[0].y, InViewport.originY);
-		Rect.width = FMath::Min(Scissor[0].width, InViewport.width);
-		Rect.height = FMath::Min(Scissor[0].height, InViewport.height);
-	}
+	
+	Rect.x = InViewport.originX;
+	Rect.y = InViewport.originY;
+	Rect.width = InViewport.width;
+	Rect.height = InViewport.height;
+	
 	SetScissorRect(bScissorRectEnabled, Rect);
 }
 
