Metal RHI

Change to use private memory for volatile, fixes Slate rendering issues

diff --git a/Engine/Source/Runtime/Apple/MetalRHI/Private/MetalRHIBuffer.cpp b/Engine/Source/Runtime/Apple/MetalRHI/Private/MetalRHIBuffer.cpp
--- a/Engine/Source/Runtime/Apple/MetalRHI/Private/MetalRHIBuffer.cpp
+++ b/Engine/Source/Runtime/Apple/MetalRHI/Private/MetalRHIBuffer.cpp
@@ -78,9 +78,9 @@ bool FMetalRHIBuffer::UsePrivateMemory() const
 	if(EnumHasAnyFlags(GetUsage(), BUF_KeepCPUAccessible) && FMetalCommandQueue::IsUMASystem())
 		return false;
 	
-	return (Device.SupportsFeature(EMetalFeaturesEfficientBufferBlits) && EnumHasAnyFlags(GetUsage(), BUF_Dynamic | BUF_Static))
-	|| (Device.SupportsFeature(EMetalFeaturesIABs) && EnumHasAnyFlags(GetUsage(), BUF_ShaderResource|BUF_UnorderedAccess))
-	&& !FMetalCommandQueue::IsUMASystem();
+	return Device.SupportsFeature(EMetalFeaturesEfficientBufferBlits)
+			|| (Device.SupportsFeature(EMetalFeaturesIABs) && EnumHasAnyFlags(GetUsage(), BUF_ShaderResource|BUF_UnorderedAccess))
+			&& !FMetalCommandQueue::IsUMASystem();
 }
 
 FMetalRHIBuffer::FMetalRHIBuffer(FRHICommandListBase& RHICmdList, FMetalDevice& MetalDevice, FRHIBufferDesc const& InBufferDesc, FRHIResourceCreateInfo& CreateInfo)
