Metal RHI

Fixed iOS compile errors caused by bindless changes

diff --git a/Engine/Source/Runtime/Apple/MetalRHI/Private/MetalHashedVertexDescriptor.cpp b/Engine/Source/Runtime/Apple/MetalRHI/Private/MetalHashedVertexDescriptor.cpp
--- a/Engine/Source/Runtime/Apple/MetalRHI/Private/MetalHashedVertexDescriptor.cpp
+++ b/Engine/Source/Runtime/Apple/MetalRHI/Private/MetalHashedVertexDescriptor.cpp
@@ -19,6 +19,7 @@ FMetalHashedVertexDescriptor::FMetalHashedVertexDescriptor()
 	// void
 }
 
+#if PLATFORM_SUPPORTS_BINDLESS_RENDERING
 FMetalHashedVertexDescriptor::FMetalHashedVertexDescriptor(IRVersionedInputLayoutDescriptor& Desc, uint32 Hash)
 	: VertexDescHash(Hash)
 	, IRVertexDesc(Desc)
@@ -26,6 +27,7 @@ FMetalHashedVertexDescriptor::FMetalHashedVertexDescriptor(IRVersionedInputLayou
 {
 	// void
 }
+#endif
 
 FMetalHashedVertexDescriptor::FMetalHashedVertexDescriptor(MTLVertexDescriptorPtr Desc, uint32 Hash)
 	: VertexDescHash(Hash)
@@ -51,7 +53,9 @@ FMetalHashedVertexDescriptor& FMetalHashedVertexDescriptor::operator=(FMetalHash
 	{
 		VertexDescHash = Other.VertexDescHash;
 		VertexDesc = Other.VertexDesc;
+#if METAL_USE_METAL_SHADER_CONVERTER
 		IRVertexDesc = Other.IRVertexDesc;
+#endif
 	}
 	return *this;
 }
@@ -66,6 +70,7 @@ bool FMetalHashedVertexDescriptor::operator==(FMetalHashedVertexDescriptor const
 		{
 			bEqual = true;
 			
+#if PLATFORM_SUPPORTS_BINDLESS_RENDERING
 			if(bUsesIRVertexDesc)
 			{
 				bEqual &= (IRVertexDesc.desc_1_0.numElements == Other.IRVertexDesc.desc_1_0.numElements);
@@ -84,7 +89,9 @@ bool FMetalHashedVertexDescriptor::operator==(FMetalHashedVertexDescriptor const
 					}
 				}
 			}
-			else if (VertexDesc != Other.VertexDesc)
+			else
+#endif
+			if (VertexDesc != Other.VertexDesc)
 			{
                 MTL::VertexBufferLayoutDescriptorArray* Layouts = VertexDesc->layouts();
 				MTL::VertexAttributeDescriptorArray* Attributes = VertexDesc->attributes();
diff --git a/Engine/Source/Runtime/Apple/MetalRHI/Private/MetalHashedVertexDescriptor.h b/Engine/Source/Runtime/Apple/MetalRHI/Private/MetalHashedVertexDescriptor.h
--- a/Engine/Source/Runtime/Apple/MetalRHI/Private/MetalHashedVertexDescriptor.h
+++ b/Engine/Source/Runtime/Apple/MetalRHI/Private/MetalHashedVertexDescriptor.h
@@ -22,11 +22,15 @@ struct FMetalHashedVertexDescriptor
 {
 	NS::UInteger VertexDescHash;
 	MTLVertexDescriptorPtr VertexDesc;
+	
+#if PLATFORM_SUPPORTS_BINDLESS_RENDERING
+	FMetalHashedVertexDescriptor(IRVersionedInputLayoutDescriptor& Desc, uint32 Hash);
 	IRVersionedInputLayoutDescriptor IRVertexDesc;
 	bool bUsesIRVertexDesc = false;
+#endif
 	
 	FMetalHashedVertexDescriptor();
-	FMetalHashedVertexDescriptor(IRVersionedInputLayoutDescriptor& Desc, uint32 Hash);
+	
 	FMetalHashedVertexDescriptor(MTLVertexDescriptorPtr Desc, uint32 Hash);
 	FMetalHashedVertexDescriptor(FMetalHashedVertexDescriptor const& Other);
 	~FMetalHashedVertexDescriptor();
diff --git a/Engine/Source/Runtime/Apple/MetalRHI/Private/MetalUAV.cpp b/Engine/Source/Runtime/Apple/MetalRHI/Private/MetalUAV.cpp
--- a/Engine/Source/Runtime/Apple/MetalRHI/Private/MetalUAV.cpp
+++ b/Engine/Source/Runtime/Apple/MetalRHI/Private/MetalUAV.cpp
@@ -252,12 +252,14 @@ void FMetalShaderResourceView::UpdateView()
 			break;
 		}
 	}
+#if METAL_USE_METAL_SHADER_CONVERTER
 	else if (SurfaceOverride != nullptr)
 	{
 		MTLTexturePtr View = SurfaceOverride->Texture;
 		InitAsTextureView(View);
 		bOwnsResource = false;
 	}
+#endif
 	else
 	{
 		FMetalSurface* Texture = ResourceCast(GetTexture());
