Metal RHI

Removed other usage of volatile creating sub 4k buffers

diff --git a/Engine/Source/Runtime/Apple/MetalRHI/Private/MetalRHIBuffer.cpp b/Engine/Source/Runtime/Apple/MetalRHI/Private/MetalRHIBuffer.cpp
--- a/Engine/Source/Runtime/Apple/MetalRHI/Private/MetalRHIBuffer.cpp
+++ b/Engine/Source/Runtime/Apple/MetalRHI/Private/MetalRHIBuffer.cpp
@@ -118,7 +118,9 @@ FMetalRHIBuffer::FMetalRHIBuffer(FRHICommandListBase& RHICmdList, FMetalDevice&
 		
 		// Temporary buffers less than the buffer page size - currently 4Kb - is better off going through the set*Bytes API if available.
 		// These can't be used for shader resources or UAVs if we want to use the 'Linear Texture' code path
-		if (!bWantsView
+		
+		// TODO: Carl - Strip this code as Volatile is not used with buffer uploads
+		/*if (!bWantsView
 			&& bIsVolatile
 			&& !EnumHasAnyFlags(MetalUsage, EMetalBufferUsage::GPUOnly)
 			&& InBufferDesc.Size < MetalBufferPageSize
@@ -129,6 +131,7 @@ FMetalRHIBuffer::FMetalRHIBuffer(FRHICommandListBase& RHICmdList, FMetalDevice&
 			METAL_INC_DWORD_STAT_BY(MemAlloc, InBufferDesc.Size, InBufferDesc.Usage);
 		}
 		else
+		*/
 		{
 			// Static buffers will never be discarded. You can update them directly.
 			if(bIsStatic)
