Metal RHI

- Fixes for page faults
- Fixed validation due to missing useResource calls

diff --git a/Engine/Source/Runtime/Apple/MetalRHI/Private/MetalBindlessDescriptors.cpp b/Engine/Source/Runtime/Apple/MetalRHI/Private/MetalBindlessDescriptors.cpp
--- a/Engine/Source/Runtime/Apple/MetalRHI/Private/MetalBindlessDescriptors.cpp
+++ b/Engine/Source/Runtime/Apple/MetalRHI/Private/MetalBindlessDescriptors.cpp
@@ -86,6 +86,7 @@ void FMetalDescriptorHeap::Reset()
 
 void FMetalDescriptorHeap::FreeDescriptor(FRHIDescriptorHandle DescriptorHandle)
 {
+	FScopeLock ScopeLock(&FreeListCS);
 	DeferredDeletionList[DeferredDeletionListIndex].Enqueue(DescriptorHandle);
 }
 
@@ -427,20 +428,19 @@ void FMetalBindlessDescriptorManager::BindTexture(FRHICommandListBase& RHICmdLis
 	UpdateType = !GIsRHIInitialized ? EDescriptorUpdateType_Immediate : UpdateType;
 #endif
 	
-	if(UpdateType == EDescriptorUpdateType_Immediate)
+	RHICmdList.EnqueueLambda([this, InUpdateType=UpdateType, Data=DescriptorData, Handle=DescriptorHandle](FRHICommandListBase&)
 	{
-		StandardResources.UpdateDescriptor(DescriptorHandle, DescriptorData);
-	}
-	else if(UpdateType == EDescriptorUpdateType_GPU)
-	{
-		RHICmdList.EnqueueLambda([this, Data=DescriptorData, Handle=DescriptorHandle](FRHICommandListBase&)
+		FScopeLock ScopeLock(&ComputeDescriptorCS);
+		if(InUpdateType == EDescriptorUpdateType_Immediate)
+		{
+			StandardResources.UpdateDescriptor(Handle, Data);
+		}
+		else
 		{
-			FScopeLock ScopeLock(&ComputeDescriptorCS);
-			
 			StandardResources.ComputeDescriptorEntries.Add(Data);
 			StandardResources.ComputeDescriptorIndices.Add(Handle.GetIndex());
-		});
-	}
+		}
+	});
 }
 
 void FMetalBindlessDescriptorManager::BindDescriptorHeapsToEncoder(FMetalCommandEncoder* Encoder, MTL::FunctionType FunctionType, EMetalShaderStages Frequency)
diff --git a/Engine/Source/Runtime/Apple/MetalRHI/Private/MetalBindlessDescriptors.h b/Engine/Source/Runtime/Apple/MetalRHI/Private/MetalBindlessDescriptors.h
--- a/Engine/Source/Runtime/Apple/MetalRHI/Private/MetalBindlessDescriptors.h
+++ b/Engine/Source/Runtime/Apple/MetalRHI/Private/MetalBindlessDescriptors.h
@@ -1,4 +1,5 @@
 // Copyright Epic Games, Inc. All Rights Reserved.
+// Copyright Epic Games, Inc. All Rights Reserved.
 
 #pragma once
 
@@ -36,7 +37,7 @@ struct FMetalDescriptorHeap
 
     static constexpr int32 NumPendingFrame = 3;
 
-    TQueue<FRHIDescriptorHandle>                                            DeferredDeletionList[NumPendingFrame];
+    TQueue<FRHIDescriptorHandle>    DeferredDeletionList[NumPendingFrame];
 
     std::atomic<uint32>             PeakDescriptorCount;
     struct IRDescriptorTableEntry*  Descriptors;
diff --git a/Engine/Source/Runtime/Apple/MetalRHI/Private/MetalCommands.cpp b/Engine/Source/Runtime/Apple/MetalRHI/Private/MetalCommands.cpp
--- a/Engine/Source/Runtime/Apple/MetalRHI/Private/MetalCommands.cpp
+++ b/Engine/Source/Runtime/Apple/MetalRHI/Private/MetalCommands.cpp
@@ -61,22 +61,23 @@ MTL::PrimitiveType TranslatePrimitiveType(uint32 PrimitiveType)
 	}
 }
 
-static FORCEINLINE EMetalShaderStages GetShaderStage(FRHIGraphicsShader* ShaderRHI)
+static FORCEINLINE EMetalShaderStages GetShaderStage(EShaderFrequency ShaderFrequency)
 {
 	EMetalShaderStages Stage = EMetalShaderStages::Num;
-	switch (ShaderRHI->GetFrequency())
+	switch (ShaderFrequency)
 	{
 	case SF_Vertex:		Stage = EMetalShaderStages::Vertex; break;
 	case SF_Pixel:		Stage = EMetalShaderStages::Pixel; break;
+	case SF_Compute:	Stage = EMetalShaderStages::Compute; break;
 #if PLATFORM_SUPPORTS_GEOMETRY_SHADERS
-    case SF_Geometry:   Stage = EMetalShaderStages::Geometry; break;
+	case SF_Geometry:   Stage = EMetalShaderStages::Geometry; break;
 #endif
 #if PLATFORM_SUPPORTS_MESH_SHADERS
-    case SF_Mesh:               Stage = EMetalShaderStages::Mesh; break;
-    case SF_Amplification:      Stage = EMetalShaderStages::Amplification; break;
+	case SF_Mesh:               Stage = EMetalShaderStages::Mesh; break;
+	case SF_Amplification:      Stage = EMetalShaderStages::Amplification; break;
 #endif
 	default:
-		checkf(0, TEXT("FRHIShader Type %d is invalid or unsupported!"), (int32)ShaderRHI->GetFrequency());
+		checkf(0, TEXT("FRHIShader Type %d is invalid or unsupported!"), (int32)ShaderFrequency);
 		NOT_SUPPORTED("RHIShaderStage");
 		break;
 	}
@@ -84,6 +85,12 @@ static FORCEINLINE EMetalShaderStages GetShaderStage(FRHIGraphicsShader* ShaderR
 	return Stage;
 }
 
+static FORCEINLINE EMetalShaderStages GetShaderStage(FRHIGraphicsShader* ShaderRHI)
+{
+	EMetalShaderStages Stage = EMetalShaderStages::Num;
+	return GetShaderStage(ShaderRHI->GetFrequency());
+}
+
 void FMetalRHICommandContext::RHISetStreamSource(uint32 StreamIndex, FRHIBuffer* VertexBufferRHI,uint32 Offset)
 {
     MTL_SCOPED_AUTORELEASE_POOL;;
@@ -572,8 +579,27 @@ static void SetShaderParameters(
             Binder.SetBindlessHandle(Handle, Parameter.Index);
         }
     }
-#endif
-
+	
+	for (const FRHIShaderParameterResource& Parameter : InBindlessParameters)
+	{
+		if (Parameter.Type == FRHIShaderParameterResource::EType::UnorderedAccessView)
+		{
+			if (ShaderFrequency == SF_Pixel || ShaderFrequency == SF_Compute)
+			{
+				Binder.StateCache.IRMakeUAVResident(GetShaderStage(ShaderFrequency), static_cast<FMetalUnorderedAccessView*>(Parameter.Resource));
+			}
+			else
+			{
+				checkf(false, TEXT("TShaderRHI Can't have compute shader to be set. UAVs are not supported on vertex, tessellation and geometry shaders."));
+			}
+		}
+		else if(Parameter.Type == FRHIShaderParameterResource::EType::ResourceView)
+		{
+			Binder.StateCache.IRMakeSRVResident(GetShaderStage(ShaderFrequency), static_cast<FMetalShaderResourceView*>(Parameter.Resource));
+		}
+	}
+#endif 
+	
     for (const FRHIShaderParameterResource& Parameter : InResourceParameters)
     {
         if (Parameter.Type == FRHIShaderParameterResource::EType::UnorderedAccessView)
diff --git a/Engine/Source/Runtime/Apple/MetalRHI/Private/MetalStateCache.cpp b/Engine/Source/Runtime/Apple/MetalRHI/Private/MetalStateCache.cpp
--- a/Engine/Source/Runtime/Apple/MetalRHI/Private/MetalStateCache.cpp
+++ b/Engine/Source/Runtime/Apple/MetalRHI/Private/MetalStateCache.cpp
@@ -1590,10 +1590,8 @@ void FMetalStateCache::IRMakeUAVResident(EMetalShaderStages const Frequency, FMe
 			{
 				CacheOrSkipResourceResidencyUpdate(View->parentTexture(), Frequency, false);
 			}
-			else
-			{
-				CacheOrSkipResourceResidencyUpdate(View.get(), Frequency, false);
-			}
+			
+			CacheOrSkipResourceResidencyUpdate(View.get(), Frequency, false);
 			break;
         }
             
@@ -1644,7 +1642,7 @@ void FMetalStateCache::IRMakeTextureResident(EMetalShaderStages const Frequency,
 	{
 		CacheOrSkipResourceResidencyUpdate(Texture->parentTexture(), Frequency, true, CanMakeTextureResidentViaHeaps(Texture->parentTexture()));
 	}
-	else
+	
 	{
 		CacheOrSkipResourceResidencyUpdate(Texture, Frequency, true, CanMakeTextureResidentViaHeaps(Texture));
 	}
diff --git a/Engine/Source/Runtime/Apple/MetalRHI/Private/MetalTexture.cpp b/Engine/Source/Runtime/Apple/MetalRHI/Private/MetalTexture.cpp
--- a/Engine/Source/Runtime/Apple/MetalRHI/Private/MetalTexture.cpp
+++ b/Engine/Source/Runtime/Apple/MetalRHI/Private/MetalTexture.cpp
@@ -1995,26 +1995,28 @@ FTextureReferenceRHIRef FMetalDynamicRHI::RHICreateTextureReference(FRHICommandL
 
 void FMetalDynamicRHI::RHIUpdateTextureReference(FRHICommandListBase& RHICmdList, FRHITextureReference* TextureRef, FRHITexture* InNewTexture)
 {
-    FRHITexture* NewTexture = InNewTexture ? InNewTexture : FRHITextureReference::GetDefaultTexture();
-
+	FRHITexture* NewTexture = InNewTexture ? InNewTexture : FRHITextureReference::GetDefaultTexture();
+	
 #if PLATFORM_SUPPORTS_BINDLESS_RENDERING
 	if (TextureRef && TextureRef->IsBindless())
 	{
-		FMetalTextureReference* MetalTextureReference = ResourceCast(TextureRef);
-		FMetalShaderResourceView* MetalTextureRefSRV = MetalTextureReference->BindlessView;
-		
-        FRHIDescriptorHandle DestHandle = MetalTextureRefSRV->GetBindlessHandle();
-        
-        if(DestHandle.IsValid())
-        {
-            FMetalSurface* NewSurface = GetMetalSurfaceFromRHITexture(NewTexture);
-            
-			MetalTextureRefSRV->SurfaceOverride = NewSurface;
-			MetalTextureRefSRV->UpdateView();
+		RHICmdList.EnqueueLambda(TEXT("FMetalDynamicRHI::RHIUpdateTextureReference"), [NewTexture, TextureRef, InNewTexture, InDevice=Device](FRHICommandListBase& ExecutingCmdList)
+		{		
+			FMetalTextureReference* MetalTextureReference = ResourceCast(TextureRef);
+			FMetalShaderResourceView* MetalTextureRefSRV = MetalTextureReference->BindlessView;
 			
-			FMetalBindlessDescriptorManager* BindlessDescriptorManager = Device->GetBindlessDescriptorManager();
-			BindlessDescriptorManager->BindResource(DestHandle, MetalTextureRefSRV);
-        }
+			FRHIDescriptorHandle DestHandle = MetalTextureRefSRV->GetBindlessHandle();
+			if(DestHandle.IsValid())
+			{
+				FMetalSurface* NewSurface = GetMetalSurfaceFromRHITexture(NewTexture);
+				
+				MetalTextureRefSRV->SurfaceOverride = NewSurface;
+				MetalTextureRefSRV->UpdateView();
+				
+				FMetalBindlessDescriptorManager* BindlessDescriptorManager = InDevice->GetBindlessDescriptorManager();
+				BindlessDescriptorManager->BindResource(DestHandle, MetalTextureRefSRV);
+			}
+		});
 	}
 #endif // PLATFORM_SUPPORTS_BINDLESS_RENDERING
 
