Metal RHI

- Fixed scissor dims when rendering with RT-less rasterization
- Removed redundant Metal state calls

diff --git a/Engine/Source/Runtime/Apple/MetalRHI/Private/MetalCommandEncoder.cpp b/Engine/Source/Runtime/Apple/MetalRHI/Private/MetalCommandEncoder.cpp
--- a/Engine/Source/Runtime/Apple/MetalRHI/Private/MetalCommandEncoder.cpp
+++ b/Engine/Source/Runtime/Apple/MetalRHI/Private/MetalCommandEncoder.cpp
@@ -1054,12 +1054,18 @@ void FMetalCommandEncoder::SetShaderBuffer(MTL::FunctionType const FunctionType,
 	FenceResource(Buffer->GetMTLBuffer(), FunctionType);
 	check(index < ML_MaxBuffers);
     
-    if(Device.SupportsFeature(EMetalFeaturesSetBufferOffset) && Buffer && (ShaderBuffers[uint32(FunctionType)].Bound & (1 << index)) && ShaderBuffers[uint32(FunctionType)].Buffers[index] == Buffer)
+    if(Device.SupportsFeature(EMetalFeaturesSetBufferOffset) &&
+	   Buffer && (ShaderBuffers[uint32(FunctionType)].Bound & (1 << index)) &&
+	   ShaderBuffers[uint32(FunctionType)].Buffers[index] == Buffer)
     {
-		SetShaderBufferOffset(FunctionType, Offset, Length, index);
-		ShaderBuffers[uint32(FunctionType)].Usage[index] = Usage;
-		ShaderBuffers[uint32(FunctionType)].ReferencedResources[index] = ReferencedResources;
-		ShaderBuffers[uint32(FunctionType)].SetBufferMetaData(index, Length, GMetalBufferFormats[Format].DataFormat, ElementRowPitch);
+		if(ShaderBuffers[uint32(FunctionType)].Offsets[index] != Offset ||
+		   ShaderBuffers[uint32(FunctionType)].Usage[index] != Usage)
+		{
+			SetShaderBufferOffset(FunctionType, Offset, Length, index);
+			ShaderBuffers[uint32(FunctionType)].Usage[index] = Usage;
+			ShaderBuffers[uint32(FunctionType)].ReferencedResources[index] = ReferencedResources;
+			ShaderBuffers[uint32(FunctionType)].SetBufferMetaData(index, Length, GMetalBufferFormats[Format].DataFormat, ElementRowPitch);
+		}
 	}
     else
     {
@@ -1191,6 +1197,7 @@ void FMetalCommandEncoder::SetShaderBufferOffset(MTL::FunctionType FunctionType,
 	check(Device.SupportsFeature(EMetalFeaturesSetBufferOffset));
 	ShaderBuffers[uint32(FunctionType)].Offsets[index] = Offset;
 	ShaderBuffers[uint32(FunctionType)].SetBufferMetaData(index, Length, GMetalBufferFormats[PF_Unknown].DataFormat, 0);
+	
 	switch (FunctionType)
 	{
 		case MTL::FunctionTypeVertex:
diff --git a/Engine/Source/Runtime/Apple/MetalRHI/Private/MetalRHIContext.cpp b/Engine/Source/Runtime/Apple/MetalRHI/Private/MetalRHIContext.cpp
--- a/Engine/Source/Runtime/Apple/MetalRHI/Private/MetalRHIContext.cpp
+++ b/Engine/Source/Runtime/Apple/MetalRHI/Private/MetalRHIContext.cpp
@@ -79,6 +79,7 @@ void FMetalRHICommandContext::BeginComputeEncoder()
 	
 	if(!CurrentEncoder.IsComputeCommandEncoderActive())
 	{
+		StateCache.ClearPreviousComputeState();
 		if(CurrentEncoder.IsAnyCommandEncoderActive())
 		{
 			CurrentEncoderFence = CurrentEncoder.EndEncoding();
diff --git a/Engine/Source/Runtime/Apple/MetalRHI/Private/MetalStateCache.cpp b/Engine/Source/Runtime/Apple/MetalRHI/Private/MetalStateCache.cpp
--- a/Engine/Source/Runtime/Apple/MetalRHI/Private/MetalStateCache.cpp
+++ b/Engine/Source/Runtime/Apple/MetalRHI/Private/MetalStateCache.cpp
@@ -134,7 +134,6 @@ FMetalStateCache::FMetalStateCache(FMetalDevice& MetalDevice, bool const bInImme
 , bIsRenderTargetActive(false)
 , bHasValidRenderTarget(false)
 , bHasValidColorTarget(false)
-, bScissorRectEnabled(false)
 , bImmediate(bInImmediate)
 {
 	FMemory::Memzero(Viewport);
@@ -232,7 +231,6 @@ void FMetalStateCache::Reset()
 	bIsRenderTargetActive = false;
 	bHasValidRenderTarget = false;
 	bHasValidColorTarget = false;
-	bScissorRectEnabled = false;
 	
 	FMemory::Memzero(DirtyUniformBuffers);
 	FMemory::Memzero(BoundUniformBuffers);
@@ -285,6 +283,7 @@ void FMetalStateCache::Reset()
 	RasterizerState.SafeRelease();
 	GraphicsPSO.SafeRelease();
 	ComputeShader.SafeRelease();
+	PreviousComputeShader.SafeRelease();
 	DepthStencilSurface.SafeRelease();
 	StencilRef = 0;
 	
@@ -353,30 +352,31 @@ static bool MTLScissorRectEqual(MTL::ScissorRect const& Left, MTL::ScissorRect c
 
 void FMetalStateCache::SetScissorRect(bool const bEnable, MTL::ScissorRect const& Rect)
 {
-	if (bScissorRectEnabled != bEnable || !MTLScissorRectEqual(Scissor[0], Rect))
+	if (bEnable)
+	{
+		Scissor[0] = Rect;
+	}
+	else
+	{
+		Scissor[0].x = Viewport[0].originX;
+		Scissor[0].y = Viewport[0].originY;
+		Scissor[0].width = Viewport[0].width;
+		Scissor[0].height = Viewport[0].height;
+	}
+	
+	// Clamp to framebuffer size - Metal doesn't allow scissor to be larger.
+	Scissor[0].x = Scissor[0].x;
+	Scissor[0].y = Scissor[0].y;
+	
+	// FrameBufferSize will be 0 if using RT-less rasterization, so ignore
+	if(FrameBufferSize.width != 0 && FrameBufferSize.height != 0)
 	{
-		bScissorRectEnabled = bEnable;
-		if (bEnable)
-		{
-			Scissor[0] = Rect;
-		}
-		else
-		{
-			Scissor[0].x = Viewport[0].originX;
-			Scissor[0].y = Viewport[0].originY;
-			Scissor[0].width = Viewport[0].width;
-			Scissor[0].height = Viewport[0].height;
-		}
-		
-		// Clamp to framebuffer size - Metal doesn't allow scissor to be larger.
-		Scissor[0].x = Scissor[0].x;
-		Scissor[0].y = Scissor[0].y;
 		Scissor[0].width = FMath::Max((Scissor[0].x + Scissor[0].width <= FMath::RoundToInt32(FrameBufferSize.width)) ? Scissor[0].width : FMath::RoundToInt32(FrameBufferSize.width) - Scissor[0].x, (NS::UInteger)1u);
 		Scissor[0].height = FMath::Max((Scissor[0].y + Scissor[0].height <= FMath::RoundToInt32(FrameBufferSize.height)) ? Scissor[0].height : FMath::RoundToInt32(FrameBufferSize.height) - Scissor[0].y, (NS::UInteger)1u);
-		
-		RasterBits |= EMetalRenderFlagScissorRect;
 	}
 	
+	RasterBits |= EMetalRenderFlagScissorRect;
+	
 	ActiveScissors = 1;
 }
 
@@ -1038,7 +1038,6 @@ void FMetalStateCache::SetViewport(const MTL::Viewport& InViewport)
 	
 	ActiveViewports = 1;
 	
-	bScissorRectEnabled = false;
 	MTL::ScissorRect Rect;
 	
 	Rect.x = InViewport.originX;
@@ -1046,7 +1045,7 @@ void FMetalStateCache::SetViewport(const MTL::Viewport& InViewport)
 	Rect.width = InViewport.width;
 	Rect.height = InViewport.height;
 	
-	SetScissorRect(bScissorRectEnabled, Rect);
+	SetScissorRect(false, Rect);
 }
 
 void FMetalStateCache::SetViewport(uint32 Index, const MTL::Viewport& InViewport)
@@ -2690,11 +2689,15 @@ void FMetalStateCache::SetComputePipelineState(FMetalCommandEncoder& CommandEnco
 {
 	if ((PipelineBits & EMetalPipelineFlagComputeMask) != 0)
 	{
-		FMetalShaderPipelinePtr Pipeline = ComputeShader->GetPipeline();
-	    check(Pipeline);
-	    CommandEncoder.SetComputePipelineState(Pipeline);
-        
-        PipelineBits &= EMetalPipelineFlagRasterMask;
+		if(PreviousComputeShader != ComputeShader)
+		{
+			FMetalShaderPipelinePtr Pipeline = ComputeShader->GetPipeline();
+			check(Pipeline);
+			CommandEncoder.SetComputePipelineState(Pipeline);
+			
+			PipelineBits &= EMetalPipelineFlagRasterMask;
+			PreviousComputeShader = ComputeShader;
+		}
     }
 	
 	if (Device.GetRuntimeDebuggingLevel() >= EMetalDebugLevelFastValidation)
@@ -2881,3 +2884,8 @@ void FMetalStateCache::ReleaseDescriptor(MTL::RenderPassDescriptor* Desc)
 {
 	FMetalRenderPassDescriptorPool::Get().ReleaseDescriptor(Desc);
 }
+
+void FMetalStateCache::ClearPreviousComputeState() 
+{
+	PreviousComputeShader.SafeRelease();
+}
diff --git a/Engine/Source/Runtime/Apple/MetalRHI/Private/MetalStateCache.h b/Engine/Source/Runtime/Apple/MetalRHI/Private/MetalStateCache.h
--- a/Engine/Source/Runtime/Apple/MetalRHI/Private/MetalStateCache.h
+++ b/Engine/Source/Runtime/Apple/MetalRHI/Private/MetalStateCache.h
@@ -186,7 +186,6 @@ public:
 	uint32 GetVertexBufferSize(uint32 const Index);
 	uint32 GetRenderTargetArraySize() const { return RenderTargetArraySize; }
 	FMetalQueryBuffer* GetVisibilityResultsBuffer() const { return VisibilityResults; }
-	bool GetScissorRectEnabled() const { return bScissorRectEnabled; }
 	bool NeedsToSetRenderTarget(const FRHIRenderPassInfo& RenderPassInfo);
 	bool HasValidDepthStencilSurface() const { return IsValidRef(DepthStencilSurface); }
 
@@ -204,6 +203,8 @@ public:
 	void DiscardRenderTargets(bool Depth, bool Stencil, uint32 ColorBitMask);
 	
 	void ReleaseDescriptor(MTL::RenderPassDescriptor* Desc);
+	void ClearPreviousComputeState();
+	
 private:
 	void ConditionalUpdateBackBuffer(FMetalSurface& Surface);
 	
@@ -331,6 +332,8 @@ private:
 	TRefCountPtr<FMetalRasterizerState> RasterizerState;
 	TRefCountPtr<FMetalGraphicsPipelineState> GraphicsPSO;
 	TRefCountPtr<FMetalComputeShader> ComputeShader;
+	TRefCountPtr<FMetalComputeShader> PreviousComputeShader;
+	
 	uint32 StencilRef;
 	
 	FLinearColor BlendFactor;
@@ -356,7 +359,6 @@ private:
 	bool bIsRenderTargetActive;
 	bool bHasValidRenderTarget;
 	bool bHasValidColorTarget;
-	bool bScissorRectEnabled;
     bool bImmediate;
 };
 
