Metal RHI

Removed check for valid rendertarget in PrepareDraw for rasterization without RT

diff --git a/Engine/Source/Runtime/Apple/MetalRHI/Private/MetalCommands.cpp b/Engine/Source/Runtime/Apple/MetalRHI/Private/MetalCommands.cpp
--- a/Engine/Source/Runtime/Apple/MetalRHI/Private/MetalCommands.cpp
+++ b/Engine/Source/Runtime/Apple/MetalRHI/Private/MetalCommands.cpp
@@ -799,12 +799,6 @@ bool FMetalRHICommandContext::PrepareToDraw(uint32 PrimitiveType)
 	TRefCountPtr<FMetalGraphicsPipelineState> CurrentPSO = StateCache.GetGraphicsPSO();
 	check(IsValidRef(CurrentPSO));
 	
-	// Enforce calls to SetRenderTarget prior to issuing draw calls.
-	if (!StateCache.GetHasValidRenderTarget())
-	{
-		return false;
-	}
-	
 	FMetalHashedVertexDescriptor const& VertexDesc = CurrentPSO->VertexDeclaration->Layout;
 	
 	// Validate the vertex layout in debug mode, or when the validation layer is enabled for development builds.
