MetalRHI

Refactored Metal deferred deletion to occur after command buffer submission.

diff --git a/Engine/Source/Runtime/Apple/MetalRHI/Private/MetalBindlessDescriptors.cpp b/Engine/Source/Runtime/Apple/MetalRHI/Private/MetalBindlessDescriptors.cpp
--- a/Engine/Source/Runtime/Apple/MetalRHI/Private/MetalBindlessDescriptors.cpp
+++ b/Engine/Source/Runtime/Apple/MetalRHI/Private/MetalBindlessDescriptors.cpp
@@ -10,6 +10,7 @@
 
 #include "MetalDevice.h"
 #include "MetalCommandEncoder.h"
+#include "MetalDynamicRHI.h"
 
 #define USE_CPU_DESCRIPTOR_COPY 0
 #define USE_DESCRIPTOR_BUFFER_COPY !USE_CPU_DESCRIPTOR_COPY
@@ -32,7 +33,6 @@ static FAutoConsoleVariableRef CVarBindlessSamplerDescriptorHeapSize(
 
 FMetalDescriptorHeap::FMetalDescriptorHeap(FMetalDevice& MetalDevice, const ERHIDescriptorHeapType DescriptorType)
 	: Device(MetalDevice)
-	, DeferredDeletionListIndex(0)
 	, ResourceHeap(nullptr)
 	, Type(DescriptorType)
 {
@@ -68,26 +68,10 @@ void FMetalDescriptorHeap::Init(const int32 HeapSize)
 	MaxDirtyIndex = 0;
 }
 
-void FMetalDescriptorHeap::Reset()
-{
-	DeferredDeletionListIndex = (Device.GetFrameNumberRHIThread() % NumPendingFrame);
-
-	{
-		FScopeLock ScopeLock(&FreeListCS);
-
-		while (!DeferredDeletionList[DeferredDeletionListIndex].IsEmpty())
-		{
-			FRHIDescriptorHandle DeletionIndex;
-			DeferredDeletionList[DeferredDeletionListIndex].Dequeue(DeletionIndex);
-			FreeList.Enqueue(DeletionIndex.GetIndex());
-		}
-	}
-}
-
 void FMetalDescriptorHeap::FreeDescriptor(FRHIDescriptorHandle DescriptorHandle)
 {
 	FScopeLock ScopeLock(&FreeListCS);
-	DeferredDeletionList[DeferredDeletionListIndex].Enqueue(DescriptorHandle);
+	FreeList.Enqueue(DescriptorHandle.GetIndex());
 }
 
 uint32 FMetalDescriptorHeap::GetFreeResourceIndex()
@@ -157,17 +141,6 @@ void FMetalBindlessDescriptorManager::Init()
 	bIsSupported = true;
 }
 
-void FMetalBindlessDescriptorManager::Reset()
-{
-	if(!bIsSupported)
-	{
-		return;
-	}
-
-	StandardResources.Reset();
-	SamplerResources.Reset();
-}
-
 FRHIDescriptorHandle FMetalBindlessDescriptorManager::ReserveDescriptor(ERHIDescriptorHeapType InType)
 {
 	switch (InType)
@@ -413,7 +386,7 @@ void FMetalBindlessDescriptorManager::UpdateDescriptorsWithCompute()
 		InContext->WaitForEvent(InEvt, 1);
 	});
 	
-	Device.ReleaseFunction([Evt](){});
+	FMetalDynamicRHI::Get().DeferredDelete([Evt](){});
 #endif
 }
 
diff --git a/Engine/Source/Runtime/Apple/MetalRHI/Private/MetalBindlessDescriptors.h b/Engine/Source/Runtime/Apple/MetalRHI/Private/MetalBindlessDescriptors.h
--- a/Engine/Source/Runtime/Apple/MetalRHI/Private/MetalBindlessDescriptors.h
+++ b/Engine/Source/Runtime/Apple/MetalRHI/Private/MetalBindlessDescriptors.h
@@ -19,8 +19,6 @@ struct FMetalDescriptorHeap
 
     void                            Init(const int32 HeapSize);
 
-    void                            Reset();
-
     FRHIDescriptorHandle            ReserveDescriptor();
     void                            FreeDescriptor(FRHIDescriptorHandle DescriptorHandle);
     uint32                          GetFreeResourceIndex();
@@ -33,11 +31,6 @@ struct FMetalDescriptorHeap
 	
     FCriticalSection                FreeListCS;
     TQueue<uint32>                  FreeList;
-    uint32                          DeferredDeletionListIndex;
-
-    static constexpr int32 NumPendingFrame = 3;
-
-    TQueue<FRHIDescriptorHandle>    DeferredDeletionList[NumPendingFrame];
 
     std::atomic<uint32>             PeakDescriptorCount;
     struct IRDescriptorTableEntry*  Descriptors;
@@ -66,9 +59,8 @@ public:
                             ~FMetalBindlessDescriptorManager();
 
     void                    Init();
-    void                    Reset();
 
-    FRHIDescriptorHandle     ReserveDescriptor(ERHIDescriptorHeapType InType);
+    FRHIDescriptorHandle    ReserveDescriptor(ERHIDescriptorHeapType InType);
     void                    FreeDescriptor(FRHIDescriptorHandle DescriptorHandle);
 
     void                    BindSampler(FRHIDescriptorHandle DescriptorHandle, MTL::SamplerState* Sampler);
diff --git a/Engine/Source/Runtime/Apple/MetalRHI/Private/MetalBuffer.cpp b/Engine/Source/Runtime/Apple/MetalRHI/Private/MetalBuffer.cpp
--- a/Engine/Source/Runtime/Apple/MetalRHI/Private/MetalBuffer.cpp
+++ b/Engine/Source/Runtime/Apple/MetalRHI/Private/MetalBuffer.cpp
@@ -7,6 +7,7 @@
 #include "MetalLLM.h"
 #include "MetalProfiler.h"
 #include "MetalRHIPrivate.h"
+#include "MetalDynamicRHI.h"
 #include "Templates/AlignmentTemplates.h"
 #include <objc/runtime.h>
 
@@ -915,7 +916,7 @@ FMetalRingBufferRef::FMetalRingBufferRef(FMetalDevice& InDevice, FMetalBufferPtr
 
 FMetalRingBufferRef::~FMetalRingBufferRef()
 {
-	Device.ReleaseBuffer(Buffer);
+	FMetalDynamicRHI::Get().DeferredDelete(Buffer);
 }
 
 FMetalSubBufferRing::FMetalSubBufferRing(FMetalDevice& MetalDevice, NS::UInteger Size, NS::UInteger Alignment, MTL::ResourceOptions InOptions)
diff --git a/Engine/Source/Runtime/Apple/MetalRHI/Private/MetalCommandQueue.cpp b/Engine/Source/Runtime/Apple/MetalRHI/Private/MetalCommandQueue.cpp
--- a/Engine/Source/Runtime/Apple/MetalRHI/Private/MetalCommandQueue.cpp
+++ b/Engine/Source/Runtime/Apple/MetalRHI/Private/MetalCommandQueue.cpp
@@ -12,6 +12,7 @@
 #include "MetalProfiler.h"
 #include "MetalRHIPrivate.h"
 #include "Misc/ConfigCacheIni.h"
+#include "MetalDynamicRHI.h"
 
 #if !UE_BUILD_SHIPPING
 #import "MetalThirdParty.h"
@@ -116,7 +117,6 @@ FMetalCommandBuffer* FMetalCommandQueue::CreateCommandBuffer(void)
     CmdBufferDesc->release();
                                                            
     FMetalCommandBuffer* CommandBuffer = new FMetalCommandBuffer(CmdBuffer);
-	Device.AddCommandBufferFence(CommandBuffer->GetCompletionFence());
     
 	INC_DWORD_STAT(STAT_MetalCommandBufferCreatedPerFrame);
 	return CommandBuffer;
@@ -134,10 +134,6 @@ void FMetalCommandQueue::CommitCommandBuffer(FMetalCommandBuffer* CommandBuffer)
 	{
 		CommandBuffer->GetMTLCmdBuffer()->waitUntilCompleted();
 	}
-    
-    Device.ReleaseFunction([CommandBuffer]() {
-        delete CommandBuffer;
-    });
 }
 
 FMetalFence* FMetalCommandQueue::CreateFence(NS::String* Label) const
diff --git a/Engine/Source/Runtime/Apple/MetalRHI/Private/MetalDevice.cpp b/Engine/Source/Runtime/Apple/MetalRHI/Private/MetalDevice.cpp
--- a/Engine/Source/Runtime/Apple/MetalRHI/Private/MetalDevice.cpp
+++ b/Engine/Source/Runtime/Apple/MetalRHI/Private/MetalDevice.cpp
@@ -52,7 +52,7 @@ static FAutoConsoleVariableRef CVarMetalBufferScribble(
 	GMetalBufferScribble,
 	TEXT("Debug option: when enabled will scribble over the buffer contents with a single value when releasing buffer objects, or regions thereof. (Default: 0, Off)"));
 
-static int32 GMetalResourcePurgeOnDelete = 0;
+int32 GMetalResourcePurgeOnDelete = 0;
 static FAutoConsoleVariableRef CVarMetalResourcePurgeOnDelete(
 	TEXT("rhi.Metal.ResourcePurgeOnDelete"),
 	GMetalResourcePurgeOnDelete,
@@ -411,9 +411,6 @@ FMetalDevice::~FMetalDevice()
 {
 	FRHICommandListImmediate& RHICmdList = FRHICommandListImmediate::Get();
 	RHICmdList.SubmitAndBlockUntilGPUIdle();
-
-	FlushFreeList(true);
-	ClearFreeList();
 	
 	delete &(GetCommandQueue());
 	delete PSOManager;
@@ -639,202 +636,19 @@ void FMetalDevice::EndDrawingViewport(bool bPresent)
 	}
 }
 
-void FMetalDevice::ClearFreeList()
-{
-	FreeListMutex.Lock();
-	
-	uint32 Index = 0;
-	while(Index < DelayedFreeLists.Num())
-	{
-		FMetalDelayedFreeList* Pair = DelayedFreeLists[Index];
-		if(METAL_DEBUG_OPTION(Pair->DeferCount-- <= 0 &&) Pair->IsComplete())
-		{
-			for(NS::Object* Entry : Pair->ObjectFreeList )
-			{
-				Entry->release();
-			}
-			for (FMetalBufferPtr Buffer : Pair->UsedBuffers)
-			{
-                Buffer->MarkDeleted();
-				
-#if METAL_DEBUG_OPTIONS
-                MTL::Buffer* MTLBuffer = Buffer->GetMTLBuffer();
-				if (GMetalResourcePurgeOnDelete && !MTLBuffer->heap() &&
-                    Buffer->GetOffset() == 0 && Buffer->GetLength() == MTLBuffer->length())
-				{
-                    MTLBuffer->setPurgeableState(MTL::PurgeableStateEmpty);
-				}
-#endif
-			}
-			for (MTLTexturePtr Texture : Pair->UsedTextures)
-			{
-                if (!Texture->buffer() && !Texture->parentTexture())
-				{
-#if METAL_DEBUG_OPTIONS
-					if (GMetalResourcePurgeOnDelete && !Texture->heap())
-					{
-						Texture->setPurgeableState(MTL::PurgeableStateEmpty);
-					}
-#endif
-					Heap.ReleaseTexture(nullptr, Texture);
-				}
-			}
-			for (FMetalFence* Fence : Pair->FenceFreeList)
-			{
-				FMetalFencePool::Get().ReleaseFence(Fence);
-			}
-            for (TFunction<void()>& Function : Pair->FunctionFreeList)
-            {
-                Function();
-            }
-			delete Pair;
-			DelayedFreeLists.RemoveAt(Index, EAllowShrinking::No);
-		}
-		else
-		{
-			Index++;
-		}
-	}
-	
-	FreeListMutex.Unlock();
-}
-
 void FMetalDevice::DrainHeap()
 {
 	Heap.Compact(false);
 }
 
-bool FMetalDevice::FMetalDelayedFreeList::IsComplete() const
-{
-	bool bFinished = true;
-	for (TSharedPtr<FMetalCommandBufferFence, ESPMode::ThreadSafe> Fence : Fences)
-	{
-		bFinished &= Fence->Wait(0);
-
-		if (!bFinished)
-			break;
-	}
-	return bFinished;
-}
-
-void FMetalDevice::AddCommandBufferFence(TSharedPtr<FMetalCommandBufferFence, ESPMode::ThreadSafe> Fence)
-{
-	FreeListMutex.Lock();
-	
-	// TODO: Temporary, removing this code for 5.5, moving to RHI deferred delete
-	FreeListFences.Add(Fence);
-	if(bPendingGarbageCollect)
-	{
-		FlushFreeList(true);
-		bPendingGarbageCollect = false;
-	}
-	FreeListMutex.Unlock();
-}
-
-void FMetalDevice::FlushFreeList(bool const bFlushFences)
-{
-	FreeListMutex.Lock();
-	
-	FMetalDelayedFreeList* NewList = new FMetalDelayedFreeList;
-	
-	// Get the committed command buffer fences and clear the array in the command-queue
-	NewList->Fences = MoveTemp(FreeListFences);
-	
-	METAL_DEBUG_OPTION(NewList->DeferCount = GMetalResourceDeferDeleteNumFrames);
-	NewList->UsedBuffers = MoveTemp(UsedBuffers);
-	NewList->UsedTextures = MoveTemp(UsedTextures);
-	NewList->ObjectFreeList = ObjectFreeList;
-	if (bFlushFences)
-	{
-		TArray<FMetalFence*> Fences;
-		FenceFreeList.PopAll(Fences);
-		for (FMetalFence* Fence : Fences)
-		{
-			if(!UsedFences.Contains(Fence))
-			{
-				UsedFences.Add(Fence);
-			}
-		}
-		NewList->FenceFreeList = MoveTemp(UsedFences);
-        
-	}
-    NewList->FunctionFreeList = MoveTemp(FunctionFreeList);
-    
-#if METAL_DEBUG_OPTIONS
-	if (FrameFences.Num())
-	{
-		FrameFences.Empty();
-	}
-#endif
-	ObjectFreeList.Empty(ObjectFreeList.Num());
-	
-	DelayedFreeLists.Add(NewList);
-	
-	FreeListMutex.Unlock();
-}
-
 void FMetalDevice::GarbageCollect()
 {
-	MarkForGarbageCollect();
-	ClearFreeList();
 	DrainHeap();
 	
 	TransferBufferAllocator->Cleanup();
 	UniformBufferAllocator->Cleanup();
 }
 
-void FMetalDevice::ReleaseObject(NS::Object* obj)
-{
-	check(GIsMetalInitialized);
-	 
-	check(obj);
-	FreeListMutex.Lock();
-	if(!ObjectFreeList.Contains(obj))
-	{
-		ObjectFreeList.Add(obj);
-	}
-	else
-	{
-		obj->release();
-	}
-	FreeListMutex.Unlock();
-}
-
-void FMetalDevice::ReleaseTexture(MTLTexturePtr Texture)
-{
-	check(GIsMetalInitialized);
-	
-	check(Texture);
-	FreeListMutex.Lock();
-	if(!UsedTextures.Contains(Texture))
-	{
-		UsedTextures.Add(Texture);
-	}
-	FreeListMutex.Unlock();
-}
-
-void FMetalDevice::ReleaseFence(FMetalFence* Fence)
-{
-#if METAL_DEBUG_OPTIONS
-	if(GetRuntimeDebuggingLevel() >= EMetalDebugLevelValidation)
-	{
-		FScopeLock Lock(&FreeListMutex);
-		FrameFences.Add(Fence);
-	}
-#endif
-	
-	check(GIsMetalInitialized);
-	
-	check(Fence);
-	FenceFreeList.Push(Fence);
-}
-
-void FMetalDevice::ReleaseFunction(TFunction<void()> Func)
-{
-	check(GIsMetalInitialized);
-	FunctionFreeList.Push(Func);
-}
-
 MTLTexturePtr FMetalDevice::CreateTexture(FMetalSurface* Surface, MTL::TextureDescriptor* Descriptor)
 {
 	MTLTexturePtr Tex = Heap.CreateTexture(Descriptor, Surface);
@@ -888,20 +702,6 @@ MTLEventPtr FMetalDevice::CreateEvent()
 	return Event;
 }
 
-void FMetalDevice::ReleaseBuffer(FMetalBufferPtr Buffer)
-{
-	if(GIsMetalInitialized)
-	{
-		check(Buffer);
-		FreeListMutex.Lock();
-		if(!UsedBuffers.Contains(Buffer))
-		{
-			UsedBuffers.Add(Buffer);
-		}
-		FreeListMutex.Unlock();
-	}
-}
-
 uint32 FMetalDevice::GetDeviceIndex(void) const
 {
 	return DeviceIndex;
diff --git a/Engine/Source/Runtime/Apple/MetalRHI/Private/MetalDevice.h b/Engine/Source/Runtime/Apple/MetalRHI/Private/MetalDevice.h
--- a/Engine/Source/Runtime/Apple/MetalRHI/Private/MetalDevice.h
+++ b/Engine/Source/Runtime/Apple/MetalRHI/Private/MetalDevice.h
@@ -113,19 +113,6 @@ public:
 	FMetalBufferPtr CreatePooledBuffer(FMetalPooledBufferArgs const& Args);
 	MTLEventPtr CreateEvent();
 	
-	void ReleaseBuffer(FMetalBufferPtr Buf);
-	void ReleaseObject(NS::Object* Obj);
-	void ReleaseTexture(MTLTexturePtr Texture);
-	void ReleaseFence(FMetalFence* Fence);
-    void ReleaseFunction(TFunction<void()>);
-	
-	void AddCommandBufferFence(TSharedPtr<FMetalCommandBufferFence, ESPMode::ThreadSafe> Fence);
-	void MarkForGarbageCollect()
-	{
-		bPendingGarbageCollect = true;
-	};
-	
-	void ClearFreeList();
 	void DrainHeap();
 	void GarbageCollect();
 	
@@ -239,32 +226,6 @@ private:
 	/** GPU Frame Capture Manager */
 	FMetalCaptureManager* CaptureManager;
 	
-	/** Free lists for releasing objects only once it is safe to do so */
-	TSet<FMetalBufferPtr> UsedBuffers;
-	TSet<MTLTexturePtr> UsedTextures;
-	TSet<FMetalFence*> UsedFences;
-	TLockFreePointerListLIFO<FMetalFence> FenceFreeList;
-    TArray<TFunction<void()>> FunctionFreeList;
-	TSet<NS::Object*> ObjectFreeList;
-	
-	struct FMetalDelayedFreeList
-	{
-		bool IsComplete() const;
-		TArray<TSharedPtr<FMetalCommandBufferFence, ESPMode::ThreadSafe>> Fences;
-		TSet<FMetalBufferPtr> UsedBuffers;
-		TSet<MTLTexturePtr> UsedTextures;
-		TSet<FMetalFence*> FenceFreeList;
-		TSet<NS::Object*> ObjectFreeList;
-        TArray<TFunction<void()>> FunctionFreeList;
-#if METAL_DEBUG_OPTIONS
-		int32 DeferCount;
-#endif
-	};
-	
-	TArray<TSharedPtr<FMetalCommandBufferFence, ESPMode::ThreadSafe>> FreeListFences;
-	TArray<FMetalDelayedFreeList*> DelayedFreeLists;
-	bool bPendingGarbageCollect = false;
-	
     FMetalTempAllocator* UniformBufferAllocator;
 	FMetalTempAllocator* TransferBufferAllocator;
 	
diff --git a/Engine/Source/Runtime/Apple/MetalRHI/Private/MetalDynamicRHI.cpp b/Engine/Source/Runtime/Apple/MetalRHI/Private/MetalDynamicRHI.cpp
--- a/Engine/Source/Runtime/Apple/MetalRHI/Private/MetalDynamicRHI.cpp
+++ b/Engine/Source/Runtime/Apple/MetalRHI/Private/MetalDynamicRHI.cpp
@@ -268,3 +268,17 @@ bool FMetalDynamicRHI::RHIMatchPrecachePSOInitializers(const FGraphicsPipelineSt
 
 	return true;
 }
+
+void FMetalDynamicRHI::AddDeferredDeleteFence(TSharedPtr<FMetalCommandBufferFence, ESPMode::ThreadSafe> Fence)
+{
+	FScopeLock Lock(&ObjectsToDeleteCS);
+	DeferredDeleteFences.Add(Fence);
+}
+
+void FMetalDynamicRHI::GatherDeferredDeleteObjects(TArray<FMetalDeferredDeleteObject>& DeferredDeleteObjects,
+											   TArray<TSharedPtr<FMetalCommandBufferFence, ESPMode::ThreadSafe>>& WaitFences)
+{
+	FScopeLock Lock(&ObjectsToDeleteCS);
+	WaitFences = MoveTemp(DeferredDeleteFences);
+	DeferredDeleteObjects = MoveTemp(ObjectsToDelete);
+}
diff --git a/Engine/Source/Runtime/Apple/MetalRHI/Private/MetalDynamicRHI.h b/Engine/Source/Runtime/Apple/MetalRHI/Private/MetalDynamicRHI.h
--- a/Engine/Source/Runtime/Apple/MetalRHI/Private/MetalDynamicRHI.h
+++ b/Engine/Source/Runtime/Apple/MetalRHI/Private/MetalDynamicRHI.h
@@ -22,6 +22,41 @@ class FMetalDevice;
 class FMetalRayTracingCompactionRequestHandler;
 #endif // METAL_RHI_RAYTRACING
 
+struct FMetalDeferredDeleteObject
+{
+	typedef TVariant<FMetalBufferPtr, 
+					MTLTexturePtr,
+					NS::Object*,
+#if PLATFORM_SUPPORTS_BINDLESS_RENDERING
+					FRHIDescriptorHandle,
+#endif
+					FMetalFence*,
+					TUniqueFunction<void()>*> TObjectStorage;
+	
+	TObjectStorage Storage;
+	
+	explicit FMetalDeferredDeleteObject(FMetalBufferPtr InBuffer) : Storage(TInPlaceType<FMetalBufferPtr>(), InBuffer)
+	{}
+
+	explicit FMetalDeferredDeleteObject(MTLTexturePtr InTexture) : Storage(TInPlaceType<MTLTexturePtr>(), InTexture)
+	{}
+	
+	explicit FMetalDeferredDeleteObject(NS::Object* InObject) : Storage(TInPlaceType<NS::Object*>(), InObject)
+	{}
+
+#if PLATFORM_SUPPORTS_BINDLESS_RENDERING
+	explicit FMetalDeferredDeleteObject(FRHIDescriptorHandle InHandle) : Storage(TInPlaceType<FRHIDescriptorHandle>(), InHandle)
+	{}
+#endif
+
+	explicit FMetalDeferredDeleteObject(FMetalFence* InFence) : Storage(TInPlaceType<FMetalFence*>(), InFence)
+	{}
+	
+	explicit FMetalDeferredDeleteObject(TUniqueFunction<void()>&& Func) : 
+			Storage(TInPlaceType<TUniqueFunction<void()>*>(), new TUniqueFunction<void()>(MoveTemp(Func)))
+	{}
+};
+
 /** The interface which is implemented by the dynamically bound RHI. */
 class FMetalDynamicRHI : public FDynamicRHI
 {
@@ -126,6 +161,7 @@ public:
 	virtual IRHIUploadContext* RHIGetUploadContext() final override;
 	
 	virtual IRHIComputeContext* RHIGetCommandContext(ERHIPipeline Pipeline, FRHIGPUMask GPUMask) final override;
+	virtual void RHIProcessDeleteQueue() final override;
 	virtual void RHIFinalizeContext(FRHIFinalizeContextArgs&& Args, TRHIPipelineArray<IRHIPlatformCommandList*>& Output) final override;
 	virtual void RHISubmitCommandLists(FRHISubmitCommandListsArgs&& Args) final override;
 
@@ -175,6 +211,19 @@ public:
 	
 	virtual void RHIBeginRenderQuery_TopOfPipe(FRHICommandListBase& RHICmdList, FRHIRenderQuery* RenderQuery) override;
 	virtual void RHIEndRenderQuery_TopOfPipe  (FRHICommandListBase& RHICmdList, FRHIRenderQuery* RenderQuery) override;
+	
+	template <typename ...Args>
+	void DeferredDelete(Args&&... InArgs)
+	{
+		check(!IsInGameThread() || !IsRunningRHIInSeparateThread());
+		FScopeLock Lock(&ObjectsToDeleteCS);
+		ObjectsToDelete.Emplace(Forward<Args>(InArgs)...);
+	}
+	
+	void AddDeferredDeleteFence(TSharedPtr<FMetalCommandBufferFence, ESPMode::ThreadSafe> Fence);
+	void GatherDeferredDeleteObjects(TArray<FMetalDeferredDeleteObject>& DeferredDeleteObjects,
+									 TArray<TSharedPtr<FMetalCommandBufferFence, ESPMode::ThreadSafe>>& WaitFences);
+	void ProcessDeferredDeleteQueue();
 
 private:
 	FMetalDevice* Device;
@@ -183,6 +232,18 @@ private:
 	TMap<uint32, FVertexDeclarationRHIRef> VertexDeclarationCache;
 	TLockFreePointerListUnordered<FMetalRHICommandContext, PLATFORM_CACHE_LINE_SIZE> MetalCommandContextPool;
 	
+	struct FDeferredDeleteData
+	{
+		TArray<FMetalDeferredDeleteObject> DeferredDeleteObjects;
+		TArray<TSharedPtr<FMetalCommandBufferFence, ESPMode::ThreadSafe>> WaitFences;
+	};
+	
+	TArray<FDeferredDeleteData> DeferredDeleteQueue;
+	
+	FCriticalSection ObjectsToDeleteCS;
+	TArray<FMetalDeferredDeleteObject> ObjectsToDelete;
+	TArray<TSharedPtr<FMetalCommandBufferFence, ESPMode::ThreadSafe>> DeferredDeleteFences;
+	
 #if METAL_USE_METAL_SHADER_CONVERTER
     struct IRCompiler* CompilerInstance;
 #endif
diff --git a/Engine/Source/Runtime/Apple/MetalRHI/Private/MetalRHI.cpp b/Engine/Source/Runtime/Apple/MetalRHI/Private/MetalRHI.cpp
--- a/Engine/Source/Runtime/Apple/MetalRHI/Private/MetalRHI.cpp
+++ b/Engine/Source/Runtime/Apple/MetalRHI/Private/MetalRHI.cpp
@@ -7,6 +7,7 @@
 #include "MetalRHI.h"
 #include "MetalDynamicRHI.h"
 #include "MetalRHIPrivate.h"
+#include "MetalCommandBuffer.h"
 #include "Misc/MessageDialog.h"
 #include "Modules/ModuleManager.h"
 #include "RenderUtils.h"
@@ -54,6 +55,7 @@ static TAutoConsoleVariable<bool> CVarEnableMetalPSOFileCacheWhenPrecachingActiv
 	TEXT("true: Allow both PSO file cache and precaching."),
 	ECVF_RenderThreadSafe | ECVF_ReadOnly);
 
+extern int32 GMetalResourcePurgeOnDelete;
 
 static void ValidateTargetedRHIFeatureLevelExists(EShaderPlatform Platform)
 {
@@ -1428,9 +1430,6 @@ bool FMetalDynamicRHI::RHIGetAvailableResolutions(FScreenResolutionArray& Resolu
 void FMetalDynamicRHI::RHIFlushResources()
 {
     MTL_SCOPED_AUTORELEASE_POOL;
-    
-	Device->MarkForGarbageCollect();
-	Device->ClearFreeList();
 	Device->DrainHeap();
 }
 
@@ -1477,6 +1476,7 @@ uint16 FMetalDynamicRHI::RHIGetPlatformTextureMaxSampleCount()
 void FMetalDynamicRHI::RHIBlockUntilGPUIdle()
 {
 	Device->WaitForGPUIdle();
+	ProcessDeferredDeleteQueue();
 }
 
 uint32 FMetalDynamicRHI::RHIGetGPUFrameCycles(uint32 GPUIndex)
@@ -1514,6 +1514,115 @@ public:
 
 static TLockFreePointerListUnordered<FMetalRHIUploadContext, PLATFORM_CACHE_LINE_SIZE> MetalUploadContextPool;
 
+void FMetalDynamicRHI::ProcessDeferredDeleteQueue()
+{
+	uint32 Index = 0;
+	while(Index < DeferredDeleteQueue.Num())
+	{
+		FDeferredDeleteData& Data = DeferredDeleteQueue[Index];
+		
+		bool bFencesReady = true;
+		for (TSharedPtr<FMetalCommandBufferFence, ESPMode::ThreadSafe> Fence : Data.WaitFences)
+		{
+			bFencesReady &= Fence->Wait(0);
+		}
+		
+		if(!bFencesReady)
+		{
+			Index++;
+			break;
+		}
+		
+		for(FMetalDeferredDeleteObject& Object : Data.DeferredDeleteObjects)
+		{	
+			switch (Object.Storage.GetIndex())
+			{
+				case FMetalDeferredDeleteObject::TObjectStorage::IndexOfType<NS::Object*>():
+				{
+					Object.Storage.Get<NS::Object*>()->release();
+					break;
+				}
+				case FMetalDeferredDeleteObject::TObjectStorage::IndexOfType<FMetalBufferPtr>():
+				{
+					FMetalBufferPtr Buffer = Object.Storage.Get<FMetalBufferPtr>();
+					
+					Buffer->MarkDeleted();
+					
+#if METAL_DEBUG_OPTIONS
+					MTL::Buffer* MTLBuffer = Buffer->GetMTLBuffer();
+					if (GMetalResourcePurgeOnDelete && !MTLBuffer->heap() &&
+						Buffer->GetOffset() == 0 && Buffer->GetLength() == MTLBuffer->length())
+					{
+						MTLBuffer->setPurgeableState(MTL::PurgeableStateEmpty);
+					}
+#endif
+					break;
+				}
+				case FMetalDeferredDeleteObject::TObjectStorage::IndexOfType<MTLTexturePtr>():
+				{
+					MTLTexturePtr Texture = Object.Storage.Get<MTLTexturePtr>();
+					
+					if (!Texture->buffer() && !Texture->parentTexture())
+					{
+#if METAL_DEBUG_OPTIONS
+						if (GMetalResourcePurgeOnDelete && !Texture->heap())
+						{
+							Texture->setPurgeableState(MTL::PurgeableStateEmpty);
+						}
+#endif
+						Device->GetResourceHeap().ReleaseTexture(nullptr, Texture);
+					}
+					break;
+				}
+#if PLATFORM_SUPPORTS_BINDLESS_RENDERING
+				case FMetalDeferredDeleteObject::TObjectStorage::IndexOfType<FRHIDescriptorHandle>():
+				{
+					FRHIDescriptorHandle Handle = Object.Storage.Get<FRHIDescriptorHandle>();
+					
+					FMetalBindlessDescriptorManager* BindlessDescriptorManager = Device->GetBindlessDescriptorManager();
+					check(BindlessDescriptorManager);
+	
+					BindlessDescriptorManager->FreeDescriptor(Handle);
+					
+					break;
+				}
+#endif
+				case FMetalDeferredDeleteObject::TObjectStorage::IndexOfType<FMetalFence*>():
+				{
+					FMetalFence* Fence = Object.Storage.Get<FMetalFence*>();
+					FMetalFencePool::Get().ReleaseFence(Fence);
+					
+					break;
+				}
+				case FMetalDeferredDeleteObject::TObjectStorage::IndexOfType<TUniqueFunction<void()>*>():
+				{
+					TUniqueFunction<void()>* Func = Object.Storage.Get<TUniqueFunction<void()>*>();
+					(*Func)();
+					delete Func;
+					
+					break;
+				}
+				default:
+				{
+					checkNoEntry();
+				}
+			}
+		}
+		
+		DeferredDeleteQueue.RemoveAt(Index, EAllowShrinking::No);
+	}
+}
+
+void FMetalDynamicRHI::RHIProcessDeleteQueue()
+{
+	FDeferredDeleteData NewData;	
+	GatherDeferredDeleteObjects(NewData.DeferredDeleteObjects, NewData.WaitFences);
+	
+	DeferredDeleteQueue.Add(MoveTemp(NewData));
+	
+	ProcessDeferredDeleteQueue();
+}
+
 void FMetalDynamicRHI::RHIFinalizeContext(FRHIFinalizeContextArgs&& Args, TRHIPipelineArray<IRHIPlatformCommandList*>& Output)
 {
 	MTL_SCOPED_AUTORELEASE_POOL;
@@ -1573,6 +1682,7 @@ void FMetalDynamicRHI::RHISubmitCommandLists(FRHISubmitCommandListsArgs&& Args)
 		{
 			if(CommandBuffer)
 			{
+				AddDeferredDeleteFence(CommandBuffer->GetCompletionFence());
 				Device->GetCommandQueue().CommitCommandBuffer(CommandBuffer);
 			}
 		}
diff --git a/Engine/Source/Runtime/Apple/MetalRHI/Private/MetalRHIBuffer.cpp b/Engine/Source/Runtime/Apple/MetalRHI/Private/MetalRHIBuffer.cpp
--- a/Engine/Source/Runtime/Apple/MetalRHI/Private/MetalRHIBuffer.cpp
+++ b/Engine/Source/Runtime/Apple/MetalRHI/Private/MetalRHIBuffer.cpp
@@ -330,7 +330,7 @@ void FMetalRHIBuffer::ReleaseBuffers()
 		check(Buffer);
 		
 		METAL_INC_DWORD_STAT_BY(MemFreed, Buffer->GetLength(), GetUsage());
-		Device.ReleaseBuffer(Buffer);
+		FMetalDynamicRHI::Get().DeferredDelete(Buffer);
 	}
 } 
 
@@ -514,7 +514,7 @@ void FMetalRHIBuffer::Unlock(FRHICommandListBase& RHICmdList)
 				UploadContext.EnqueueFunction([&InDevice=Device, Size=LockSize, Dest=CurrentBuffer, InTransferBuffer=TransferBuffer](FMetalRHICommandContext* Context)
 				{
 					Context->CopyFromBufferToBuffer(InTransferBuffer, 0, Dest, 0, Size);
-					InDevice.ReleaseBuffer(InTransferBuffer);
+					FMetalDynamicRHI::Get().DeferredDelete(InTransferBuffer);
 				});
 				
 				TransferBuffer = nullptr;
@@ -539,7 +539,7 @@ void FMetalRHIBuffer::Unlock(FRHICommandListBase& RHICmdList)
 			if(TransferBuffer)
 			{
 				check(Mode == MTL::StorageModePrivate);
-				Device.ReleaseBuffer(TransferBuffer);
+				FMetalDynamicRHI::Get().DeferredDelete(TransferBuffer);
 				TransferBuffer = nullptr;
 			}
 		}
@@ -590,7 +590,7 @@ void FMetalRHIBuffer::ReleaseOwnership()
     if(TransferBuffer)
     {
         METAL_INC_DWORD_STAT_BY(MemFreed, TransferBuffer->GetLength(), GetUsage());
-		Device.ReleaseBuffer(TransferBuffer);
+		FMetalDynamicRHI::Get().DeferredDelete(TransferBuffer);
     }
     
 	ReleaseBuffers();
@@ -601,13 +601,13 @@ void FMetalRHIBuffer::ReleaseOwnership()
         auto ReleaseFunction = [ReleaseData=Data](){
             delete ReleaseData;
         };
-        Device.ReleaseFunction(ReleaseFunction);
+		FMetalDynamicRHI::Get().DeferredDelete(ReleaseFunction);
     }
 
 #if METAL_RHI_RAYTRACING
     if (EnumHasAnyFlags(GetUsage(), BUF_AccelerationStructure))
     {
-        Device.ReleaseObject(AccelerationStructureHandle);
+        Device.DeferredDelete(AccelerationStructureHandle);
         AccelerationStructureHandle = nullptr;
     }
 #endif // METAL_RHI_RAYTRACING
diff --git a/Engine/Source/Runtime/Apple/MetalRHI/Private/MetalRHIContext.cpp b/Engine/Source/Runtime/Apple/MetalRHI/Private/MetalRHIContext.cpp
--- a/Engine/Source/Runtime/Apple/MetalRHI/Private/MetalRHIContext.cpp
+++ b/Engine/Source/Runtime/Apple/MetalRHI/Private/MetalRHIContext.cpp
@@ -17,12 +17,8 @@ void METALRHI_API SafeReleaseMetalObject(NS::Object* Object)
 {
 	if(GIsMetalInitialized && GDynamicRHI && Object)
 	{
-		FMetalRHICommandContext* Context = static_cast<FMetalRHICommandContext*>(RHIGetDefaultContext());
-		if(Context)
-		{
-			Context->GetDevice().ReleaseObject(Object);
-			return;
-		}
+		FMetalDynamicRHI::Get().DeferredDelete(Object);
+		return;
 	}
 	Object->release();
 }
diff --git a/Engine/Source/Runtime/Apple/MetalRHI/Private/MetalRHIStagingBuffer.cpp b/Engine/Source/Runtime/Apple/MetalRHI/Private/MetalRHIStagingBuffer.cpp
--- a/Engine/Source/Runtime/Apple/MetalRHI/Private/MetalRHIStagingBuffer.cpp
+++ b/Engine/Source/Runtime/Apple/MetalRHI/Private/MetalRHIStagingBuffer.cpp
@@ -9,7 +9,7 @@
 #include "MetalResources.h"
 #include "MetalDevice.h"
 #include "MetalRHIPrivate.h"
-
+#include "MetalDynamicRHI.h"
 
 //------------------------------------------------------------------------------
 
@@ -27,7 +27,7 @@ FMetalRHIStagingBuffer::~FMetalRHIStagingBuffer()
 {
 	if (ShadowBuffer)
 	{
-		Device.ReleaseBuffer(ShadowBuffer);
+		FMetalDynamicRHI::Get().DeferredDelete(ShadowBuffer);
 		ShadowBuffer = nullptr;
 	}
 }
diff --git a/Engine/Source/Runtime/Apple/MetalRHI/Private/MetalRenderTarget.cpp b/Engine/Source/Runtime/Apple/MetalRHI/Private/MetalRenderTarget.cpp
--- a/Engine/Source/Runtime/Apple/MetalRHI/Private/MetalRenderTarget.cpp
+++ b/Engine/Source/Runtime/Apple/MetalRHI/Private/MetalRenderTarget.cpp
@@ -230,7 +230,7 @@ void FMetalDynamicRHI::RHIReadSurfaceData(FRHITexture* TextureRHI, FIntRect Rect
 		
 		if (TempTexture)
 		{
-			Device->ReleaseTexture(TempTexture);
+			FMetalDynamicRHI::Get().DeferredDelete(TempTexture);
 		}
 	}
 	else
@@ -274,7 +274,7 @@ void FMetalDynamicRHI::RHIReadSurfaceData(FRHITexture* TextureRHI, FIntRect Rect
 		
 		ConvertSurfaceDataToFColor(Surface->GetDesc().Format, SizeX, SizeY, (uint8*)Buffer->Contents(), AlignedStride, OutDataPtr, InFlags);
 		
-		Device->ReleaseBuffer(Buffer);
+		FMetalDynamicRHI::Get().DeferredDelete(Buffer);
 	}
 }
 
@@ -388,7 +388,7 @@ void FMetalDynamicRHI::RHIReadSurfaceFloatData(FRHITexture* TextureRHI, FIntRect
 		FMemory::Memcpy(OutDataPtr, FloatBGRAData, FloatBGRADataSize);
 	}
 	
-	Device->ReleaseBuffer(Buffer);
+	FMetalDynamicRHI::Get().DeferredDelete(Buffer);
 }
 
 void FMetalDynamicRHI::RHIRead3DSurfaceFloatData(FRHITexture* TextureRHI,FIntRect InRect,FIntPoint ZMinMax,TArray<FFloat16Color>& OutData)
@@ -463,5 +463,5 @@ void FMetalDynamicRHI::RHIRead3DSurfaceFloatData(FRHITexture* TextureRHI,FIntRec
 		FMemory::Memcpy(OutDataPtr, FloatBGRAData, FloatBGRADataSize);
 	}
 	
-	Device->ReleaseBuffer(Buffer);
+	FMetalDynamicRHI::Get().DeferredDelete(Buffer);
 }
diff --git a/Engine/Source/Runtime/Apple/MetalRHI/Private/MetalState.cpp b/Engine/Source/Runtime/Apple/MetalRHI/Private/MetalState.cpp
--- a/Engine/Source/Runtime/Apple/MetalRHI/Private/MetalState.cpp
+++ b/Engine/Source/Runtime/Apple/MetalRHI/Private/MetalState.cpp
@@ -320,12 +320,9 @@ FMetalSamplerState::FMetalSamplerState(FMetalDevice& MetalDevice, const FSampler
 FMetalSamplerState::~FMetalSamplerState()
 {
 #if PLATFORM_SUPPORTS_BINDLESS_RENDERING
-    FMetalBindlessDescriptorManager* BindlessDescriptorManager = Device.GetBindlessDescriptorManager();
-    check(BindlessDescriptorManager);
-
 	if(IsMetalBindlessEnabled())
 	{
-		BindlessDescriptorManager->FreeDescriptor(BindlessHandle);
+		FMetalDynamicRHI::Get().DeferredDelete(BindlessHandle);
 	}
 #endif
 }
diff --git a/Engine/Source/Runtime/Apple/MetalRHI/Private/MetalStateCache.cpp b/Engine/Source/Runtime/Apple/MetalRHI/Private/MetalStateCache.cpp
--- a/Engine/Source/Runtime/Apple/MetalRHI/Private/MetalStateCache.cpp
+++ b/Engine/Source/Runtime/Apple/MetalRHI/Private/MetalStateCache.cpp
@@ -13,6 +13,7 @@
 #include "MetalBindlessDescriptors.h"
 #include "RHIShaderParametersShared.h"
 #include "MetalResourceCollection.h"
+#include "MetalDynamicRHI.h"
 
 #if PLATFORM_MAC
 	#ifndef UINT128_MAX
@@ -315,8 +316,6 @@ void FMetalStateCache::Reset()
 		
 		if(BindlessDescriptorManager->IsSupported())
 		{
-			BindlessDescriptorManager->Reset();
-			
 			for (uint32 i = 0; i < UE_ARRAY_COUNT(VertexBufferVAs); i++)
 			{
 				VertexBufferVAs[i].addr = 0;
@@ -336,7 +335,7 @@ void FMetalStateCache::Reset()
 			// Free temporary allocations
 			for (FMetalBufferPtr TemporaryBuffer : TemporaryBuffers)
 			{
-				Device.ReleaseBuffer(TemporaryBuffer);
+				FMetalDynamicRHI::Get().DeferredDelete(TemporaryBuffer);
 			}
 			TemporaryBuffers.Reset();
 			UniformBufferVAs.Reset();
diff --git a/Engine/Source/Runtime/Apple/MetalRHI/Private/MetalTempAllocator.cpp b/Engine/Source/Runtime/Apple/MetalRHI/Private/MetalTempAllocator.cpp
--- a/Engine/Source/Runtime/Apple/MetalRHI/Private/MetalTempAllocator.cpp
+++ b/Engine/Source/Runtime/Apple/MetalRHI/Private/MetalTempAllocator.cpp
@@ -3,6 +3,7 @@
 #include "MetalTempAllocator.h"
 #include "MetalDevice.h"
 #include "MetalRHIPrivate.h"
+#include "MetalDynamicRHI.h"
 #include "MetalProfiler.h"
 
 FMetalTempAllocator::FMetalTempAllocator(FMetalDevice& InDevice, uint32_t InMinAllocationSize, uint32_t InTargetAllocationLimit, uint32_t InAlignment)
@@ -72,7 +73,7 @@ void FMetalTempAllocator::Cleanup()
 	Buffers.RemoveAll([](const FTempBufferInfo& TempBuffer) { return TempBuffer.Offset != 0; });
 	
 	// Ensure that buffers are re-added to the pool when fences are complete (if we are below the target limit)
-	Device.ReleaseFunction([this, OldBuffers]()
+	FMetalDynamicRHI::Get().DeferredDelete([this, OldBuffers]()
 	{
 		FScopeLock lock(&AllocatorLock);
 		for(const FTempBufferInfo& Buffer : OldBuffers)
diff --git a/Engine/Source/Runtime/Apple/MetalRHI/Private/MetalTexture.cpp b/Engine/Source/Runtime/Apple/MetalRHI/Private/MetalTexture.cpp
--- a/Engine/Source/Runtime/Apple/MetalRHI/Private/MetalTexture.cpp
+++ b/Engine/Source/Runtime/Apple/MetalRHI/Private/MetalTexture.cpp
@@ -666,7 +666,7 @@ FMetalSurface::FMetalSurface(FMetalDevice& MetalDevice, FRHICommandListBase* RHI
 
 			Texture = NS::TransferPtr(Buffer->GetMTLBuffer()->newTexture(NewCreateDesc.Desc.get(), Buffer->GetOffset(), BytesPerRow));
             
-            Device.ReleaseBuffer(Buffer);
+			FMetalDynamicRHI::Get().DeferredDelete(Buffer);
 		}
         else if (bTextureArrayWithAtomics)
         {
@@ -686,7 +686,7 @@ FMetalSurface::FMetalSurface(FMetalDevice& MetalDevice, FRHICommandListBase* RHI
             NewCreateDesc.Desc->setTextureType(MTL::TextureType2D);
             Texture = NS::TransferPtr(Buffer->GetMTLBuffer()->newTexture(NewCreateDesc.Desc.get(), Buffer->GetOffset(), BytesPerRow));
             
-            Device.ReleaseBuffer(Buffer);
+			FMetalDynamicRHI::Get().DeferredDelete(Buffer);
         }
 		else
 		{
@@ -886,7 +886,7 @@ FMetalSurface::~FMetalSurface()
     DeferredStats->Flags = GetDesc().Flags;
     DeferredStats->TextureSize = TotalTextureSize;
     
-	Device.ReleaseFunction([DeferredStats]() {
+	FMetalDynamicRHI::Get().DeferredDelete([DeferredStats]() {
         delete DeferredStats;
     });
 	
@@ -894,7 +894,7 @@ FMetalSurface::~FMetalSurface()
 	{
 		// CFArray can contain CFType objects and is toll-free bridged with NSArray
 		CFArrayRef Temp = CFArrayCreate(kCFAllocatorSystemDefault, &ImageSurfaceRef, 1, &kCFTypeArrayCallBacks);
-		Device.ReleaseObject((__bridge NS::Array*)Temp);
+		FMetalDynamicRHI::Get().DeferredDelete((__bridge NS::Array*)Temp);
 		CFRelease(ImageSurfaceRef);
 	}
 	
@@ -908,7 +908,7 @@ FMetalSurface::~FMetalSurface()
 	{
 		if (!(GetDesc().Flags & TexCreate_Presentable))
 		{
-			BindlessDescriptorManager->FreeDescriptor(BindlessHandle);
+			FMetalDynamicRHI::Get().DeferredDelete(BindlessHandle);
 		}
 	}
 #endif
@@ -918,7 +918,7 @@ void FMetalSurface::SafeRelease(MTLTexturePtr InTexture)
 {
 	if(GIsMetalInitialized && GDynamicRHI)
 	{
-		Device.ReleaseTexture(InTexture);
+		FMetalDynamicRHI::Get().DeferredDelete(InTexture);
 	}
 }
 
@@ -1044,11 +1044,6 @@ void FMetalSurface::UpdateSurfaceAndDestroySourceBuffer(FMetalRHICommandContext*
 		Context->AddCompletionHandler(CompletionHandler);
 		
 		INC_DWORD_STAT_BY(STAT_MetalTextureMemUpdate, Size);
-		
-		if (bWait)
-		{
-			Device.ClearFreeList();
-		}
 	}
 	else
 	{
@@ -1066,7 +1061,7 @@ void FMetalSurface::UpdateSurfaceAndDestroySourceBuffer(FMetalRHICommandContext*
 		INC_DWORD_STAT_BY(STAT_MetalTextureMemUpdate, BytesPerImage);
 	}
 	
-	Device.ReleaseBuffer(Source);
+	FMetalDynamicRHI::Get().DeferredDelete(Source);
 	
 	FPlatformAtomics::InterlockedExchange(&Written, 1);
 	
@@ -1661,7 +1656,7 @@ void FMetalDynamicRHI::RHIUpdateTexture2D(FRHICommandListBase& RHICmdList, FRHIT
     RHICmdList.EnqueueLambda([TextureRHI, MipIndex, UpdateRegion, SourcePitch, Buffer, InDevice=Device](FRHICommandListBase& InRHICmdList) mutable
     {
 		InternalUpdateTexture2D(FMetalRHICommandContext::Get(InRHICmdList), TextureRHI, MipIndex, UpdateRegion, SourcePitch, Buffer);
-		InDevice->ReleaseBuffer(Buffer);
+		FMetalDynamicRHI::Get().DeferredDelete(Buffer);
     });
 
     INC_DWORD_STAT_BY(STAT_MetalTextureMemUpdate, UpdateRegion.Height*SourcePitch);
@@ -1761,7 +1756,7 @@ void FMetalDynamicRHI::RHIUpdateTexture3D(FRHICommandListBase& RHICmdList, FRHIT
     RHICmdList.EnqueueLambda([TextureRHI, MipIndex, UpdateRegion, SourceRowPitch, SourceDepthPitch, Buffer, InDevice=Device](FRHICommandListBase& InRHICmdList) mutable
     {
         InternalUpdateTexture3D(FMetalRHICommandContext::Get(InRHICmdList), TextureRHI, MipIndex, UpdateRegion, SourceRowPitch, SourceDepthPitch, Buffer);
-		InDevice->ReleaseBuffer(Buffer);
+		FMetalDynamicRHI::Get().DeferredDelete(Buffer);
     });
 
     INC_DWORD_STAT_BY(STAT_MetalTextureMemUpdate, UpdateRegion.Height * UpdateRegion.Width * SourceDepthPitch);
@@ -1911,14 +1906,14 @@ void FMetalRHICommandContext::RHICopyTexture(FRHITexture* SourceTextureRHI, FRHI
                     CopyFromTextureToBuffer(MetalSrcTexture->Texture.get(), SourceSliceIndex, SourceMipIndex, SourceOrigin, SourceSize, Buffer, 0, AlignedStride, BytesPerImage, Options);
                     CopyFromBufferToTexture(Buffer, 0, Stride, BytesPerImage, DestSize, MetalDestTexture->Texture.get(), DestSliceIndex, DestMipIndex, DestinationOrigin, Options);
                     
-                    Device.ReleaseBuffer(Buffer);
+					FMetalDynamicRHI::Get().DeferredDelete(Buffer);
                 }
             }
         }
         
         if (SrcTexture && (SrcTexture != MetalSrcTexture->Texture))
         {
-			Device.ReleaseTexture(SrcTexture);
+			FMetalDynamicRHI::Get().DeferredDelete(SrcTexture);
         }
     }
     else
diff --git a/Engine/Source/Runtime/Apple/MetalRHI/Private/MetalUAV.cpp b/Engine/Source/Runtime/Apple/MetalRHI/Private/MetalUAV.cpp
--- a/Engine/Source/Runtime/Apple/MetalRHI/Private/MetalUAV.cpp
+++ b/Engine/Source/Runtime/Apple/MetalRHI/Private/MetalUAV.cpp
@@ -34,11 +34,11 @@ void FMetalResourceViewBase::Invalidate()
 		switch (GetMetalType())
 		{
 		case EMetalType::TextureView:
-			Device.ReleaseTexture(Storage.Get<MTLTexturePtr>());
+			FMetalDynamicRHI::Get().DeferredDelete(Storage.Get<MTLTexturePtr>());
 			break;
 
 		case EMetalType::BufferView:
-			Device.ReleaseBuffer(Storage.Get<FBufferView>().Buffer);
+			FMetalDynamicRHI::Get().DeferredDelete(Storage.Get<FBufferView>().Buffer);
 			break;
                 
         case EMetalType::TextureBufferBacked:
@@ -46,12 +46,12 @@ void FMetalResourceViewBase::Invalidate()
 			// If it is a buffer we don't own the resource
 			if (View.bIsBuffer)
 			{
-				Device.ReleaseTexture(View.Texture);
+				FMetalDynamicRHI::Get().DeferredDelete(View.Texture);
 			}
 			else
 			{
-				Device.ReleaseBuffer(View.Buffer);
-				Device.ReleaseTexture(View.Texture);
+				FMetalDynamicRHI::Get().DeferredDelete(View.Buffer);
+				FMetalDynamicRHI::Get().DeferredDelete(View.Texture);
 			}
             break;
 		}
@@ -651,7 +651,7 @@ void FMetalUnorderedAccessView::ClearUAVWithBlitEncoder(TRHICommandList_Recursiv
 		}
 
 		Context.CopyFromBufferToBuffer(Temp, 0, Buffer, Info.OffsetInBytes, Size);
-		Device.ReleaseBuffer(Temp);
+		FMetalDynamicRHI::Get().DeferredDelete(Temp);
 	});
 }
 #endif // UE_METAL_RHI_SUPPORT_CLEAR_UAV_WITH_BLIT_ENCODER
@@ -810,7 +810,7 @@ void FMetalRHICommandContext::RHICopyToStagingBuffer(FRHIBuffer* SourceBufferRHI
     {
         if (ReadbackBuffer)
         {
-            Device.ReleaseBuffer(ReadbackBuffer);
+			FMetalDynamicRHI::Get().DeferredDelete(ReadbackBuffer);
         }
         FMetalPooledBufferArgs ArgsCPU(&Device, NumBytes, BUF_Dynamic, MTL::StorageModeShared);
         ReadbackBuffer = Device.CreatePooledBuffer(ArgsCPU);
