Metal RHI

Fixed resourceOptions when using buffer backed textures

diff --git a/Engine/Source/Runtime/Apple/MetalRHI/Private/MetalTexture.cpp b/Engine/Source/Runtime/Apple/MetalRHI/Private/MetalTexture.cpp
--- a/Engine/Source/Runtime/Apple/MetalRHI/Private/MetalTexture.cpp
+++ b/Engine/Source/Runtime/Apple/MetalRHI/Private/MetalTexture.cpp
@@ -661,7 +661,7 @@ FMetalSurface::FMetalSurface(FMetalDevice& MetalDevice, FRHICommandListBase* RHI
 			const NS::UInteger BytesPerRow = Align(NewCreateDesc.Desc->width() * GPixelFormats[NewCreateDesc.Format].BlockBytes, MinimumByteAlignment);
 
 			// Backing buffer resource options must match the texture we are going to create from it
-			FMetalPooledBufferArgs Args(&Device, BytesPerRow * NewCreateDesc.Desc->height(), BUF_Dynamic, MTL::StorageModePrivate, NewCreateDesc.Desc->cpuCacheMode());
+			FMetalPooledBufferArgs Args(&Device, BytesPerRow * NewCreateDesc.Desc->height(), BUF_Dynamic, NewCreateDesc.Desc->storageMode(), NewCreateDesc.Desc->cpuCacheMode());
 			FMetalBufferPtr Buffer = Device.CreatePooledBuffer(Args);
 
 			Texture = NS::TransferPtr(Buffer->GetMTLBuffer()->newTexture(NewCreateDesc.Desc.get(), Buffer->GetOffset(), BytesPerRow));
@@ -678,7 +678,7 @@ FMetalSurface::FMetalSurface(FMetalDevice& MetalDevice, FRHICommandListBase* RHI
             const NS::UInteger BytesPerRow = Align(NewCreateDesc.Desc->width() * NewCreateDesc.Desc->arrayLength() * GPixelFormats[NewCreateDesc.Format].BlockBytes, MinimumByteAlignment);
 
             // Backing buffer resource options must match the texture we are going to create from it
-            FMetalPooledBufferArgs Args(&Device, BytesPerRow * NewCreateDesc.Desc->height(), BUF_Dynamic, MTL::StorageModePrivate, NewCreateDesc.Desc->cpuCacheMode());
+            FMetalPooledBufferArgs Args(&Device, BytesPerRow * NewCreateDesc.Desc->height(), BUF_Dynamic, NewCreateDesc.Desc->storageMode(), NewCreateDesc.Desc->cpuCacheMode());
             FMetalBufferPtr Buffer = Device.CreatePooledBuffer(Args);
 
             NewCreateDesc.Desc->setWidth(NewCreateDesc.Desc->width() * NewCreateDesc.Desc->arrayLength());
