Metal RHI

Added support for 3D textures with R8 to RGBA8 conversion

diff --git a/Engine/Source/Runtime/Apple/MetalRHI/Private/MetalTexture.cpp b/Engine/Source/Runtime/Apple/MetalRHI/Private/MetalTexture.cpp
--- a/Engine/Source/Runtime/Apple/MetalRHI/Private/MetalTexture.cpp
+++ b/Engine/Source/Runtime/Apple/MetalRHI/Private/MetalTexture.cpp
@@ -941,12 +941,12 @@ MTL::Buffer* FMetalSurface::AllocSurface(uint32 MipIndex, uint32 ArrayIndex, ERe
 	
 #if PLATFORM_MAC
 	// Expand R8_sRGB into RGBA8_sRGB for non Apple Silicon Mac.
-	if (   GetDesc().Format == PF_G8 
-		&& GetDesc().Dimension == ETextureDimension::Texture2D
+	if (GetDesc().Format == PF_G8 
 		&& EnumHasAnyFlags(GetDesc().Flags, TexCreate_SRGB) 
 		&& LockMode == RLM_WriteOnly 
 		&& Texture->pixelFormat() == MTL::PixelFormatRGBA8Unorm_sRGB)
 	{
+		check(GetDesc().Dimension == ETextureDimension::Texture2D || GetDesc().Dimension == ETextureDimension::Texture3D);
 		DestStride = FMath::Max<uint32>(GetDesc().Extent.X >> MipIndex, 1);
 	}
 #endif
@@ -989,27 +989,33 @@ void FMetalSurface::UpdateSurfaceAndDestroySourceBuffer(FMetalRHICommandContext*
 
 #if PLATFORM_MAC
 	// Expand R8_sRGB into RGBA8_sRGB for non Apple Silicon Mac.
-	if (   GetDesc().Format == PF_G8 
-		&& GetDesc().Dimension == ETextureDimension::Texture2D
+	if (GetDesc().Format == PF_G8 
 		&& EnumHasAnyFlags(GetDesc().Flags, TexCreate_SRGB) 
 		&& Texture->pixelFormat() == MTL::PixelFormatRGBA8Unorm_sRGB)
 	{
+		check(GetDesc().Dimension == ETextureDimension::Texture2D || GetDesc().Dimension == ETextureDimension::Texture3D);
+		
 		TArray<uint8> Data;
 		uint8* ExpandedMem = (uint8*) SourceBuffer->contents();
 		check(ExpandedMem);
 		Data.Append(ExpandedMem, BytesPerImage);
 		uint32 SrcStride = FMath::Max<uint32>(GetDesc().Extent.X >> MipIndex, 1);
-		for(uint y = 0; y < FMath::Max<uint32>(GetDesc().Extent.Y >> MipIndex, 1); y++)
+		
+		for(uint z = 0; z < FMath::Max<uint32>(GetDesc().Depth >> MipIndex, 1); z++)
 		{
-			uint8* RowDest = ExpandedMem;
-			for(uint x = 0; x < FMath::Max<uint32>(GetDesc().Extent.X >> MipIndex, 1); x++)
+			uint32_t DepthOffset = z * (GetDesc().Extent.X * GetDesc().Extent.Y);
+			for(uint y = 0; y < FMath::Max<uint32>(GetDesc().Extent.Y >> MipIndex, 1); y++)
 			{
-				*(RowDest++) = Data[(y * SrcStride) + x];
-				*(RowDest++) = Data[(y * SrcStride) + x];
-				*(RowDest++) = Data[(y * SrcStride) + x];
-				*(RowDest++) = Data[(y * SrcStride) + x];
+				uint8* RowDest = ExpandedMem;
+				for(uint x = 0; x < FMath::Max<uint32>(GetDesc().Extent.X >> MipIndex, 1); x++)
+				{
+					*(RowDest++) = Data[(y * SrcStride) + DepthOffset + x];
+					*(RowDest++) = Data[(y * SrcStride) + DepthOffset + x];
+					*(RowDest++) = Data[(y * SrcStride) + DepthOffset + x];
+					*(RowDest++) = Data[(y * SrcStride) + DepthOffset + x];
+				}
+				ExpandedMem = (ExpandedMem + Stride);
 			}
-			ExpandedMem = (ExpandedMem + Stride);
 		}
 	}
 #endif
@@ -1553,6 +1559,29 @@ static void InternalExpandR8ToStandardRGBA(uint32* pDest, const struct FUpdateTe
 	
 	InOutSourcePitch = ExpandedPitch;
 }
+
+static void InternalExpandR8ToStandardRGBA3D(uint32* pDest, const struct FUpdateTextureRegion3D& UpdateRegion, uint32& InOutSourcePitch, uint32& InOutSourceDepthPitch, const uint8* pSrc)
+{
+	// Should only be required for non Apple Silicon Macs
+	const uint32 ExpandedPitch = UpdateRegion.Width * sizeof(uint32);
+	
+	for(uint z = 0; z < UpdateRegion.Depth; z++)
+	{
+		uint32_t DepthOffset = z * InOutSourceDepthPitch;
+		
+		for(uint y = 0; y < UpdateRegion.Height; y++)
+		{
+			for(uint x = 0; x < UpdateRegion.Width; x++)
+			{
+				uint8 Value = pSrc[(y * InOutSourcePitch) + DepthOffset + x];
+				*(pDest++) = (Value | (Value << 8) | (Value << 16) | (Value << 24));
+			}
+		}
+	}
+	
+	InOutSourceDepthPitch = UpdateRegion.Width * UpdateRegion.Height * sizeof(uint32);
+	InOutSourcePitch = ExpandedPitch;
+}
 #endif
 
 static FMetalBufferPtr Internal_CreateBufferAndCopyTexture2DUpdateRegionData(FMetalDevice& Device, FRHITexture* TextureRHI, const struct FUpdateTextureRegion2D& UpdateRegion, uint32& InOutSourcePitch, const uint8* SourceData)
@@ -1662,33 +1691,60 @@ void FMetalDynamicRHI::RHIUpdateTexture2D(FRHICommandListBase& RHICmdList, FRHIT
     INC_DWORD_STAT_BY(STAT_MetalTextureMemUpdate, UpdateRegion.Height*SourcePitch);
 }
 
-static FMetalBufferPtr Internal_CreateBufferAndCopyTexture3DUpdateRegionData(FMetalDevice& Device, FRHITexture* TextureRHI, const struct FUpdateTextureRegion3D& UpdateRegion, uint32 SourceRowPitch, uint32 SourceDepthPitch, const uint8* SourceData)
+static FMetalBufferPtr Internal_CreateBufferAndCopyTexture3DUpdateRegionData(FMetalDevice& Device, FRHITexture* TextureRHI, const struct FUpdateTextureRegion3D& UpdateRegion, uint32& SourceRowPitch, uint32& SourceDepthPitch, const uint8* SourceData)
 {
 	FMetalSurface* Texture = ResourceCast(TextureRHI);
 	
-	const uint32 BufferSize = SourceDepthPitch * UpdateRegion.Depth;
-	FMetalBufferPtr OutBuffer = Device.CreatePooledBuffer(FMetalPooledBufferArgs(&Device, BufferSize, BUF_Static, MTL::StorageModeShared));
-
-	const FPixelFormatInfo& FormatInfo = GPixelFormats[TextureRHI->GetFormat()];
-	uint32 CopyPitch = FMath::DivideAndRoundUp(UpdateRegion.Width, (uint32)FormatInfo.BlockSizeX) * FormatInfo.BlockBytes;
+	const EPixelFormat PixelFormat = TextureRHI->GetFormat();
+	const FPixelFormatInfo& FormatInfo = GPixelFormats[PixelFormat];
 	
-	check(FormatInfo.BlockSizeZ == 1);
-	check(CopyPitch <= SourceRowPitch);
+	FMetalBufferPtr OutBuffer;
+	
+#if PLATFORM_MAC
+	// Expand R8_sRGB into RGBA8_sRGB for non Apple Silicon Mac.
+	if (PixelFormat == PF_G8
+		&& EnumHasAnyFlags(Texture->GetFlags(), TexCreate_SRGB)
+		&& Texture->Texture->pixelFormat() == MTL::PixelFormatRGBA8Unorm_sRGB)
+	{
+		const uint32 ExpandedBufferSize = UpdateRegion.Height * UpdateRegion.Width * UpdateRegion.Depth * sizeof(uint32);
 		
-	uint8_t* DestData = (uint8_t*)OutBuffer->Contents();
-	const uint32 NumRows = FMath::DivideAndRoundUp(UpdateRegion.Height, (uint32)FormatInfo.BlockSizeY);
+		const uint32 SrcXInBlocks   = FMath::DivideAndRoundUp<uint32>(UpdateRegion.SrcX,   	FormatInfo.BlockSizeX);
+		const uint32 SrcYInBlocks   = FMath::DivideAndRoundUp<uint32>(UpdateRegion.SrcY,   	FormatInfo.BlockSizeY);
+		const uint32 SrcZInBlocks   = FMath::DivideAndRoundUp<uint32>(UpdateRegion.SrcZ,   	FormatInfo.BlockSizeZ);
+
+		const uint8* OffsetSourceData = SourceData + FormatInfo.BlockBytes * SrcXInBlocks +
+										SourceRowPitch * SrcYInBlocks * FormatInfo.BlockSizeY +
+										SourceDepthPitch * SrcZInBlocks * FormatInfo.BlockSizeZ;
 		
-	// Perform safe line copy
-	for (uint32 i = 0;i < UpdateRegion.Depth;++i)
+		OutBuffer = Device.CreatePooledBuffer(FMetalPooledBufferArgs(&Device, ExpandedBufferSize, BUF_Static, MTL::StorageModeShared));
+		InternalExpandR8ToStandardRGBA3D((uint32*)OutBuffer->Contents(), UpdateRegion, SourceRowPitch, SourceDepthPitch, OffsetSourceData);
+	}
+	else
+#endif
 	{
-		const uint8* pSourceRowData = SourceData + (SourceDepthPitch * i);
-		uint8* pDestRowData = DestData + (SourceDepthPitch * i);
-
-		for (uint32 j = 0;j < NumRows;++j)
+		const uint32 BufferSize = SourceDepthPitch * UpdateRegion.Depth;
+		OutBuffer = Device.CreatePooledBuffer(FMetalPooledBufferArgs(&Device, BufferSize, BUF_Static, MTL::StorageModeShared));
+		
+		uint32 CopyPitch = FMath::DivideAndRoundUp(UpdateRegion.Width, (uint32)FormatInfo.BlockSizeX) * FormatInfo.BlockBytes;
+		
+		check(FormatInfo.BlockSizeZ == 1);
+		check(CopyPitch <= SourceRowPitch);
+		
+		uint8_t* DestData = (uint8_t*)OutBuffer->Contents();
+		const uint32 NumRows = FMath::DivideAndRoundUp(UpdateRegion.Height, (uint32)FormatInfo.BlockSizeY);
+		
+		// Perform safe line copy
+		for (uint32 i = 0;i < UpdateRegion.Depth;++i)
 		{
-			FMemory::Memcpy(pDestRowData, pSourceRowData, CopyPitch);
-			pSourceRowData += SourceRowPitch;
-			pDestRowData += SourceRowPitch;
+			const uint8* pSourceRowData = SourceData + (SourceDepthPitch * i);
+			uint8* pDestRowData = DestData + (SourceDepthPitch * i);
+			
+			for (uint32 j = 0;j < NumRows;++j)
+			{
+				FMemory::Memcpy(pDestRowData, pSourceRowData, CopyPitch);
+				pSourceRowData += SourceRowPitch;
+				pDestRowData += SourceRowPitch;
+			}
 		}
 	}
 	
