Metal RHI

- Fixed missing CommandBuffer delete.
- Fixed Sub-ring buffer never calling commit which never reset head.

diff --git a/Engine/Source/Runtime/Apple/MetalRHI/Private/MetalCommandEncoder.cpp b/Engine/Source/Runtime/Apple/MetalRHI/Private/MetalCommandEncoder.cpp
--- a/Engine/Source/Runtime/Apple/MetalRHI/Private/MetalCommandEncoder.cpp
+++ b/Engine/Source/Runtime/Apple/MetalRHI/Private/MetalCommandEncoder.cpp
@@ -275,7 +275,7 @@ FMetalCommandBuffer* FMetalCommandEncoder::Finalize()
 		CommandBuffer->GetMTLCmdBuffer()->setLabel(DebugGroups.Last());
 	}
 	
-	RingBuffer.Submit();
+	RingBuffer.Commit(CommandBuffer);
     
 #if METAL_DEBUG_OPTIONS
     if(Device.GetRuntimeDebuggingLevel() >= EMetalDebugLevelValidation)
@@ -687,6 +687,7 @@ TRefCountPtr<FMetalFence> FMetalCommandEncoder::EndEncoding(void)
 		FMemory::Memzero(ShaderBuffers[Frequency].Usage);
 		ShaderBuffers[Frequency].Bound = 0;
 	}
+	
     return Fence;
 }
 
diff --git a/Engine/Source/Runtime/Apple/MetalRHI/Private/MetalCommandQueue.cpp b/Engine/Source/Runtime/Apple/MetalRHI/Private/MetalCommandQueue.cpp
--- a/Engine/Source/Runtime/Apple/MetalRHI/Private/MetalCommandQueue.cpp
+++ b/Engine/Source/Runtime/Apple/MetalRHI/Private/MetalCommandQueue.cpp
@@ -134,6 +134,26 @@ void FMetalCommandQueue::CommitCommandBuffer(FMetalCommandBuffer* CommandBuffer)
 	{
 		CommandBuffer->GetMTLCmdBuffer()->waitUntilCompleted();
 	}
+	
+	if(IsInRHIThread() || !IsRunningRHIInSeparateThread())
+	{
+		FMetalDynamicRHI::Get().DeferredDelete([CommandBuffer]() {
+			delete CommandBuffer;
+		});
+	}
+	else
+	{
+		// Deferred deletes need to run on the RHI thread
+		FFunctionGraphTask::CreateAndDispatchWhenReady(
+		   [CommandBuffer]()
+		   {
+			   FMetalDynamicRHI::Get().DeferredDelete([CommandBuffer]() {
+				   delete CommandBuffer;
+			   });
+		   },
+		   QUICK_USE_CYCLE_STAT(FExecuteRHIThreadTask, STATGROUP_TaskGraphTasks), nullptr, ENamedThreads::RHIThread);	
+	}
+		
 }
 
 FMetalFence* FMetalCommandQueue::CreateFence(NS::String* Label) const
