Fix for semaphore blocking renderthread when editor is in the background

diff --git a/Engine/Source/Runtime/Apple/MetalRHI/Private/MetalRHI.cpp b/Engine/Source/Runtime/Apple/MetalRHI/Private/MetalRHI.cpp
--- a/Engine/Source/Runtime/Apple/MetalRHI/Private/MetalRHI.cpp
+++ b/Engine/Source/Runtime/Apple/MetalRHI/Private/MetalRHI.cpp
@@ -1269,9 +1269,6 @@ void FMetalDynamicRHI::RHIEndFrame_RenderThread(FRHICommandListImmediate& RHICmd
 		[this](FMetalContextArray const& Contexts)
 	{
 		MTL_SCOPED_AUTORELEASE_POOL;
-
-		// Wait for the frame semaphore
-		dispatch_semaphore_wait(Device->GetFrameSemaphore(), DISPATCH_TIME_FOREVER);
 		
 #if ENABLE_METAL_GPUPROFILE
 		Contexts[ERHIPipeline::Graphics]->GetProfiler()->BeginFrame();
diff --git a/Engine/Source/Runtime/Apple/MetalRHI/Private/MetalViewport.cpp b/Engine/Source/Runtime/Apple/MetalRHI/Private/MetalViewport.cpp
--- a/Engine/Source/Runtime/Apple/MetalRHI/Private/MetalViewport.cpp
+++ b/Engine/Source/Runtime/Apple/MetalRHI/Private/MetalViewport.cpp
@@ -576,6 +576,9 @@ void FMetalViewport::Present(FMetalCommandQueue& CommandQueue, bool bLockToVsync
 #endif // PLATFORM_MAC
                         METAL_GPUPROFILE(Stats->End(CurrentCommandBuffer->GetMTLCmdBuffer()));
                         CommandQueue.CommitCommandBuffer(CurrentCommandBuffer);
+						
+						// Wait for the frame semaphore
+						dispatch_semaphore_wait(Device.GetFrameSemaphore(), DISPATCH_TIME_FOREVER);
 					}
 				}
 			}
@@ -728,6 +731,9 @@ void FMetalViewport::PresentImmersive(const MetalRHIVisionOS::PresentImmersivePa
 				Context.ResetContext();
 				
 				cp_frame_end_submission(CompositorServicesFrame);
+				
+				// Wait for the frame semaphore
+				dispatch_semaphore_wait(Device.GetFrameSemaphore(), DISPATCH_TIME_FOREVER);
 			}
 		}
 	}
