Fix failed check when calling SafeReleaseMetalObject with RHI thread.

diff --git a/Engine/Source/Runtime/Apple/MetalRHI/Private/MetalRHIContext.cpp b/Engine/Source/Runtime/Apple/MetalRHI/Private/MetalRHIContext.cpp
--- a/Engine/Source/Runtime/Apple/MetalRHI/Private/MetalRHIContext.cpp
+++ b/Engine/Source/Runtime/Apple/MetalRHI/Private/MetalRHIContext.cpp
@@ -17,7 +17,20 @@ void METALRHI_API SafeReleaseMetalObject(NS::Object* Object)
 {
 	if(GIsMetalInitialized && GDynamicRHI && Object)
 	{
-		FMetalDynamicRHI::Get().DeferredDelete(Object);
+		if(!IsRunningRHIInSeparateThread())
+		{
+			FMetalDynamicRHI::Get().DeferredDelete(Object);
+		}
+		else
+		{
+			FFunctionGraphTask::CreateAndDispatchWhenReady(
+			   [Object]()
+			   {
+				   FMetalDynamicRHI::Get().DeferredDelete(Object);
+			   },
+			   QUICK_USE_CYCLE_STAT(FExecuteRHIThreadTask, STATGROUP_TaskGraphTasks), nullptr, ENamedThreads::RHIThread);
+		}
+		
 		return;
 	}
 	Object->release();
