Fix Main Thread Checker issue in FMetalViewport::Present(..) when not using GMetalSeparatePresentThread

#tests Use the editor for a couple of days with the change

diff --git a/Engine/Source/Runtime/Apple/MetalRHI/Private/MetalViewport.cpp b/Engine/Source/Runtime/Apple/MetalRHI/Private/MetalViewport.cpp
--- a/Engine/Source/Runtime/Apple/MetalRHI/Private/MetalViewport.cpp
+++ b/Engine/Source/Runtime/Apple/MetalRHI/Private/MetalViewport.cpp
@@ -449,11 +449,9 @@ void FMetalViewport::Present(FMetalCommandQueue& CommandQueue, bool bLockToVsync
 {
 	FScopeLock Lock(&Mutex);
 	
-	bool bIsLiveResize = false;
 #if PLATFORM_MAC
 	NSNumber* ScreenId = [View.window.screen.deviceDescription objectForKey:@"NSScreenNumber"];
 	DisplayID = ScreenId.unsignedIntValue;
-	bIsLiveResize = View.inLiveResize;
 	{
 		FCAMetalLayer* CurrentLayer = (FCAMetalLayer*)[View layer];
 		CurrentLayer.displaySyncEnabled = bLockToVsync || (!(IsRunningGame() && bIsFullScreen));
@@ -473,7 +471,12 @@ void FMetalViewport::Present(FMetalCommandQueue& CommandQueue, bool bLockToVsync
 #endif
 			bool bIsInLiveResize = false;
 #if PLATFORM_MAC
-			bIsInLiveResize = View.inLiveResize;
+			// DisplayID is always equal to 0 when GMetalSeparatePresentThread == false so we don't need to access inLiveResize outside of main thread here
+			// @TODO: UE-227397
+			if (GMetalSeparatePresentThread)
+			{
+				bIsInLiveResize = View.inLiveResize;
+			}
 #endif
 			if (FrameAvailable > 0 && (InDisplayID == 0 || (DisplayID == InDisplayID && !bIsInLiveResize)))
 			{
@@ -485,7 +488,12 @@ void FMetalViewport::Present(FMetalCommandQueue& CommandQueue, bool bLockToVsync
 				{
 					FScopeLock BlockLock(&Mutex);
 #if PLATFORM_MAC
-					bIsInLiveResize = View.inLiveResize;
+					// DisplayID is always equal to 0 when GMetalSeparatePresentThread == false so we don't need to access inLiveResize outside of main thread here
+					// @TODO: UE-227397
+					if (GMetalSeparatePresentThread)
+					{
+						bIsInLiveResize = View.inLiveResize;
+					}
 #endif
 					
 					if (DrawableTexture && (InDisplayID == 0 || !bIsInLiveResize))
@@ -590,7 +598,11 @@ void FMetalViewport::Present(FMetalCommandQueue& CommandQueue, bool bLockToVsync
 		}
 	}
 	
-	if (bIsLiveResize || !GMetalSeparatePresentThread)
+	if (!GMetalSeparatePresentThread
+#if PLATFORM_MAC
+		|| View.inLiveResize
+#endif
+		)
 	{
 		Block(0, 0.0, 0.0);
 	}
