Fixing GPU Profiling on Mac

- Added support for splitting commmand buffers to match profiler scopes
- Changed how command buffer timings are added and recorded for GPUTIme

diff --git a/Engine/Source/Runtime/Apple/MetalRHI/Private/MetalCommandEncoder.cpp b/Engine/Source/Runtime/Apple/MetalRHI/Private/MetalCommandEncoder.cpp
--- a/Engine/Source/Runtime/Apple/MetalRHI/Private/MetalCommandEncoder.cpp
+++ b/Engine/Source/Runtime/Apple/MetalRHI/Private/MetalCommandEncoder.cpp
@@ -259,12 +259,12 @@ void FMetalCommandEncoder::StartCommandBuffer(void)
 	}
 }
 	
-FMetalCommandBuffer* FMetalCommandEncoder::Finalize()
+void FMetalCommandEncoder::EndCommandBuffer()
 {
 	check(CommandBuffer);
 	check(IsRenderCommandEncoderActive() == false
-          && IsComputeCommandEncoderActive() == false
-          && IsBlitCommandEncoderActive() == false
+		  && IsComputeCommandEncoderActive() == false
+		  && IsBlitCommandEncoderActive() == false
 #if METAL_RHI_RAYTRACING
 		  && IsAccelerationStructureCommandEncoderActive() == false
 #endif // METAL_RHI_RAYTRACING
@@ -276,39 +276,46 @@ FMetalCommandBuffer* FMetalCommandEncoder::Finalize()
 	}
 	
 	RingBuffer.Commit(CommandBuffer);
-    
+	
 #if METAL_DEBUG_OPTIONS
-    if(Device.GetRuntimeDebuggingLevel() >= EMetalDebugLevelValidation)
-    {
-        for (FMetalBufferPtr Buffer : ActiveBuffers)
-        {
-            Device.AddActiveBuffer(Buffer->GetMTLBuffer(), Buffer->GetRange());
-        }
-        
-        TSet<FMetalBufferPtr> NewActiveBuffers = MoveTemp(ActiveBuffers);
-        
-        FMetalCommandBufferCompletionHandler CompletionHander;
-        CompletionHander.BindLambda([&InDevice = Device, NewActiveBuffers](MTL::CommandBuffer*)
-        {
-            for (FMetalBufferPtr Buffer : NewActiveBuffers)
-            {
+	if(Device.GetRuntimeDebuggingLevel() >= EMetalDebugLevelValidation)
+	{
+		for (FMetalBufferPtr Buffer : ActiveBuffers)
+		{
+			Device.AddActiveBuffer(Buffer->GetMTLBuffer(), Buffer->GetRange());
+		}
+		
+		TSet<FMetalBufferPtr> NewActiveBuffers = MoveTemp(ActiveBuffers);
+		
+		FMetalCommandBufferCompletionHandler CompletionHander;
+		CompletionHander.BindLambda([&InDevice = Device, NewActiveBuffers](MTL::CommandBuffer*)
+		{
+			for (FMetalBufferPtr Buffer : NewActiveBuffers)
+			{
 				InDevice.RemoveActiveBuffer(Buffer->GetMTLBuffer(), Buffer->GetRange());
-            }
-        });
-        
-        AddCompletionHandler(CompletionHander);
-    }
+			}
+		});
+		
+		AddCompletionHandler(CompletionHander);
+	}
 #endif
 #if ENABLE_METAL_GPUPROFILE
 	CommandBufferStats->End(CommandBuffer->GetMTLCmdBuffer());
 	CommandBufferStats = nullptr;
 #endif
-
+	
 	CommandList.FinalizeCommandBuffer(CommandBuffer, MoveTemp(CompletionHandlers));
-	FMetalCommandBuffer* Output = CommandBuffer;
+	CommandBuffers.Add(CommandBuffer);
 	
 	CommandBuffer = nullptr;
 	EncoderNum = 0;
+}
+
+TArray<FMetalCommandBuffer*> FMetalCommandEncoder::Finalize()
+{
+	EndCommandBuffer();
+	
+	TArray<FMetalCommandBuffer*> Output = MoveTemp(CommandBuffers);
 	
 	return Output;
 }
@@ -411,7 +418,7 @@ void FMetalCommandEncoder::BeginRenderCommandEncoding(void)
 	
 	if(GetEmitDrawEvents())
 	{
-        Label = FStringToNSString(FString::Printf(TEXT("RenderEncoder: %s"), DebugGroups.Num() > 0 ? *NSStringToFString(DebugGroups.Last()) : TEXT("InitialPass")));
+		Label = FStringToNSString(FString::Printf(TEXT("RenderEncoder: %s"), DebugGroups.Num() > 0 ? *NSStringToFString(DebugGroups.Last()) : TEXT("InitialPass")));
 		RenderCommandEncoder->setLabel(Label);
 		
         for (NS::String* Group : DebugGroups)
@@ -533,6 +540,16 @@ void FMetalCommandEncoder::BeginAccelerationStructureCommandEncoding(void)
 }
 #endif // METAL_RHI_RAYTRACING
 
+void FMetalCommandEncoder::SplitCommandBuffers()
+{
+	if(IsAnyCommandEncoderActive())
+	{
+		EndEncoding();
+	}
+	EndCommandBuffer();
+	StartCommandBuffer();	
+}
+
 TRefCountPtr<FMetalFence> FMetalCommandEncoder::EndEncoding(void)
 {
 	static bool bSupportsFences = Device.SupportsFeature(EMetalFeaturesFences);
diff --git a/Engine/Source/Runtime/Apple/MetalRHI/Private/MetalCommandEncoder.h b/Engine/Source/Runtime/Apple/MetalRHI/Private/MetalCommandEncoder.h
--- a/Engine/Source/Runtime/Apple/MetalRHI/Private/MetalCommandEncoder.h
+++ b/Engine/Source/Runtime/Apple/MetalRHI/Private/MetalCommandEncoder.h
@@ -170,10 +170,15 @@ public:
 	 */
 	void StartCommandBuffer(void);
 	
+	/**
+	 * Will deprecate in 5.6, splits the command buffer into chunks, required for profiling until we switch to encoder counter buffers
+	 */
+	void EndCommandBuffer();
+	
 	/**
 	 * Commit the existing command buffer if there is one & optionally waiting for completion, if there isn't a current command buffer this is a no-op.
  	 */
-	FMetalCommandBuffer* Finalize();
+	TArray<FMetalCommandBuffer*> Finalize();
 
 #pragma mark - Public Command Buffer Accessors -
 	
@@ -505,6 +510,11 @@ public:
 	 */
 	FMetalSubBufferRing& GetRingBuffer(void);
 	
+	/**
+	 * Splits the current command buffer
+	 */
+	void SplitCommandBuffers();
+	
 #pragma mark - Public Resource query Access -
 	
 private:
@@ -596,6 +606,8 @@ public:
     MTL::RenderPassDescriptor* RenderPassDesc = nullptr;
 	
     FMetalCommandBuffer* CommandBuffer = nullptr;
+	TArray<FMetalCommandBuffer*> CommandBuffers;
+	
     MTLRenderCommandEncoderPtr RenderCommandEncoder;
     MTLComputeCommandEncoderPtr ComputeCommandEncoder;
     MTLBlitCommandEncoderPtr BlitCommandEncoder;
diff --git a/Engine/Source/Runtime/Apple/MetalRHI/Private/MetalCommandList.cpp b/Engine/Source/Runtime/Apple/MetalRHI/Private/MetalCommandList.cpp
--- a/Engine/Source/Runtime/Apple/MetalRHI/Private/MetalCommandList.cpp
+++ b/Engine/Source/Runtime/Apple/MetalRHI/Private/MetalCommandList.cpp
@@ -208,14 +208,11 @@ void FMetalCommandList::HandleMetalCommandBufferFailure(MTL::CommandBuffer* Comp
 void FMetalCommandList::FinalizeCommandBuffer(FMetalCommandBuffer* Buffer, TArray<FMetalCommandBufferCompletionHandler> CompletionHandlers)
 {
 	check(Buffer);
-
-	// The lifetime of this should be for the entire game
-	if (!LastCompletedBufferTiming.IsValid())
-	{
-		LastCompletedBufferTiming = MakeShared<FMetalCommandBufferTiming, ESPMode::ThreadSafe>();
-	}
-    
-    MTL::HandlerFunction CompletionHandler = [CompletionHandlers, &InDevice = CommandQueue.GetDevice(), FrameBufferTimingsLocal = FMetalGPUProfiler::GetFrameBufferTimings(), LastCompletedBufferTimingLocal = LastCompletedBufferTiming, Buffer](MTL::CommandBuffer* CompletedBuffer)
+	
+	FMetalCommandBufferTimer& Timer = FMetalGPUProfiler::GetFrameBufferTimer();
+	Timer.Submit();
+	
+	MTL::HandlerFunction CompletionHandler = [CompletionHandlers, &InDevice = CommandQueue.GetDevice(), &FrameBufferTimer = Timer, Buffer](MTL::CommandBuffer* CompletedBuffer)
 	{
 		if (CompletedBuffer->status() == MTL::CommandBufferStatusError)
 		{
@@ -231,19 +228,13 @@ void FMetalCommandList::FinalizeCommandBuffer(FMetalCommandBuffer* Buffer, TArra
 
 		if (CompletedBuffer->status() == MTL::CommandBufferStatusCompleted)
 		{
-			FrameBufferTimingsLocal->Add({CompletedBuffer->GPUStartTime(), CompletedBuffer->GPUEndTime()});
-		}
-
-		// If this is the last reference, then it is the last command buffer to return, so record the frame
-		if (FrameBufferTimingsLocal.IsUnique())
-		{
-			FMetalGPUProfiler::RecordFrame(*FrameBufferTimingsLocal, *LastCompletedBufferTimingLocal);
+			FrameBufferTimer.AddTiming({CompletedBuffer->GPUStartTime(), CompletedBuffer->GPUEndTime()});
 		}
 		
 		InDevice.RemoveInflightCommandBuffer(CompletedBuffer);
 	};
-    
-    Buffer->GetMTLCmdBuffer()->addCompletedHandler(CompletionHandler);
+	
+	Buffer->GetMTLCmdBuffer()->addCompletedHandler(CompletionHandler);
 	
 	CommandQueue.GetDevice().AddInflightCommandBuffer(Buffer->GetMTLCmdBuffer().get());
 }
diff --git a/Engine/Source/Runtime/Apple/MetalRHI/Private/MetalProfiler.cpp b/Engine/Source/Runtime/Apple/MetalRHI/Private/MetalProfiler.cpp
--- a/Engine/Source/Runtime/Apple/MetalRHI/Private/MetalProfiler.cpp
+++ b/Engine/Source/Runtime/Apple/MetalRHI/Private/MetalProfiler.cpp
@@ -2,6 +2,7 @@
 
 #include "MetalProfiler.h"
 #include "MetalRHIPrivate.h"
+#include "MetalDynamicRHI.h"
 #include "EngineGlobals.h"
 #include "StaticBoundShaderState.h"
 #include "MetalCommandBuffer.h"
@@ -111,9 +112,11 @@ FMetalCommandBufferCompletionHandler FMetalEventNode::Start(void)
 		const double CyclesPerSecond = 1.0 / FPlatformTime::GetSecondsPerCycle();
 		StartTime = GpuTimeSeconds * CyclesPerSecond;
 	});
+	
     return Block;
 }
 
+
 void FMetalEventNode::StopTiming()
 {
 	Context.EndTiming(this);
@@ -128,14 +131,6 @@ FMetalCommandBufferCompletionHandler FMetalEventNode::Stop(void)
 		const CFTimeInterval GpuTimeSeconds = CompletedBuffer->GPUEndTime();
 		const double CyclesPerSecond = 1.0 / FPlatformTime::GetSecondsPerCycle();
 		EndTime = GpuTimeSeconds * CyclesPerSecond;
-		
-		if(bRoot)
-		{
-			if(!bFullProfiling)
-			{
-				delete this;
-			}
-		}
 	});
 	return Block;
 }
@@ -178,7 +173,38 @@ FGPUProfilerEventNode* FMetalGPUProfiler::CreateEventNode(const TCHAR* InName, F
 #endif
 }
 
-TSharedPtr<TArray<FMetalCommandBufferTiming>, ESPMode::ThreadSafe> FMetalGPUProfiler::FrameBufferTimings;
+void FMetalCommandBufferTimer::AddTiming(FMetalCommandBufferTiming Timing)
+{
+	FScopeLock Lock(&Mutex);
+	Counter--;
+	Timings.Add(Timing);
+	
+	if(Counter == 0 && bFrameEnded)
+	{
+		RecordFrame();
+	}
+}
+
+void FMetalCommandBufferTimer::FrameEnd()
+{
+	FScopeLock Lock(&Mutex);
+	
+	bFrameEnded = true;
+	
+	if(Counter == 0)
+	{
+		RecordFrame();
+	}
+}
+
+void FMetalCommandBufferTimer::RecordFrame()
+{
+	Timings.Sort();
+	FMetalProfiler* Profiler = FMetalProfiler::GetProfiler();
+	Profiler->RecordFrame(*this);
+}
+
+FMetalCommandBufferTimer* FMetalGPUProfiler::Timer = nullptr;
 
 void FMetalGPUProfiler::Cleanup()
 {
@@ -206,57 +232,56 @@ TGlobalResource<FTexture> GMetalLongTaskRT;
 
 void FMetalGPUProfiler::BeginFrame()
 {
-	if(!CurrentEventNodeFrame)
+	if(GTriggerGPUProfile)
+	{
+		bTrackingEvents = true;
+		bLatchedGProfilingGPU = true;
+		GTriggerGPUProfile = false;
+	}
+	
+	if (bLatchedGProfilingGPU)
 	{
 		// Start tracking the frame
 		CurrentEventNodeFrame = new FMetalEventNodeFrame(Context, GTriggerGPUProfile);
 		CurrentEventNodeFrame->StartFrame();
-		
-		if(GTriggerGPUProfile)
-		{
-			bTrackingEvents = true;
-			bLatchedGProfilingGPU = true;
-			GTriggerGPUProfile = false;
-		}
 	}
-	NumNestedFrames++;
 }
 
 void FMetalGPUProfiler::EndFrame()
 {
-	if(--NumNestedFrames == 0)
-	{
-		dispatch_async(dispatch_get_global_queue(QOS_CLASS_UTILITY, 0), ^{
+	dispatch_async(dispatch_get_global_queue(QOS_CLASS_UTILITY, 0), ^{
 #if PLATFORM_MAC
-			FPlatformMisc::UpdateDriverMonitorStatistics(Context.GetDevice().GetDeviceIndex());
+	FPlatformMisc::UpdateDriverMonitorStatistics(Context.GetDevice().GetDeviceIndex());
 #endif
-		});
+	});
+	
 #if STATS
-		SET_CYCLE_COUNTER(STAT_MetalTexturePageOnTime, GMetalTexturePageOnTime);
-		GMetalTexturePageOnTime = 0;
-		
-		SET_CYCLE_COUNTER(STAT_MetalGPUIdleTime, GMetalGPUIdleTime);
-		SET_CYCLE_COUNTER(STAT_MetalGPUWorkTime, GMetalGPUWorkTime);
-		SET_CYCLE_COUNTER(STAT_MetalPresentTime, GMetalPresentTime);
+	SET_CYCLE_COUNTER(STAT_MetalTexturePageOnTime, GMetalTexturePageOnTime);
+	GMetalTexturePageOnTime = 0;
+	
+	SET_CYCLE_COUNTER(STAT_MetalGPUIdleTime, GMetalGPUIdleTime);
+	SET_CYCLE_COUNTER(STAT_MetalGPUWorkTime, GMetalGPUWorkTime);
+	SET_CYCLE_COUNTER(STAT_MetalPresentTime, GMetalPresentTime);
 #endif
+	
+	if(CurrentEventNodeFrame)
+	{
+		Context.GetDevice().WaitForGPUIdle();
 		
-		if(CurrentEventNodeFrame)
+		CurrentEventNodeFrame->EndFrame();
+		
+		if(bLatchedGProfilingGPU)
 		{
-			CurrentEventNodeFrame->EndFrame();
-			
-			if(bLatchedGProfilingGPU)
-			{
-				bTrackingEvents = false;
-				bLatchedGProfilingGPU = false;
-			
-				UE_LOG(LogRHI, Warning, TEXT(""));
-				UE_LOG(LogRHI, Warning, TEXT(""));
-				CurrentEventNodeFrame->DumpEventTree();
-			}
-			
-			delete CurrentEventNodeFrame;
-			CurrentEventNodeFrame = NULL;
+			bTrackingEvents = false;
+			bLatchedGProfilingGPU = false;
+		
+			UE_LOG(LogRHI, Warning, TEXT(""));
+			UE_LOG(LogRHI, Warning, TEXT(""));
+			CurrentEventNodeFrame->DumpEventTree();
 		}
+		
+		delete CurrentEventNodeFrame;
+		CurrentEventNodeFrame = NULL;
 	}
 }
 
@@ -264,35 +289,26 @@ void FMetalGPUProfiler::EndFrame()
 // All these recording functions MUST be called from within scheduled/completion handlers.
 // Ordering is enforced by libdispatch so calling these outside of that context WILL result in
 // incorrect values.
-void FMetalGPUProfiler::RecordFrame(TArray<FMetalCommandBufferTiming>& CommandBufferTimings, FMetalCommandBufferTiming& LastBufferTiming)
+void FMetalGPUProfiler::RecordFrame(FMetalCommandBufferTimer& InTimer)
 {
+	FMetalCommandBufferTiming LastBufferTiming;
 	const double CyclesPerSecond = 1.0 / FPlatformTime::GetSecondsPerCycle();
 
 	double RunningFrameTimeSeconds = 0.0;
 	uint64 FrameStartGPUCycles = 0;
 	uint64 FrameEndGPUCycles = 0;
 
-	// Sort the timings
-	CommandBufferTimings.Sort();
+	const TArray<FMetalCommandBufferTiming>& Timings = InTimer.GetTimings();
 
 	CFTimeInterval FirstStartTime = 0.0;
 
-	double FrameStartTime = DBL_MAX;
-	double FrameEndTime = 0.0;
-	
-	for (const FMetalCommandBufferTiming& Timing : CommandBufferTimings)
-	{
-		FrameStartTime = FMath::Min(FrameStartTime, Timing.StartTime);
-		FrameEndTime = FMath::Max(FrameStartTime, Timing.EndTime);
-	}
-	double FrameTime = FrameEndTime - FrameStartTime;
-	
 	// Add the timings excluding any overlapping time
-	for (const FMetalCommandBufferTiming& Timing : CommandBufferTimings)
+	for (const FMetalCommandBufferTiming& Timing : Timings)
 	{
 		if (FirstStartTime == 0.0)
 		{
 			FirstStartTime = Timing.StartTime;
+			LastBufferTiming = Timing;
 		}
 		
 		// Only process if the previous buffer finished before the end of this one
@@ -313,7 +329,7 @@ void FMetalGPUProfiler::RecordFrame(TArray<FMetalCommandBufferTiming>& CommandBu
 			LastBufferTiming = Timing;
 		}
 	}
-    
+	
 	FrameStartGPUCycles = FirstStartTime * CyclesPerSecond;
 	FrameEndGPUCycles = LastBufferTiming.EndTime * CyclesPerSecond;
     
@@ -339,19 +355,27 @@ void FMetalGPUProfiler::RecordPresent(MTL::CommandBuffer* CommandBuffer)
 }
 // END WARNING
 
+
 void FMetalGPUProfiler::ResetFrameBufferTimings()
 {
-	FrameBufferTimings = MakeShared<TArray<FMetalCommandBufferTiming>, ESPMode::ThreadSafe>();
+	if(Timer)
+	{
+		Timer->FrameEnd();
+		FMetalDynamicRHI::Get().DeferredDelete([InTimer=Timer](){
+			delete InTimer;
+		});
+	}
+	
+	Timer = new FMetalCommandBufferTimer();
 }
 
-TSharedPtr<TArray<FMetalCommandBufferTiming>, ESPMode::ThreadSafe> FMetalGPUProfiler::GetFrameBufferTimings()
+FMetalCommandBufferTimer& FMetalGPUProfiler::GetFrameBufferTimer()
 {
-	if(!FrameBufferTimings)
+	if(!Timer)
 	{
-		ResetFrameBufferTimings();
+		Timer = new FMetalCommandBufferTimer();
 	}
-		
-	return FrameBufferTimings;
+	return *Timer;
 }
 
 IMetalStatsScope::~IMetalStatsScope()
diff --git a/Engine/Source/Runtime/Apple/MetalRHI/Private/MetalProfiler.h b/Engine/Source/Runtime/Apple/MetalRHI/Private/MetalProfiler.h
--- a/Engine/Source/Runtime/Apple/MetalRHI/Private/MetalProfiler.h
+++ b/Engine/Source/Runtime/Apple/MetalRHI/Private/MetalProfiler.h
@@ -118,10 +118,8 @@ public:
 	
 	virtual ~FMetalEventNodeFrame()
 	{
-        if(bFullProfiling)
-        {
-            delete RootNode;
-        }
+        delete RootNode;
+		RootNode = nullptr;
 	}
 	
 	/** Start this frame of per tracking */
@@ -251,6 +249,31 @@ struct FMetalCommandBufferTiming
 	}
 };
 
+class FMetalCommandBufferTimer
+{
+public:
+	void Submit()
+	{
+		FScopeLock Lock(&Mutex);
+		Counter++;
+	}
+	
+	void AddTiming(FMetalCommandBufferTiming Timing);
+	void FrameEnd();
+	void RecordFrame();
+	
+	const TArray<FMetalCommandBufferTiming>& GetTimings()
+	{
+		return Timings;
+	}
+	
+private:
+	bool bFrameEnded;
+	uint32_t Counter = 0;
+	TArray<FMetalCommandBufferTiming> Timings;
+	FCriticalSection Mutex;
+};
+
 class FMetalRHICommandContext;
 
 /**
@@ -266,7 +289,6 @@ struct FMetalGPUProfiler : public FGPUProfiler
 	:	FGPUProfiler()
 	,	TimingSupport(InContext)
 	,	Context(InContext)
-	,   NumNestedFrames(0)
 	{}
 	
 	virtual ~FMetalGPUProfiler() {}
@@ -284,18 +306,17 @@ struct FMetalGPUProfiler : public FGPUProfiler
 	// WARNING:
 	// These functions MUST be called from within Metal scheduled/completion handlers
 	// since they depend on libdispatch to enforce ordering.
-	static void RecordFrame(TArray<FMetalCommandBufferTiming>& CommandBufferTimings, FMetalCommandBufferTiming& LastPresentBufferTiming);
+	static void RecordFrame(FMetalCommandBufferTimer& Timer);
 	static void RecordPresent(MTL::CommandBuffer* CommandBuffer);
 	// END WARNING
 	
 	static void ResetFrameBufferTimings();
-	static TSharedPtr<TArray<FMetalCommandBufferTiming>, ESPMode::ThreadSafe> GetFrameBufferTimings();
+	static FMetalCommandBufferTimer& GetFrameBufferTimer();
 	
 	FMetalGPUTiming TimingSupport;
 	FMetalRHICommandContext& Context;
-	int32 NumNestedFrames;
 	
-	static TSharedPtr<TArray<FMetalCommandBufferTiming>, ESPMode::ThreadSafe> FrameBufferTimings;
+	static FMetalCommandBufferTimer* Timer;
 };
 
 class FMetalProfiler : public FMetalGPUProfiler
diff --git a/Engine/Source/Runtime/Apple/MetalRHI/Private/MetalRHI.cpp b/Engine/Source/Runtime/Apple/MetalRHI/Private/MetalRHI.cpp
--- a/Engine/Source/Runtime/Apple/MetalRHI/Private/MetalRHI.cpp
+++ b/Engine/Source/Runtime/Apple/MetalRHI/Private/MetalRHI.cpp
@@ -1186,7 +1186,9 @@ FMetalDynamicRHI::FMetalDynamicRHI(ERHIFeatureLevel::Type RequestedFeatureLevel)
 	ImmediateContext.SetProfiler(Profiler);
 	
 	if (Profiler)
+	{
 		Profiler->BeginFrame();
+	}
 #endif
 
 #if METAL_USE_METAL_SHADER_CONVERTER
@@ -1200,13 +1202,6 @@ FMetalDynamicRHI::FMetalDynamicRHI(ERHIFeatureLevel::Type RequestedFeatureLevel)
 		BindlessDescriptorManager->Init();
 	}
 #endif
-	
-#if ENABLE_METAL_GPUPROFILE
-	if (Profiler)
-	{
-		Profiler->EndFrame();
-	}
-#endif
 }
 
 FMetalDynamicRHI::~FMetalDynamicRHI()
@@ -1254,6 +1249,8 @@ void FMetalDynamicRHI::RHIEndFrame_RenderThread(FRHICommandListImmediate& RHICmd
 	{
 		MTL_SCOPED_AUTORELEASE_POOL;
 
+		FMetalGPUProfiler::ResetFrameBufferTimings();
+		
 #if ENABLE_METAL_GPUPROFILE
 		Contexts[ERHIPipeline::Graphics]->GetProfiler()->EndFrame();
 #endif
@@ -1662,8 +1659,7 @@ void FMetalDynamicRHI::RHIFinalizeContext(FRHIFinalizeContextArgs&& Args, TRHIPi
 		
 		if(!CmdContext->IsInsideRenderPass())
 		{
-			FMetalCommandBuffer* CmdBuffer = CmdContext->Finalize();
-			PlatformCmdList->CommandBuffers.Add(CmdBuffer);
+			PlatformCmdList->CommandBuffers.Append(CmdContext->Finalize());
 
 			CmdContext->ResetContext();
 			if(GRHISupportsParallelRHIExecute)
diff --git a/Engine/Source/Runtime/Apple/MetalRHI/Private/MetalRHIContext.cpp b/Engine/Source/Runtime/Apple/MetalRHI/Private/MetalRHIContext.cpp
--- a/Engine/Source/Runtime/Apple/MetalRHI/Private/MetalRHIContext.cpp
+++ b/Engine/Source/Runtime/Apple/MetalRHI/Private/MetalRHIContext.cpp
@@ -447,7 +447,7 @@ void FMetalRHICommandContext::CopyFromBufferToBuffer(FMetalBufferPtr SourceBuffe
 	EndBlitEncoder();
 }
 
-FMetalCommandBuffer* FMetalRHICommandContext::Finalize()
+TArray<FMetalCommandBuffer*> FMetalRHICommandContext::Finalize()
 {
 	GetQueryBufferPool()->ReleaseCurrentQueryBuffer();
 	
@@ -456,13 +456,14 @@ FMetalCommandBuffer* FMetalRHICommandContext::Finalize()
 		CurrentEncoderFence = CurrentEncoder.EndEncoding();
 	}
 	
+	TArray<FMetalCommandBuffer*> CommandBuffers;
+	
 	if (CurrentEncoder.GetCommandBuffer())
 	{
-		FMetalCommandBuffer* CommandBuffer = CurrentEncoder.Finalize();
-		return CommandBuffer;
+		CommandBuffers = CurrentEncoder.Finalize();
 	}
 	
-	return nullptr;
+	return CommandBuffers;
 }
 
 void FMetalRHICommandContext::InsertCommandBufferFence(TSharedPtr<FMetalCommandBufferFence, ESPMode::ThreadSafe>& Fence, FMetalCommandBufferCompletionHandler Handler)
@@ -494,6 +495,11 @@ void FMetalRHICommandContext::StartTiming(class FMetalEventNode* EventNode)
 		{
 			CurrentEncoder.AddCompletionHandler(Handler);
 		}
+		
+		if(!bWithinRenderPass)
+		{
+			CurrentEncoder.SplitCommandBuffers();
+		}
 	}
 	
 	if (Handler.IsBound() && !bHasCurrentCommandBuffer)
@@ -584,13 +590,13 @@ TArray<FMetalCommandBuffer*>* FMetalRHIUploadContext::Finalize()
 	
 	TArray<FMetalCommandBuffer*>* CommandBuffers = new TArray<FMetalCommandBuffer*>();
 	
-	CommandBuffers->Add(UploadContext->Finalize());
+	CommandBuffers->Append(UploadContext->Finalize());
 	
 	UploadFunctions.Reset();
 	UploadContext->ResetContext();
 	
 	WaitContext->WaitForEvent(UploadSyncEvent, UploadSyncCounter);
-	CommandBuffers->Add(WaitContext->Finalize());
+	CommandBuffers->Append(WaitContext->Finalize());
 	
 	WaitContext->ResetContext();
 	
diff --git a/Engine/Source/Runtime/Apple/MetalRHI/Private/MetalRHIRenderQuery.cpp b/Engine/Source/Runtime/Apple/MetalRHI/Private/MetalRHIRenderQuery.cpp
--- a/Engine/Source/Runtime/Apple/MetalRHI/Private/MetalRHIRenderQuery.cpp
+++ b/Engine/Source/Runtime/Apple/MetalRHI/Private/MetalRHIRenderQuery.cpp
@@ -339,8 +339,11 @@ void FMetalRHIRenderQuery::End(FMetalRHICommandContext* Context)
 				this->Release();
 			});
 
-            Context->InsertCommandBufferFence(CommandBufferFence, Handler);
+			Context->InsertCommandBufferFence(CommandBufferFence, Handler);
 			Buffer.CommandBufferFence = CommandBufferFence;
+			
+			Context->SplitCommandBuffers();
+            
 			break;
 		}
 		default:
diff --git a/Engine/Source/Runtime/Apple/MetalRHI/Private/MetalViewport.cpp b/Engine/Source/Runtime/Apple/MetalRHI/Private/MetalViewport.cpp
--- a/Engine/Source/Runtime/Apple/MetalRHI/Private/MetalViewport.cpp
+++ b/Engine/Source/Runtime/Apple/MetalRHI/Private/MetalViewport.cpp
@@ -607,8 +607,6 @@ void FMetalViewport::Present(FMetalCommandQueue& CommandQueue, bool bLockToVsync
 		Block(0, 0.0, 0.0);
 	}
 	
-	FMetalGPUProfiler::ResetFrameBufferTimings();
-	
 	if (!(GRHISupportsRHIThread && IsRunningRHIInSeparateThread()))
 	{
 		Swap();
diff --git a/Engine/Source/Runtime/Apple/MetalRHI/Public/MetalRHIContext.h b/Engine/Source/Runtime/Apple/MetalRHI/Public/MetalRHIContext.h
--- a/Engine/Source/Runtime/Apple/MetalRHI/Public/MetalRHIContext.h
+++ b/Engine/Source/Runtime/Apple/MetalRHI/Public/MetalRHIContext.h
@@ -192,7 +192,7 @@ public:
 	bool PrepareToDraw(uint32 PrimitiveType);
 	void PrepareToDispatch();
 
-	FMetalCommandBuffer* Finalize();
+	TArray<FMetalCommandBuffer*> Finalize();
 	
 	void InsertCommandBufferFence(TSharedPtr<FMetalCommandBufferFence, ESPMode::ThreadSafe>& Fence, FMetalCommandBufferCompletionHandler Handler);
 	
@@ -264,6 +264,14 @@ public:
 		return bWithinRenderPass;
 	}
 	
+	void SplitCommandBuffers()
+	{
+		if(!bWithinRenderPass)
+		{
+			CurrentEncoder.SplitCommandBuffers();
+		}
+	}
+	
 protected:
 	FMetalDevice& Device;
 	
