Fix vision OS command buffer submission in PresentImmersive

diff --git a/Engine/Source/Runtime/Apple/MetalRHI/Private/MetalViewport.cpp b/Engine/Source/Runtime/Apple/MetalRHI/Private/MetalViewport.cpp
--- a/Engine/Source/Runtime/Apple/MetalRHI/Private/MetalViewport.cpp
+++ b/Engine/Source/Runtime/Apple/MetalRHI/Private/MetalViewport.cpp
@@ -708,15 +708,16 @@ void FMetalViewport::PresentImmersive(const MetalRHIVisionOS::PresentImmersivePa
 					}
 				}
 				
-				FMetalCommandBuffer* Buffer = Context.Finalize();
-				Context.GetDevice().GetCommandQueue().CommitCommandBuffer(Buffer);
-				Context.ResetContext();
+				TArray<FMetalCommandBuffer*> Buffers = Context.Finalize();
+				check(Buffers.Num() > 0);
 				
-				FMetalCommandBuffer* CurrentCommandBuffer = Context.GetCurrentCommandBuffer();
+				// We need to attach the completion handler and the present signal to the final
+				// command buffer
+				FMetalCommandBuffer* FinalCommandBuffer = Buffers.Last();
 				
 #if ENABLE_METAL_GPUPROFILE
 				FMetalProfiler* Profiler = FMetalProfiler::GetProfiler();
-				FMetalCommandBufferStats* Stats = Profiler->AllocateCommandBuffer(CurrentCommandBuffer->GetMTLCmdBuffer(), 0);
+				FMetalCommandBufferStats* Stats = Profiler->AllocateCommandBuffer(FinalCommandBuffer->GetMTLCmdBuffer(), 0);
 #endif
 				
 				{
@@ -729,15 +730,18 @@ void FMetalViewport::PresentImmersive(const MetalRHIVisionOS::PresentImmersivePa
 						
 						FMetalGPUProfiler::RecordPresent(cmd_buf);
 					};
-					CurrentCommandBuffer->GetMTLCmdBuffer()->addCompletedHandler(CommandBufferHandler);
+					FinalCommandBuffer->GetMTLCmdBuffer()->addCompletedHandler(CommandBufferHandler);
 				}
 				
-				cp_drawable_encode_present(VisionOSParams.SwiftDrawable, (__bridge id<MTLCommandBuffer>)CurrentCommandBuffer->GetMTLCmdBuffer().get());
+				cp_drawable_encode_present(VisionOSParams.SwiftDrawable, (__bridge id<MTLCommandBuffer>)FinalCommandBuffer->GetMTLCmdBuffer().get());
 				cp_frame_t CompositorServicesFrame = VisionOSParams.SwiftFrame;
+
+				METAL_GPUPROFILE(Stats->End(FinalCommandBuffer->GetMTLCmdBuffer()));			
 				
-				Buffer = Context.Finalize();
-				METAL_GPUPROFILE(Stats->End(CurrentCommandBuffer->GetMTLCmdBuffer()));			
-				Context.GetDevice().GetCommandQueue().CommitCommandBuffer(Buffer);
+				for(FMetalCommandBuffer* Buffer : Buffers)
+				{
+					Context.GetDevice().GetCommandQueue().CommitCommandBuffer(Buffer);
+				}
 				Context.ResetContext();
 				
 				cp_frame_end_submission(CompositorServicesFrame);
