Fixed compilation errors on iOS caused by the file mapping API change

diff --git a/Engine/Source/Runtime/Apple/MetalRHI/Private/MetalDynamicRHI_Shaders.cpp b/Engine/Source/Runtime/Apple/MetalRHI/Private/MetalDynamicRHI_Shaders.cpp
--- a/Engine/Source/Runtime/Apple/MetalRHI/Private/MetalDynamicRHI_Shaders.cpp
+++ b/Engine/Source/Runtime/Apple/MetalRHI/Private/MetalDynamicRHI_Shaders.cpp
@@ -158,8 +158,8 @@ FRHIShaderLibraryRef FMetalDynamicRHI::RHICreateShaderLibrary(EShaderPlatform Pl
     TUniquePtr<FMetalShaderLibrary::FShaderLibDataOwner> MemOwner = MakeUnique<FMetalShaderLibrary::FShaderLibDataOwner>();
 	if (FPlatformProperties::SupportsMemoryMappedFiles())
 	{
-		FOpenMappedResult Result = FPlatformFileManager::Get().GetPlatformFile().OpenMappedEx(*BinaryShaderFile);
-		MemOwner->MappedCacheFile = Result.HasError() ? nullptr : Result.StealValue();
+		FOpenMappedResult Res = FPlatformFileManager::Get().GetPlatformFile().OpenMappedEx(*BinaryShaderFile);
+		MemOwner->MappedCacheFile = Res.HasError() ? nullptr : Res.StealValue();
 		if (MemOwner->MappedCacheFile.IsValid())
 		{
 			MemOwner->MappedRegion = TUniquePtr<IMappedFileRegion>(MemOwner->MappedCacheFile->MapRegion(0, MemOwner->MappedCacheFile->GetFileSize()));
diff --git a/Engine/Source/Runtime/Core/Private/IOS/IOSPlatformFile.cpp b/Engine/Source/Runtime/Core/Private/IOS/IOSPlatformFile.cpp
--- a/Engine/Source/Runtime/Core/Private/IOS/IOSPlatformFile.cpp
+++ b/Engine/Source/Runtime/Core/Private/IOS/IOSPlatformFile.cpp
@@ -944,7 +944,7 @@ FOpenMappedResult FIOSPlatformFile::OpenMappedEx(const TCHAR* Filename, EOpenRea
 			// check the read path
 			if(fstat(Handle, &FileInfo) == -1)
 			{
-				return NULL;
+				return MakeError(FString::Printf(TEXT("FIOSPlatformFile::OpenMappedEx failed to get file info for file '%s'"), *NormalizedFilename));
 			}
 			uint64 FileSize = FileInfo.st_size;
 
