Fix CIS MetalRHI

diff --git a/Engine/Source/Runtime/Apple/MetalRHI/Private/MetalRHI.cpp b/Engine/Source/Runtime/Apple/MetalRHI/Private/MetalRHI.cpp
--- a/Engine/Source/Runtime/Apple/MetalRHI/Private/MetalRHI.cpp
+++ b/Engine/Source/Runtime/Apple/MetalRHI/Private/MetalRHI.cpp
@@ -836,7 +836,8 @@ FMetalDynamicRHI::FMetalDynamicRHI(ERHIFeatureLevel::Type RequestedFeatureLevel)
 	GMetalBufferFormats[PF_ASTC_8x8_NORM_RG		] = { MTL::PixelFormatInvalid, (uint8)EMetalBufferFormat::Unknown };
 	GMetalBufferFormats[PF_ASTC_10x10_NORM_RG	] = { MTL::PixelFormatInvalid, (uint8)EMetalBufferFormat::Unknown };
 	GMetalBufferFormats[PF_ASTC_12x12_NORM_RG	] = { MTL::PixelFormatInvalid, (uint8)EMetalBufferFormat::Unknown };
-	static_assert(PF_MAX == 93, "Please setup GMetalBufferFormats properly for the new pixel format");
+	GMetalBufferFormats[PF_R8G8B8				] = { MTL::PixelFormatInvalid, (uint8)EMetalBufferFormat::Unknown };
+	static_assert(PF_MAX == 94, "Please setup GMetalBufferFormats properly for the new pixel format");
 
 	// Initialize the platform pixel format map.
 	GPixelFormats[PF_Unknown			].PlatformFormat	= (uint32)MTL::PixelFormatInvalid;
