Adding CVar and logic for extra frame of latency in deferred delete
Added additional logging for crash in MetalUniformBuffer

diff --git a/Engine/Source/Runtime/Apple/MetalRHI/Private/MetalBuffer.cpp b/Engine/Source/Runtime/Apple/MetalRHI/Private/MetalBuffer.cpp
--- a/Engine/Source/Runtime/Apple/MetalRHI/Private/MetalBuffer.cpp
+++ b/Engine/Source/Runtime/Apple/MetalRHI/Private/MetalBuffer.cpp
@@ -71,7 +71,7 @@ void FMetalBuffer::Release()
 {
     check(!GIsRHIInitialized || (OnFreePolicy != FreePolicy::Owner || bMarkedDeleted));
     
-	if (OnFreePolicy == FreePolicy::BufferAllocator)
+	if (OnFreePolicy == FreePolicy::BufferAllocator || Allocator)
 	{
 		check(Allocator);
 		Allocator->ReleaseBuffer(this);
diff --git a/Engine/Source/Runtime/Apple/MetalRHI/Private/MetalDynamicRHI.h b/Engine/Source/Runtime/Apple/MetalRHI/Private/MetalDynamicRHI.h
--- a/Engine/Source/Runtime/Apple/MetalRHI/Private/MetalDynamicRHI.h
+++ b/Engine/Source/Runtime/Apple/MetalRHI/Private/MetalDynamicRHI.h
@@ -233,6 +233,7 @@ private:
 	{
 		TArray<FMetalDeferredDeleteObject> DeferredDeleteObjects;
 		TArray<TSharedPtr<FMetalCommandBufferFence, ESPMode::ThreadSafe>> WaitFences;
+		bool bWaitExtraFrame = false;
 	};
 	
 	TArray<FDeferredDeleteData> DeferredDeleteQueue;
diff --git a/Engine/Source/Runtime/Apple/MetalRHI/Private/MetalRHI.cpp b/Engine/Source/Runtime/Apple/MetalRHI/Private/MetalRHI.cpp
--- a/Engine/Source/Runtime/Apple/MetalRHI/Private/MetalRHI.cpp
+++ b/Engine/Source/Runtime/Apple/MetalRHI/Private/MetalRHI.cpp
@@ -55,6 +55,13 @@ static TAutoConsoleVariable<bool> CVarEnableMetalPSOFileCacheWhenPrecachingActiv
 	TEXT("true: Allow both PSO file cache and precaching."),
 	ECVF_RenderThreadSafe | ECVF_ReadOnly);
 
+static TAutoConsoleVariable<bool> CVarEnableMetalDeferredDeleteLatency(
+	TEXT("r.Metal.EnableMetalDeferredDeleteLatency"),
+	true,
+	TEXT("false: No added latency on deferred delete \n")
+	TEXT("true: Extra latency on deferred delete"),
+	ECVF_RenderThreadSafe | ECVF_ReadOnly);
+
 extern int32 GMetalResourcePurgeOnDelete;
 
 static void ValidateTargetedRHIFeatureLevelExists(EShaderPlatform Platform)
@@ -1540,6 +1547,13 @@ void FMetalDynamicRHI::ProcessDeferredDeleteQueue()
 			break;
 		}
 		
+		if(Data.bWaitExtraFrame)
+		{
+			Data.bWaitExtraFrame = false;
+			Index++;
+			continue;
+		}
+		
 		for(FMetalDeferredDeleteObject& Object : Data.DeferredDeleteObjects)
 		{	
 			switch (Object.Storage.GetIndex())
@@ -1622,7 +1636,12 @@ void FMetalDynamicRHI::ProcessDeferredDeleteQueue()
 
 void FMetalDynamicRHI::RHIProcessDeleteQueue()
 {
-	FDeferredDeleteData NewData;	
+	FDeferredDeleteData NewData;
+
+	if(CVarEnableMetalDeferredDeleteLatency.GetValueOnAnyThread() > 0)
+	{
+		NewData.bWaitExtraFrame = true;
+	}
 	GatherDeferredDeleteObjects(NewData.DeferredDeleteObjects, NewData.WaitFences);
 	
 	DeferredDeleteQueue.Add(MoveTemp(NewData));
diff --git a/Engine/Source/Runtime/Apple/MetalRHI/Private/MetalTempAllocator.cpp b/Engine/Source/Runtime/Apple/MetalRHI/Private/MetalTempAllocator.cpp
--- a/Engine/Source/Runtime/Apple/MetalRHI/Private/MetalTempAllocator.cpp
+++ b/Engine/Source/Runtime/Apple/MetalRHI/Private/MetalTempAllocator.cpp
@@ -44,12 +44,20 @@ FMetalBufferPtr FMetalTempAllocator::Allocate(const uint32_t Size)
 		TempBuffer->Buffer = Device.GetDevice()->newBuffer(BufferSize, 
 														MTL::ResourceCPUCacheModeWriteCombined |
 														MTL::ResourceStorageModeShared);
+		
+		if(!TempBuffer->Buffer)
+		{
+			UE_LOG(LogMetal, Fatal, TEXT("Failed to allocate MTL::Buffer in MetalTempAllocator::Allocate"));
+		}
 	}
 	
 	FMetalBufferPtr Buffer = FMetalBufferPtr(new FMetalBuffer(TempBuffer->Buffer, NS::Range(TempBuffer->Offset, AlignedSize), this));
 	TempBuffer->Offset += AlignedSize;
 	
-	check(Buffer);
+	if(!Buffer)
+	{
+		UE_LOG(LogMetal, Fatal, TEXT("Failed to allocate FMetalBuffer in MetalTempAllocator::Allocate"));
+	}
 	
 	return Buffer;
 }
diff --git a/Engine/Source/Runtime/Apple/MetalRHI/Private/MetalUniformBuffer.cpp b/Engine/Source/Runtime/Apple/MetalRHI/Private/MetalUniformBuffer.cpp
--- a/Engine/Source/Runtime/Apple/MetalRHI/Private/MetalUniformBuffer.cpp
+++ b/Engine/Source/Runtime/Apple/MetalRHI/Private/MetalUniformBuffer.cpp
@@ -53,11 +53,15 @@ void FMetalSuballocatedUniformBuffer::Update(const void* Contents)
 void FMetalSuballocatedUniformBuffer::PushToGPUBacking(const void* Contents)
 {    
 	FMetalTempAllocator* Allocator = Device.GetUniformAllocator();
-	FMetalBufferPtr Buffer = Allocator->Allocate(GetSize());
-	// copy contents into backing
-	BackingBuffer = Buffer;
 	
-	uint8* ConstantSpace = reinterpret_cast<uint8*>(Buffer->Contents());
+	BackingBuffer = Allocator->Allocate(GetSize());
+	
+	if(!BackingBuffer)
+	{
+		UE_LOG(LogMetal, Fatal, TEXT("BackingBuffer returned from MetalTempAllocator::Allocate is nullptr"));
+	}
+	
+	uint8* ConstantSpace = reinterpret_cast<uint8*>(BackingBuffer->Contents());
 	FMemory::Memcpy(ConstantSpace, Contents, GetSize());
 	LastFrameUpdated = Device.GetFrameNumberRHIThread();
 }
