Texture locking cleanup
- Merging LockTexture2D, LockTexture2DArray and LockTextureCubeFace into a single LockTexture function.
- LockTexture and UnlockTexture now take a single argument that has all the locking arguments.
- LockTexture returns a structure with the data pointer, the data stride and the lock data size instead of having multiple output arguments.
- FRHICommandList versions not deprecated yet, but they all use LockTexture directly.

diff --git a/Engine/Source/Runtime/Apple/MetalRHI/Private/MetalDynamicRHI.h b/Engine/Source/Runtime/Apple/MetalRHI/Private/MetalDynamicRHI.h
--- a/Engine/Source/Runtime/Apple/MetalRHI/Private/MetalDynamicRHI.h
+++ b/Engine/Source/Runtime/Apple/MetalRHI/Private/MetalDynamicRHI.h
@@ -122,14 +122,12 @@ public:
 	virtual FTextureRHIRef RHIAsyncCreateTexture2D(uint32 SizeX, uint32 SizeY, uint8 Format, uint32 NumMips, ETextureCreateFlags Flags, ERHIAccess InResourceState, void** InitialMipData, uint32 NumInitialMips, const TCHAR* DebugName, FGraphEventRef& OutCompletionEvent) final override;
 	virtual uint32 RHIComputeMemorySize(FRHITexture* TextureRHI) final override;
 	virtual FTextureRHIRef RHIAsyncReallocateTexture2D(FRHITexture* Texture2D, int32 NewMipCount, int32 NewSizeX, int32 NewSizeY, FThreadSafeCounter* RequestStatus) final override;
-	virtual void* RHILockTexture2D(FRHITexture* Texture, uint32 MipIndex, EResourceLockMode LockMode, uint32& DestStride, bool bLockWithinMiptail, uint64* OutLockedByteCount = nullptr) final override;
-	virtual void RHIUnlockTexture2D(FRHITexture* Texture, uint32 MipIndex, bool bLockWithinMiptail) final override;
-	virtual void* RHILockTexture2DArray(FRHITexture* Texture, uint32 TextureIndex, uint32 MipIndex, EResourceLockMode LockMode, uint32& DestStride, bool bLockWithinMiptail) final override;
-	virtual void RHIUnlockTexture2DArray(FRHITexture* Texture, uint32 TextureIndex, uint32 MipIndex, bool bLockWithinMiptail) final override;
+
+	virtual FRHILockTextureResult RHILockTexture(FRHICommandListImmediate& RHICmdList, const FRHILockTextureArgs& Arguments) final override;
+	virtual void                  RHIUnlockTexture(FRHICommandListImmediate& RHICmdList, const FRHILockTextureArgs& Arguments) final override;
+
 	virtual void RHIUpdateTexture2D(FRHICommandListBase& RHICmdList, FRHITexture* Texture, uint32 MipIndex, const struct FUpdateTextureRegion2D& UpdateRegion, uint32 SourcePitch, const uint8* SourceData) final override;
 	virtual void RHIUpdateTexture3D(FRHICommandListBase& RHICmdList, FRHITexture* Texture, uint32 MipIndex, const struct FUpdateTextureRegion3D& UpdateRegion, uint32 SourceRowPitch, uint32 SourceDepthPitch, const uint8* SourceData) final override;
-	virtual void* RHILockTextureCubeFace(FRHITexture* Texture, uint32 FaceIndex, uint32 ArrayIndex, uint32 MipIndex, EResourceLockMode LockMode, uint32& DestStride, bool bLockWithinMiptail) final override;
-	virtual void RHIUnlockTextureCubeFace(FRHITexture* Texture, uint32 FaceIndex, uint32 ArrayIndex, uint32 MipIndex, bool bLockWithinMiptail) final override;
 	virtual void RHIBindDebugLabelName(FRHICommandListBase& RHICmdList, FRHITexture* Texture, const TCHAR* Name) final override;
 	virtual void RHIReadSurfaceData(FRHITexture* Texture,FIntRect Rect,TArray<FColor>& OutData,FReadSurfaceDataFlags InFlags) final override;
 	virtual void RHIReadSurfaceData(FRHITexture* TextureRHI, FIntRect InRect, TArray<FLinearColor>& OutData, FReadSurfaceDataFlags InFlags) final override;
@@ -165,9 +163,6 @@ public:
 
 	virtual FTextureRHIRef AsyncReallocateTexture2D_RenderThread(class FRHICommandListImmediate& RHICmdList, FRHITexture* Texture2D, int32 NewMipCount, int32 NewSizeX, int32 NewSizeY, FThreadSafeCounter* RequestStatus) final override;
 
-	virtual void* LockTexture2D_RenderThread(class FRHICommandListImmediate& RHICmdList, FRHITexture* Texture, uint32 MipIndex, EResourceLockMode LockMode, uint32& DestStride, bool bLockWithinMiptail, bool bNeedsDefaultRHIFlush = true, uint64* OutLockedByteCount = nullptr) final override;
-	virtual void UnlockTexture2D_RenderThread(class FRHICommandListImmediate& RHICmdList, FRHITexture* Texture, uint32 MipIndex, bool bLockWithinMiptail, bool bNeedsDefaultRHIFlush = true) final override;
-
 	virtual FUpdateTexture3DData RHIBeginUpdateTexture3D(FRHICommandListBase& RHICmdList, FRHITexture* Texture, uint32 MipIndex, const struct FUpdateTextureRegion3D& UpdateRegion) final override;
 	virtual void RHIEndUpdateTexture3D(FRHICommandListBase& RHICmdList, FUpdateTexture3DData& UpdateData) final override;
 
diff --git a/Engine/Source/Runtime/Apple/MetalRHI/Private/MetalRenderTarget.cpp b/Engine/Source/Runtime/Apple/MetalRHI/Private/MetalRenderTarget.cpp
--- a/Engine/Source/Runtime/Apple/MetalRHI/Private/MetalRenderTarget.cpp
+++ b/Engine/Source/Runtime/Apple/MetalRHI/Private/MetalRenderTarget.cpp
@@ -291,12 +291,12 @@ void FMetalDynamicRHI::RHIMapStagingSurface(FRHITexture* TextureRHI, FRHIGPUFenc
 	}
 	
 	FMetalSurface* Surface = GetMetalSurfaceFromRHITexture(TextureRHI);
-    
-	uint32 Stride = 0;
+
+	FRHILockTextureResult Result = Surface->Lock(FRHILockTextureArgs::Lock2D(Surface, 0, RLM_ReadOnly, false), false);
+
 	OutWidth = Surface->GetSizeX();
 	OutHeight = Surface->GetSizeY();
-		
-	OutData = Surface->Lock(0, 0, RLM_ReadOnly, Stride);
+	OutData = Result.Data;
 }
 
 void FMetalDynamicRHI::RHIUnmapStagingSurface(FRHITexture* TextureRHI, uint32 GPUIndex)
@@ -304,7 +304,7 @@ void FMetalDynamicRHI::RHIUnmapStagingSurface(FRHITexture* TextureRHI, uint32 GP
     MTL_SCOPED_AUTORELEASE_POOL;
     
     FMetalSurface* Surface = GetMetalSurfaceFromRHITexture(TextureRHI);
-    Surface->Unlock(0, 0, false);
+    Surface->Unlock(FRHILockTextureArgs::Lock2D(Surface, 0, RLM_ReadOnly, false));
 }
 
 void FMetalDynamicRHI::RHIReadSurfaceFloatData(FRHITexture* TextureRHI, FIntRect Rect, TArray<FFloat16Color>& OutData, ECubeFace CubeFace,int32 ArrayIndex,int32 MipIndex)
diff --git a/Engine/Source/Runtime/Apple/MetalRHI/Private/MetalTexture.cpp b/Engine/Source/Runtime/Apple/MetalRHI/Private/MetalTexture.cpp
--- a/Engine/Source/Runtime/Apple/MetalRHI/Private/MetalTexture.cpp
+++ b/Engine/Source/Runtime/Apple/MetalRHI/Private/MetalTexture.cpp
@@ -220,13 +220,6 @@ MTL::PixelFormat UEToMetalFormat(FMetalDevice& Device, EPixelFormat UEFormat, bo
 	return MTLFormat;
 }
 
-static inline uint32 ComputeLockIndex(uint32 MipIndex, uint32 ArrayIndex)
-{
-	check(MipIndex < MAX_uint16);
-	check(ArrayIndex < MAX_uint16);
-	return (MipIndex & MAX_uint16) | ((ArrayIndex & MAX_uint16) << 16);
-}
-
 MTLTexturePtr FMetalSurface::Reallocate(MTLTexturePtr InTexture, MTL::TextureUsage UsageModifier)
 {
 	MTL::TextureDescriptor* Desc = MTL::TextureDescriptor::alloc()->init();
@@ -713,11 +706,12 @@ FMetalSurface::FMetalSurface(FMetalDevice& MetalDevice, FRHICommandListBase* RHI
 		FRHICommandListImmediate& RHICmdListImmediate = RHICmdList->GetAsImmediate();
 
 		// lock, copy, unlock
-		uint32 Stride;
-		void* LockedData = FMetalDynamicRHI::Get().LockTexture2D_RenderThread(RHICmdListImmediate, this, 0, RLM_WriteOnly, Stride, false);
-		check(LockedData);
-		FMemory::Memcpy(LockedData, BulkData->GetResourceBulkData(), BulkData->GetResourceBulkDataSize());
-		FMetalDynamicRHI::Get().UnlockTexture2D_RenderThread(RHICmdListImmediate, this, 0, false);
+		const FRHILockTextureArgs LockArgs = FRHILockTextureArgs::Lock2D(this, 0, RLM_WriteOnly, false);
+
+		FRHILockTextureResult LockResult = RHICmdListImmediate.LockTexture(LockArgs);
+		check(LockResult.Data);
+		FMemory::Memcpy(LockResult.Data, BulkData->GetResourceBulkData(), BulkData->GetResourceBulkDataSize());
+		RHICmdListImmediate.UnlockTexture(LockArgs);
 
 		// bulk data can be unloaded now
 		BulkData->Discard();
@@ -922,13 +916,27 @@ void FMetalSurface::SafeRelease(MTLTexturePtr InTexture)
 	}
 }
 
-MTL::Buffer* FMetalSurface::AllocSurface(uint32 MipIndex, uint32 ArrayIndex, EResourceLockMode LockMode, uint32& DestStride, bool SingleLayer /*= false*/)
+static uint32 GetMetalResourceIndex(const FRHILockTextureArgs& Arguments)
+{
+	const FRHITextureDesc& Desc = Arguments.Texture->GetDesc();
+
+	if (Desc.IsTextureCube())
+	{
+		return Arguments.FaceIndex + Arguments.ArrayIndex * 6;
+	}
+
+	if (Desc.IsTextureArray())
+	{
+		return Arguments.ArrayIndex;
+	}
+
+	return 0;
+}
+
+MTL::Buffer* FMetalSurface::AllocSurface(const FRHILockTextureArgs& Arguments, uint32 MipBytes, uint32 DestStride)
 {
 	check(IsInRenderingThread());
 
-	// get size and stride
-	uint32 MipBytes = GetMipSize(MipIndex, &DestStride, SingleLayer);
-	
 	// allocate some temporary memory
 	// This should really be pooled and texture transfers should be their own pool
 	MTL::Device* MTLDevice = Device.GetDevice();
@@ -936,18 +944,17 @@ MTL::Buffer* FMetalSurface::AllocSurface(uint32 MipIndex, uint32 ArrayIndex, ERe
 	Buffer->setLabel(NS::String::string("Temporary Surface Backing", NS::UTF8StringEncoding));
 	
 	// Note: while the lock is active, this map owns the backing store.
-	const uint32 LockIndex = ComputeLockIndex(MipIndex, ArrayIndex);
-	GRHILockTracker.Lock(this, Buffer, LockIndex, MipBytes, LockMode, false);
+	GRHILockTracker.Lock(Arguments, Buffer, false);
 	
 #if PLATFORM_MAC
 	// Expand R8_sRGB into RGBA8_sRGB for non Apple Silicon Mac.
 	if (GetDesc().Format == PF_G8 
 		&& EnumHasAnyFlags(GetDesc().Flags, TexCreate_SRGB) 
-		&& LockMode == RLM_WriteOnly 
+		&& Arguments.LockMode == RLM_WriteOnly
 		&& Texture->pixelFormat() == MTL::PixelFormatRGBA8Unorm_sRGB)
 	{
 		check(GetDesc().Dimension == ETextureDimension::Texture2D || GetDesc().Dimension == ETextureDimension::Texture3D);
-		DestStride = FMath::Max<uint32>(GetDesc().Extent.X >> MipIndex, 1);
+		DestStride = FMath::Max<uint32>(GetDesc().Extent.X >> Arguments.MipIndex, 1);
 	}
 #endif
 	
@@ -1076,19 +1083,20 @@ void FMetalSurface::UpdateSurfaceAndDestroySourceBuffer(FMetalRHICommandContext*
 #endif
 }
 
-void* FMetalSurface::Lock(uint32 MipIndex, uint32 ArrayIndex, EResourceLockMode LockMode, uint32& DestStride, bool SingleLayer /*= false*/, uint64* OutLockedByteCount /* =nullptr */)
+FRHILockTextureResult FMetalSurface::Lock(const FRHILockTextureArgs& Arguments, bool bSingleLayer)
 {
+	const uint32 ResourceIndex = GetMetalResourceIndex(Arguments);
+
+	FRHILockTextureResult Result;
+
 	// get size and stride
-	uint32 MipBytes = GetMipSize(MipIndex, &DestStride, false);
-	if (OutLockedByteCount)
-	{
-		*OutLockedByteCount = (uint64)MipBytes;
-	}
-	
+	uint32 MipBytes = GetMipSize(Arguments.MipIndex, &Result.Stride, bSingleLayer);
+	Result.ByteCount = MipBytes;
+
 	// allocate some temporary memory
-	MTL::Buffer* SourceData = AllocSurface(MipIndex, ArrayIndex, LockMode, DestStride, SingleLayer);
-	
-	switch(LockMode)
+	MTL::Buffer* SourceData = AllocSurface(Arguments, MipBytes, Result.Stride);
+
+	switch(Arguments.LockMode)
 	{
 		case RLM_ReadOnly:
 		{
@@ -1102,50 +1110,46 @@ void* FMetalSurface::Lock(uint32 MipIndex, uint32 ArrayIndex, EResourceLockMode
 				// upload the texture to the texture slice
 				Region = MTL::Region(
 					0, 0, 0,
-					FMath::Max<uint32>(GetDesc().Extent.X >> MipIndex, 1),
-					FMath::Max<uint32>(GetDesc().Extent.Y >> MipIndex, 1),
-					FMath::Max<uint32>(GetDesc().Depth    >> MipIndex, 1));
+					FMath::Max<uint32>(GetDesc().Extent.X >> Arguments.MipIndex, 1),
+					FMath::Max<uint32>(GetDesc().Extent.Y >> Arguments.MipIndex, 1),
+					FMath::Max<uint32>(GetDesc().Depth    >> Arguments.MipIndex, 1));
 			}
 			else
 			{
 				// upload the texture to the texture slice
 				Region = MTL::Region(
 					0, 0,
-					FMath::Max<uint32>(GetDesc().Extent.X >> MipIndex, 1),
-					FMath::Max<uint32>(GetDesc().Extent.Y >> MipIndex, 1)
+					FMath::Max<uint32>(GetDesc().Extent.X >> Arguments.MipIndex, 1),
+					FMath::Max<uint32>(GetDesc().Extent.Y >> Arguments.MipIndex, 1)
 				);
 			}
 			
 			if (Texture->storageMode() == MTL::StorageModePrivate)
 			{
 				// If we are running with command lists or the RHI thread is enabled we have to execute GFX commands in that context.
-				auto CopyTexToBuf =
-				[this, ArrayIndex, MipIndex, Region, SourceData, DestStride, MipBytes](FRHICommandListImmediate& RHICmdList)
-				{
-                    FMetalBufferPtr Source = FMetalBufferPtr(new FMetalBuffer(SourceData, FMetalBuffer::FreePolicy::Temporary));
-                                        
-					FMetalRHICommandContext& Context = FMetalRHICommandContext::Get(RHICmdList);
-					Context.CopyFromTextureToBuffer(this->Texture.get(), ArrayIndex, MipIndex, Region.origin, Region.size, Source, 0, DestStride, MipBytes, MTL::BlitOptionNone);
-				};
-				
-				RHICmdList.EnqueueLambda(MoveTemp(CopyTexToBuf));
-				RHICmdList.GetAsImmediate().SubmitAndBlockUntilGPUIdle();
+				RHICmdList.EnqueueLambda(
+					[this, ResourceIndex, MipIndex = Arguments.MipIndex, Region, SourceData, Stride = Result.Stride, MipBytes](FRHICommandListImmediate& RHICmdList)
+					{
+						FMetalBufferPtr Source = FMetalBufferPtr(new FMetalBuffer(SourceData, FMetalBuffer::FreePolicy::Temporary));
+
+						FMetalRHICommandContext& Context = FMetalRHICommandContext::Get(RHICmdList);
+						Context.CopyFromTextureToBuffer(this->Texture.get(), ResourceIndex, MipIndex, Region.origin, Region.size, Source, 0, Stride, MipBytes, MTL::BlitOptionNone);
+					});
+				RHICmdList.SubmitAndBlockUntilGPUIdle();
 			}
 			else
 			{
 #if PLATFORM_MAC
-				if(this->Texture->storageMode() == MTL::StorageModeManaged)
+				if (this->Texture->storageMode() == MTL::StorageModeManaged)
 				{
 					// Managed texture - need to sync GPU -> CPU before access as it could have been written to by the GPU
-					auto SyncReadbackToCPU =
-					[this, ArrayIndex, MipIndex](FRHICommandListImmediate& RHICmdList)
-					{
-						FMetalRHICommandContext& Context = FMetalRHICommandContext::Get(RHICmdList);
-						Context.SynchronizeTexture(this->Texture.get(), ArrayIndex, MipIndex);
-					};
-					
-					RHICmdList.EnqueueLambda(MoveTemp(SyncReadbackToCPU));
-					RHICmdList.GetAsImmediate().SubmitAndBlockUntilGPUIdle();
+					RHICmdList.EnqueueLambda(
+						[this, ResourceIndex, MipIndex = Arguments.MipIndex](FRHICommandListImmediate& RHICmdList)
+						{
+							FMetalRHICommandContext& Context = FMetalRHICommandContext::Get(RHICmdList);
+							Context.SynchronizeTexture(this->Texture.get(), ResourceIndex, MipIndex);
+						});
+					RHICmdList.SubmitAndBlockUntilGPUIdle();
 				}
 #endif
 				
@@ -1154,11 +1158,11 @@ void* FMetalSurface::Lock(uint32 MipIndex, uint32 ArrayIndex, EResourceLockMode
 				if (GetDesc().Format == PF_PVRTC2 || GetDesc().Format == PF_PVRTC4)
 				{
 					// for compressed textures metal debug RT expects 0 for rowBytes and imageBytes.
-					DestStride = 0;
+					Result.Stride = 0;
 					MipBytes = 0;
 				}
 #endif
-				uint32 BytesPerRow = DestStride;
+				uint32 BytesPerRow = Result.Stride;
 				if (GetDesc().Format == PF_PVRTC2 || GetDesc().Format == PF_PVRTC4)
 				{
 					// for compressed textures metal debug RT expects 0 for rowBytes and imageBytes.
@@ -1171,7 +1175,7 @@ void* FMetalSurface::Lock(uint32 MipIndex, uint32 ArrayIndex, EResourceLockMode
 				
 				void * Contents = SourceData->contents();
 				
-				Texture->getBytes(Contents, BytesPerRow, BytesPerImage, Region, MipIndex, ArrayIndex);
+				Texture->getBytes(Contents, BytesPerRow, BytesPerImage, Region, Arguments.MipIndex, ResourceIndex);
 			}
 			
 #if PLATFORM_MAC
@@ -1184,16 +1188,16 @@ void* FMetalSurface::Lock(uint32 MipIndex, uint32 ArrayIndex, EResourceLockMode
 				TArray<uint8> Data;
 				uint8* ExpandedMem = (uint8*)SourceData->contents();
 				Data.Append(ExpandedMem, MipBytes);
-				uint32 SrcStride = DestStride;
-				DestStride = FMath::Max<uint32>(GetDesc().Extent.X >> MipIndex, 1);
-				for(uint y = 0; y < FMath::Max<uint32>(GetDesc().Extent.Y >> MipIndex, 1); y++)
+				uint32 SrcStride = Result.Stride;
+				Result.Stride = FMath::Max<uint32>(GetDesc().Extent.X >> Arguments.MipIndex, 1);
+				for(uint y = 0; y < FMath::Max<uint32>(GetDesc().Extent.Y >> Arguments.MipIndex, 1); y++)
 				{
 					uint8* RowDest = ExpandedMem;
-					for(uint x = 0; x < FMath::Max<uint32>(GetDesc().Extent.X >> MipIndex, 1); x++)
+					for(uint x = 0; x < FMath::Max<uint32>(GetDesc().Extent.X >> Arguments.MipIndex, 1); x++)
 					{
 						*(RowDest++) = Data[(y * SrcStride) + (x * 4)];
 					}
-					ExpandedMem = (ExpandedMem + DestStride);
+					ExpandedMem = (ExpandedMem + Result.Stride);
 				}
 			}
 #endif
@@ -1209,50 +1213,46 @@ void* FMetalSurface::Lock(uint32 MipIndex, uint32 ArrayIndex, EResourceLockMode
 			break;
 	}
 	
-	return SourceData->contents();
+	Result.Data = SourceData->contents();
+
+	return Result;
 }
 
-void FMetalSurface::Unlock(uint32 MipIndex, uint32 ArrayIndex, bool bTryAsync)
+void FMetalSurface::Unlock(const FRHILockTextureArgs& Arguments)
 {
 	check(IsInRenderingThread());
-	
-	const uint32 LockIndex = ComputeLockIndex(MipIndex, ArrayIndex);
-	FRHILockTracker::FLockParams Params = GRHILockTracker.Unlock(this, LockIndex);
-	
-	MTL::Buffer* SourceData = (MTL::Buffer*) Params.Buffer;
-	
+
+	FRHILockTracker::FTextureLockParams Params = GRHILockTracker.Unlock(Arguments);
+	MTL::Buffer* SourceData = reinterpret_cast<MTL::Buffer*>(Params.Data);
+
+	const uint32 ResourceIndex = GetMetalResourceIndex(Arguments);
+
 	FRHICommandListImmediate& RHICmdList = FRHICommandListImmediate::Get();
-	
-	RHICmdList.EnqueueLambda([
-		Surface = this,
-		MipIndex,
-		ArrayIndex,
-		SourceData](FRHICommandListBase& ExecutingCmdList)
-	{
-		FMetalRHICommandContext& Context = FMetalRHICommandContext::Get(ExecutingCmdList);
-		Surface->UpdateSurfaceAndDestroySourceBuffer(&Context, SourceData, MipIndex, ArrayIndex);
-	});
+	RHICmdList.EnqueueLambda(
+		[Surface = this, SourceData, MipIndex = Arguments.MipIndex, ResourceIndex](FRHICommandListBase& ExecutingCmdList)
+		{
+			FMetalRHICommandContext& Context = FMetalRHICommandContext::Get(ExecutingCmdList);
+			Surface->UpdateSurfaceAndDestroySourceBuffer(&Context, SourceData, MipIndex, ResourceIndex);
+		});
 }
 
-void* FMetalSurface::AsyncLock(class FRHICommandListImmediate& RHICmdList, uint32 MipIndex, uint32 ArrayIndex, EResourceLockMode LockMode, uint32& DestStride, bool bNeedsDefaultRHIFlush, uint64* OutLockedByteCount)
+FRHILockTextureResult FMetalSurface::AsyncLock(FRHICommandListBase& RHICmdList, const FRHILockTextureArgs& Arguments)
 {
-	bool bDirectLock = (LockMode == RLM_ReadOnly || !GIsRHIInitialized);
-	
-	void* BufferData = nullptr;
+	bool bDirectLock = (Arguments.LockMode == RLM_ReadOnly || !GIsRHIInitialized);
 	
 	// Never flush for writing, it is unnecessary
-	if (bDirectLock && bNeedsDefaultRHIFlush)
+	if (bDirectLock && Arguments.bNeedsDefaultRHIFlush)
     {
         // @todo Not all read locks need to flush either, but that'll require resource use tracking
-        QUICK_SCOPE_CYCLE_COUNTER(STAT_RHIMETHOD_LockTexture2D_Flush);
-        RHICmdList.ImmediateFlush(EImmediateFlushType::FlushRHIThread);
+        QUICK_SCOPE_CYCLE_COUNTER(STAT_RHIMETHOD_LockTexture_Flush);
+        RHICmdList.GetAsImmediate().ImmediateFlush(EImmediateFlushType::FlushRHIThread);
     }
     
-    BufferData = Lock(MipIndex, ArrayIndex, LockMode, DestStride, false, OutLockedByteCount);
+	const FRHILockTextureResult Result = Lock(Arguments, false);
 	
-	check(BufferData);
+	check(Result.Data);
 	
-	return BufferData;
+	return Result;
 }
 
 uint32 FMetalSurface::GetMipSize(uint32 MipIndex, uint32* Stride, bool bSingleLayer)
@@ -1470,54 +1470,28 @@ FTextureRHIRef FMetalDynamicRHI::RHIAsyncReallocateTexture2D(FRHITexture* OldTex
     return NewTexture;
 }
 
-void* FMetalDynamicRHI::LockTexture2D_RenderThread(class FRHICommandListImmediate& RHICmdList, FRHITexture* Texture, uint32 MipIndex, EResourceLockMode LockMode, uint32& DestStride, bool bLockWithinMiptail, bool bNeedsDefaultRHIFlush, uint64* OutLockedByteCount)
-{
-    MTL_SCOPED_AUTORELEASE_POOL;
-    
-    check(IsInRenderingThread());
-    
-    FMetalSurface* TextureMTL = ResourceCast(Texture);
-    void* BufferData = TextureMTL->AsyncLock(RHICmdList, MipIndex, 0, LockMode, DestStride, bNeedsDefaultRHIFlush, OutLockedByteCount);
-    return BufferData;
-}
-
-void FMetalDynamicRHI::UnlockTexture2D_RenderThread(class FRHICommandListImmediate& RHICmdList, FRHITexture* Texture, uint32 MipIndex, bool bLockWithinMiptail, bool bNeedsDefaultRHIFlush)
+FRHILockTextureResult FMetalDynamicRHI::RHILockTexture(FRHICommandListImmediate& RHICmdList, const FRHILockTextureArgs& Arguments)
 {
-    MTL_SCOPED_AUTORELEASE_POOL;
+	MTL_SCOPED_AUTORELEASE_POOL;
 
-    check(IsInRenderingThread());
-    
-    FMetalSurface* TextureMTL = ResourceCast(Texture);
-    TextureMTL->Unlock(MipIndex, 0, true);
-}
+	FMetalSurface* Texture = ResourceCast(Arguments.Texture);
+	const uint32 ResourceIndex = GetMetalResourceIndex(Arguments);
 
+	if (Texture->GetDesc().Dimension == ETextureDimension::Texture2D)
+	{
+		return Texture->AsyncLock(RHICmdList, Arguments);
+	}
 
-void* FMetalDynamicRHI::RHILockTexture2D(FRHITexture* TextureRHI,uint32 MipIndex,EResourceLockMode LockMode,uint32& DestStride,bool bLockWithinMiptail, uint64* OutLockedByteCount)
-{
-    MTL_SCOPED_AUTORELEASE_POOL;
-    FMetalSurface* Texture = ResourceCast(TextureRHI);
-	return Texture->Lock(MipIndex, 0, LockMode, DestStride, OutLockedByteCount);
+	return Texture->Lock(Arguments, false);
 }
 
-void FMetalDynamicRHI::RHIUnlockTexture2D(FRHITexture* TextureRHI,uint32 MipIndex,bool bLockWithinMiptail)
+void FMetalDynamicRHI::RHIUnlockTexture(FRHICommandListImmediate& RHICmdList, const FRHILockTextureArgs& Arguments)
 {
-    MTL_SCOPED_AUTORELEASE_POOL;
-    FMetalSurface* Texture = ResourceCast(TextureRHI);
-	Texture->Unlock(MipIndex, 0, false);
-}
+	MTL_SCOPED_AUTORELEASE_POOL;
 
-void* FMetalDynamicRHI::RHILockTexture2DArray(FRHITexture* TextureRHI,uint32 TextureIndex,uint32 MipIndex,EResourceLockMode LockMode,uint32& DestStride,bool bLockWithinMiptail)
-{
-    MTL_SCOPED_AUTORELEASE_POOL;
-    FMetalSurface* Texture = ResourceCast(TextureRHI);
-	return Texture->Lock(MipIndex, TextureIndex, LockMode, DestStride);
-}
+	FMetalSurface* Texture = ResourceCast(Arguments.Texture);
 
-void FMetalDynamicRHI::RHIUnlockTexture2DArray(FRHITexture* TextureRHI,uint32 TextureIndex,uint32 MipIndex,bool bLockWithinMiptail)
-{
-    MTL_SCOPED_AUTORELEASE_POOL;
-    FMetalSurface* Texture = ResourceCast(TextureRHI);
-	Texture->Unlock(MipIndex, TextureIndex, false);
+	Texture->Unlock(Arguments);
 }
 
 #if PLATFORM_MAC
@@ -1799,22 +1773,6 @@ void FMetalDynamicRHI::RHIUpdateTexture3D(FRHICommandListBase& RHICmdList, FRHIT
 /*-----------------------------------------------------------------------------
  Cubemap texture support.
  -----------------------------------------------------------------------------*/
-void* FMetalDynamicRHI::RHILockTextureCubeFace(FRHITexture* TextureCubeRHI,uint32 FaceIndex,uint32 ArrayIndex,uint32 MipIndex,EResourceLockMode LockMode,uint32& DestStride,bool bLockWithinMiptail)
-{
-    MTL_SCOPED_AUTORELEASE_POOL;
-    
-    FMetalSurface* TextureCube = ResourceCast(TextureCubeRHI);
-    uint32 MetalFace = GetMetalCubeFace((ECubeFace)FaceIndex);
-    return TextureCube->Lock(MipIndex, MetalFace + (6 * ArrayIndex), LockMode, DestStride, true);
-}
-
-void FMetalDynamicRHI::RHIUnlockTextureCubeFace(FRHITexture* TextureCubeRHI,uint32 FaceIndex,uint32 ArrayIndex,uint32 MipIndex,bool bLockWithinMiptail)
-{
-    MTL_SCOPED_AUTORELEASE_POOL;
-    FMetalSurface* TextureCube = ResourceCast(TextureCubeRHI);
-    uint32 MetalFace = GetMetalCubeFace((ECubeFace)FaceIndex);
-    TextureCube->Unlock(MipIndex, MetalFace + (ArrayIndex * 6), false);
-}
 
 void FMetalDynamicRHI::RHIBindDebugLabelName(FRHICommandListBase& RHICmdList, FRHITexture* TextureRHI, const TCHAR* Name)
 {
diff --git a/Engine/Source/Runtime/Apple/MetalRHI/Public/MetalResources.h b/Engine/Source/Runtime/Apple/MetalRHI/Public/MetalResources.h
--- a/Engine/Source/Runtime/Apple/MetalRHI/Public/MetalResources.h
+++ b/Engine/Source/Runtime/Apple/MetalRHI/Public/MetalResources.h
@@ -173,7 +173,7 @@ public:
 	virtual ~FMetalSurface();
 
 	/** @returns A newly allocated buffer object large enough for the surface within the texture specified. */
-	MTL::Buffer* AllocSurface(uint32 MipIndex, uint32 ArrayIndex, EResourceLockMode LockMode, uint32& DestStride, bool SingleLayer = false);
+	MTL::Buffer* AllocSurface(const FRHILockTextureArgs& Arguments, uint32 MipBytes, uint32 DestStride);
 
 	/** Apply the data in Buffer to the surface specified.
 	 * Will also handle destroying SourceBuffer appropriately.
@@ -185,25 +185,20 @@ public:
 	 * @param ArrayIndex Index of the texture array/face in the form Index*6+Face
 	 * @return A pointer to the specified texture data.
 	 */
-	void* Lock(uint32 MipIndex, uint32 ArrayIndex, EResourceLockMode LockMode, uint32& DestStride, bool SingleLayer = false, uint64* OutLockedByteCount = nullptr);
+	FRHILockTextureResult Lock(const FRHILockTextureArgs& Arguments, bool bSingleLayer);
 	
 	/** Unlocks a previously locked mip-map.
 	 * @param ArrayIndex Index of the texture array/face in the form Index*6+Face
 	 */
-	void Unlock(uint32 MipIndex, uint32 ArrayIndex, bool bTryAsync);
+	void Unlock(const FRHILockTextureArgs& Arguments);
 	
 	/**
 	 * Locks one of the texture's mip-maps.
 	 * @param ArrayIndex Index of the texture array/face in the form Index*6+Face
 	 * @return A pointer to the specified texture data.
 	 */
-	void* AsyncLock(class FRHICommandListImmediate& RHICmdList, uint32 MipIndex, uint32 ArrayIndex, EResourceLockMode LockMode, uint32& DestStride, bool bNeedsDefaultRHIFlush, uint64* OutLockedByteCount = nullptr);
+	FRHILockTextureResult AsyncLock(FRHICommandListBase& RHICmdList, const FRHILockTextureArgs& Arguments);
 	
-	/** Unlocks a previously locked mip-map.
-	 * @param ArrayIndex Index of the texture array/face in the form Index*6+Face
-	 */
-	void AsyncUnlock(MTL::Buffer* SourceData, uint32 MipIndex, uint32 ArrayIndex);
-
 	/**
 	 * Returns how much memory a single mip uses, and optionally returns the stride
 	 */
diff --git a/Engine/Source/Runtime/D3D12RHI/Private/D3D12RHIPrivate.h b/Engine/Source/Runtime/D3D12RHI/Private/D3D12RHIPrivate.h
--- a/Engine/Source/Runtime/D3D12RHI/Private/D3D12RHIPrivate.h
+++ b/Engine/Source/Runtime/D3D12RHI/Private/D3D12RHIPrivate.h
@@ -354,19 +354,15 @@ public:
 	virtual FTextureRHIRef RHICreateTexture(FRHICommandListBase& RHICmdList, const FRHITextureCreateDesc& CreateDesc) override;
 	virtual uint32 RHIComputeMemorySize(FRHITexture* TextureRHI) final override;
 	virtual FTextureRHIRef RHIAsyncReallocateTexture2D(FRHITexture* Texture2D, int32 NewMipCount, int32 NewSizeX, int32 NewSizeY, FThreadSafeCounter* RequestStatus) override;
-	virtual void* RHILockTexture2D(FRHITexture* Texture, uint32 MipIndex, EResourceLockMode LockMode, uint32& DestStride, bool bLockWithinMiptail, uint64* OutLockedByteCount = nullptr) final override;
-	virtual void RHIUnlockTexture2D(FRHITexture* Texture, uint32 MipIndex, bool bLockWithinMiptail) final override;
-	virtual void* RHILockTexture2DArray(FRHITexture* Texture, uint32 TextureIndex, uint32 MipIndex, EResourceLockMode LockMode, uint32& DestStride, bool bLockWithinMiptail) final override;
-	virtual void RHIUnlockTexture2DArray(FRHITexture* Texture, uint32 TextureIndex, uint32 MipIndex, bool bLockWithinMiptail) final override;
-	virtual void* RHILockTextureCubeFace_RenderThread(class FRHICommandListImmediate& RHICmdList, FRHITexture* TextureRHI, uint32 FaceIndex, uint32 ArrayIndex, uint32 MipIndex, EResourceLockMode LockMode, uint32& DestStride, bool bLockWithinMiptail) final override;
-	virtual void RHIUnlockTextureCubeFace_RenderThread(class FRHICommandListImmediate& RHICmdList, FRHITexture* TextureRHI, uint32 FaceIndex, uint32 ArrayIndex, uint32 MipIndex, bool bLockWithinMiptail) final override;
+
+	virtual FRHILockTextureResult RHILockTexture(FRHICommandListImmediate& RHICmdList, const FRHILockTextureArgs& Arguments) final override;
+	virtual void                  RHIUnlockTexture(FRHICommandListImmediate& RHICmdList, const FRHILockTextureArgs& Arguments) final override;
+
 	virtual void RHIUpdateTexture2D(FRHICommandListBase& RHICmdList, FRHITexture* Texture, uint32 MipIndex, const struct FUpdateTextureRegion2D& UpdateRegion, uint32 SourcePitch, const uint8* SourceData) final override;
 	virtual void RHIUpdateTexture3D(FRHICommandListBase& RHICmdList, FRHITexture* Texture, uint32 MipIndex, const struct FUpdateTextureRegion3D& UpdateRegion, uint32 SourceRowPitch, uint32 SourceDepthPitch, const uint8* SourceData) final override;
 	virtual FUpdateTexture3DData RHIBeginUpdateTexture3D(FRHICommandListBase& RHICmdList, FRHITexture* Texture, uint32 MipIndex, const struct FUpdateTextureRegion3D& UpdateRegion) final override;
 	virtual void RHIEndUpdateTexture3D(FRHICommandListBase& RHICmdList, FUpdateTexture3DData& UpdateData) final override;
 	virtual void RHIEndMultiUpdateTexture3D(FRHICommandListBase& RHICmdList, TArray<FUpdateTexture3DData>& UpdateDataArray) final override;
-	virtual void* RHILockTextureCubeFace(FRHITexture* Texture, uint32 FaceIndex, uint32 ArrayIndex, uint32 MipIndex, EResourceLockMode LockMode, uint32& DestStride, bool bLockWithinMiptail) final override;
-	virtual void RHIUnlockTextureCubeFace(FRHITexture* Texture, uint32 FaceIndex, uint32 ArrayIndex, uint32 MipIndex, bool bLockWithinMiptail) final override;
 	virtual void RHIBindDebugLabelName(FRHICommandListBase& RHICmdList, FRHITexture* Texture, const TCHAR* Name) final override;
 	virtual void RHIBindDebugLabelName(FRHICommandListBase& RHICmdList, FRHIBuffer* Buffer, const TCHAR* Name) final override;
 	virtual void RHIReadSurfaceData(FRHITexture* Texture, FIntRect Rect, TArray<FColor>& OutData, FReadSurfaceDataFlags InFlags) final override;
@@ -460,12 +456,6 @@ public:
 	// These will be un-commented as they are implemented.
 	//
 
-	virtual void* LockTexture2D_RenderThread(class FRHICommandListImmediate& RHICmdList, FRHITexture* Texture, uint32 MipIndex, EResourceLockMode LockMode, uint32& DestStride, bool bLockWithinMiptail, bool bNeedsDefaultRHIFlush = true, uint64* OutLockedByteCount = nullptr) override final;
-	virtual void UnlockTexture2D_RenderThread(class FRHICommandListImmediate& RHICmdList, FRHITexture* Texture, uint32 MipIndex, bool bLockWithinMiptail, bool bNeedsDefaultRHIFlush = true) override final;
-
-	virtual void* LockTexture2DArray_RenderThread(class FRHICommandListImmediate& RHICmdList, FRHITexture* Texture, uint32 ArrayIndex, uint32 MipIndex, EResourceLockMode LockMode, uint32& DestStride, bool bLockWithinMiptail) override final;
-	virtual void UnlockTexture2DArray_RenderThread(class FRHICommandListImmediate& RHICmdList, FRHITexture* Texture, uint32 ArrayIndex, uint32 MipIndex, bool bLockWithinMiptail) override final;
-
 	virtual FTextureRHIRef AsyncReallocateTexture2D_RenderThread(class FRHICommandListImmediate& RHICmdList, FRHITexture* Texture2D, int32 NewMipCount, int32 NewSizeX, int32 NewSizeY, FThreadSafeCounter* RequestStatus);
 
 #if D3D12_RHI_RAYTRACING
@@ -514,8 +504,8 @@ public:
 	virtual void* CreateVirtualTexture(ETextureCreateFlags InFlags, D3D12_RESOURCE_DESC& ResourceDesc, const struct FD3D12TextureLayout& TextureLayout, FD3D12Resource** ppResource, FPlatformMemory::FPlatformVirtualMemoryBlock& RawTextureBlock, D3D12_RESOURCE_STATES InitialUsage = D3D12_RESOURCE_STATE_PIXEL_SHADER_RESOURCE) = 0;
 	virtual void DestroyVirtualTexture(ETextureCreateFlags InFlags, void* RawTextureMemory, FPlatformMemory::FPlatformVirtualMemoryBlock& RawTextureBlock, uint64 CommittedTextureSize) = 0;
 #endif
-	virtual bool HandleSpecialLock(void*& MemoryOut, uint32 MipIndex, uint32 ArrayIndex, FD3D12Texture* InTexture, EResourceLockMode LockMode, uint32& DestStride, uint64* OutLockedByteCount = nullptr) { return false; }
-	virtual bool HandleSpecialUnlock(FRHICommandListBase& RHICmdList, uint32 MipIndex, FD3D12Texture* InTexture) { return false; }
+	virtual bool HandleSpecialLock(FRHILockTextureResult& Result, FD3D12Texture* InTexture, const FRHILockTextureArgs& Arguments) { return false; }
+	virtual bool HandleSpecialUnlock(FRHICommandListBase& RHICmdList, FD3D12Texture* InTexture, const FRHILockTextureArgs& Arguments) { return false; }
 
 	FD3D12Adapter& GetAdapter(uint32_t Index = 0) { return *ChosenAdapters[Index]; }
 	const FD3D12Adapter& GetAdapter(uint32_t Index = 0) const { return *ChosenAdapters[Index]; }
diff --git a/Engine/Source/Runtime/D3D12RHI/Private/D3D12Texture.cpp b/Engine/Source/Runtime/D3D12RHI/Private/D3D12Texture.cpp
--- a/Engine/Source/Runtime/D3D12RHI/Private/D3D12Texture.cpp
+++ b/Engine/Source/Runtime/D3D12RHI/Private/D3D12Texture.cpp
@@ -1509,7 +1509,23 @@ void FD3D12Texture::AliasResources(FD3D12Texture* Texture)
 	}
 }
 
-void* FD3D12Texture::Lock(FRHICommandListImmediate& RHICmdList, uint32 MipIndex, uint32 ArrayIndex, EResourceLockMode LockMode, uint32& DestStride, uint64* OutLockedByteCount)
+static uint32 GetD3D12ArrayIndex(const FRHILockTextureArgs& Arguments, const FRHITextureDesc& Desc)
+{
+	if (Desc.IsTextureCube())
+	{
+		const uint32 D3DFace = GetD3D12CubeFace(static_cast<ECubeFace>(Arguments.FaceIndex));
+		return D3DFace + Arguments.ArrayIndex * 6;
+	}
+
+	if (Desc.IsTextureArray())
+	{
+		return Arguments.ArrayIndex;
+	}
+
+	return 0;
+}
+
+FRHILockTextureResult FD3D12Texture::Lock(FRHICommandListBase& RHICmdList, const FRHILockTextureArgs& Arguments)
 {
 	SCOPE_CYCLE_COUNTER(STAT_D3D12LockTextureTime);
 
@@ -1519,8 +1535,11 @@ void* FD3D12Texture::Lock(FRHICommandListImmediate& RHICmdList, uint32 MipIndex,
 	FD3D12Device* Device = GetParentDevice();
 	FD3D12Adapter* Adapter = Device->GetParentAdapter();
 
+	const FRHITextureDesc& Desc = GetDesc();
+	const uint32 ArrayIndex = GetD3D12ArrayIndex(Arguments, Desc);
+
 	// Calculate the subresource index corresponding to the specified mip-map.
-	const uint32 Subresource = CalcSubresource(MipIndex, ArrayIndex, this->GetNumMips());
+	const uint32 Subresource = CalcSubresource(Arguments.MipIndex, ArrayIndex, this->GetNumMips());
 
 	check(LockedMap.Find(Subresource) == nullptr);
 	TUniquePtr<FD3D12LockedResource> LockedResource = MakeUnique<FD3D12LockedResource>(Device);
@@ -1537,25 +1556,18 @@ void* FD3D12Texture::Lock(FRHICommandListImmediate& RHICmdList, uint32 MipIndex,
 	// Store the footprint information so we don't have to recompute it in Unlock.
 	LockedResource->Footprint = PlacedFootprint.Footprint;
 
-	DestStride = LockedResource->Footprint.RowPitch;
-	const uint64 SubresourceSize = LockedResource->Footprint.RowPitch * LockedResource->Footprint.Height * LockedResource->Footprint.Depth;
-	if (OutLockedByteCount)
-	{
-		*OutLockedByteCount = SubresourceSize;
-	}
-
-#if !PLATFORM_SUPPORTS_VIRTUAL_TEXTURES
-	void* RawTextureMemory = (void*)ResourceLocation.GetGPUVirtualAddress();
-#endif
+	FRHILockTextureResult Result;
+	Result.Stride = LockedResource->Footprint.RowPitch;
 
-	void* Data = nullptr;
+	const uint64 SubresourceSize = LockedResource->Footprint.RowPitch * LockedResource->Footprint.Height * LockedResource->Footprint.Depth;
+	Result.ByteCount = SubresourceSize;
 
-	if (FD3D12DynamicRHI::GetD3DRHI()->HandleSpecialLock(Data, MipIndex, ArrayIndex, this, LockMode, DestStride, OutLockedByteCount))
+	if (FD3D12DynamicRHI::GetD3DRHI()->HandleSpecialLock(Result, this, Arguments))
 	{
 		// nothing left to do...
-		check(Data != nullptr);
+		check(Result.Data != nullptr);
 	}
-	else if (LockMode == RLM_WriteOnly)
+	else if (Arguments.LockMode == RLM_WriteOnly)
 	{
 		// If we're writing to the texture, allocate a system memory buffer to receive the new contents.
 		// Use an upload heap to copy data to a default resource.
@@ -1565,10 +1577,10 @@ void* FD3D12Texture::Lock(FRHICommandListImmediate& RHICmdList, uint32 MipIndex,
 		if (nullptr == pData)
 		{
 			check(false);
-			return nullptr;
+			return Result;
 		}
 
-		Data = LockedResource->ResourceLocation.GetMappedBaseAddress();
+		Result.Data = LockedResource->ResourceLocation.GetMappedBaseAddress();
 	}
 	else
 	{
@@ -1590,7 +1602,7 @@ void* FD3D12Texture::Lock(FRHICommandListImmediate& RHICmdList, uint32 MipIndex,
 		CD3DX12_TEXTURE_COPY_LOCATION DestCopyLocation(StagingTexture->GetResource(), PlacedFootprint);
 		CD3DX12_TEXTURE_COPY_LOCATION SourceCopyLocation(GetResource()->GetResource(), Subresource);
 
-		RHICmdList.EnqueueLambda([this, DestCopyLocation, SourceCopyLocation](FRHICommandListImmediate& ExecutingCmdList)
+		RHICmdList.EnqueueLambda([this, DestCopyLocation, SourceCopyLocation](FRHICommandListBase& ExecutingCmdList)
 		{
 			FD3D12CommandContext& Context = FD3D12CommandContext::Get(ExecutingCmdList, 0);
 			FScopedResourceBarrier ScopeResourceBarrierSource(Context, GetResource(), &ResourceLocation, D3D12_RESOURCE_STATE_COPY_SOURCE, SourceCopyLocation.SubresourceIndex);
@@ -1606,15 +1618,15 @@ void* FD3D12Texture::Lock(FRHICommandListImmediate& RHICmdList, uint32 MipIndex,
 		});
 
 		// We need to execute the command list so we can read the data from the map below
-		RHICmdList.SubmitAndBlockUntilGPUIdle();
+		RHICmdList.GetAsImmediate().SubmitAndBlockUntilGPUIdle();
 
-		Data = LockedResource->ResourceLocation.GetMappedBaseAddress();
+		Result.Data = LockedResource->ResourceLocation.GetMappedBaseAddress();
 	}
 
 	LockedMap.Emplace(Subresource, MoveTemp(LockedResource));
 
-	check(Data != nullptr);
-	return Data;
+	check(Result.Data != nullptr);
+	return Result;
 }
 
 void FD3D12Texture::UpdateTexture(FD3D12CommandContext& Context, uint32 MipIndex, uint32 DestX, uint32 DestY, uint32 DestZ, const D3D12_TEXTURE_COPY_LOCATION& SourceCopyLocation)
@@ -1784,18 +1796,21 @@ void FD3D12Texture::InitializeTextureData(FRHICommandListBase& RHICmdList, const
 	});
 }
 
-void FD3D12Texture::Unlock(FRHICommandListBase& RHICmdList, uint32 MipIndex, uint32 ArrayIndex)
+void FD3D12Texture::Unlock(FRHICommandListBase& RHICmdList, const FRHILockTextureArgs& Arguments)
 {
 	SCOPE_CYCLE_COUNTER(STAT_D3D12UnlockTextureTime);
 	check(IsHeadLink());
 
+	const FRHITextureDesc& Desc = GetDesc();
+	const uint32 ArrayIndex = GetD3D12ArrayIndex(Arguments, Desc);
+
 	// Calculate the subresource index corresponding to the specified mip-map.
-	const uint32 Subresource = CalcSubresource(MipIndex, ArrayIndex, GetNumMips());
+	const uint32 Subresource = CalcSubresource(Arguments.MipIndex, ArrayIndex, Desc.NumMips);
 
 	TUniquePtr<FD3D12LockedResource> LockedResource = LockedMap.FindAndRemoveChecked(Subresource);
 	check(LockedResource);
 
-	if (FD3D12DynamicRHI::GetD3DRHI()->HandleSpecialUnlock(RHICmdList, MipIndex, this))
+	if (FD3D12DynamicRHI::GetD3DRHI()->HandleSpecialUnlock(RHICmdList, this, Arguments))
 	{
 		// nothing left to do...
 	}
@@ -1934,87 +1949,18 @@ void FD3D12Texture::GetReadBackHeapDesc(D3D12_PLACED_SUBRESOURCE_FOOTPRINT& OutF
 	}
 }
 
-void* FD3D12DynamicRHI::LockTexture2D_RenderThread(FRHICommandListImmediate& RHICmdList, FRHITexture* TextureRHI, uint32 MipIndex, EResourceLockMode LockMode, uint32& DestStride, bool bLockWithinMiptail, bool bNeedsDefaultRHIFlush, uint64* OutLockedByteCount)
-{
-	check(TextureRHI);
-	FD3D12Texture* Texture = FD3D12DynamicRHI::ResourceCast(TextureRHI);
-	return Texture->Lock(RHICmdList, MipIndex, 0, LockMode, DestStride, OutLockedByteCount);
-}
-
-void FD3D12DynamicRHI::UnlockTexture2D_RenderThread(FRHICommandListImmediate& RHICmdList, FRHITexture* TextureRHI, uint32 MipIndex, bool bLockWithinMiptail, bool bNeedsDefaultRHIFlush)
-{
-	check(TextureRHI);
-	FD3D12Texture* Texture = FD3D12DynamicRHI::ResourceCast(TextureRHI);
-	Texture->Unlock(RHICmdList, MipIndex, 0);
-}
-
-void* FD3D12DynamicRHI::LockTexture2DArray_RenderThread(FRHICommandListImmediate& RHICmdList, FRHITexture* TextureRHI, uint32 TextureIndex, uint32 MipIndex, EResourceLockMode LockMode, uint32& DestStride, bool bLockWithinMiptail)
+FRHILockTextureResult FD3D12DynamicRHI::RHILockTexture(FRHICommandListImmediate& RHICmdList, const FRHILockTextureArgs& Arguments)
 {
-	check(TextureRHI);
-	FD3D12Texture* Texture = FD3D12DynamicRHI::ResourceCast(TextureRHI);
-	return Texture->Lock(RHICmdList, MipIndex, TextureIndex, LockMode, DestStride);
+	FD3D12Texture* Texture = FD3D12DynamicRHI::ResourceCast(Arguments.Texture);
+	return Texture->Lock(RHICmdList, Arguments);
 }
 
-void FD3D12DynamicRHI::UnlockTexture2DArray_RenderThread(FRHICommandListImmediate& RHICmdList, FRHITexture* TextureRHI, uint32 TextureIndex, uint32 MipIndex, bool bLockWithinMiptail)
+void FD3D12DynamicRHI::RHIUnlockTexture(FRHICommandListImmediate& RHICmdList, const FRHILockTextureArgs& Arguments)
 {
-	check(TextureRHI);
-	FD3D12Texture* Texture = FD3D12DynamicRHI::ResourceCast(TextureRHI);
-	Texture->Unlock(RHICmdList, MipIndex, TextureIndex);
+	FD3D12Texture* Texture = FD3D12DynamicRHI::ResourceCast(Arguments.Texture);
+	Texture->Unlock(RHICmdList, Arguments);
 }
 
-void* FD3D12DynamicRHI::RHILockTextureCubeFace_RenderThread(FRHICommandListImmediate& RHICmdList, FRHITexture* TextureRHI, uint32 FaceIndex, uint32 ArrayIndex, uint32 MipIndex, EResourceLockMode LockMode, uint32& DestStride, bool bLockWithinMiptail)
-{
-	check(TextureRHI);
-	FD3D12Texture* TextureCube = FD3D12DynamicRHI::ResourceCast(TextureRHI);
-	uint32 D3DFace = GetD3D12CubeFace((ECubeFace)FaceIndex);
-	return TextureCube->Lock(RHICmdList, MipIndex, D3DFace + ArrayIndex * 6, LockMode, DestStride);
-}
-
-void FD3D12DynamicRHI::RHIUnlockTextureCubeFace_RenderThread(FRHICommandListImmediate& RHICmdList, FRHITexture* TextureRHI, uint32 FaceIndex, uint32 ArrayIndex, uint32 MipIndex, bool bLockWithinMiptail)
-{
-	check(TextureRHI);
-	FD3D12Texture* TextureCube = FD3D12DynamicRHI::ResourceCast(TextureRHI);
-	uint32 D3DFace = GetD3D12CubeFace((ECubeFace)FaceIndex);
-	TextureCube->Unlock(RHICmdList, MipIndex, D3DFace + ArrayIndex * 6);
-}
-
-
-
-void* FD3D12DynamicRHI::RHILockTexture2D(FRHITexture* TextureRHI, uint32 MipIndex, EResourceLockMode LockMode, uint32& DestStride, bool bLockWithinMiptail, uint64* OutLockedByteCount)
-{
-	UE_LOG(LogD3D12RHI, Fatal, TEXT("RHILockTexture2D is unimplemented and should never be called. Use LockTexture2D_RenderThread."));
-	return nullptr;
-}
-
-void FD3D12DynamicRHI::RHIUnlockTexture2D(FRHITexture* TextureRHI, uint32 MipIndex, bool bLockWithinMiptail)
-{
-	UE_LOG(LogD3D12RHI, Fatal, TEXT("RHIUnlockTexture2D is unimplemented and should never be called. Use UnlockTexture2D_RenderThread."));
-}
-
-void* FD3D12DynamicRHI::RHILockTexture2DArray(FRHITexture* TextureRHI, uint32 TextureIndex, uint32 MipIndex, EResourceLockMode LockMode, uint32& DestStride, bool bLockWithinMiptail)
-{
-	UE_LOG(LogD3D12RHI, Fatal, TEXT("RHILockTexture2DArray is unimplemented and should never be called. Use LockTexture2DArray_RenderThread."));
-	return nullptr;
-}
-
-void FD3D12DynamicRHI::RHIUnlockTexture2DArray(FRHITexture* TextureRHI, uint32 TextureIndex, uint32 MipIndex, bool bLockWithinMiptail)
-{
-	UE_LOG(LogD3D12RHI, Fatal, TEXT("RHIUnlockTexture2DArray is unimplemented and should never be called. Use UnlockTexture2DArray_RenderThread."));
-}
-
-void* FD3D12DynamicRHI::RHILockTextureCubeFace(FRHITexture* TextureCubeRHI, uint32 FaceIndex, uint32 ArrayIndex, uint32 MipIndex, EResourceLockMode LockMode, uint32& DestStride, bool bLockWithinMiptail)
-{
-	UE_LOG(LogD3D12RHI, Fatal, TEXT("RHIUnlockTexture2DArray is unimplemented and should never be called. Use UnlockTexture2DArray_RenderThread."));
-	return nullptr;
-}
-
-void FD3D12DynamicRHI::RHIUnlockTextureCubeFace(FRHITexture* TextureCubeRHI, uint32 FaceIndex, uint32 ArrayIndex, uint32 MipIndex, bool bLockWithinMiptail)
-{
-	UE_LOG(LogD3D12RHI, Fatal, TEXT("RHIUnlockTextureCubeFace is unimplemented and should never be called. Use RHIUnlockTextureCubeFace_RenderThread."));
-}
-
-
-
 void FD3D12DynamicRHI::RHIUpdateTexture2D(FRHICommandListBase& RHICmdList, FRHITexture* TextureRHI, uint32 MipIndex, const FUpdateTextureRegion2D& UpdateRegion, uint32 SourcePitch, const uint8* SourceData)
 {
 	check(TextureRHI);
diff --git a/Engine/Source/Runtime/D3D12RHI/Private/D3D12Texture.h b/Engine/Source/Runtime/D3D12RHI/Private/D3D12Texture.h
--- a/Engine/Source/Runtime/D3D12RHI/Private/D3D12Texture.h
+++ b/Engine/Source/Runtime/D3D12RHI/Private/D3D12Texture.h
@@ -119,8 +119,8 @@ public:
 	}
 
 	// Locking/update functions
-	void* Lock(FRHICommandListImmediate& RHICmdList, uint32 MipIndex, uint32 ArrayIndex, EResourceLockMode LockMode, uint32& DestStride, uint64* OutLockedByteCount = nullptr);
-	void Unlock(FRHICommandListBase& RHICmdList, uint32 MipIndex, uint32 ArrayIndex);
+	FRHILockTextureResult Lock(FRHICommandListBase& RHICmdList, const FRHILockTextureArgs& Arguments);
+	void Unlock(FRHICommandListBase& RHICmdList, const FRHILockTextureArgs& Arguments);
 
 	void UpdateTexture2D(FRHICommandListBase& RHICmdList, uint32 MipIndex, const struct FUpdateTextureRegion2D& UpdateRegion, uint32 SourcePitch, const uint8* SourceData);
 	void UpdateTexture(FD3D12CommandContext& Context, uint32 MipIndex, uint32 DestX, uint32 DestY, uint32 DestZ, const D3D12_TEXTURE_COPY_LOCATION& SourceCopyLocation);
diff --git a/Engine/Source/Runtime/NullDrv/Public/NullRHI.h b/Engine/Source/Runtime/NullDrv/Public/NullRHI.h
--- a/Engine/Source/Runtime/NullDrv/Public/NullRHI.h
+++ b/Engine/Source/Runtime/NullDrv/Public/NullRHI.h
@@ -215,24 +215,20 @@ public:
 
 		return new FNullTexture(Desc);
 	}
-	virtual void* RHILockTexture2D(FRHITexture* Texture, uint32 MipIndex, EResourceLockMode LockMode, uint32& DestStride, bool bLockWithinMiptail, uint64* OutLockedByteCount) final override
-	{ 
-		DestStride = 0; 
-		return GetStaticTextureBuffer(Texture->GetSizeX(), Texture->GetSizeY(), Texture->GetFormat(), DestStride, OutLockedByteCount);
-	}
-	virtual void RHIUnlockTexture2D(FRHITexture* Texture, uint32 MipIndex, bool bLockWithinMiptail) final override
+
+	virtual FRHILockTextureResult RHILockTexture(FRHICommandListImmediate& RHICmdList, const FRHILockTextureArgs& Arguments) final override
 	{
+		FRHITexture* Texture = Arguments.Texture;
 
+		FRHILockTextureResult Result{};
+		Result.Data = GetStaticTextureBuffer(Texture->GetSizeX(), Texture->GetSizeY(), Texture->GetFormat(), Result.Stride, &Result.ByteCount);
+
+		return Result;
 	}
-	virtual void* RHILockTexture2DArray(FRHITexture* Texture, uint32 TextureIndex, uint32 MipIndex, EResourceLockMode LockMode, uint32& DestStride, bool bLockWithinMiptail) final override
-	{ 
-		DestStride = 0; 
-		return GetStaticTextureBuffer(Texture->GetSizeX(), Texture->GetSizeY(), Texture->GetFormat(), DestStride);
-	}
-	virtual void RHIUnlockTexture2DArray(FRHITexture* Texture, uint32 TextureIndex, uint32 MipIndex, bool bLockWithinMiptail) final override
+	virtual void RHIUnlockTexture(FRHICommandListImmediate& RHICmdList, const FRHILockTextureArgs& Arguments) final override
 	{
-
 	}
+
 	virtual void RHIUpdateTexture2D(FRHICommandListBase&, FRHITexture* Texture, uint32 MipIndex, const struct FUpdateTextureRegion2D& UpdateRegion, uint32 SourcePitch, const uint8* SourceData) final override
 	{
 
@@ -242,15 +238,6 @@ public:
 
 	}
 
-	virtual void* RHILockTextureCubeFace(FRHITexture* Texture, uint32 FaceIndex, uint32 ArrayIndex, uint32 MipIndex, EResourceLockMode LockMode, uint32& DestStride, bool bLockWithinMiptail) final override
-	{ 
-		DestStride = 0; 
-		return GetStaticTextureBuffer(Texture->GetSize(), Texture->GetSize(), Texture->GetFormat(), DestStride);
-	}
-	virtual void RHIUnlockTextureCubeFace(FRHITexture* Texture, uint32 FaceIndex, uint32 ArrayIndex, uint32 MipIndex, bool bLockWithinMiptail) final override
-	{
-	}
-
 	virtual void RHICopyTexture(FRHITexture* SourceTexture, FRHITexture* DestTexture, const FRHICopyTextureInfo& CopyInfo) final override
 	{
 
diff --git a/Engine/Source/Runtime/OpenGLDrv/Private/OpenGLDrv.h b/Engine/Source/Runtime/OpenGLDrv/Private/OpenGLDrv.h
--- a/Engine/Source/Runtime/OpenGLDrv/Private/OpenGLDrv.h
+++ b/Engine/Source/Runtime/OpenGLDrv/Private/OpenGLDrv.h
@@ -344,14 +344,12 @@ public:
 	virtual void RHIGenerateMips(FRHITexture* Texture) final override;
 	virtual uint32 RHIComputeMemorySize(FRHITexture* TextureRHI) final override;
 	virtual FTextureRHIRef RHIAsyncReallocateTexture2D(FRHITexture* Texture2D, int32 NewMipCount, int32 NewSizeX, int32 NewSizeY, FThreadSafeCounter* RequestStatus) final override;
-	virtual void* RHILockTexture2D(FRHITexture* Texture, uint32 MipIndex, EResourceLockMode LockMode, uint32& DestStride, bool bLockWithinMiptail, uint64* OutLockedByteCount = nullptr) final override;
-	virtual void RHIUnlockTexture2D(FRHITexture* Texture, uint32 MipIndex, bool bLockWithinMiptail) final override;
-	virtual void* RHILockTexture2DArray(FRHITexture* Texture, uint32 TextureIndex, uint32 MipIndex, EResourceLockMode LockMode, uint32& DestStride, bool bLockWithinMiptail) final override;
-	virtual void RHIUnlockTexture2DArray(FRHITexture* Texture, uint32 TextureIndex, uint32 MipIndex, bool bLockWithinMiptail) final override;
+
+	virtual FRHILockTextureResult RHILockTexture(FRHICommandListImmediate& RHICmdList, const FRHILockTextureArgs& Arguments) final override;
+	virtual void RHIUnlockTexture(FRHICommandListImmediate& RHICmdList, const FRHILockTextureArgs& Arguments) final override;
+
 	virtual void RHIUpdateTexture2D(FRHICommandListBase& RHICmdList, FRHITexture* Texture, uint32 MipIndex, const struct FUpdateTextureRegion2D& UpdateRegion, uint32 SourcePitch, const uint8* SourceData) final override;
 	virtual void RHIUpdateTexture3D(FRHICommandListBase& RHICmdList, FRHITexture* Texture, uint32 MipIndex, const struct FUpdateTextureRegion3D& UpdateRegion, uint32 SourceRowPitch, uint32 SourceDepthPitch, const uint8* SourceData) final override;
-	virtual void* RHILockTextureCubeFace(FRHITexture* Texture, uint32 FaceIndex, uint32 ArrayIndex, uint32 MipIndex, EResourceLockMode LockMode, uint32& DestStride, bool bLockWithinMiptail) final override;
-	virtual void RHIUnlockTextureCubeFace(FRHITexture* Texture, uint32 FaceIndex, uint32 ArrayIndex, uint32 MipIndex, bool bLockWithinMiptail) final override;
 	virtual void RHIBindDebugLabelName(FRHICommandListBase& RHICmdList, FRHITexture* Texture, const TCHAR* Name) final override;
 	virtual void RHIReadSurfaceData(FRHITexture* Texture,FIntRect Rect,TArray<FColor>& OutData,FReadSurfaceDataFlags InFlags) final override;
 	virtual void RHIReadSurfaceData(FRHITexture* Texture, FIntRect Rect, TArray<FLinearColor>& OutData, FReadSurfaceDataFlags InFlags) final override;
@@ -630,13 +628,6 @@ public:
 		}
 	};
 
-	virtual void* LockTexture2D_RenderThread(class FRHICommandListImmediate& RHICmdList, FRHITexture* Texture, uint32 MipIndex, EResourceLockMode LockMode, uint32& DestStride, bool bLockWithinMiptail, bool bNeedsDefaultRHIFlush, uint64* OutLockedByteCount = nullptr) final override;
-	virtual void UnlockTexture2D_RenderThread(class FRHICommandListImmediate& RHICmdList, FRHITexture* Texture, uint32 MipIndex, bool bLockWithinMiptail, bool bNeedsDefaultRHIFlush) final override;
-	virtual void* RHILockTextureCubeFace_RenderThread(class FRHICommandListImmediate& RHICmdList, FRHITexture* Texture, uint32 FaceIndex, uint32 ArrayIndex, uint32 MipIndex, EResourceLockMode LockMode, uint32& DestStride, bool bLockWithinMiptail) final override;
-	virtual void RHIUnlockTextureCubeFace_RenderThread(class FRHICommandListImmediate& RHICmdList, FRHITexture* Texture, uint32 FaceIndex, uint32 ArrayIndex, uint32 MipIndex, bool bLockWithinMiptail) final override;
-	virtual void* LockTexture2DArray_RenderThread(class FRHICommandListImmediate& RHICmdList, FRHITexture* Texture, uint32 ArrayIndex, uint32 MipIndex, EResourceLockMode LockMode, uint32& DestStride, bool bLockWithinMiptail) final override;
-	virtual void UnlockTexture2DArray_RenderThread(class FRHICommandListImmediate& RHICmdList, FRHITexture* Texture, uint32 ArrayIndex, uint32 MipIndex, bool bLockWithinMiptail) final override;
-
 	virtual FTextureRHIRef AsyncReallocateTexture2D_RenderThread(class FRHICommandListImmediate& RHICmdList, FRHITexture* Texture2D, int32 NewMipCount, int32 NewSizeX, int32 NewSizeY, FThreadSafeCounter* RequestStatus) final override;
 
 	virtual FSamplerStateRHIRef RHICreateSamplerState(const FSamplerStateInitializerRHI& Initializer) final override;
diff --git a/Engine/Source/Runtime/OpenGLDrv/Private/OpenGLTexture.cpp b/Engine/Source/Runtime/OpenGLDrv/Private/OpenGLTexture.cpp
--- a/Engine/Source/Runtime/OpenGLDrv/Private/OpenGLTexture.cpp
+++ b/Engine/Source/Runtime/OpenGLDrv/Private/OpenGLTexture.cpp
@@ -1457,32 +1457,6 @@ FTextureRHIRef FOpenGLDynamicRHI::RHIAsyncReallocateTexture2D(FRHITexture* Textu
 	return NewTexture;
 }
 
-void* FOpenGLDynamicRHI::RHILockTexture2D(FRHITexture* TextureRHI,uint32 MipIndex,EResourceLockMode LockMode,uint32& DestStride,bool bLockWithinMiptail, uint64* OutLockedByteCount)
-{
-	FOpenGLTexture* Texture = ResourceCast(TextureRHI);
-	if (OutLockedByteCount)
-	{
-		*OutLockedByteCount = Texture->GetLockSize(MipIndex, 0, LockMode, DestStride);
-	}
-
-	return Texture->Lock(MipIndex, 0, LockMode, DestStride);
-}
-
-void FOpenGLDynamicRHI::RHIUnlockTexture2D(FRHITexture* TextureRHI, uint32 MipIndex, bool bLockWithinMiptail)
-{
-	ResourceCast(TextureRHI)->Unlock(MipIndex, 0);
-}
-
-void* FOpenGLDynamicRHI::RHILockTexture2DArray(FRHITexture* TextureRHI, uint32 TextureIndex, uint32 MipIndex, EResourceLockMode LockMode, uint32& DestStride, bool bLockWithinMiptail)
-{
-	return ResourceCast(TextureRHI)->Lock(MipIndex, TextureIndex, LockMode, DestStride);
-}
-
-void FOpenGLDynamicRHI::RHIUnlockTexture2DArray(FRHITexture* TextureRHI, uint32 TextureIndex, uint32 MipIndex, bool bLockWithinMiptail)
-{
-	ResourceCast(TextureRHI)->Unlock(MipIndex, TextureIndex);
-}
-
 void FOpenGLDynamicRHI::RHIUpdateTexture2D(FRHICommandListBase& RHICmdList, FRHITexture* TextureRHI, uint32 MipIndex, const FUpdateTextureRegion2D& UpdateRegion, uint32 SourcePitch, const uint8* SourceData)
 {
 	const FPixelFormatInfo& FormatInfo = GPixelFormats[TextureRHI->GetFormat()];
@@ -1699,15 +1673,6 @@ void FOpenGLDynamicRHI::InvalidateUAVResourceInCache(GLuint Resource)
 /*-----------------------------------------------------------------------------
 	Cubemap texture support.
 -----------------------------------------------------------------------------*/
-void* FOpenGLDynamicRHI::RHILockTextureCubeFace(FRHITexture* TextureCubeRHI, uint32 FaceIndex, uint32 ArrayIndex, uint32 MipIndex, EResourceLockMode LockMode, uint32& DestStride, bool bLockWithinMiptail)
-{
-	return ResourceCast(TextureCubeRHI)->Lock(MipIndex, FaceIndex + 6 * ArrayIndex, LockMode, DestStride);
-}
-
-void FOpenGLDynamicRHI::RHIUnlockTextureCubeFace(FRHITexture* TextureCubeRHI, uint32 FaceIndex, uint32 ArrayIndex, uint32 MipIndex, bool bLockWithinMiptail)
-{
-	ResourceCast(TextureCubeRHI)->Unlock(MipIndex, FaceIndex + ArrayIndex * 6);
-}
 
 void FOpenGLDynamicRHI::RHIBindDebugLabelName(FRHICommandListBase& RHICmdList, FRHITexture* TextureRHI, const TCHAR* Name)
 {
@@ -1879,174 +1844,83 @@ FTextureRHIRef FOpenGLDynamicRHI::RHICreateAliasedTexture(FTextureRHIRef& Source
 	return new FOpenGLTexture(*ResourceCast(SourceTexture), *Name, FOpenGLTexture::AliasResource);
 }
 
-void* FOpenGLDynamicRHI::LockTexture2D_RenderThread(class FRHICommandListImmediate& RHICmdList, FRHITexture* Texture, uint32 MipIndex, EResourceLockMode LockMode, uint32& DestStride, bool bLockWithinMiptail, bool bNeedsDefaultRHIFlush, uint64* OutLockedByteCount)
+static uint32 GetGLArrayIndex(const FRHILockTextureArgs& Arguments)
 {
-	check(IsInRenderingThread());
-	check(LockMode != EResourceLockMode::RLM_WriteOnly_NoOverwrite);
-	void* Result;
-	uint32 MipBytes = 0;
-	if (LockMode != RLM_WriteOnly)
-	{
-		RHICmdList.EnqueueLambda([this, &Result, Texture, MipIndex, LockMode, &DestStride, bLockWithinMiptail, OutLockedByteCount](FRHICommandListImmediate&)
-		{
-			Result = this->RHILockTexture2D(Texture, MipIndex, LockMode, DestStride, bLockWithinMiptail, OutLockedByteCount);
-		});
-		RHITHREAD_GLTRACE_BLOCKING;
-		RHICmdList.ImmediateFlush(EImmediateFlushType::FlushRHIThread);
-		MipBytes = ResourceCast(Texture)->GetLockSize(MipIndex, 0, LockMode, DestStride);
-	}
-	else
-	{
-		MipBytes = ResourceCast(Texture)->GetLockSize(MipIndex, 0, LockMode, DestStride);
-		Result = FMemory::Malloc(MipBytes, 16);
-	}
-	check(Result);
+	const FRHITextureDesc& Desc = Arguments.Texture->GetDesc();
 
-	if (OutLockedByteCount)
+	if (Desc.IsTextureCube())
 	{
-		*OutLockedByteCount = MipBytes;
+		return Arguments.FaceIndex + Arguments.ArrayIndex * 6;
 	}
 
-	GLLockTracker.Lock(Texture, Result, 0, MipIndex, DestStride, MipBytes, LockMode);
-	return Result;
-}
-
-void FOpenGLDynamicRHI::UnlockTexture2D_RenderThread(class FRHICommandListImmediate& RHICmdList, FRHITexture* Texture, uint32 MipIndex, bool bLockWithinMiptail, bool bNeedsDefaultRHIFlush)
-{
-	check(IsInRenderingThread());
-	FTextureLockTracker::FLockParams Params = GLLockTracker.Unlock(Texture, 0, MipIndex);
-	if (Params.LockMode != RLM_WriteOnly)
+	if (Desc.IsTextureArray())
 	{
-		GLLockTracker.TotalMemoryOutstanding = 0;
-		RHICmdList.EnqueueLambda([this, Texture, MipIndex, bLockWithinMiptail](FRHICommandListImmediate&)
-		{
-			this->RHIUnlockTexture2D(Texture, MipIndex, bLockWithinMiptail);
-		});
+		return Arguments.ArrayIndex;
 	}
-	else
-	{
-		RHICmdList.EnqueueLambda([this, Params, Texture, MipIndex, bLockWithinMiptail](FRHICommandListImmediate&)
-		{
-			uint32 DestStride;
-			uint64 LockedByteCount = ~0ULL;
-			uint8* TexMem = (uint8*)this->RHILockTexture2D(Texture, MipIndex, Params.LockMode, DestStride, bLockWithinMiptail, &LockedByteCount);
-			check(LockedByteCount != ~0ULL);
-			uint8* BuffMem = (uint8*)Params.Buffer;
-			check(DestStride == Params.Stride);
-			check(LockedByteCount >= Params.BufferSize);
-			FMemory::Memcpy(TexMem, BuffMem, Params.BufferSize);
-			FMemory::Free(Params.Buffer);
-			this->RHIUnlockTexture2D(Texture, MipIndex, bLockWithinMiptail);
-		});
-	}
-}
 
-void* FOpenGLDynamicRHI::RHILockTextureCubeFace_RenderThread(class FRHICommandListImmediate& RHICmdList, FRHITexture* Texture, uint32 FaceIndex, uint32 ArrayIndex, uint32 MipIndex, EResourceLockMode LockMode, uint32& DestStride, bool bLockWithinMiptail)
-{
-	check(IsInRenderingThread());
-	check(LockMode != EResourceLockMode::RLM_WriteOnly_NoOverwrite);
-	void* Result;
-	uint32 MipBytes = 0;
-	if (LockMode != RLM_WriteOnly)
-	{
-		RHICmdList.EnqueueLambda([this, &Result, Texture, FaceIndex, ArrayIndex, MipIndex, LockMode, &DestStride, bLockWithinMiptail](FRHICommandListImmediate&)
-		{
-			Result = this->RHILockTextureCubeFace(Texture, FaceIndex, ArrayIndex, MipIndex, LockMode, DestStride, bLockWithinMiptail);
-		});
-		RHITHREAD_GLTRACE_BLOCKING;
-		RHICmdList.ImmediateFlush(EImmediateFlushType::FlushRHIThread);
-		MipBytes = ResourceCast(Texture)->GetLockSize(MipIndex, 0, LockMode, DestStride);
-	}
-	else
-	{
-		MipBytes = ResourceCast(Texture)->GetLockSize(MipIndex, 0, LockMode, DestStride);
-		Result = FMemory::Malloc(MipBytes, 16);
-	}
-	check(Result);
-	GLLockTracker.Lock(Texture, Result, ArrayIndex, MipIndex, DestStride, MipBytes, LockMode);
-	return Result;
+	return 0;
 }
 
-void FOpenGLDynamicRHI::RHIUnlockTextureCubeFace_RenderThread(class FRHICommandListImmediate& RHICmdList, FRHITexture* Texture, uint32 FaceIndex, uint32 ArrayIndex, uint32 MipIndex, bool bLockWithinMiptail)
+FRHILockTextureResult FOpenGLDynamicRHI::RHILockTexture(FRHICommandListImmediate& RHICmdList, const FRHILockTextureArgs& Arguments)
 {
 	check(IsInRenderingThread());
+	check(Arguments.LockMode != EResourceLockMode::RLM_WriteOnly_NoOverwrite);
 
-	FTextureLockTracker::FLockParams Params = GLLockTracker.Unlock(Texture, ArrayIndex, MipIndex);
-	if (Params.LockMode != RLM_WriteOnly)
-	{
-		GLLockTracker.TotalMemoryOutstanding = 0;
-		RHICmdList.EnqueueLambda([this, Texture, FaceIndex, ArrayIndex, MipIndex, bLockWithinMiptail](FRHICommandListImmediate&)
-		{
-			this->RHIUnlockTextureCubeFace(Texture, FaceIndex, ArrayIndex, MipIndex, bLockWithinMiptail);
-		});
-	}
-	else
-	{
-		RHICmdList.EnqueueLambda([this, Params, Texture, FaceIndex, ArrayIndex, MipIndex, bLockWithinMiptail](FRHICommandListImmediate&)
-		{
-			uint32 DestStride;
-			uint8* TexMem = (uint8*)this->RHILockTextureCubeFace(Texture, FaceIndex, ArrayIndex, MipIndex, RLM_WriteOnly, DestStride, bLockWithinMiptail);
-			uint8* BuffMem = (uint8*)Params.Buffer;
-			check(DestStride == Params.Stride);
-			FMemory::Memcpy(TexMem, BuffMem, Params.BufferSize);
-			FMemory::Free(Params.Buffer);
-			this->RHIUnlockTextureCubeFace(Texture, FaceIndex, ArrayIndex, MipIndex, bLockWithinMiptail);
-		});
-	}
-}
+	FOpenGLTexture* Texture = ResourceCast(Arguments.Texture);
+	const uint32 ArrayIndex = GetGLArrayIndex(Arguments);
 
-void* FOpenGLDynamicRHI::LockTexture2DArray_RenderThread(class FRHICommandListImmediate& RHICmdList, FRHITexture* Texture, uint32 ArrayIndex, uint32 MipIndex, EResourceLockMode LockMode, uint32& DestStride, bool bLockWithinMiptail)
-{
-	check(IsInRenderingThread());
-	check(LockMode != EResourceLockMode::RLM_WriteOnly_NoOverwrite);
-	void* Result;
-	uint32 MipBytes = 0;
+	FRHILockTextureResult Result;
 
-	if (LockMode != RLM_WriteOnly)
+	if (Arguments.LockMode == RLM_ReadOnly)
 	{
-		RHICmdList.EnqueueLambda([this, &Result, Texture, ArrayIndex, MipIndex, LockMode, &DestStride, bLockWithinMiptail](FRHICommandListImmediate&)
+		RHICmdList.EnqueueLambda([this, &Result, Texture, MipIndex = Arguments.MipIndex, ArrayIndex](FRHICommandListImmediate&)
 		{
-			Result = this->RHILockTexture2DArray(Texture, ArrayIndex, MipIndex, LockMode, DestStride, bLockWithinMiptail);
+			Result.Data = Texture->Lock(MipIndex, ArrayIndex, RLM_ReadOnly, Result.Stride);
 		});
+
 		RHITHREAD_GLTRACE_BLOCKING;
 		RHICmdList.ImmediateFlush(EImmediateFlushType::FlushRHIThread);
-
-		MipBytes = ResourceCast(Texture)->GetLockSize(MipIndex, ArrayIndex, LockMode, DestStride);
+		Result.ByteCount = Texture->GetLockSize(Arguments.MipIndex, ArrayIndex, RLM_ReadOnly, Result.Stride);
 	}
 	else
 	{
-		MipBytes = ResourceCast(Texture)->GetLockSize(MipIndex, ArrayIndex, LockMode, DestStride);
-		Result = FMemory::Malloc(MipBytes, 16);
+		Result.ByteCount = Texture->GetLockSize(Arguments.MipIndex, 0, RLM_WriteOnly, Result.Stride);
+		Result.Data = FMemory::Malloc(Result.ByteCount, 16);
 	}
-	check(Result);
 
-	GLLockTracker.Lock(Texture, Result, ArrayIndex, MipIndex, DestStride, MipBytes, LockMode);
+	check(Result.Data);
+	GLLockTracker.Lock(Texture, Result.Data, ArrayIndex, Arguments.MipIndex, Result.Stride, Result.ByteCount, Arguments.LockMode);
 	return Result;
 }
 
-void FOpenGLDynamicRHI::UnlockTexture2DArray_RenderThread(class FRHICommandListImmediate& RHICmdList, FRHITexture* Texture, uint32 ArrayIndex, uint32 MipIndex, bool bLockWithinMiptail)
+void FOpenGLDynamicRHI::RHIUnlockTexture(FRHICommandListImmediate& RHICmdList, const FRHILockTextureArgs& Arguments)
 {
 	check(IsInRenderingThread());
-	FTextureLockTracker::FLockParams Params = GLLockTracker.Unlock(Texture, ArrayIndex, MipIndex);
-	if (Params.LockMode != RLM_WriteOnly)
+
+	FOpenGLTexture* Texture = ResourceCast(Arguments.Texture);
+	const uint32 ArrayIndex = GetGLArrayIndex(Arguments);
+
+	FTextureLockTracker::FLockParams Params = GLLockTracker.Unlock(Texture, ArrayIndex, Arguments.MipIndex);
+	if (Params.LockMode == RLM_ReadOnly)
 	{
 		GLLockTracker.TotalMemoryOutstanding = 0;
-		RHICmdList.EnqueueLambda([this, Texture, ArrayIndex, MipIndex, bLockWithinMiptail](FRHICommandListImmediate&)
+		RHICmdList.EnqueueLambda([this, Texture, MipIndex = Arguments.MipIndex, ArrayIndex](FRHICommandListImmediate&)
 		{
-			this->RHIUnlockTexture2DArray(Texture, ArrayIndex, MipIndex, bLockWithinMiptail);
+			Texture->Unlock(MipIndex, ArrayIndex);
 		});
 	}
 	else
 	{
-		RHICmdList.EnqueueLambda([this, Params, Texture, ArrayIndex, MipIndex, bLockWithinMiptail](FRHICommandListImmediate&)
+		RHICmdList.EnqueueLambda([this, Params, Texture, MipIndex = Arguments.MipIndex, ArrayIndex](FRHICommandListImmediate&)
 		{
-			uint32 DestStride;
-			uint8* TexMem = (uint8*)this->RHILockTexture2DArray(Texture, ArrayIndex, MipIndex, Params.LockMode, DestStride, bLockWithinMiptail);
-			uint8* BuffMem = (uint8*)Params.Buffer;
+			uint32 DestStride = 0;
+			void* DestinationMemory = Texture->Lock(MipIndex, ArrayIndex, RLM_WriteOnly, DestStride);
 			check(DestStride == Params.Stride);
-			FMemory::Memcpy(TexMem, BuffMem, Params.BufferSize);
+
+			FMemory::Memcpy(DestinationMemory, Params.Buffer, Params.BufferSize);
 			FMemory::Free(Params.Buffer);
-			this->RHIUnlockTexture2DArray(Texture, ArrayIndex, MipIndex, bLockWithinMiptail);
+
+			Texture->Unlock(MipIndex, ArrayIndex);
 		});
 	}
 }
diff --git a/Engine/Source/Runtime/RHI/Private/RHICommandList.cpp b/Engine/Source/Runtime/RHI/Private/RHICommandList.cpp
--- a/Engine/Source/Runtime/RHI/Private/RHICommandList.cpp
+++ b/Engine/Source/Runtime/RHI/Private/RHICommandList.cpp
@@ -1996,73 +1996,12 @@ void FDynamicRHI::RHIEndMultiUpdateTexture3D(FRHICommandListBase& RHICmdList, TA
 	}
 }
 
-void* FDynamicRHI::LockTexture2D_RenderThread(class FRHICommandListImmediate& RHICmdList, FRHITexture* Texture, uint32 MipIndex, EResourceLockMode LockMode, uint32& DestStride, bool bLockWithinMiptail, bool bNeedsDefaultRHIFlush, uint64* OutLockedByteCount)
-{
-	if (bNeedsDefaultRHIFlush) 
-	{
-		QUICK_SCOPE_CYCLE_COUNTER(STAT_RHIMETHOD_LockTexture2D_Flush);
-		CSV_SCOPED_TIMING_STAT(RHITFlushes, LockTexture2D_RenderThread);
-		RHICmdList.ImmediateFlush(EImmediateFlushType::FlushRHIThread);
-		return GDynamicRHI->RHILockTexture2D(Texture, MipIndex, LockMode, DestStride, bLockWithinMiptail, OutLockedByteCount);
-	}
-	CSV_SCOPED_TIMING_STAT(RHITStalls, LockTexture2D_RenderThread);
-	FScopedRHIThreadStaller StallRHIThread(RHICmdList);
-	return GDynamicRHI->RHILockTexture2D(Texture, MipIndex, LockMode, DestStride, bLockWithinMiptail, OutLockedByteCount);
-}
-
-void FDynamicRHI::UnlockTexture2D_RenderThread(class FRHICommandListImmediate& RHICmdList, FRHITexture* Texture, uint32 MipIndex, bool bLockWithinMiptail, bool bNeedsDefaultRHIFlush)
-{
-	if (bNeedsDefaultRHIFlush)
-	{
-		QUICK_SCOPE_CYCLE_COUNTER(STAT_RHIMETHOD_UnlockTexture2D_Flush);
-		CSV_SCOPED_TIMING_STAT(RHITFlushes, UnlockTexture2D_RenderThread);
-		RHICmdList.ImmediateFlush(EImmediateFlushType::FlushRHIThread);
-		GDynamicRHI->RHIUnlockTexture2D(Texture, MipIndex, bLockWithinMiptail);
-		return;
-	}
-	CSV_SCOPED_TIMING_STAT(RHITStalls, UnlockTexture2D_RenderThread);
-	FScopedRHIThreadStaller StallRHIThread(RHICmdList);
-	GDynamicRHI->RHIUnlockTexture2D(Texture, MipIndex, bLockWithinMiptail);
-}
-
-void* FDynamicRHI::LockTexture2DArray_RenderThread(class FRHICommandListImmediate& RHICmdList, FRHITexture* Texture, uint32 ArrayIndex, uint32 MipIndex, EResourceLockMode LockMode, uint32& DestStride, bool bLockWithinMiptail)
-{
-	QUICK_SCOPE_CYCLE_COUNTER(STAT_RHIMETHOD_LockTexture2DArray_Flush);
-	CSV_SCOPED_TIMING_STAT(RHITFlushes, LockTexture2DArray_RenderThread);
-	RHICmdList.ImmediateFlush(EImmediateFlushType::FlushRHIThread);
-	return GDynamicRHI->RHILockTexture2DArray(Texture, ArrayIndex, MipIndex, LockMode, DestStride, bLockWithinMiptail);
-}
-
-void FDynamicRHI::UnlockTexture2DArray_RenderThread(class FRHICommandListImmediate& RHICmdList, FRHITexture* Texture, uint32 ArrayIndex, uint32 MipIndex, bool bLockWithinMiptail)
-{
-	QUICK_SCOPE_CYCLE_COUNTER(STAT_RHIMETHOD_UnlockTexture2DArray_Flush);
-	CSV_SCOPED_TIMING_STAT(RHITFlushes, UnlockTexture2DArray_RenderThread);
-	RHICmdList.ImmediateFlush(EImmediateFlushType::FlushRHIThread);
-	GDynamicRHI->RHIUnlockTexture2DArray(Texture, ArrayIndex, MipIndex, bLockWithinMiptail);
-}
-
 FRHIShaderLibraryRef FDynamicRHI::RHICreateShaderLibrary_RenderThread(class FRHICommandListImmediate& RHICmdList, EShaderPlatform Platform, FString FilePath, FString Name)
 {
 	FScopedRHIThreadStaller StallRHIThread(RHICmdList);
 	return GDynamicRHI->RHICreateShaderLibrary(Platform, FilePath, Name);
 }
 
-void* FDynamicRHI::RHILockTextureCubeFace_RenderThread(class FRHICommandListImmediate& RHICmdList, FRHITexture* Texture, uint32 FaceIndex, uint32 ArrayIndex, uint32 MipIndex, EResourceLockMode LockMode, uint32& DestStride, bool bLockWithinMiptail)
-{
-	QUICK_SCOPE_CYCLE_COUNTER(STAT_RHIMETHOD_LockTextureCubeFace_Flush);
-	CSV_SCOPED_TIMING_STAT(RHITFlushes, RHILockTextureCubeFace_RenderThread);
-	RHICmdList.ImmediateFlush(EImmediateFlushType::FlushRHIThread);
-	return GDynamicRHI->RHILockTextureCubeFace(Texture, FaceIndex, ArrayIndex, MipIndex, LockMode, DestStride, bLockWithinMiptail);
-}
-
-void FDynamicRHI::RHIUnlockTextureCubeFace_RenderThread(class FRHICommandListImmediate& RHICmdList, FRHITexture* Texture, uint32 FaceIndex, uint32 ArrayIndex, uint32 MipIndex, bool bLockWithinMiptail)
-{
-	QUICK_SCOPE_CYCLE_COUNTER(STAT_RHIMETHOD_UnlockTextureCubeFace_Flush);
-	CSV_SCOPED_TIMING_STAT(RHITFlushes, RHIUnlockTextureCubeFace_RenderThread);
-	RHICmdList.ImmediateFlush(EImmediateFlushType::FlushRHIThread);
-	GDynamicRHI->RHIUnlockTextureCubeFace(Texture, FaceIndex, ArrayIndex, MipIndex, bLockWithinMiptail);
-}
-
 void FDynamicRHI::RHIMapStagingSurface_RenderThread(class FRHICommandListImmediate& RHICmdList, FRHITexture* Texture, uint32 GPUIndex, FRHIGPUFence* Fence, void*& OutData, int32& OutWidth, int32& OutHeight)
 {
 	if (Fence == nullptr || !Fence->Poll() || Fence->NumPendingWriteCommands.GetValue() != 0)
diff --git a/Engine/Source/Runtime/RHI/Private/RHIValidation.cpp b/Engine/Source/Runtime/RHI/Private/RHIValidation.cpp
--- a/Engine/Source/Runtime/RHI/Private/RHIValidation.cpp
+++ b/Engine/Source/Runtime/RHI/Private/RHIValidation.cpp
@@ -972,6 +972,48 @@ void* FValidationRHI::RHILockBufferMGPU(class FRHICommandListBase& RHICmdList, F
 	return RHI->RHILockBufferMGPU(RHICmdList, Buffer, GPUIndex, Offset, SizeRHI, LockMode);
 }
 
+FRHILockTextureResult FValidationRHI::RHILockTexture(FRHICommandListImmediate& RHICmdList, const FRHILockTextureArgs& Arguments)
+{
+	FRHITexture* Texture = Arguments.Texture;
+	const FRHITextureDesc& Desc = Arguments.Texture->GetDesc();
+
+	RHI_VALIDATION_CHECK(Arguments.MipIndex < Desc.NumMips, TEXT("Out of bounds MipIndex"));
+
+	switch (Desc.Dimension)
+	{
+	case ETextureDimension::Texture2D:
+		RHI_VALIDATION_CHECK(Arguments.ArrayIndex == 0, TEXT("Texture2D locks do not support array indexing"));
+		RHI_VALIDATION_CHECK(Arguments.FaceIndex == 0, TEXT("Texture2D locks do not support face indexing"));
+		break;
+	case ETextureDimension::Texture2DArray:
+		RHI_VALIDATION_CHECK(Arguments.ArrayIndex < Desc.ArraySize, TEXT("Texture2DArray lock out of bounds ArrayIndex"));
+		RHI_VALIDATION_CHECK(Arguments.FaceIndex == 0, TEXT("Texture2DArray locks do not support face indexing"));
+		break;
+	case ETextureDimension::Texture3D:
+		RHI_VALIDATION_CHECK(false, TEXT("Texture3D locks have not been fully tested"));
+		RHI_VALIDATION_CHECK(Arguments.FaceIndex == 0, TEXT("Texture3D locks do not support face indexing"));
+		break;
+	case ETextureDimension::TextureCube:
+		RHI_VALIDATION_CHECK(Arguments.ArrayIndex == 0, TEXT("TextureCube locks do not support array indexing"));
+		RHI_VALIDATION_CHECK(Arguments.FaceIndex < 6, TEXT("TextureCube lock out of bounds FaceIndex"));
+		break;
+	case ETextureDimension::TextureCubeArray:
+		RHI_VALIDATION_CHECK(Arguments.ArrayIndex < Desc.ArraySize, TEXT("Out of bounds ArrayIndex"));
+		RHI_VALIDATION_CHECK(Arguments.FaceIndex < 6, TEXT("TextureCubeArray lock out of bounds Face Index"));
+		break;
+	default:
+		checkNoEntry();
+		break;
+	}
+
+	return RHI->RHILockTexture(RHICmdList, Arguments);
+}
+
+void FValidationRHI::RHIUnlockTexture(FRHICommandListImmediate& RHICmdList, const FRHILockTextureArgs& Arguments)
+{
+	return RHI->RHIUnlockTexture(RHICmdList, Arguments);
+}
+
 class FRHIValidationQueueScope
 {
 	RHIValidation::FOpQueueState* Prev;
diff --git a/Engine/Source/Runtime/RHI/Public/DynamicRHI.h b/Engine/Source/Runtime/RHI/Public/DynamicRHI.h
--- a/Engine/Source/Runtime/RHI/Public/DynamicRHI.h
+++ b/Engine/Source/Runtime/RHI/Public/DynamicRHI.h
@@ -121,6 +121,82 @@ struct FRHICalcTextureSizeResult
 	uint32 Align;
 };
 
+struct FRHILockedTextureDesc
+{
+	FRHITexture* Texture = nullptr;
+	uint32       FaceIndex = 0;
+	uint32       ArrayIndex = 0;
+	uint32       MipIndex = 0;
+
+	bool operator==(const FRHILockedTextureDesc& Other) const
+	{
+		return Texture    == Other.Texture
+			&& FaceIndex  == Other.FaceIndex
+			&& ArrayIndex == Other.ArrayIndex
+			&& MipIndex   == Other.MipIndex;
+	}
+	bool operator!=(const FRHILockedTextureDesc& Other) const
+	{
+		return !(*this == Other);
+	}
+};
+
+struct FRHILockTextureArgs : public FRHILockedTextureDesc
+{
+	static inline FRHILockTextureArgs Lock2D(FRHITexture* InTexture, uint32 InMipIndex, EResourceLockMode InLockMode, bool bInLockWithinMiptail, bool bFlushRHIThread = true)
+	{
+		FRHILockTextureArgs Result;
+
+		Result.Texture = InTexture;
+		Result.LockMode = InLockMode;
+		Result.MipIndex = InMipIndex;
+		Result.bLockWithinMiptail = bInLockWithinMiptail;
+		Result.bNeedsDefaultRHIFlush = bFlushRHIThread;
+
+		return Result;
+	}
+	static inline FRHILockTextureArgs Lock2DArray(FRHITexture* InTexture, uint32 InArrayIndex, uint32 InMipIndex, EResourceLockMode InLockMode, bool bInLockWithinMiptail)
+	{
+		FRHILockTextureArgs Result;
+
+		Result.Texture = InTexture;
+		Result.LockMode = InLockMode;
+		Result.ArrayIndex = InArrayIndex;
+		Result.MipIndex = InMipIndex;
+		Result.bLockWithinMiptail = bInLockWithinMiptail;
+
+		return Result;
+	}
+	static inline FRHILockTextureArgs LockCubeFace(FRHITexture* InTexture, uint32 InFaceIndex, uint32 InArrayIndex, uint32 InMipIndex, EResourceLockMode InLockMode, bool bInLockWithinMiptail)
+	{
+		FRHILockTextureArgs Result;
+
+		Result.Texture = InTexture;
+		Result.LockMode = InLockMode;
+		Result.FaceIndex = InFaceIndex;
+		Result.ArrayIndex = InArrayIndex;
+		Result.MipIndex = InMipIndex;
+		Result.bLockWithinMiptail = bInLockWithinMiptail;
+
+		return Result;
+	}
+
+	EResourceLockMode LockMode = RLM_Num;
+	bool              bLockWithinMiptail = false;
+	bool              bNeedsDefaultRHIFlush = false;
+
+protected:
+	// Default constructor restricted to use in static functions
+	FRHILockTextureArgs() = default;
+};
+
+struct FRHILockTextureResult
+{
+	void*  Data      = nullptr;
+	uint64 ByteCount = 0;
+	uint32 Stride    = 0;
+};
+
 /** The interface which is implemented by the dynamically bound RHI. */
 class FDynamicRHI
 {
@@ -479,47 +555,8 @@ public:
 	// NP: Note that no RHI currently implements this as an async call, we should simplify the API.
 	virtual FTextureRHIRef RHIAsyncReallocateTexture2D(FRHITexture* Texture2D, int32 NewMipCount, int32 NewSizeX, int32 NewSizeY, FThreadSafeCounter* RequestStatus) = 0;
 
-	/**
-	* Locks an RHI texture's mip-map for read/write operations on the CPU
-	* @param Texture - the RHI texture resource to lock, must not be 0
-	* @param MipIndex - index of the mip level to lock
-	* @param LockMode - Whether to lock the texture read-only instead of write-only
-	* @param DestStride - output to retrieve the textures row stride (pitch)
-	* @param bLockWithinMiptail - for platforms that support packed miptails allow locking of individual mip levels within the miptail
-	* @return pointer to the CPU accessible resource data
-	*/
-	// FlushType: Flush RHI Thread
-	virtual void* RHILockTexture2D(FRHITexture* Texture, uint32 MipIndex, EResourceLockMode LockMode, uint32& DestStride, bool bLockWithinMiptail, uint64* OutLockedByteCount = nullptr) = 0;
-
-	/**
-	* Unlocks a previously locked RHI texture resource
-	* @param Texture - the RHI texture resource to unlock, must not be 0
-	* @param MipIndex - index of the mip level to unlock
-	* @param bLockWithinMiptail - for platforms that support packed miptails allow locking of individual mip levels within the miptail
-	*/
-	// FlushType: Flush RHI Thread
-	virtual void RHIUnlockTexture2D(FRHITexture* Texture, uint32 MipIndex, bool bLockWithinMiptail) = 0;
-
-	/**
-	* Locks an RHI texture's mip-map for read/write operations on the CPU
-	* @param Texture - the RHI texture resource to lock
-	* @param MipIndex - index of the mip level to lock
-	* @param LockMode - Whether to lock the texture read-only instead of write-only
-	* @param DestStride - output to retrieve the textures row stride (pitch)
-	* @param bLockWithinMiptail - for platforms that support packed miptails allow locking of individual mip levels within the miptail
-	* @return pointer to the CPU accessible resource data
-	*/
-	// FlushType: Flush RHI Thread
-	virtual void* RHILockTexture2DArray(FRHITexture* Texture, uint32 TextureIndex, uint32 MipIndex, EResourceLockMode LockMode, uint32& DestStride, bool bLockWithinMiptail) = 0;
-
-	/**
-	* Unlocks a previously locked RHI texture resource
-	* @param Texture - the RHI texture resource to unlock
-	* @param MipIndex - index of the mip level to unlock
-	* @param bLockWithinMiptail - for platforms that support packed miptails allow locking of individual mip levels within the miptail
-	*/
-	// FlushType: Flush RHI Thread
-	virtual void RHIUnlockTexture2DArray(FRHITexture* Texture, uint32 TextureIndex, uint32 MipIndex, bool bLockWithinMiptail) = 0;
+	virtual FRHILockTextureResult RHILockTexture  (FRHICommandListImmediate& RHICmdList, const FRHILockTextureArgs& Arguments) = 0;
+	virtual void                  RHIUnlockTexture(FRHICommandListImmediate& RHICmdList, const FRHILockTextureArgs& Arguments) = 0;
 
 	/**
 	* Updates a region of a 2D texture from system memory
@@ -555,26 +592,6 @@ public:
 	*/
 	virtual void RHIUpdateTexture3D(FRHICommandListBase& RHICmdList, FRHITexture* Texture, uint32 MipIndex, const struct FUpdateTextureRegion3D& UpdateRegion, uint32 SourceRowPitch, uint32 SourceDepthPitch, const uint8* SourceData) = 0;
 
-	/**
-	* Locks an RHI texture's mip-map for read/write operations on the CPU
-	* @param Texture - the RHI texture resource to lock
-	* @param MipIndex - index of the mip level to lock
-	* @param LockMode - Whether to lock the texture read-only instead of write-only.
-	* @param DestStride - output to retrieve the textures row stride (pitch)
-	* @param bLockWithinMiptail - for platforms that support packed miptails allow locking of individual mip levels within the miptail
-	* @return pointer to the CPU accessible resource data
-	*/
-	// FlushType: Flush RHI Thread
-	virtual void* RHILockTextureCubeFace(FRHITexture* Texture, uint32 FaceIndex, uint32 ArrayIndex, uint32 MipIndex, EResourceLockMode LockMode, uint32& DestStride, bool bLockWithinMiptail) = 0;
-
-	/**
-	* Unlocks a previously locked RHI texture resource
-	* @param Texture - the RHI texture resource to unlock
-	* @param MipIndex - index of the mip level to unlock
-	* @param bLockWithinMiptail - for platforms that support packed miptails allow locking of individual mip levels within the miptail
-	*/
-	// FlushType: Flush RHI Thread
-	virtual void RHIUnlockTextureCubeFace(FRHITexture* Texture, uint32 FaceIndex, uint32 ArrayIndex, uint32 MipIndex, bool bLockWithinMiptail) = 0;
 
 	// FlushType: Thread safe
 	virtual void RHIBindDebugLabelName(FRHICommandListBase& RHICmdList,FRHITexture* Texture, const TCHAR* Name) = 0;
@@ -876,12 +893,6 @@ public:
 
 	RHI_API virtual FTextureRHIRef AsyncReallocateTexture2D_RenderThread(class FRHICommandListImmediate& RHICmdList, FRHITexture* Texture2D, int32 NewMipCount, int32 NewSizeX, int32 NewSizeY, FThreadSafeCounter* RequestStatus);
 
-	RHI_API virtual void* LockTexture2D_RenderThread(class FRHICommandListImmediate& RHICmdList, FRHITexture* Texture, uint32 MipIndex, EResourceLockMode LockMode, uint32& DestStride, bool bLockWithinMiptail, bool bNeedsDefaultRHIFlush = true, uint64* OutLockedByteCount = nullptr);
-	RHI_API virtual void UnlockTexture2D_RenderThread(class FRHICommandListImmediate& RHICmdList, FRHITexture* Texture, uint32 MipIndex, bool bLockWithinMiptail, bool bNeedsDefaultRHIFlush = true);
-
-	RHI_API virtual void* LockTexture2DArray_RenderThread(class FRHICommandListImmediate& RHICmdList, FRHITexture* Texture, uint32 ArrayIndex, uint32 MipIndex, EResourceLockMode LockMode, uint32& DestStride, bool bLockWithinMiptail);
-	RHI_API virtual void UnlockTexture2DArray_RenderThread(class FRHICommandListImmediate& RHICmdList, FRHITexture* Texture, uint32 ArrayIndex, uint32 MipIndex, bool bLockWithinMiptail);
-
 	RHI_API virtual FUpdateTexture3DData RHIBeginUpdateTexture3D(FRHICommandListBase& RHICmdList, FRHITexture* Texture, uint32 MipIndex, const struct FUpdateTextureRegion3D& UpdateRegion);
 	RHI_API virtual void RHIEndUpdateTexture3D(FRHICommandListBase& RHICmdList, FUpdateTexture3DData& UpdateData);
 
@@ -889,9 +900,6 @@ public:
 
 	RHI_API virtual FRHIShaderLibraryRef RHICreateShaderLibrary_RenderThread(class FRHICommandListImmediate& RHICmdList, EShaderPlatform Platform, FString FilePath, FString Name);
 
-	RHI_API virtual void* RHILockTextureCubeFace_RenderThread(class FRHICommandListImmediate& RHICmdList, FRHITexture* Texture, uint32 FaceIndex, uint32 ArrayIndex, uint32 MipIndex, EResourceLockMode LockMode, uint32& DestStride, bool bLockWithinMiptail);
-	RHI_API virtual void RHIUnlockTextureCubeFace_RenderThread(class FRHICommandListImmediate& RHICmdList, FRHITexture* Texture, uint32 FaceIndex, uint32 ArrayIndex, uint32 MipIndex, bool bLockWithinMiptail);
-
 	RHI_API virtual void RHIMapStagingSurface_RenderThread(class FRHICommandListImmediate& RHICmdList, FRHITexture* Texture, uint32 GPUIndex, FRHIGPUFence* Fence, void*& OutData, int32& OutWidth, int32& OutHeight);
 	RHI_API virtual void RHIUnmapStagingSurface_RenderThread(class FRHICommandListImmediate& RHICmdList, FRHITexture* Texture, uint32 GPUIndex);
 	RHI_API virtual void RHIReadSurfaceFloatData_RenderThread(class FRHICommandListImmediate& RHICmdList, FRHITexture* Texture, FIntRect Rect, TArray<FFloat16Color>& OutData, ECubeFace CubeFace, int32 ArrayIndex, int32 MipIndex);
diff --git a/Engine/Source/Runtime/RHI/Public/RHICommandList.h b/Engine/Source/Runtime/RHI/Public/RHICommandList.h
--- a/Engine/Source/Runtime/RHI/Public/RHICommandList.h
+++ b/Engine/Source/Runtime/RHI/Public/RHICommandList.h
@@ -4939,28 +4939,69 @@ public:
 	{
 		return TexRealloc_Succeeded;
 	}
+
+	FORCEINLINE FRHILockTextureResult LockTexture(const FRHILockTextureArgs& Arguments)
+	{
+		LLM_SCOPE(ELLMTag::Textures);
+		return GDynamicRHI->RHILockTexture(*this, Arguments);
+	}
+
+	FORCEINLINE void UnlockTexture(const FRHILockTextureArgs& Arguments)
+	{
+		LLM_SCOPE(ELLMTag::Textures);
+		GDynamicRHI->RHIUnlockTexture(*this, Arguments);
+	}
 	
+	//UE_DEPRECATED(5.6, "LockTexture should be used for all texture type locks.")
 	FORCEINLINE void* LockTexture2D(FRHITexture* Texture, uint32 MipIndex, EResourceLockMode LockMode, uint32& DestStride, bool bLockWithinMiptail, bool bFlushRHIThread = true, uint64* OutLockedByteCount = nullptr)
 	{
 		LLM_SCOPE(ELLMTag::Textures);
-		return GDynamicRHI->LockTexture2D_RenderThread(*this, Texture, MipIndex, LockMode, DestStride, bLockWithinMiptail, bFlushRHIThread, OutLockedByteCount);
+
+		const FRHILockTextureResult Result = LockTexture(FRHILockTextureArgs::Lock2D(Texture, MipIndex, LockMode, bLockWithinMiptail, bFlushRHIThread));
+		DestStride = Result.Stride;
+		if (OutLockedByteCount)
+		{
+			*OutLockedByteCount = Result.ByteCount;
+		}
+		return Result.Data;
 	}
 	
+	//UE_DEPRECATED(5.6, "UnlockTexture should be used for all texture type unlocks.")
 	FORCEINLINE void UnlockTexture2D(FRHITexture* Texture, uint32 MipIndex, bool bLockWithinMiptail, bool bFlushRHIThread = true)
-	{		
-		GDynamicRHI->UnlockTexture2D_RenderThread(*this, Texture, MipIndex, bLockWithinMiptail, bFlushRHIThread);
+	{
+		UnlockTexture(FRHILockTextureArgs::Lock2D(Texture, MipIndex, RLM_Num, bLockWithinMiptail, bFlushRHIThread));
 	}
 	
-	FORCEINLINE void* LockTexture2DArray(FRHITexture* Texture, uint32 TextureIndex, uint32 MipIndex, EResourceLockMode LockMode, uint32& DestStride, bool bLockWithinMiptail)
+	//UE_DEPRECATED(5.6, "LockTexture should be used for all texture type locks.")
+	FORCEINLINE void* LockTexture2DArray(FRHITexture* Texture, uint32 ArrayIndex, uint32 MipIndex, EResourceLockMode LockMode, uint32& DestStride, bool bLockWithinMiptail)
 	{
 		LLM_SCOPE(ELLMTag::Textures);
-		return GDynamicRHI->LockTexture2DArray_RenderThread(*this, Texture, TextureIndex, MipIndex, LockMode, DestStride, bLockWithinMiptail);
+		const FRHILockTextureResult Result = LockTexture(FRHILockTextureArgs::Lock2DArray(Texture, ArrayIndex, MipIndex, LockMode, bLockWithinMiptail));
+		DestStride = Result.Stride;
+		return Result.Data;
 	}
 	
-	FORCEINLINE void UnlockTexture2DArray(FRHITexture* Texture, uint32 TextureIndex, uint32 MipIndex, bool bLockWithinMiptail)
+	//UE_DEPRECATED(5.6, "UnlockTexture should be used for all texture type unlocks.")
+	FORCEINLINE void UnlockTexture2DArray(FRHITexture* Texture, uint32 ArrayIndex, uint32 MipIndex, bool bLockWithinMiptail)
 	{
 		LLM_SCOPE(ELLMTag::Textures);
-		GDynamicRHI->UnlockTexture2DArray_RenderThread(*this, Texture, TextureIndex, MipIndex, bLockWithinMiptail);
+		UnlockTexture(FRHILockTextureArgs::Lock2DArray(Texture, ArrayIndex, MipIndex, RLM_Num, bLockWithinMiptail));
+	}
+
+	//UE_DEPRECATED(5.6, "LockTexture should be used for all texture type locks.")
+	FORCEINLINE void* LockTextureCubeFace(FRHITexture* Texture, uint32 FaceIndex, uint32 ArrayIndex, uint32 MipIndex, EResourceLockMode LockMode, uint32& DestStride, bool bLockWithinMiptail)
+	{
+		LLM_SCOPE(ELLMTag::Textures);
+		const FRHILockTextureResult Result = LockTexture(FRHILockTextureArgs::LockCubeFace(Texture, FaceIndex, ArrayIndex, MipIndex, LockMode, bLockWithinMiptail));
+		DestStride = Result.Stride;
+		return Result.Data;
+	}
+
+	//UE_DEPRECATED(5.6, "UnlockTexture should be used for all texture type unlocks.")
+	FORCEINLINE void UnlockTextureCubeFace(FRHITexture* Texture, uint32 FaceIndex, uint32 ArrayIndex, uint32 MipIndex, bool bLockWithinMiptail)
+	{
+		LLM_SCOPE(ELLMTag::Textures);
+		UnlockTexture(FRHILockTextureArgs::LockCubeFace(Texture, FaceIndex, ArrayIndex, MipIndex, RLM_Num, bLockWithinMiptail));
 	}
 	
 	FORCEINLINE FUpdateTexture3DData BeginUpdateTexture3D(FRHITexture* Texture, uint32 MipIndex, const struct FUpdateTextureRegion3D& UpdateRegion)
@@ -4983,18 +5024,6 @@ public:
 		LLM_SCOPE(ELLMTag::Textures);
 		GDynamicRHI->RHIEndMultiUpdateTexture3D(*this, UpdateDataArray);
 	}
-
-	FORCEINLINE void* LockTextureCubeFace(FRHITexture* Texture, uint32 FaceIndex, uint32 ArrayIndex, uint32 MipIndex, EResourceLockMode LockMode, uint32& DestStride, bool bLockWithinMiptail)
-	{
-		LLM_SCOPE(ELLMTag::Textures);
-		return GDynamicRHI->RHILockTextureCubeFace_RenderThread(*this, Texture, FaceIndex, ArrayIndex, MipIndex, LockMode, DestStride, bLockWithinMiptail);
-	}
-	
-	FORCEINLINE void UnlockTextureCubeFace(FRHITexture* Texture, uint32 FaceIndex, uint32 ArrayIndex, uint32 MipIndex, bool bLockWithinMiptail)
-	{
-		LLM_SCOPE(ELLMTag::Textures);
-		GDynamicRHI->RHIUnlockTextureCubeFace_RenderThread(*this, Texture, FaceIndex, ArrayIndex, MipIndex, bLockWithinMiptail);
-	}
 	
 	// ReadSurfaceFloatData reads texture data into FColor
 	//	pixels in other formats are converted to FColor
@@ -5745,26 +5774,42 @@ FORCEINLINE ETextureReallocationStatus RHICancelAsyncReallocateTexture2D(FRHITex
 	return TexRealloc_Succeeded;
 }
 
+//UE_DEPRECATED(5.6, "RHICmdList.LockTexture should be used for all texture type locks.")
 FORCEINLINE void* RHILockTexture2D(FRHITexture* Texture, uint32 MipIndex, EResourceLockMode LockMode, uint32& DestStride, bool bLockWithinMiptail, bool bFlushRHIThread = true, uint64* OutLockedByteCount = nullptr)
 {
 	return FRHICommandListExecutor::GetImmediateCommandList().LockTexture2D(Texture, MipIndex, LockMode, DestStride, bLockWithinMiptail, bFlushRHIThread, OutLockedByteCount);
 }
 
+//UE_DEPRECATED(5.6, "RHICmdList.UnlockTexture should be used for all texture type unlocks.")
 FORCEINLINE void RHIUnlockTexture2D(FRHITexture* Texture, uint32 MipIndex, bool bLockWithinMiptail, bool bFlushRHIThread = true)
 {
-	 FRHICommandListExecutor::GetImmediateCommandList().UnlockTexture2D(Texture, MipIndex, bLockWithinMiptail, bFlushRHIThread);
+	FRHICommandListExecutor::GetImmediateCommandList().UnlockTexture2D(Texture, MipIndex, bLockWithinMiptail, bFlushRHIThread);
 }
 
+//UE_DEPRECATED(5.6, "RHICmdList.LockTexture should be used for all texture type locks.")
 FORCEINLINE void* RHILockTexture2DArray(FRHITexture* Texture, uint32 TextureIndex, uint32 MipIndex, EResourceLockMode LockMode, uint32& DestStride, bool bLockWithinMiptail)
 {
 	return FRHICommandListExecutor::GetImmediateCommandList().LockTexture2DArray(Texture, TextureIndex, MipIndex, LockMode, DestStride, bLockWithinMiptail);
 }
 
+//UE_DEPRECATED(5.6, "RHICmdList.UnlockTexture should be used for all texture type unlocks.")
 FORCEINLINE void RHIUnlockTexture2DArray(FRHITexture* Texture, uint32 TextureIndex, uint32 MipIndex, bool bLockWithinMiptail)
 {
 	 FRHICommandListExecutor::GetImmediateCommandList().UnlockTexture2DArray(Texture, TextureIndex, MipIndex, bLockWithinMiptail);
 }
 
+//UE_DEPRECATED(5.6, "RHICmdList.LockTexture should be used for all texture type locks.")
+FORCEINLINE void* RHILockTextureCubeFace(FRHITexture* Texture, uint32 FaceIndex, uint32 ArrayIndex, uint32 MipIndex, EResourceLockMode LockMode, uint32& DestStride, bool bLockWithinMiptail)
+{
+	return FRHICommandListExecutor::GetImmediateCommandList().LockTextureCubeFace(Texture, FaceIndex, ArrayIndex, MipIndex, LockMode, DestStride, bLockWithinMiptail);
+}
+
+//UE_DEPRECATED(5.6, "RHICmdList.UnlockTexture should be used for all texture type unlocks.")
+FORCEINLINE void RHIUnlockTextureCubeFace(FRHITexture* Texture, uint32 FaceIndex, uint32 ArrayIndex, uint32 MipIndex, bool bLockWithinMiptail)
+{
+	FRHICommandListExecutor::GetImmediateCommandList().UnlockTextureCubeFace(Texture, FaceIndex, ArrayIndex, MipIndex, bLockWithinMiptail);
+}
+
 FORCEINLINE void RHIUpdateTexture2D(FRHITexture* Texture, uint32 MipIndex, const struct FUpdateTextureRegion2D& UpdateRegion, uint32 SourcePitch, const uint8* SourceData)
 {
 	 FRHICommandListExecutor::GetImmediateCommandList().UpdateTexture2D(Texture, MipIndex, UpdateRegion, SourcePitch, SourceData);
@@ -5790,16 +5835,6 @@ FORCEINLINE void RHIUpdateTexture3D(FRHITexture* Texture, uint32 MipIndex, const
 	 FRHICommandListExecutor::GetImmediateCommandList().UpdateTexture3D(Texture, MipIndex, UpdateRegion, SourceRowPitch, SourceDepthPitch, SourceData);
 }
 
-FORCEINLINE void* RHILockTextureCubeFace(FRHITexture* Texture, uint32 FaceIndex, uint32 ArrayIndex, uint32 MipIndex, EResourceLockMode LockMode, uint32& DestStride, bool bLockWithinMiptail)
-{
-	return FRHICommandListExecutor::GetImmediateCommandList().LockTextureCubeFace(Texture, FaceIndex, ArrayIndex, MipIndex, LockMode, DestStride, bLockWithinMiptail);
-}
-
-FORCEINLINE void RHIUnlockTextureCubeFace(FRHITexture* Texture, uint32 FaceIndex, uint32 ArrayIndex, uint32 MipIndex, bool bLockWithinMiptail)
-{
-	 FRHICommandListExecutor::GetImmediateCommandList().UnlockTextureCubeFace(Texture, FaceIndex, ArrayIndex, MipIndex, bLockWithinMiptail);
-}
-
 UE_DEPRECATED(5.5, "RHIAcquireThreadOwnership() is deprecated. Thread ownership of the RHI is automatic, so this function is redundant.")
 FORCEINLINE void RHIAcquireThreadOwnership()
 {}
diff --git a/Engine/Source/Runtime/RHI/Public/RHILockTracker.h b/Engine/Source/Runtime/RHI/Public/RHILockTracker.h
--- a/Engine/Source/Runtime/RHI/Public/RHILockTracker.h
+++ b/Engine/Source/Runtime/RHI/Public/RHILockTracker.h
@@ -6,6 +6,58 @@
 
 struct FRHILockTracker
 {
+	struct FTextureLockParams
+	{
+		FTextureLockParams() = default;
+
+		FTextureLockParams(const FRHILockedTextureDesc& InDesc, EResourceLockMode InLockMode, void* InData, bool bInbDirectLock)
+			: Desc(InDesc)
+			, LockMode(InLockMode)
+			, Data(InData)
+			, bDirectLock(bInbDirectLock)
+		{
+		}
+
+		FRHILockedTextureDesc Desc;
+		EResourceLockMode     LockMode = RLM_WriteOnly;
+		void*                 Data = nullptr;
+
+		//did we call the normal flushing/updating lock?
+		bool                  bDirectLock = false;
+	};
+
+	TArray<FTextureLockParams, TInlineAllocator<16>> OutstandingTextureLocks;
+
+	void Lock(const FRHILockTextureArgs& InArguments, void* InData, bool bInDirectBufferWrite)
+	{
+		const FRHILockedTextureDesc& Desc = InArguments;
+#if DO_CHECK
+		for (const FTextureLockParams& OutstandingLock : OutstandingTextureLocks)
+		{
+			check(Desc != OutstandingLock.Desc || (OutstandingLock.bDirectLock && bInDirectBufferWrite));
+		}
+#endif
+		OutstandingTextureLocks.Emplace(Desc, InArguments.LockMode, InData, bInDirectBufferWrite);
+	}
+
+	FORCEINLINE_DEBUGGABLE FTextureLockParams Unlock(const FRHILockedTextureDesc& InDesc)
+	{
+		for (int32 Index = 0; Index < OutstandingTextureLocks.Num(); Index++)
+		{
+			if (OutstandingTextureLocks[Index].Desc == InDesc)
+			{
+				FTextureLockParams Result = OutstandingTextureLocks[Index];
+				OutstandingTextureLocks.RemoveAtSwap(Index, EAllowShrinking::No);
+				return Result;
+			}
+		}
+		RaiseMismatchError();
+		return FTextureLockParams();
+	}
+
+
+
+
 	struct FLockParams
 	{
 		void* RHIBuffer;
diff --git a/Engine/Source/Runtime/RHI/Public/RHIValidation.h b/Engine/Source/Runtime/RHI/Public/RHIValidation.h
--- a/Engine/Source/Runtime/RHI/Public/RHIValidation.h
+++ b/Engine/Source/Runtime/RHI/Public/RHIValidation.h
@@ -592,59 +592,8 @@ public:
 		return NewTexture2D;
 	}
 
-	/**
-	* Locks an RHI texture's mip-map for read/write operations on the CPU
-	* @param Texture - the RHI texture resource to lock, must not be 0
-	* @param MipIndex - index of the mip level to lock
-	* @param LockMode - Whether to lock the texture read-only instead of write-only
-	* @param DestStride - output to retrieve the textures row stride (pitch)
-	* @param bLockWithinMiptail - for platforms that support packed miptails allow locking of individual mip levels within the miptail
-	* @return pointer to the CPU accessible resource data
-	*/
-	// FlushType: Flush RHI Thread
-	virtual void* RHILockTexture2D(FRHITexture* Texture, uint32 MipIndex, EResourceLockMode LockMode, uint32& DestStride, bool bLockWithinMiptail, uint64* OutLockedByteCount) override final
-	{
-		return RHI->RHILockTexture2D(Texture, MipIndex, LockMode, DestStride, bLockWithinMiptail, OutLockedByteCount);
-	}
-
-	/**
-	* Unlocks a previously locked RHI texture resource
-	* @param Texture - the RHI texture resource to unlock, must not be 0
-	* @param MipIndex - index of the mip level to unlock
-	* @param bLockWithinMiptail - for platforms that support packed miptails allow locking of individual mip levels within the miptail
-	*/
-	// FlushType: Flush RHI Thread
-	virtual void RHIUnlockTexture2D(FRHITexture* Texture, uint32 MipIndex, bool bLockWithinMiptail) override final
-	{
-		RHI->RHIUnlockTexture2D(Texture, MipIndex, bLockWithinMiptail);
-	}
-
-	/**
-	* Locks an RHI texture's mip-map for read/write operations on the CPU
-	* @param Texture - the RHI texture resource to lock
-	* @param MipIndex - index of the mip level to lock
-	* @param LockMode - Whether to lock the texture read-only instead of write-only
-	* @param DestStride - output to retrieve the textures row stride (pitch)
-	* @param bLockWithinMiptail - for platforms that support packed miptails allow locking of individual mip levels within the miptail
-	* @return pointer to the CPU accessible resource data
-	*/
-	// FlushType: Flush RHI Thread
-	virtual void* RHILockTexture2DArray(FRHITexture* Texture, uint32 TextureIndex, uint32 MipIndex, EResourceLockMode LockMode, uint32& DestStride, bool bLockWithinMiptail) override final
-	{
-		return RHI->RHILockTexture2DArray(Texture, TextureIndex, MipIndex, LockMode, DestStride, bLockWithinMiptail);
-	}
-
-	/**
-	* Unlocks a previously locked RHI texture resource
-	* @param Texture - the RHI texture resource to unlock
-	* @param MipIndex - index of the mip level to unlock
-	* @param bLockWithinMiptail - for platforms that support packed miptails allow locking of individual mip levels within the miptail
-	*/
-	// FlushType: Flush RHI Thread
-	virtual void RHIUnlockTexture2DArray(FRHITexture* Texture, uint32 TextureIndex, uint32 MipIndex, bool bLockWithinMiptail) override final
-	{
-		RHI->RHIUnlockTexture2DArray(Texture, TextureIndex, MipIndex, bLockWithinMiptail);
-	}
+	virtual FRHILockTextureResult RHILockTexture(FRHICommandListImmediate& RHICmdList, const FRHILockTextureArgs& Arguments) final override;
+	virtual void RHIUnlockTexture(FRHICommandListImmediate& RHICmdList, const FRHILockTextureArgs& Arguments) final override;
 
 	/**
 	* Updates a region of a 2D texture from system memory
@@ -680,33 +629,6 @@ public:
 		RHI->RHIUpdateTexture3D(RHICmdList, Texture, MipIndex, UpdateRegion, SourceRowPitch, SourceDepthPitch, SourceData);
 	}
 
-	/**
-	* Locks an RHI texture's mip-map for read/write operations on the CPU
-	* @param Texture - the RHI texture resource to lock
-	* @param MipIndex - index of the mip level to lock
-	* @param LockMode - Whether to lock the texture read-only instead of write-only.
-	* @param DestStride - output to retrieve the textures row stride (pitch)
-	* @param bLockWithinMiptail - for platforms that support packed miptails allow locking of individual mip levels within the miptail
-	* @return pointer to the CPU accessible resource data
-	*/
-	// FlushType: Flush RHI Thread
-	virtual void* RHILockTextureCubeFace(FRHITexture* Texture, uint32 FaceIndex, uint32 ArrayIndex, uint32 MipIndex, EResourceLockMode LockMode, uint32& DestStride, bool bLockWithinMiptail) override final
-	{
-		return RHI->RHILockTextureCubeFace(Texture, FaceIndex, ArrayIndex, MipIndex, LockMode, DestStride, bLockWithinMiptail);
-	}
-
-	/**
-	* Unlocks a previously locked RHI texture resource
-	* @param Texture - the RHI texture resource to unlock
-	* @param MipIndex - index of the mip level to unlock
-	* @param bLockWithinMiptail - for platforms that support packed miptails allow locking of individual mip levels within the miptail
-	*/
-	// FlushType: Flush RHI Thread
-	virtual void RHIUnlockTextureCubeFace(FRHITexture* Texture, uint32 FaceIndex, uint32 ArrayIndex, uint32 MipIndex, bool bLockWithinMiptail) override final
-	{
-		RHI->RHIUnlockTextureCubeFace(Texture, FaceIndex, ArrayIndex, MipIndex, bLockWithinMiptail);
-	}
-
 	// FlushType: Thread safe
 	virtual void RHIBindDebugLabelName(FRHICommandListBase& RHICmdList, FRHITexture* Texture, const TCHAR* Name) override final;
 
@@ -1093,26 +1015,6 @@ public:
 		RHI->UnlockBuffer_BottomOfPipe(RHICmdList, Buffer);
 	}
 
-	virtual void* LockTexture2D_RenderThread(class FRHICommandListImmediate& RHICmdList, FRHITexture* Texture, uint32 MipIndex, EResourceLockMode LockMode, uint32& DestStride, bool bLockWithinMiptail, bool bNeedsDefaultRHIFlush = true, uint64* OutLockedByteCount = nullptr) override final
-	{
-		return RHI->LockTexture2D_RenderThread(RHICmdList, Texture, MipIndex, LockMode, DestStride, bLockWithinMiptail, bNeedsDefaultRHIFlush, OutLockedByteCount);
-	}
-
-	virtual void UnlockTexture2D_RenderThread(class FRHICommandListImmediate& RHICmdList, FRHITexture* Texture, uint32 MipIndex, bool bLockWithinMiptail, bool bNeedsDefaultRHIFlush = true) override final
-	{
-		RHI->UnlockTexture2D_RenderThread(RHICmdList, Texture, MipIndex, bLockWithinMiptail, bNeedsDefaultRHIFlush);
-	}
-
-	virtual void* LockTexture2DArray_RenderThread(class FRHICommandListImmediate& RHICmdList, FRHITexture* Texture, uint32 ArrayIndex, uint32 MipIndex, EResourceLockMode LockMode, uint32& DestStride, bool bLockWithinMiptail) override final
-	{
-		return RHI->LockTexture2DArray_RenderThread(RHICmdList, Texture, ArrayIndex, MipIndex, LockMode, DestStride, bLockWithinMiptail);
-	}
-
-	virtual void UnlockTexture2DArray_RenderThread(class FRHICommandListImmediate& RHICmdList, FRHITexture* Texture, uint32 ArrayIndex, uint32 MipIndex, bool bLockWithinMiptail) override final
-	{
-		RHI->UnlockTexture2DArray_RenderThread(RHICmdList, Texture, ArrayIndex, MipIndex, bLockWithinMiptail);
-	}
-
 	virtual FUpdateTexture3DData RHIBeginUpdateTexture3D(FRHICommandListBase& RHICmdList, FRHITexture* Texture, uint32 MipIndex, const struct FUpdateTextureRegion3D& UpdateRegion) override final
 	{
 		return RHI->RHIBeginUpdateTexture3D(RHICmdList, Texture, MipIndex, UpdateRegion);
@@ -1133,15 +1035,6 @@ public:
 		return RHI->RHICreateShaderLibrary_RenderThread(RHICmdList, Platform, FilePath, Name);
 	}
 
-	virtual void* RHILockTextureCubeFace_RenderThread(class FRHICommandListImmediate& RHICmdList, FRHITexture* Texture, uint32 FaceIndex, uint32 ArrayIndex, uint32 MipIndex, EResourceLockMode LockMode, uint32& DestStride, bool bLockWithinMiptail) override final
-	{
-		return RHI->RHILockTextureCubeFace_RenderThread(RHICmdList, Texture, FaceIndex, ArrayIndex, MipIndex, LockMode, DestStride, bLockWithinMiptail);
-	}
-	virtual void RHIUnlockTextureCubeFace_RenderThread(class FRHICommandListImmediate& RHICmdList, FRHITexture* Texture, uint32 FaceIndex, uint32 ArrayIndex, uint32 MipIndex, bool bLockWithinMiptail) override final
-	{
-		RHI->RHIUnlockTextureCubeFace_RenderThread(RHICmdList, Texture, FaceIndex, ArrayIndex, MipIndex, bLockWithinMiptail);
-	}
-
 	virtual void RHIReadSurfaceFloatData_RenderThread(class FRHICommandListImmediate& RHICmdList, FRHITexture* Texture, FIntRect Rect, TArray<FFloat16Color>& OutData, ECubeFace CubeFace, int32 ArrayIndex, int32 MipIndex) override final
 	{
 		RHI->RHIReadSurfaceFloatData_RenderThread(RHICmdList, Texture, Rect, OutData, CubeFace, ArrayIndex, MipIndex);
diff --git a/Engine/Source/Runtime/VulkanRHI/Private/VulkanResources.h b/Engine/Source/Runtime/VulkanRHI/Private/VulkanResources.h
--- a/Engine/Source/Runtime/VulkanRHI/Private/VulkanResources.h
+++ b/Engine/Source/Runtime/VulkanRHI/Private/VulkanResources.h
@@ -388,7 +388,6 @@ struct FVulkanCpuReadbackBuffer
 {
 	VkBuffer Buffer;
 	uint32 MipOffsets[MAX_TEXTURE_MIP_COUNT];
-	uint32 MipSize[MAX_TEXTURE_MIP_COUNT];
 };
 
 class FVulkanView
@@ -659,15 +658,10 @@ public:
 	 */
 	void GetMipStride(uint32 MipIndex, uint32& Stride);
 
-	/*
-	 * Returns the memory offset to the texture's mip-map.
-	 */
-	void GetMipOffset(uint32 MipIndex, uint32& Offset);
-
 	/**
 	* Returns how much memory a single mip uses.
 	*/
-	void GetMipSize(uint32 MipIndex, uint32& MipBytes);
+	void GetMipSize(uint32 MipIndex, uint64& MipBytes);
 
 	inline VkImageViewType GetViewType() const
 	{
diff --git a/Engine/Source/Runtime/VulkanRHI/Private/VulkanTexture.cpp b/Engine/Source/Runtime/VulkanRHI/Private/VulkanTexture.cpp
--- a/Engine/Source/Runtime/VulkanRHI/Private/VulkanTexture.cpp
+++ b/Engine/Source/Runtime/VulkanRHI/Private/VulkanTexture.cpp
@@ -13,6 +13,7 @@
 #include "VulkanBarriers.h"
 #include "VulkanTransientResourceAllocator.h"
 #include "RHICoreStats.h"
+#include "RHILockTracker.h"
 
 int32 GVulkanSubmitOnTextureUnlock = 1;
 static FAutoConsoleVariableRef CVarVulkanSubmitOnTextureUnlock(
@@ -37,22 +38,6 @@ static FAutoConsoleVariableRef CVarVulkanDepthStencilForceStorageBit(
 
 extern int32 GVulkanLogDefrag;
 
-static FCriticalSection GTextureMapLock;
-
-struct FTextureLock
-{
-	FRHIResource* Texture;
-	uint32 MipIndex;
-	uint32 LayerIndex;
-
-	FTextureLock(FRHIResource* InTexture, uint32 InMipIndex, uint32 InLayerIndex = 0)
-		: Texture(InTexture)
-		, MipIndex(InMipIndex)
-		, LayerIndex(InLayerIndex)
-	{
-	}
-};
-
 #if ENABLE_LOW_LEVEL_MEM_TRACKER
 inline ELLMTagVulkan GetMemoryTagForTextureFlags(ETextureCreateFlags UEFlags)
 {
@@ -61,18 +46,6 @@ inline ELLMTagVulkan GetMemoryTagForTextureFlags(ETextureCreateFlags UEFlags)
 }
 #endif // ENABLE_LOW_LEVEL_MEM_TRACKER
 
-inline bool operator == (const FTextureLock& A, const FTextureLock& B)
-{
-	return A.Texture == B.Texture && A.MipIndex == B.MipIndex && A.LayerIndex == B.LayerIndex;
-}
-
-inline uint32 GetTypeHash(const FTextureLock& Lock)
-{
-	return GetTypeHash(Lock.Texture) ^ (Lock.MipIndex << 16) ^ (Lock.LayerIndex << 8);
-}
-
-static TMap<FTextureLock, VulkanRHI::FStagingBuffer*> GPendingLockedBuffers;
-
 static const VkImageTiling GVulkanViewTypeTilingMode[VK_IMAGE_VIEW_TYPE_RANGE_SIZE] =
 {
 	VK_IMAGE_TILING_LINEAR,		// VK_IMAGE_VIEW_TYPE_1D
@@ -863,17 +836,8 @@ void FVulkanTexture::GetMipStride(uint32 MipIndex, uint32& Stride)
 	Stride = NumBlocksX * BlockBytes;
 }
 
-void FVulkanTexture::GetMipOffset(uint32 MipIndex, uint32& Offset)
-{
-	uint32 offset = Offset = 0;
-	for(uint32 i = 0; i < MipIndex; i++)
-	{
-		GetMipSize(i, offset);
-		Offset += offset;
-	}
-}
 
-void FVulkanTexture::GetMipSize(uint32 MipIndex, uint32& MipBytes)
+void FVulkanTexture::GetMipSize(uint32 MipIndex, uint64& MipBytes)
 {
 	// Calculate the dimensions of mip-map level.
 	const FRHITextureDesc& Desc = GetDesc();
@@ -1244,148 +1208,79 @@ FTextureRHIRef FVulkanDynamicRHI::RHIAsyncReallocateTexture2D(FRHITexture* OldTe
 	return NewTexture;
 }
 
-void* FVulkanDynamicRHI::RHILockTexture2D(FRHITexture* TextureRHI, uint32 MipIndex, EResourceLockMode LockMode, uint32& DestStride, bool bLockWithinMiptail, uint64* OutLockedByteCount)
+static FCriticalSection GTextureMapLock;
+
+FRHILockTextureResult FVulkanDynamicRHI::RHILockTexture(FRHICommandListImmediate& RHICmdList, const FRHILockTextureArgs& Arguments)
 {
 	LLM_SCOPE_VULKAN(ELLMTagVulkan::VulkanTextures);
-	FVulkanTexture* Texture = ResourceCast(TextureRHI);
-	check(Texture);
-
-	VulkanRHI::FStagingBuffer** StagingBuffer = nullptr;
-	{
-		FScopeLock Lock(&GTextureMapLock);
-		StagingBuffer = &GPendingLockedBuffers.FindOrAdd(FTextureLock(TextureRHI, MipIndex));
-		checkf(!*StagingBuffer, TEXT("Can't lock the same texture twice!"));
-	}
 
-	// No locks for read allowed yet
-	check(LockMode == RLM_WriteOnly);
+	FVulkanTexture* Texture = ResourceCast(Arguments.Texture);
 
-	uint32 BufferSize = 0;
-	DestStride = 0;
-	Texture->GetMipSize(MipIndex, BufferSize);
-	Texture->GetMipStride(MipIndex, DestStride);
-	*StagingBuffer = Device->GetStagingManager().AcquireBuffer(BufferSize);
+	FRHILockTextureResult Result{};
+	Texture->GetMipSize(Arguments.MipIndex, Result.ByteCount);
+	Texture->GetMipStride(Arguments.MipIndex, Result.Stride);
 
-	if (OutLockedByteCount)
+	VulkanRHI::FStagingBuffer* StagingBuffer = Device->GetStagingManager().AcquireBuffer(Result.ByteCount);
 	{
-		*OutLockedByteCount = BufferSize;
+		FScopeLock Lock(&GTextureMapLock);
+		GRHILockTracker.Lock(Arguments, reinterpret_cast<void*>(StagingBuffer), false);
 	}
 
-	void* Data = (*StagingBuffer)->GetMappedPointer();
-	return Data;
+	Result.Data = StagingBuffer->GetMappedPointer();
+	
+	return Result;
 }
 
-void FVulkanDynamicRHI::InternalUnlockTexture2D(bool bFromRenderingThread, FRHITexture* TextureRHI, uint32 MipIndex, bool bLockWithinMiptail)
+static uint32 GetVulkanArrayIndex(const FRHILockTextureArgs& Arguments, const FRHITextureDesc& Desc)
 {
-	LLM_SCOPE_VULKAN(ELLMTagVulkan::VulkanTextures);
-	FVulkanTexture* Texture = ResourceCast(TextureRHI);
-	check(Texture);
-
-	VkDevice LogicalDevice = Device->GetInstanceHandle();
-
-	VulkanRHI::FStagingBuffer* StagingBuffer = nullptr;
+	if (Desc.IsTextureCube())
 	{
-		FScopeLock Lock(&GTextureMapLock);
-		bool bFound = GPendingLockedBuffers.RemoveAndCopyValue(FTextureLock(TextureRHI, MipIndex), StagingBuffer);
-		checkf(bFound, TEXT("Texture was not locked!"));
+		return Arguments.FaceIndex + Arguments.ArrayIndex * 6;
 	}
 
-	const FRHITextureDesc& Desc = Texture->GetDesc();
-	const EPixelFormat Format = Desc.Format;
-	uint32 MipWidth = FMath::Max<uint32>(Desc.Extent.X >> MipIndex, 0);
-	uint32 MipHeight = FMath::Max<uint32>(Desc.Extent.Y >> MipIndex, 0);
-	ensure(!(MipHeight == 0 && MipWidth == 0));
-	MipWidth = FMath::Max<uint32>(MipWidth, 1);
-	MipHeight = FMath::Max<uint32>(MipHeight, 1);
-	uint32 LayerCount = Texture->GetNumberOfArrayLevels();
-
-	VkBufferImageCopy Region;
-	FMemory::Memzero(Region);
-	//#todo-rco: Might need an offset here?
-	//Region.bufferOffset = 0;
-	//Region.bufferRowLength = 0;
-	//Region.bufferImageHeight = 0;
-	Region.imageSubresource.aspectMask = VK_IMAGE_ASPECT_COLOR_BIT;
-	Region.imageSubresource.mipLevel = MipIndex;
-	//Region.imageSubresource.baseArrayLayer = 0;
-	Region.imageSubresource.layerCount = LayerCount;
-	Region.imageExtent.width = MipWidth;
-	Region.imageExtent.height = MipHeight;
-	Region.imageExtent.depth = 1;
-
-	FRHICommandList& RHICmdList = FRHICommandListExecutor::GetImmediateCommandList();
-	if (!bFromRenderingThread || (RHICmdList.Bypass() || !IsRunningRHIInSeparateThread()))
-	{
-		FVulkanTexture::InternalLockWrite(Device->GetImmediateContext(), Texture, Region, StagingBuffer);
-	}
-	else
+	if (Desc.IsTextureArray())
 	{
-		check(IsInRenderingThread());
-		ALLOC_COMMAND_CL(RHICmdList, FRHICommandLockWriteTexture)(Texture, Region, StagingBuffer);
+		return Arguments.ArrayIndex;
 	}
+
+	return 0;
 }
 
-void* FVulkanDynamicRHI::RHILockTexture2DArray(FRHITexture* TextureRHI,uint32 TextureIndex,uint32 MipIndex,EResourceLockMode LockMode,uint32& DestStride,bool bLockWithinMiptail)
+void FVulkanDynamicRHI::RHIUnlockTexture(FRHICommandListImmediate& RHICmdList, const FRHILockTextureArgs& Arguments)
 {
 	LLM_SCOPE_VULKAN(ELLMTagVulkan::VulkanTextures);
-	FVulkanTexture* Texture = ResourceCast(TextureRHI);
-	check(Texture);
 
-	VulkanRHI::FStagingBuffer** StagingBuffer = nullptr;
+	VulkanRHI::FStagingBuffer* StagingBuffer = nullptr;
 	{
 		FScopeLock Lock(&GTextureMapLock);
-		StagingBuffer = &GPendingLockedBuffers.FindOrAdd(FTextureLock(TextureRHI, MipIndex, TextureIndex));
-		checkf(!*StagingBuffer, TEXT("Can't lock the same texture twice!"));
+		const FRHILockTracker::FTextureLockParams Params = GRHILockTracker.Unlock(Arguments);
+		StagingBuffer = reinterpret_cast<VulkanRHI::FStagingBuffer*>(Params.Data);
+		checkf(StagingBuffer, TEXT("Texture was not locked!"));
 	}
 
-	// No locks for read allowed yet
-	check(LockMode == RLM_WriteOnly);
-
-	uint32 BufferSize = 0;
-	DestStride = 0;
-	Texture->GetMipSize(MipIndex, BufferSize);
-	Texture->GetMipStride(MipIndex, DestStride);
-	*StagingBuffer = Device->GetStagingManager().AcquireBuffer(BufferSize);
+	FVulkanTexture* Texture = ResourceCast(Arguments.Texture);
+	const FRHITextureDesc& Desc = Texture->GetDesc();
 
-	void* Data = (*StagingBuffer)->GetMappedPointer();
-	return Data;
-}
+	const uint32 ArrayIndex = GetVulkanArrayIndex(Arguments, Desc);
 
-void FVulkanDynamicRHI::RHIUnlockTexture2DArray(FRHITexture* TextureRHI, uint32 TextureIndex, uint32 MipIndex, bool bLockWithinMiptail)
-{
-	LLM_SCOPE_VULKAN(ELLMTagVulkan::VulkanTextures);
-	FVulkanTexture* Texture = ResourceCast(TextureRHI);
-	check(Texture);
+	uint32 MipWidth  = FMath::Max<uint32>(Desc.Extent.X >> Arguments.MipIndex, 0);
+	uint32 MipHeight = FMath::Max<uint32>(Desc.Extent.Y >> Arguments.MipIndex, 0);
 
-	VkDevice LogicalDevice = Device->GetInstanceHandle();
+	ensure(!(MipHeight == 0 && MipWidth == 0));
 
-	VulkanRHI::FStagingBuffer* StagingBuffer = nullptr;
-	{
-		FScopeLock Lock(&GTextureMapLock);
-		bool bFound = GPendingLockedBuffers.RemoveAndCopyValue(FTextureLock(TextureRHI, MipIndex, TextureIndex), StagingBuffer);
-		checkf(bFound, TEXT("Texture was not locked!"));
-	}
+	MipWidth = FMath::Max<uint32>(MipWidth, 1);
+	MipHeight = FMath::Max<uint32>(MipHeight, 1);
 
-	const FRHITextureDesc& Desc = Texture->GetDesc();
-	const EPixelFormat Format = Desc.Format;
-	const uint32 MipWidth = FMath::Max<uint32>(Desc.Extent.X >> MipIndex, 1);
-	const uint32 MipHeight = FMath::Max<uint32>(Desc.Extent.Y >> MipIndex, 1);
+	VkBufferImageCopy Region{};
 
-	VkBufferImageCopy Region;
-	FMemory::Memzero(Region);
-	//#todo-rco: Might need an offset here?
-	//Region.bufferOffset = 0;
-	//Region.bufferRowLength = 0;
-	//Region.bufferImageHeight = 0;
 	Region.imageSubresource.aspectMask = Texture->GetPartialAspectMask();
-	Region.imageSubresource.mipLevel = MipIndex;
-	Region.imageSubresource.baseArrayLayer = TextureIndex;
+	Region.imageSubresource.mipLevel = Arguments.MipIndex;
+	Region.imageSubresource.baseArrayLayer = ArrayIndex;
 	Region.imageSubresource.layerCount = 1;
 	Region.imageExtent.width = MipWidth;
 	Region.imageExtent.height = MipHeight;
 	Region.imageExtent.depth = 1;
 
-	FRHICommandList& RHICmdList = FRHICommandListExecutor::GetImmediateCommandList();
 	if (RHICmdList.Bypass() || !IsRunningRHIInSeparateThread())
 	{
 		FVulkanTexture::InternalLockWrite(Device->GetImmediateContext(), Texture, Region, StagingBuffer);
@@ -1397,6 +1292,7 @@ void FVulkanDynamicRHI::RHIUnlockTexture2DArray(FRHITexture* TextureRHI, uint32
 	}
 }
 
+
 void FVulkanDynamicRHI::InternalUpdateTexture2D(FRHICommandListBase& RHICmdList, FRHITexture* TextureRHI, uint32 MipIndex, const FUpdateTextureRegion2D& UpdateRegion, uint32 SourcePitch, const uint8* SourceData)
 {
 	LLM_SCOPE_VULKAN(ELLMTagVulkan::VulkanTextures);
@@ -1575,13 +1471,12 @@ FVulkanTexture::FVulkanTexture(FRHICommandListBase* RHICmdList, FVulkanDevice& I
 		check(InCreateDesc.ArraySize == 1);  //not implemented
 
 		CpuReadbackBuffer = new FVulkanCpuReadbackBuffer;
-		uint32 Size = 0;
+		uint64 Size = 0;
 		for (uint32 Mip = 0; Mip < InCreateDesc.NumMips; ++Mip)
 		{
-			uint32 LocalSize;
+			uint64 LocalSize = 0;
 			GetMipSize(Mip, LocalSize);
 			CpuReadbackBuffer->MipOffsets[Mip] = Size;
-			CpuReadbackBuffer->MipSize[Mip] = LocalSize;
 			Size += LocalSize;
 		}
 
@@ -2364,81 +2259,6 @@ bool FVulkanTexture::GetTextureResourceInfo(FRHIResourceInfo& OutResourceInfo) c
 	return true;
 }
 
-/*-----------------------------------------------------------------------------
-	Cubemap texture support.
------------------------------------------------------------------------------*/
-
-void* FVulkanDynamicRHI::RHILockTextureCubeFace(FRHITexture* TextureCubeRHI, uint32 FaceIndex, uint32 ArrayIndex, uint32 MipIndex, EResourceLockMode LockMode, uint32& DestStride, bool bLockWithinMiptail)
-{
-	LLM_SCOPE_VULKAN(ELLMTagVulkan::VulkanTextures);
-	FVulkanTexture* Texture = ResourceCast(TextureCubeRHI);
-	check(Texture);
-
-	VulkanRHI::FStagingBuffer** StagingBuffer = nullptr;
-	{
-		FScopeLock Lock(&GTextureMapLock);
-		StagingBuffer = &GPendingLockedBuffers.FindOrAdd(FTextureLock(TextureCubeRHI, MipIndex));
-		checkf(!*StagingBuffer, TEXT("Can't lock the same texture twice!"));
-	}
-
-	uint32 BufferSize = 0;
-	DestStride = 0;
-	Texture->GetMipSize(MipIndex, BufferSize);
-	Texture->GetMipStride(MipIndex, DestStride);
-	*StagingBuffer = Device->GetStagingManager().AcquireBuffer(BufferSize);
-
-	void* Data = (*StagingBuffer)->GetMappedPointer();
-	return Data;
-}
-
-void FVulkanDynamicRHI::RHIUnlockTextureCubeFace(FRHITexture* TextureCubeRHI, uint32 FaceIndex, uint32 ArrayIndex, uint32 MipIndex, bool bLockWithinMiptail)
-{
-	LLM_SCOPE_VULKAN(ELLMTagVulkan::VulkanTextures);
-	FVulkanTexture* Texture = ResourceCast(TextureCubeRHI);
-	check(Texture);
-
-	VkDevice LogicalDevice = Device->GetInstanceHandle();
-
-	VulkanRHI::FStagingBuffer* StagingBuffer = nullptr;
-	{
-		FScopeLock Lock(&GTextureMapLock);
-		bool bFound = GPendingLockedBuffers.RemoveAndCopyValue(FTextureLock(TextureCubeRHI, MipIndex), StagingBuffer);
-		checkf(bFound, TEXT("Texture was not locked!"));
-	}
-
-	const FRHITextureDesc& Desc = Texture->GetDesc();
-	uint32 MipWidth = FMath::Max<uint32>(Desc.Extent.X >> MipIndex, 0);
-	uint32 MipHeight = FMath::Max<uint32>(Desc.Extent.Y >> MipIndex, 0);
-	ensure(!(MipHeight == 0 && MipWidth == 0));
-	MipWidth = FMath::Max<uint32>(MipWidth, 1);
-	MipHeight = FMath::Max<uint32>(MipHeight, 1);
-
-	VkBufferImageCopy Region;
-	FMemory::Memzero(Region);
-	//#todo-rco: Might need an offset here?
-	//Region.bufferOffset = 0;
-	//Region.bufferRowLength = 0;
-	//Region.bufferImageHeight = 0;
-	Region.imageSubresource.aspectMask = Texture->GetPartialAspectMask();
-	Region.imageSubresource.mipLevel = MipIndex;
-	Region.imageSubresource.baseArrayLayer = ArrayIndex * 6 + FaceIndex;
-	Region.imageSubresource.layerCount = 1;
-	Region.imageExtent.width = MipWidth;
-	Region.imageExtent.height = MipHeight;
-	Region.imageExtent.depth = 1;
-
-	FRHICommandList& RHICmdList = FRHICommandListExecutor::GetImmediateCommandList();
-	if (RHICmdList.Bypass() || !IsRunningRHIInSeparateThread())
-	{
-		FVulkanTexture::InternalLockWrite(Device->GetImmediateContext(), Texture, Region, StagingBuffer);
-	}
-	else
-	{
-		check(IsInRenderingThread());
-		ALLOC_COMMAND_CL(RHICmdList, FRHICommandLockWriteTexture)(Texture, Region, StagingBuffer);
-	}
-}
-
 void FVulkanDynamicRHI::RHIBindDebugLabelName(FRHICommandListBase& RHICmdList, FRHITexture* TextureRHI, const TCHAR* Name)
 {
 #if VULKAN_ENABLE_IMAGE_TRACKING_LAYER
diff --git a/Engine/Source/Runtime/VulkanRHI/Public/VulkanDynamicRHI.h b/Engine/Source/Runtime/VulkanRHI/Public/VulkanDynamicRHI.h
--- a/Engine/Source/Runtime/VulkanRHI/Public/VulkanDynamicRHI.h
+++ b/Engine/Source/Runtime/VulkanRHI/Public/VulkanDynamicRHI.h
@@ -139,13 +139,9 @@ public:
 	virtual uint32 RHIComputeMemorySize(FRHITexture* TextureRHI) final override;
 	virtual FTextureRHIRef RHIAsyncReallocateTexture2D(FRHITexture* Texture2D, int32 NewMipCount, int32 NewSizeX, int32 NewSizeY, FThreadSafeCounter* RequestStatus) final override;
 
-	virtual void* RHILockTexture2D(FRHITexture* Texture, uint32 MipIndex, EResourceLockMode LockMode, uint32& DestStride, bool bLockWithinMiptail, uint64* OutLockedByteCount = nullptr) final override;
-	virtual void RHIUnlockTexture2D(FRHITexture* Texture, uint32 MipIndex, bool bLockWithinMiptail) final override
-	{
-		InternalUnlockTexture2D(false, Texture, MipIndex, bLockWithinMiptail);
-	}
-	virtual void* RHILockTexture2DArray(FRHITexture* Texture, uint32 TextureIndex, uint32 MipIndex, EResourceLockMode LockMode, uint32& DestStride, bool bLockWithinMiptail) final override;
-	virtual void RHIUnlockTexture2DArray(FRHITexture* Texture, uint32 TextureIndex, uint32 MipIndex, bool bLockWithinMiptail) final override;
+	virtual FRHILockTextureResult RHILockTexture(FRHICommandListImmediate& RHICmdList, const FRHILockTextureArgs& Arguments) final override;
+	virtual void RHIUnlockTexture(FRHICommandListImmediate& RHICmdList, const FRHILockTextureArgs& Arguments) final override;
+
 	virtual void RHIUpdateTexture2D(FRHICommandListBase& RHICmdList, FRHITexture* Texture, uint32 MipIndex, const struct FUpdateTextureRegion2D& UpdateRegion, uint32 SourcePitch, const uint8* SourceData) final override
 	{
 		InternalUpdateTexture2D(RHICmdList, Texture, MipIndex, UpdateRegion, SourcePitch, SourceData);
@@ -154,8 +150,6 @@ public:
 	{
 		InternalUpdateTexture3D(RHICmdList, Texture, MipIndex, UpdateRegion, SourceRowPitch, SourceDepthPitch, SourceData);
 	}
-	virtual void* RHILockTextureCubeFace(FRHITexture* Texture, uint32 FaceIndex, uint32 ArrayIndex, uint32 MipIndex, EResourceLockMode LockMode, uint32& DestStride, bool bLockWithinMiptail) final override;
-	virtual void RHIUnlockTextureCubeFace(FRHITexture* Texture, uint32 FaceIndex, uint32 ArrayIndex, uint32 MipIndex, bool bLockWithinMiptail) final override;
 	virtual void RHIBindDebugLabelName(FRHICommandListBase& RHICmdList, FRHITexture* Texture, const TCHAR* Name) final override;
 	virtual void RHIBindDebugLabelName(FRHICommandListBase& RHICmdList, FRHIUnorderedAccessView* UnorderedAccessViewRHI, const TCHAR* Name) final override;
 	virtual void RHIReadSurfaceData(FRHITexture* Texture, FIntRect Rect, TArray<FColor>& OutData, FReadSurfaceDataFlags InFlags) final override;
@@ -204,16 +198,7 @@ public:
 	virtual bool RHIMatchPrecachePSOInitializers(const FGraphicsPipelineStateInitializer& LHS, const FGraphicsPipelineStateInitializer& RHS) final override;
 
 	virtual FTextureRHIRef AsyncReallocateTexture2D_RenderThread(class FRHICommandListImmediate& RHICmdList, FRHITexture* Texture2D, int32 NewMipCount, int32 NewSizeX, int32 NewSizeY, FThreadSafeCounter* RequestStatus) override final;
-	
-	virtual void* LockTexture2D_RenderThread(class FRHICommandListImmediate& RHICmdList, FRHITexture* Texture, uint32 MipIndex, EResourceLockMode LockMode, uint32& DestStride, bool bLockWithinMiptail, bool bNeedsDefaultRHIFlush = true, uint64* OutLockedByteCount = nullptr) override final
-	{
-		return RHILockTexture2D(Texture, MipIndex, LockMode, DestStride, bLockWithinMiptail, OutLockedByteCount);
-	}
-	virtual void UnlockTexture2D_RenderThread(class FRHICommandListImmediate& RHICmdList, FRHITexture* Texture, uint32 MipIndex, bool bLockWithinMiptail, bool bNeedsDefaultRHIFlush = true) override final
-	{
-		InternalUnlockTexture2D(true, Texture, MipIndex, bLockWithinMiptail);
-	}
-	
+
 	virtual FUpdateTexture3DData RHIBeginUpdateTexture3D(FRHICommandListBase& RHICmdList, FRHITexture* Texture, uint32 MipIndex, const struct FUpdateTextureRegion3D& UpdateRegion) override final;
 	virtual void RHIEndUpdateTexture3D(FRHICommandListBase& RHICmdList, FUpdateTexture3DData& UpdateData) override final;
 
@@ -369,7 +354,6 @@ protected:
 
 	FCriticalSection LockBufferCS;
 
-	void InternalUnlockTexture2D(bool bFromRenderingThread, FRHITexture* Texture, uint32 MipIndex, bool bLockWithinMiptail);
 	void InternalUpdateTexture2D(FRHICommandListBase& RHICmdList, FRHITexture* TextureRHI, uint32 MipIndex, const struct FUpdateTextureRegion2D& UpdateRegion, uint32 SourcePitch, const uint8* SourceData);
 	void InternalUpdateTexture3D(FRHICommandListBase& RHICmdList, FRHITexture* TextureRHI, uint32 MipIndex, const struct FUpdateTextureRegion3D& UpdateRegion, uint32 SourceRowPitch, uint32 SourceDepthPitch, const uint8* SourceData);
 
diff --git a/Engine/Source/Runtime/Windows/D3D11RHI/Private/D3D11RHIPrivate.h b/Engine/Source/Runtime/Windows/D3D11RHI/Private/D3D11RHIPrivate.h
--- a/Engine/Source/Runtime/Windows/D3D11RHI/Private/D3D11RHIPrivate.h
+++ b/Engine/Source/Runtime/Windows/D3D11RHI/Private/D3D11RHIPrivate.h
@@ -453,14 +453,12 @@ public:
 	virtual uint32 RHIComputeMemorySize(FRHITexture* TextureRHI) final override;
 	virtual void RHIAsyncCopyTexture2DCopy(FRHITexture* NewTexture2DRHI, FRHITexture* Texture2DRHI, int32 NewMipCount, int32 NewSizeX, int32 NewSizeY, FThreadSafeCounter* RequestStatus);
 	virtual FTextureRHIRef RHIAsyncReallocateTexture2D(FRHITexture* Texture2D, int32 NewMipCount, int32 NewSizeX, int32 NewSizeY, FThreadSafeCounter* RequestStatus) final override;
-	virtual void* RHILockTexture2D(FRHITexture* Texture, uint32 MipIndex, EResourceLockMode LockMode, uint32& DestStride, bool bLockWithinMiptail, uint64* OutLockedByteCount) final override;
-	virtual void RHIUnlockTexture2D(FRHITexture* Texture, uint32 MipIndex, bool bLockWithinMiptail) final override;
-	virtual void* RHILockTexture2DArray(FRHITexture* Texture, uint32 TextureIndex, uint32 MipIndex, EResourceLockMode LockMode, uint32& DestStride, bool bLockWithinMiptail) final override;
-	virtual void RHIUnlockTexture2DArray(FRHITexture* Texture, uint32 TextureIndex, uint32 MipIndex, bool bLockWithinMiptail) final override;
+
+	virtual FRHILockTextureResult RHILockTexture(FRHICommandListImmediate& RHICmdList, const FRHILockTextureArgs& Arguments) final override;
+	virtual void                  RHIUnlockTexture(FRHICommandListImmediate& RHICmdList, const FRHILockTextureArgs& Arguments) final override;
+
 	virtual void RHIUpdateTexture2D(FRHICommandListBase& RHICmdList, FRHITexture* Texture, uint32 MipIndex, const struct FUpdateTextureRegion2D& UpdateRegion, uint32 SourcePitch, const uint8* SourceData) final override;
 	virtual void RHIUpdateTexture3D(FRHICommandListBase& RHICmdList, FRHITexture* Texture, uint32 MipIndex, const struct FUpdateTextureRegion3D& UpdateRegion, uint32 SourceRowPitch, uint32 SourceDepthPitch, const uint8* SourceData) final override;
-	virtual void* RHILockTextureCubeFace(FRHITexture* Texture, uint32 FaceIndex, uint32 ArrayIndex, uint32 MipIndex, EResourceLockMode LockMode, uint32& DestStride, bool bLockWithinMiptail) final override;
-	virtual void RHIUnlockTextureCubeFace(FRHITexture* Texture, uint32 FaceIndex, uint32 ArrayIndex, uint32 MipIndex, bool bLockWithinMiptail) final override;
 	virtual void RHIBindDebugLabelName(FRHICommandListBase& RHICmdList, FRHITexture* Texture, const TCHAR* Name) final override;
 	virtual void RHIBindDebugLabelName(FRHICommandListBase& RHICmdList, FRHIBuffer* Buffer, const TCHAR* Name) final override;
 	virtual void RHIBindDebugLabelName(FRHICommandListBase& RHICmdList, FRHIUnorderedAccessView* UnorderedAccessViewRHI, const TCHAR* Name) final override;
@@ -570,11 +568,7 @@ public:
 	// For example, deferring the actual creation to RHI thread by sending an RHI command.
 	// For D3D11, these functions mainly just remove RHIT stalls because ID3D11Device is thread safe.
 	virtual FTextureRHIRef AsyncReallocateTexture2D_RenderThread(class FRHICommandListImmediate& RHICmdList, FRHITexture* Texture2D, int32 NewMipCount, int32 NewSizeX, int32 NewSizeY, FThreadSafeCounter* RequestStatus) final override;
-	virtual void* LockTexture2D_RenderThread(class FRHICommandListImmediate& RHICmdList, FRHITexture* Texture, uint32 MipIndex, EResourceLockMode LockMode, uint32& DestStride, bool bLockWithinMiptail, bool bNeedsDefaultRHIFlush = true, uint64* OutLockedByteCount = nullptr) final override;
-	virtual void UnlockTexture2D_RenderThread(class FRHICommandListImmediate& RHICmdList, FRHITexture* Texture, uint32 MipIndex, bool bLockWithinMiptail, bool bNeedsDefaultRHIFlush = true) final override;
 	virtual void RHIEndUpdateTexture3D(FRHICommandListBase& RHICmdList, FUpdateTexture3DData& UpdateData) final override;
-	virtual void* RHILockTextureCubeFace_RenderThread(class FRHICommandListImmediate& RHICmdList, FRHITexture* Texture, uint32 FaceIndex, uint32 ArrayIndex, uint32 MipIndex, EResourceLockMode LockMode, uint32& DestStride, bool bLockWithinMiptail) final override;
-	virtual void RHIUnlockTextureCubeFace_RenderThread(class FRHICommandListImmediate& RHICmdList, FRHITexture* Texture, uint32 FaceIndex, uint32 ArrayIndex, uint32 MipIndex, bool bLockWithinMiptail) final override;
 
 	virtual void RHIBeginRenderPass(const FRHIRenderPassInfo& InInfo, const TCHAR* InName) final override;
 	virtual void RHIEndRenderPass() final override;
diff --git a/Engine/Source/Runtime/Windows/D3D11RHI/Private/D3D11Texture.cpp b/Engine/Source/Runtime/Windows/D3D11RHI/Private/D3D11Texture.cpp
--- a/Engine/Source/Runtime/Windows/D3D11RHI/Private/D3D11Texture.cpp
+++ b/Engine/Source/Runtime/Windows/D3D11RHI/Private/D3D11Texture.cpp
@@ -986,45 +986,64 @@ FTextureRHIRef FD3D11DynamicRHI::AsyncReallocateTexture2D_RenderThread(
 	return NewTexture2D;
 }
 
-void* FD3D11Texture::Lock(FD3D11DynamicRHI* D3DRHI, uint32 MipIndex, uint32 ArrayIndex, EResourceLockMode LockMode, uint32& DestStride, bool bForceLockDeferred, uint64* OutLockedByteCount)
+static uint32 GetD3D11ArrayIndex(const FRHILockTextureArgs& Arguments, const FRHITextureDesc& Desc)
+{
+	if (Desc.IsTextureCube())
+	{
+		const uint32 D3DFace = GetD3D11CubeFace(static_cast<ECubeFace>(Arguments.FaceIndex));
+		return D3DFace + Arguments.ArrayIndex * 6;
+	}
+
+	if (Desc.IsTextureArray())
+	{
+		return Arguments.ArrayIndex;
+	}
+
+	return 0;
+}
+
+FRHILockTextureResult FD3D11Texture::Lock(FD3D11DynamicRHI* D3DRHI, const FRHILockTextureArgs& Arguments, bool bForceLockDeferred)
 {
 	check(!IsTexture3D()); // Only 2D texture locks are implemented
 
 	SCOPE_CYCLE_COUNTER(STAT_D3D11LockTextureTime);
 
+	const FRHITextureDesc& Desc = this->GetDesc();
+
+	const uint32 ArrayIndex = GetD3D11ArrayIndex(Arguments, Desc);
+
 	// Calculate the subresource index corresponding to the specified mip-map.
-	const uint32 Subresource = D3D11CalcSubresource(MipIndex,ArrayIndex,this->GetNumMips());
+	const uint32 Subresource = D3D11CalcSubresource(Arguments.MipIndex, ArrayIndex, Desc.NumMips);
 
 	// Calculate the dimensions of the mip-map.
-	const uint32 BlockSizeX = GPixelFormats[this->GetFormat()].BlockSizeX;
-	const uint32 BlockSizeY = GPixelFormats[this->GetFormat()].BlockSizeY;
-	const uint32 BlockBytes = GPixelFormats[this->GetFormat()].BlockBytes;
-	const uint32 MipSizeX = FMath::Max(this->GetSizeX() >> MipIndex,BlockSizeX);
-	const uint32 MipSizeY = FMath::Max(this->GetSizeY() >> MipIndex,BlockSizeY);
+	const uint32 BlockSizeX = GPixelFormats[Desc.Format].BlockSizeX;
+	const uint32 BlockSizeY = GPixelFormats[Desc.Format].BlockSizeY;
+	const uint32 BlockBytes = GPixelFormats[Desc.Format].BlockBytes;
+
+	const uint32 MipSizeX = FMath::Max(uint32(Desc.Extent.X) >> Arguments.MipIndex, BlockSizeX);
+	const uint32 MipSizeY = FMath::Max(uint32(Desc.Extent.Y) >> Arguments.MipIndex, BlockSizeY);
 	const uint32 NumBlocksX = (MipSizeX + BlockSizeX - 1) / BlockSizeX;
 	const uint32 NumBlocksY = (MipSizeY + BlockSizeY - 1) / BlockSizeY;
 	const uint32 MipBytes = NumBlocksX * NumBlocksY * BlockBytes;
 
-	if (OutLockedByteCount)
-	{
-		*OutLockedByteCount = MipBytes;
-	}
+	FRHILockTextureResult Result;
 
 	FD3D11LockedData LockedData;
-	if( LockMode == RLM_WriteOnly )
+	if (Arguments.LockMode == RLM_WriteOnly)
 	{
-		if (!bForceLockDeferred && EnumHasAnyFlags(GetDesc().Flags, TexCreate_CPUWritable))
+		if (!bForceLockDeferred && EnumHasAnyFlags(Desc.Flags, TexCreate_CPUWritable))
 		{
-			D3D11_MAPPED_SUBRESOURCE MappedTexture;
+			D3D11_MAPPED_SUBRESOURCE MappedTexture{};
 			VERIFYD3D11RESULT_EX(D3DRHI->GetDeviceContext()->Map(GetResource(), Subresource, D3D11_MAP_WRITE, 0, &MappedTexture), D3DRHI->GetDevice());
+
 			LockedData.SetData(MappedTexture.pData);
-			LockedData.Pitch = DestStride = MappedTexture.RowPitch;
+			LockedData.Pitch = MappedTexture.RowPitch;
 		}
 		else
 		{
 			// If we're writing to the texture, allocate a system memory buffer to receive the new contents.
 			LockedData.AllocData(MipBytes);
-			LockedData.Pitch = DestStride = NumBlocksX * BlockBytes;
+			LockedData.Pitch = NumBlocksX * BlockBytes;
 			LockedData.bLockDeferred = true;
 		}
 	}
@@ -1034,8 +1053,9 @@ void* FD3D11Texture::Lock(FD3D11DynamicRHI* D3DRHI, uint32 MipIndex, uint32 Arra
 		// If we're reading from the texture, we create a staging resource, copy the texture contents to it, and map it.
 
 		// Create the staging texture.
-		D3D11_TEXTURE2D_DESC StagingTextureDesc;
+		D3D11_TEXTURE2D_DESC StagingTextureDesc{};
 		GetD3D11Texture2D()->GetDesc(&StagingTextureDesc);
+
 		StagingTextureDesc.Width = MipSizeX;
 		StagingTextureDesc.Height = MipSizeY;
 		StagingTextureDesc.MipLevels = 1;
@@ -1044,12 +1064,15 @@ void* FD3D11Texture::Lock(FD3D11DynamicRHI* D3DRHI, uint32 MipIndex, uint32 Arra
 		StagingTextureDesc.BindFlags = 0;
 		StagingTextureDesc.CPUAccessFlags = D3D11_CPU_ACCESS_READ;
 		StagingTextureDesc.MiscFlags = 0;
+
+		const FString StagingTextureName = FString::Printf(TEXT("%s_Staging"), *GetName().ToString());
+
 		TRefCountPtr<ID3D11Texture2D> StagingTexture;
 		VERIFYD3D11CREATETEXTURERESULT(
-			D3DRHI->GetDevice()->CreateTexture2D(&StagingTextureDesc,NULL,StagingTexture.GetInitReference()),
-			GetDesc().Format,
-			this->GetSizeX(),
-			this->GetSizeY(),
+			D3DRHI->GetDevice()->CreateTexture2D(&StagingTextureDesc, NULL, StagingTexture.GetInitReference()),
+			Desc.Format,
+			Desc.Extent.X,
+			Desc.Extent.Y,
 			this->GetSizeZ(),
 			StagingTextureDesc.Format,
 			1,
@@ -1063,18 +1086,19 @@ void* FD3D11Texture::Lock(FD3D11DynamicRHI* D3DRHI, uint32 MipIndex, uint32 Arra
 			0,
 			0,
 			D3DRHI->GetDevice(),
-			*FString::Printf(TEXT("%s_Staging"), *GetName().ToString())
-			);
+			*StagingTextureName
+		);
 		LockedData.StagingResource = StagingTexture;
 
 		// Copy the mip-map data from the real resource into the staging resource
-		D3DRHI->GetDeviceContext()->CopySubresourceRegion(StagingTexture,0,0,0,0,GetResource(),Subresource,NULL);
+		D3DRHI->GetDeviceContext()->CopySubresourceRegion(StagingTexture, 0, 0, 0, 0, GetResource(), Subresource, NULL);
 
 		// Map the staging resource, and return the mapped address.
-		D3D11_MAPPED_SUBRESOURCE MappedTexture;
-		VERIFYD3D11RESULT_EX(D3DRHI->GetDeviceContext()->Map(StagingTexture,0,D3D11_MAP_READ,0,&MappedTexture), D3DRHI->GetDevice());
+		D3D11_MAPPED_SUBRESOURCE MappedTexture{};
+		VERIFYD3D11RESULT_EX(D3DRHI->GetDeviceContext()->Map(StagingTexture, 0, D3D11_MAP_READ, 0, &MappedTexture), D3DRHI->GetDevice());
+
 		LockedData.SetData(MappedTexture.pData);
-		LockedData.Pitch = DestStride = MappedTexture.RowPitch;
+		LockedData.Pitch = MappedTexture.RowPitch;
 	}
 
 	// Add the lock to the outstanding lock list.
@@ -1090,23 +1114,30 @@ void* FD3D11Texture::Lock(FD3D11DynamicRHI* D3DRHI, uint32 MipIndex, uint32 Arra
 		});
 	}
 
-	return (void*)LockedData.GetData();
+	Result.Data = LockedData.GetData();
+	Result.ByteCount = MipBytes;
+	Result.Stride = LockedData.Pitch;
+
+	return Result;
 }
 
-void FD3D11Texture::Unlock(FD3D11DynamicRHI* D3DRHI, uint32 MipIndex, uint32 ArrayIndex)
+void FD3D11Texture::Unlock(FD3D11DynamicRHI* D3DRHI, const FRHILockTextureArgs& Arguments)
 {
 	check(!IsTexture3D()); // Only 2D texture locks are implemented
 
 	SCOPE_CYCLE_COUNTER(STAT_D3D11UnlockTextureTime);
 
+	const FRHITextureDesc& Desc = this->GetDesc();
+	const uint32 ArrayIndex = GetD3D11ArrayIndex(Arguments, Desc);
+
 	// Calculate the subresource index corresponding to the specified mip-map.
-	const uint32 Subresource = D3D11CalcSubresource(MipIndex,ArrayIndex,this->GetNumMips());
+	const uint32 Subresource = D3D11CalcSubresource(Arguments.MipIndex, ArrayIndex, Desc.NumMips);
 
 	// Find the object that is tracking this lock and remove it from outstanding list
 	FD3D11LockedData LockedData;
 	verifyf(D3DRHI->RemoveLockedData(FD3D11LockedKey(GetResource(), Subresource), LockedData), TEXT("Texture is not locked"));
 
-	if (!LockedData.bLockDeferred && EnumHasAnyFlags(GetDesc().Flags, TexCreate_CPUWritable))
+	if (!LockedData.bLockDeferred && EnumHasAnyFlags(Desc.Flags, TexCreate_CPUWritable))
 	{
 		D3DRHI->GetDeviceContext()->Unmap(GetResource(), 0);
 	}
@@ -1122,83 +1153,44 @@ void FD3D11Texture::Unlock(FD3D11DynamicRHI* D3DRHI, uint32 MipIndex, uint32 Arr
 	}
 }
 
-void* FD3D11DynamicRHI::RHILockTexture2D(FRHITexture* TextureRHI, uint32 MipIndex, EResourceLockMode LockMode, uint32& DestStride, bool bLockWithinMiptail, uint64* OutLockedByteCount)
-{
-	check(TextureRHI);
-	FD3D11Texture* Texture = ResourceCast(TextureRHI);
-	ConditionalClearShaderResource(Texture, false);
-	return Texture->Lock(this, MipIndex, 0, LockMode, DestStride, false /* bForceLockDeferred */, OutLockedByteCount);
-}
-
-void* FD3D11DynamicRHI::LockTexture2D_RenderThread(
-	class FRHICommandListImmediate& RHICmdList,
-	FRHITexture* Texture,
-	uint32 MipIndex,
-	EResourceLockMode LockMode,
-	uint32& DestStride,
-	bool bLockWithinMiptail,
-	bool bNeedsDefaultRHIFlush,
-	uint64* OutLockedByteCount)
+FRHILockTextureResult FD3D11DynamicRHI::RHILockTexture(FRHICommandListImmediate& RHICmdList, const FRHILockTextureArgs& Arguments)
 {
-	void *LockedTexture = nullptr;
+	FD3D11Texture* Texture = ResourceCast(Arguments.Texture);
 
 	if (ShouldNotEnqueueRHICommand())
 	{
-		LockedTexture = RHILockTexture2D(Texture, MipIndex, LockMode, DestStride, bLockWithinMiptail, OutLockedByteCount);
+		ConditionalClearShaderResource(Texture, false);
+		return Texture->Lock(this, Arguments, false /* bForceLockDeferred */);
 	}
-	else if (LockMode == RLM_ReadOnly)
+
+	if (Arguments.LockMode == RLM_ReadOnly)
 	{
 		RHICmdList.ImmediateFlush(EImmediateFlushType::FlushRHIThread);
-		LockedTexture = RHILockTexture2D(Texture, MipIndex, LockMode, DestStride, bLockWithinMiptail, OutLockedByteCount);
-	}
-	else
-	{
-		FD3D11Texture* TextureD3D11 = ResourceCast(Texture);
-		LockedTexture = TextureD3D11->Lock(this, MipIndex, 0, LockMode, DestStride, true, OutLockedByteCount);
+
+		ConditionalClearShaderResource(Texture, false);
+		return Texture->Lock(this, Arguments, false /* bForceLockDeferred */);
 	}
-	return LockedTexture;
-}
 
-void FD3D11DynamicRHI::RHIUnlockTexture2D(FRHITexture* TextureRHI, uint32 MipIndex, bool bLockWithinMiptail)
-{
-	check(TextureRHI);
-	FD3D11Texture* Texture = ResourceCast(TextureRHI);
-	Texture->Unlock(this, MipIndex, 0);
+	return Texture->Lock(this, Arguments, true /* bForceLockDeferred */);
 }
 
-void FD3D11DynamicRHI::UnlockTexture2D_RenderThread(
-	class FRHICommandListImmediate& RHICmdList,
-	FRHITexture* Texture,
-	uint32 MipIndex,
-	bool bLockWithinMiptail,
-	bool bNeedsDefaultRHIFlush)
+void FD3D11DynamicRHI::RHIUnlockTexture(FRHICommandListImmediate& RHICmdList, const FRHILockTextureArgs& Arguments)
 {
+	FD3D11Texture* Texture = ResourceCast(Arguments.Texture);
+
 	if (ShouldNotEnqueueRHICommand())
 	{
-		RHIUnlockTexture2D(Texture, MipIndex, bLockWithinMiptail);
+		Texture->Unlock(this, Arguments);
 	}
 	else
 	{
-		RunOnRHIThread([this, Texture, MipIndex, bLockWithinMiptail]()
+		RunOnRHIThread([this, Texture, Arguments]()
 		{
-			RHIUnlockTexture2D(Texture, MipIndex, bLockWithinMiptail);
+			Texture->Unlock(this, Arguments);
 		});
 	}
 }
 
-void* FD3D11DynamicRHI::RHILockTexture2DArray(FRHITexture* TextureRHI, uint32 TextureIndex, uint32 MipIndex, EResourceLockMode LockMode, uint32& DestStride, bool bLockWithinMiptail)
-{
-	FD3D11Texture* Texture = ResourceCast(TextureRHI);
-	ConditionalClearShaderResource(Texture, false);
-	return Texture->Lock(this, MipIndex, TextureIndex, LockMode, DestStride);
-}
-
-void FD3D11DynamicRHI::RHIUnlockTexture2DArray(FRHITexture* TextureRHI, uint32 TextureIndex, uint32 MipIndex, bool bLockWithinMiptail)
-{
-	FD3D11Texture* Texture = ResourceCast(TextureRHI);
-	Texture->Unlock(this, MipIndex, TextureIndex);
-}
-
 void FD3D11DynamicRHI::RHIUpdateTexture2D(FRHICommandListBase& RHICmdList, FRHITexture* TextureRHI, uint32 MipIndex, const FUpdateTextureRegion2D& UpdateRegion, uint32 SourcePitch, const uint8* SourceData)
 {
 	const FPixelFormatInfo& FormatInfo = GPixelFormats[TextureRHI->GetFormat()];
@@ -1299,72 +1291,6 @@ void FD3D11DynamicRHI::RHIEndUpdateTexture3D(FRHICommandListBase& RHICmdList, FU
 	Cubemap texture support.
 -----------------------------------------------------------------------------*/
 
-void* FD3D11DynamicRHI::RHILockTextureCubeFace(FRHITexture* TextureCubeRHI, uint32 FaceIndex, uint32 ArrayIndex, uint32 MipIndex, EResourceLockMode LockMode, uint32& DestStride, bool bLockWithinMiptail)
-{
-	FD3D11Texture* TextureCube = ResourceCast(TextureCubeRHI);
-	ConditionalClearShaderResource(TextureCube, false);
-	uint32 D3DFace = GetD3D11CubeFace((ECubeFace)FaceIndex);
-	return TextureCube->Lock(this, MipIndex, D3DFace + ArrayIndex * 6, LockMode, DestStride);
-}
-
-void* FD3D11DynamicRHI::RHILockTextureCubeFace_RenderThread(
-	class FRHICommandListImmediate& RHICmdList,
-	FRHITexture* Texture,
-	uint32 FaceIndex,
-	uint32 ArrayIndex,
-	uint32 MipIndex,
-	EResourceLockMode LockMode,
-	uint32& DestStride,
-	bool bLockWithinMiptail)
-{
-	void *LockedTexture = nullptr;
-
-	if (ShouldNotEnqueueRHICommand())
-	{
-		LockedTexture = RHILockTextureCubeFace(Texture, FaceIndex, ArrayIndex, MipIndex, LockMode, DestStride, bLockWithinMiptail);
-	}
-	else if (LockMode == RLM_ReadOnly)
-	{
-		RHICmdList.ImmediateFlush(EImmediateFlushType::FlushRHIThread);
-		LockedTexture = RHILockTextureCubeFace(Texture, FaceIndex, ArrayIndex, MipIndex, LockMode, DestStride, bLockWithinMiptail);
-	}
-	else
-	{
-		FD3D11Texture* TextureCube = ResourceCast(Texture);
-		const uint32 D3DFace = GetD3D11CubeFace((ECubeFace)FaceIndex);
-		LockedTexture = TextureCube->Lock(this, MipIndex, D3DFace + ArrayIndex * 6, LockMode, DestStride, true);
-	}
-	return LockedTexture;
-}
-
-void FD3D11DynamicRHI::RHIUnlockTextureCubeFace(FRHITexture* TextureCubeRHI, uint32 FaceIndex, uint32 ArrayIndex, uint32 MipIndex, bool bLockWithinMiptail)
-{
-	FD3D11Texture* TextureCube = ResourceCast(TextureCubeRHI);
-	uint32 D3DFace = GetD3D11CubeFace((ECubeFace)FaceIndex);
-	TextureCube->Unlock(this, MipIndex, D3DFace + ArrayIndex * 6);
-}
-
-void FD3D11DynamicRHI::RHIUnlockTextureCubeFace_RenderThread(
-	class FRHICommandListImmediate& RHICmdList,
-	FRHITexture* Texture,
-	uint32 FaceIndex,
-	uint32 ArrayIndex,
-	uint32 MipIndex,
-	bool bLockWithinMiptail)
-{
-	if (ShouldNotEnqueueRHICommand())
-	{
-		RHIUnlockTextureCubeFace(Texture, FaceIndex, ArrayIndex, MipIndex, bLockWithinMiptail);
-	}
-	else
-	{
-		RunOnRHIThread([this, Texture, FaceIndex, ArrayIndex, MipIndex, bLockWithinMiptail]()
-		{
-			RHIUnlockTextureCubeFace(Texture, FaceIndex, ArrayIndex, MipIndex, bLockWithinMiptail);
-		});
-	}
-}
-
 void FD3D11DynamicRHI::RHIBindDebugLabelName(FRHICommandListBase& RHICmdList, FRHITexture* TextureRHI, const TCHAR* Name)
 {
 	FD3D11Texture* Texture = ResourceCast(TextureRHI);
diff --git a/Engine/Source/Runtime/Windows/D3D11RHI/Public/D3D11Resources.h b/Engine/Source/Runtime/Windows/D3D11RHI/Public/D3D11Resources.h
--- a/Engine/Source/Runtime/Windows/D3D11RHI/Public/D3D11Resources.h
+++ b/Engine/Source/Runtime/Windows/D3D11RHI/Public/D3D11Resources.h
@@ -10,6 +10,7 @@
 #include "ShaderCore.h"
 #include "Windows/D3D11ThirdParty.h"
 
+class FD3D11DynamicRHI;
 typedef ID3D11DeviceContext FD3D11DeviceContext;
 
 template <>
@@ -278,10 +279,10 @@ public:
 	* Locks one of the texture's mip-maps.
 	* @return A pointer to the specified texture data.
 	*/
-	D3D11RHI_API void* Lock(class FD3D11DynamicRHI* D3DRHI, uint32 MipIndex, uint32 ArrayIndex, EResourceLockMode LockMode, uint32& DestStride, bool bForceLockDeferred = false, uint64* OutLockedByteCount = nullptr);
+	FRHILockTextureResult Lock(FD3D11DynamicRHI* D3DRHI, const FRHILockTextureArgs& Arguments, bool bForceLockDeferred);
 
 	/** Unlocks a previously locked mip-map. */
-	D3D11RHI_API void Unlock(class FD3D11DynamicRHI* D3DRHI, uint32 MipIndex, uint32 ArrayIndex);
+	void Unlock(FD3D11DynamicRHI* D3DRHI, const FRHILockTextureArgs& Arguments);
 
 private:
 	//Resource handle for use by IHVs for SLI and other purposes.
