Set setPurgeableState Empty when releasing MTL Buffer

diff --git a/Engine/Source/Runtime/Apple/MetalRHI/Private/MetalBuffer.cpp b/Engine/Source/Runtime/Apple/MetalRHI/Private/MetalBuffer.cpp
--- a/Engine/Source/Runtime/Apple/MetalRHI/Private/MetalBuffer.cpp
+++ b/Engine/Source/Runtime/Apple/MetalRHI/Private/MetalBuffer.cpp
@@ -41,6 +41,12 @@ static FAutoConsoleVariableRef CVarMetalMaxBucketSize(
 	GMetalMaxBucketSize,
 	TEXT("Set the Max Bucket Size for the Pooled Buffer. (Default: 33547255)"));
 
+static int32 GMetalPurgeableStateEmptyOnRelease = 1;
+static FAutoConsoleVariableRef CVarMetalPurgeableStateEmptyOnRelease(
+	TEXT("rhi.Metal.PurgeableStateEmptyOnRelease"),
+	GMetalPurgeableStateEmptyOnRelease,
+	TEXT("Set the PurgeableState to PurgeableStateEmpty when releasing a Buffer  (Default: 1, On)"));
+
 
 #if METAL_DEBUG_OPTIONS
 extern int32 GMetalBufferScribble;
@@ -85,7 +91,10 @@ void FMetalBuffer::Release()
 		
 		DEC_MEMORY_STAT_BY(STAT_MetalBufferMemory, GetLength());
 		DEC_DWORD_STAT(STAT_MetalBufferCount);
-		
+		if (GMetalPurgeableStateEmptyOnRelease)
+		{
+			Buffer->setPurgeableState(MTL::PurgeableStateEmpty);
+		}
 		Buffer->release();
 	}
 }
@@ -151,6 +160,10 @@ FMetalSubBufferHeap::~FMetalSubBufferHeap()
 	
 	if(ParentBuffer)
 	{
+		if (GMetalPurgeableStateEmptyOnRelease)
+		{
+			ParentBuffer->setPurgeableState(MTL::PurgeableStateEmpty);
+		}
 		ParentBuffer->release();
 		ParentBuffer = nullptr;
 	}
@@ -263,6 +276,10 @@ void FMetalSubBufferHeap::ReleaseBuffer(FMetalBuffer* Buffer)
 {
 	if(ParentHeap)
 	{
+		if (GMetalPurgeableStateEmptyOnRelease)
+		{
+			Buffer->GetMTLBuffer()->setPurgeableState(MTL::PurgeableStateEmpty);
+		}
 		Buffer->GetMTLBuffer()->release();
 	}
 	FreeRange(Buffer->GetRange());
@@ -514,6 +531,10 @@ FMetalSubBufferLinear::~FMetalSubBufferLinear()
 	
 	if(ParentBuffer)
 	{
+		if (GMetalPurgeableStateEmptyOnRelease)
+		{
+			ParentBuffer->setPurgeableState(MTL::PurgeableStateEmpty);
+		}
 		ParentBuffer->release();
 		ParentBuffer = nullptr;
 	}
@@ -701,6 +722,10 @@ FMetalSubBufferMagazine::~FMetalSubBufferMagazine()
 		
 		if(ParentBuffer)
 		{
+			if (GMetalPurgeableStateEmptyOnRelease)
+			{
+				ParentBuffer->setPurgeableState(MTL::PurgeableStateEmpty);
+			}
 			ParentBuffer->release();
 			ParentBuffer = nullptr;
 		}
@@ -739,6 +764,10 @@ void FMetalSubBufferMagazine::ReleaseBuffer(FMetalBuffer* Buffer)
 {
 	if(ParentHeap)
 	{
+		if (GMetalPurgeableStateEmptyOnRelease)
+		{
+			Buffer->GetMTLBuffer()->setPurgeableState(MTL::PurgeableStateEmpty);
+		}
 		Buffer->GetMTLBuffer()->release();
 	}
 	FreeRange(Buffer->GetRange());
@@ -1213,6 +1242,10 @@ void FMetalBufferPoolPolicyData::FreeResource(MTL::Buffer* Resource)
 	DEC_MEMORY_STAT_BY(STAT_MetalPooledBufferUnusedMemory, Resource->length());
 	
 	CreationArgumentMap.Remove(Resource);
+	if (GMetalPurgeableStateEmptyOnRelease)
+	{
+		Resource->setPurgeableState(MTL::PurgeableStateEmpty);
+	}
 	Resource->release();
 }
 
@@ -1919,7 +1952,10 @@ void FMetalResourceHeap::ReleaseBuffer(FMetalBuffer* Buffer)
 	{
 		FreeBlock(*It);
 		AllocationHandlesLUT.Remove(MTLBuffer);
-		
+		if (GMetalPurgeableStateEmptyOnRelease)
+		{
+			MTLBuffer->setPurgeableState(MTL::PurgeableStateEmpty);
+		}
 		MTLBuffer->release();
 	}
 	else
