Fixed bug causing water flickering where buffer backed textures backing resource was freed whilst the texture was still in use.

diff --git a/Engine/Source/Runtime/Apple/MetalRHI/Private/MetalTexture.cpp b/Engine/Source/Runtime/Apple/MetalRHI/Private/MetalTexture.cpp
--- a/Engine/Source/Runtime/Apple/MetalRHI/Private/MetalTexture.cpp
+++ b/Engine/Source/Runtime/Apple/MetalRHI/Private/MetalTexture.cpp
@@ -655,11 +655,9 @@ FMetalSurface::FMetalSurface(FMetalDevice& MetalDevice, FRHICommandListBase* RHI
 
 			// Backing buffer resource options must match the texture we are going to create from it
 			FMetalPooledBufferArgs Args(&Device, BytesPerRow * NewCreateDesc.Desc->height(), BUF_Dynamic, NewCreateDesc.Desc->storageMode(), NewCreateDesc.Desc->cpuCacheMode());
-			FMetalBufferPtr Buffer = Device.CreatePooledBuffer(Args);
+			BackingBuffer = Device.CreatePooledBuffer(Args);
 
-			Texture = NS::TransferPtr(Buffer->GetMTLBuffer()->newTexture(NewCreateDesc.Desc.get(), Buffer->GetOffset(), BytesPerRow));
-            
-			FMetalDynamicRHI::Get().DeferredDelete(Buffer);
+			Texture = NS::TransferPtr(BackingBuffer->GetMTLBuffer()->newTexture(NewCreateDesc.Desc.get(), BackingBuffer->GetOffset(), BytesPerRow));
 		}
         else if (bTextureArrayWithAtomics)
         {
@@ -672,14 +670,12 @@ FMetalSurface::FMetalSurface(FMetalDevice& MetalDevice, FRHICommandListBase* RHI
 
             // Backing buffer resource options must match the texture we are going to create from it
             FMetalPooledBufferArgs Args(&Device, BytesPerRow * NewCreateDesc.Desc->height(), BUF_Dynamic, NewCreateDesc.Desc->storageMode(), NewCreateDesc.Desc->cpuCacheMode());
-            FMetalBufferPtr Buffer = Device.CreatePooledBuffer(Args);
+			BackingBuffer = Device.CreatePooledBuffer(Args);
 
             NewCreateDesc.Desc->setWidth(NewCreateDesc.Desc->width() * NewCreateDesc.Desc->arrayLength());
             NewCreateDesc.Desc->setArrayLength(1);
             NewCreateDesc.Desc->setTextureType(MTL::TextureType2D);
-            Texture = NS::TransferPtr(Buffer->GetMTLBuffer()->newTexture(NewCreateDesc.Desc.get(), Buffer->GetOffset(), BytesPerRow));
-            
-			FMetalDynamicRHI::Get().DeferredDelete(Buffer);
+            Texture = NS::TransferPtr(BackingBuffer->GetMTLBuffer()->newTexture(NewCreateDesc.Desc.get(), BackingBuffer->GetOffset(), BytesPerRow));
         }
 		else
 		{
@@ -848,6 +844,11 @@ public:
 
 FMetalSurface::~FMetalSurface()
 {
+	if (BackingBuffer) 
+	{
+		FMetalDynamicRHI::Get().DeferredDelete(BackingBuffer);	
+	}
+	
 	if (MSAATexture)
 	{
 		if (Texture != MSAATexture)
diff --git a/Engine/Source/Runtime/Apple/MetalRHI/Public/MetalResources.h b/Engine/Source/Runtime/Apple/MetalRHI/Public/MetalResources.h
--- a/Engine/Source/Runtime/Apple/MetalRHI/Public/MetalResources.h
+++ b/Engine/Source/Runtime/Apple/MetalRHI/Public/MetalResources.h
@@ -239,6 +239,9 @@ public:
 	// how much memory is allocated for this texture
 	uint64 TotalTextureSize;
 	
+	// Used for atomics
+	FMetalBufferPtr BackingBuffer;
+	
 	// For back-buffers, the owning viewport.
 	class FMetalViewport* Viewport;
 
