Metal RHI

- Removed LLM alloc tag that was tracking allocations within already allocated buffers which skewed numbers
- Added RHI buffer tracking stats
- Removed free tracking from obj-c callbacks as there are bugs when running with GPU Capture
- Now purge all deleted objects as the GPU capture system can hang onto them causing allocations to build up, we can turn this off when Apple fix this issue.
- Changed allocation of MetalRHIBuffer to on demand, reduced max by 25% and min by 8%

diff --git a/Engine/Source/Runtime/Apple/MetalRHI/Private/MetalBuffer.cpp b/Engine/Source/Runtime/Apple/MetalRHI/Private/MetalBuffer.cpp
--- a/Engine/Source/Runtime/Apple/MetalRHI/Private/MetalBuffer.cpp
+++ b/Engine/Source/Runtime/Apple/MetalRHI/Private/MetalBuffer.cpp
@@ -41,17 +41,12 @@ static FAutoConsoleVariableRef CVarMetalMaxBucketSize(
 	GMetalMaxBucketSize,
 	TEXT("Set the Max Bucket Size for the Pooled Buffer. (Default: 33547255)"));
 
-static int32 GMetalPurgeableStateEmptyOnRelease = 1;
-static FAutoConsoleVariableRef CVarMetalPurgeableStateEmptyOnRelease(
-	TEXT("rhi.Metal.PurgeableStateEmptyOnRelease"),
-	GMetalPurgeableStateEmptyOnRelease,
-	TEXT("Set the PurgeableState to PurgeableStateEmpty when releasing a Buffer  (Default: 1, On)"));
-
-
 #if METAL_DEBUG_OPTIONS
 extern int32 GMetalBufferScribble;
 #endif
 
+extern int32 GMetalResourcePurgeOnDelete;
+
 FMetalBuffer::~FMetalBuffer()
 {
     Release();
@@ -63,6 +58,11 @@ FMetalBuffer::FMetalBuffer(MTL::Buffer* Handle, FreePolicy InFreePolicy)
 , SubRange(0, Handle->length())
 , OnFreePolicy(InFreePolicy)
 {
+	if(OnFreePolicy == FreePolicy::Owner)
+	{
+		INC_MEMORY_STAT_BY(STAT_MetalDeviceBufferMemory, GetLength());
+		MetalLLM::LogAllocBufferNative(Buffer);
+	}
 }
 
 FMetalBuffer::FMetalBuffer(MTL::Buffer* Handle, NS::Range Range, IMetalBufferAllocator* InAllocator)
@@ -85,13 +85,10 @@ void FMetalBuffer::Release()
 	}
 	else if(OnFreePolicy == FreePolicy::Owner)
 	{
-		LLM_PLATFORM_SCOPE_METAL(ELLMTagMetal::Buffers);
-		
-		LLM_IF_ENABLED(FLowLevelMemTracker::Get().OnLowLevelFree(ELLMTracker::Platform, this, ELLMAllocType::System));
+		MetalLLM::LogFreeBufferNative(Buffer);
 		
-		DEC_MEMORY_STAT_BY(STAT_MetalBufferMemory, GetLength());
-		DEC_DWORD_STAT(STAT_MetalBufferCount);
-		if (GMetalPurgeableStateEmptyOnRelease)
+		DEC_MEMORY_STAT_BY(STAT_MetalDeviceBufferMemory, GetLength());
+		if (GMetalResourcePurgeOnDelete)
 		{
 			Buffer->setPurgeableState(MTL::PurgeableStateEmpty);
 		}
@@ -160,7 +157,9 @@ FMetalSubBufferHeap::~FMetalSubBufferHeap()
 	
 	if(ParentBuffer)
 	{
-		if (GMetalPurgeableStateEmptyOnRelease)
+		MetalLLM::LogFreeBufferNative(ParentBuffer);
+		
+		if (GMetalResourcePurgeOnDelete)
 		{
 			ParentBuffer->setPurgeableState(MTL::PurgeableStateEmpty);
 		}
@@ -276,7 +275,7 @@ void FMetalSubBufferHeap::ReleaseBuffer(FMetalBuffer* Buffer)
 {
 	if(ParentHeap)
 	{
-		if (GMetalPurgeableStateEmptyOnRelease)
+		if (GMetalResourcePurgeOnDelete)
 		{
 			Buffer->GetMTLBuffer()->setPurgeableState(MTL::PurgeableStateEmpty);
 		}
@@ -422,9 +421,6 @@ FMetalBufferPtr FMetalSubBufferHeap::NewBuffer(NS::UInteger length)
 		MTL::ResourceOptions Opt = MTL::ResourceOptions(Storage | Cache);
 		
 		Result = FMetalBufferPtr(new FMetalBuffer(ParentHeap->newBuffer(Size, Opt), NS::Range(0, Size), this));
-#if STATS || ENABLE_LOW_LEVEL_MEM_TRACKER
-		MetalLLM::LogAllocBufferNative(Result->GetMTLBuffer());
-#endif
 		DEC_MEMORY_STAT_BY(STAT_MetalBufferUnusedMemory, Result->GetLength());
 		DEC_MEMORY_STAT_BY(STAT_MetalHeapBufferUnusedMemory, Result->GetLength());
 		INC_MEMORY_STAT_BY(STAT_MetalHeapBufferMemory, Result->GetLength());
@@ -531,7 +527,9 @@ FMetalSubBufferLinear::~FMetalSubBufferLinear()
 	
 	if(ParentBuffer)
 	{
-		if (GMetalPurgeableStateEmptyOnRelease)
+		MetalLLM::LogFreeBufferNative(ParentBuffer);
+		
+		if (GMetalResourcePurgeOnDelete)
 		{
 			ParentBuffer->setPurgeableState(MTL::PurgeableStateEmpty);
 		}
@@ -722,7 +720,9 @@ FMetalSubBufferMagazine::~FMetalSubBufferMagazine()
 		
 		if(ParentBuffer)
 		{
-			if (GMetalPurgeableStateEmptyOnRelease)
+			MetalLLM::LogFreeBufferNative(ParentBuffer);
+			
+			if (GMetalResourcePurgeOnDelete)
 			{
 				ParentBuffer->setPurgeableState(MTL::PurgeableStateEmpty);
 			}
@@ -764,7 +764,7 @@ void FMetalSubBufferMagazine::ReleaseBuffer(FMetalBuffer* Buffer)
 {
 	if(ParentHeap)
 	{
-		if (GMetalPurgeableStateEmptyOnRelease)
+		if (GMetalResourcePurgeOnDelete)
 		{
 			Buffer->GetMTLBuffer()->setPurgeableState(MTL::PurgeableStateEmpty);
 		}
@@ -891,9 +891,6 @@ FMetalBufferPtr FMetalSubBufferMagazine::NewBuffer()
 		MTL::ResourceOptions Opt = MTL::ResourceOptions(Storage | Cache);
 		
 		Result = FMetalBufferPtr(new FMetalBuffer(ParentHeap->newBuffer(Size, Opt), NS::Range(0, Size), this));
-#if STATS || ENABLE_LOW_LEVEL_MEM_TRACKER
-		MetalLLM::LogAllocBuffer(Result);
-#endif
 		DEC_MEMORY_STAT_BY(STAT_MetalBufferUnusedMemory, Result->GetLength());
 		DEC_MEMORY_STAT_BY(STAT_MetalMagazineBufferUnusedMemory, Result->GetLength());
 		INC_MEMORY_STAT_BY(STAT_MetalMagazineBufferMemory, Result->GetLength());
@@ -1242,7 +1239,9 @@ void FMetalBufferPoolPolicyData::FreeResource(MTL::Buffer* Resource)
 	DEC_MEMORY_STAT_BY(STAT_MetalPooledBufferUnusedMemory, Resource->length());
 	
 	CreationArgumentMap.Remove(Resource);
-	if (GMetalPurgeableStateEmptyOnRelease)
+	
+	MetalLLM::LogFreeBufferNative(Resource);
+	if (GMetalResourcePurgeOnDelete)
 	{
 		Resource->setPurgeableState(MTL::PurgeableStateEmpty);
 	}
@@ -1294,7 +1293,7 @@ MTLTexturePtr FMetalTexturePool::CreateTexture(FMetalDevice& Device, MTL::Textur
 		METAL_GPUPROFILE(FScopedMetalCPUStats CPUStat(FString::Printf(TEXT("AllocTexture: %s"), TEXT("")/**FString([Desc.GetPtr() description])*/)));
         Texture = NS::TransferPtr(Device.GetDevice()->newTexture(Desc));
 #if STATS || ENABLE_LOW_LEVEL_MEM_TRACKER
-		MetalLLM::LogAllocTexture(Device, Desc, Texture.get());
+		MetalLLM::LogAllocTexture(Texture.get());
 #endif
 	}
 	return Texture;
@@ -1321,33 +1320,42 @@ void FMetalTexturePool::ReleaseTexture(MTLTexturePtr Texture)
     }
 	
 	FScopeLock Lock(&PoolMutex);
-	Pool.Add(Descriptor, Texture);
-}
 	
-void FMetalTexturePool::Drain(bool const bForce)
-{
-	FScopeLock Lock(&PoolMutex);
-	if (bForce)
+	if(Pool.Find(Descriptor))
 	{
-		Pool.Empty();
+		MetalLLM::LogFreeTexture(Texture.get());
+		if (GMetalResourcePurgeOnDelete)
+		{
+			Texture->setPurgeableState(MTL::PurgeableStateEmpty);
+		}
 	}
 	else
 	{
-		for (auto It = Pool.CreateIterator(); It; ++It)
+		Pool.Add(Descriptor, Texture);
+	}
+}
+	
+void FMetalTexturePool::Drain(bool const bForce)
+{
+	for (auto It = Pool.CreateIterator(); It; ++It)
+	{
+		if (bForce || (GFrameNumberRenderThread - It->Key.freedFrame) >= CullAfterNumFrames)
 		{
-			if ((GFrameNumberRenderThread - It->Key.freedFrame) >= CullAfterNumFrames)
+			MetalLLM::LogFreeTexture(It->Value.get());
+			if (GMetalResourcePurgeOnDelete)
 			{
-				It.RemoveCurrent();
+				It->Value->setPurgeableState(MTL::PurgeableStateEmpty);
 			}
-            else
-            {
-				if (GMetalResourcePurgeInPool && (GFrameNumberRenderThread - It->Key.freedFrame) >= PurgeAfterNumFrames)
-				{
-					It->Value->setPurgeableState(MTL::PurgeableStateEmpty);
-				}
-            }
-        }
-    }
+			It.RemoveCurrent();
+		}
+		else
+		{
+			if (GMetalResourcePurgeInPool && (GFrameNumberRenderThread - It->Key.freedFrame) >= PurgeAfterNumFrames)
+			{
+				It->Value->setPurgeableState(MTL::PurgeableStateEmpty);
+			}
+		}
+	}
 }
 
 FMetalResourceHeap::FMetalResourceHeap(FMetalDevice& MetalDevice)
@@ -1737,9 +1745,6 @@ TDoubleLinkedList<FMetalResourceHeap::MemoryBlock>::TIterator FMetalResourceHeap
 
 FMetalBufferPtr FMetalResourceHeap::CreateBuffer(uint32 Size, uint32 Alignment, EBufferUsageFlags Flags, MTL::ResourceOptions Options, bool bForceUnique)
 {
-	LLM_SCOPE_METAL(ELLMTagMetal::Buffers);
-	LLM_PLATFORM_SCOPE_METAL(ELLMTagMetal::Buffers);
-	
 	static bool bSupportsHeaps = Device.SupportsFeature(EMetalFeaturesHeaps);
 	if (bSupportsHeaps)
 	{
@@ -1928,11 +1933,6 @@ FMetalBufferPtr FMetalResourceHeap::CreateBuffer(uint32 Size, uint32 Alignment,
 	{
 		METAL_GPUPROFILE(FScopedMetalCPUStats CPUStat(FString::Printf(TEXT("AllocBuffer: %llu, %llu"), BlockSize, Options)));
         Buffer = FMetalBufferPtr(new FMetalBuffer(Device.GetDevice()->newBuffer(BlockSize, Options), FMetalBuffer::FreePolicy::Owner));
-        
-#if STATS || ENABLE_LOW_LEVEL_MEM_TRACKER
-		MetalLLM::LogAllocBuffer(Buffer);
-#endif
-		INC_MEMORY_STAT_BY(STAT_MetalDeviceBufferMemory, Buffer->GetLength());
 	}
 	
     METAL_DEBUG_OPTION(Device.ValidateIsInactiveBuffer(Buffer->GetMTLBuffer(), Buffer->GetRange()));
@@ -1952,7 +1952,7 @@ void FMetalResourceHeap::ReleaseBuffer(FMetalBuffer* Buffer)
 	{
 		FreeBlock(*It);
 		AllocationHandlesLUT.Remove(MTLBuffer);
-		if (GMetalPurgeableStateEmptyOnRelease)
+		if (GMetalResourcePurgeOnDelete)
 		{
 			MTLBuffer->setPurgeableState(MTL::PurgeableStateEmpty);
 		}
@@ -2044,13 +2044,12 @@ MTLTexturePtr FMetalResourceHeap::CreateTexture(MTL::TextureDescriptor* Desc, FM
 		METAL_GPUPROFILE(FScopedMetalCPUStats CPUStat(FString::Printf(TEXT("AllocTexture: %s"), TEXT("")/**FString([Desc.GetPtr() description])*/)));
 		MTLTexturePtr Texture = NS::TransferPtr(Heap->newTexture(Desc));
 #if STATS || ENABLE_LOW_LEVEL_MEM_TRACKER
-		MetalLLM::LogAllocTexture(Device, Desc, Texture.get());
+		MetalLLM::LogAllocTexture(Texture.get());
 #endif
 		return Texture;
 	}
 	else if (Desc->usage() & MTL::TextureUsageRenderTarget)
 	{
-		LLM_PLATFORM_SCOPE_METAL(ELLMTagMetal::RenderTargets);
 		return TargetPool.CreateTexture(Device, Desc);
 	}
 	else
diff --git a/Engine/Source/Runtime/Apple/MetalRHI/Private/MetalDevice.cpp b/Engine/Source/Runtime/Apple/MetalRHI/Private/MetalDevice.cpp
--- a/Engine/Source/Runtime/Apple/MetalRHI/Private/MetalDevice.cpp
+++ b/Engine/Source/Runtime/Apple/MetalRHI/Private/MetalDevice.cpp
@@ -52,12 +52,6 @@ static FAutoConsoleVariableRef CVarMetalBufferScribble(
 	GMetalBufferScribble,
 	TEXT("Debug option: when enabled will scribble over the buffer contents with a single value when releasing buffer objects, or regions thereof. (Default: 0, Off)"));
 
-int32 GMetalResourcePurgeOnDelete = 0;
-static FAutoConsoleVariableRef CVarMetalResourcePurgeOnDelete(
-	TEXT("rhi.Metal.ResourcePurgeOnDelete"),
-	GMetalResourcePurgeOnDelete,
-	TEXT("Debug option: when enabled all MTLResource objects will have their backing stores purged on release - any subsequent access will be invalid and cause a command-buffer failure. Useful for making intermittent resource lifetime errors more common and easier to track. (Default: 0, Off)"));
-
 static int32 GMetalResourceDeferDeleteNumFrames = 0;
 static FAutoConsoleVariableRef CVarMetalResourceDeferDeleteNumFrames(
 	TEXT("rhi.Metal.ResourceDeferDeleteNumFrames"),
@@ -65,6 +59,12 @@ static FAutoConsoleVariableRef CVarMetalResourceDeferDeleteNumFrames(
 	TEXT("Debug option: set to the number of frames that must have passed before resource free-lists are processed and resources disposed of. (Default: 0, Off)"));
 #endif
 
+int32 GMetalResourcePurgeOnDelete = 1;
+static FAutoConsoleVariableRef CVarMetalResourcePurgeOnDelete(
+	TEXT("rhi.Metal.ResourcePurgeOnDelete"),
+	GMetalResourcePurgeOnDelete,
+	TEXT("When enabled all MTLResource objects will have their backing stores purged on release - any subsequent access will be invalid and cause a command-buffer failure. Useful for making intermittent resource lifetime errors more common and easier to track. (Default: 0, Off)"));
+
 #if UE_BUILD_SHIPPING
 int32 GMetalRuntimeDebugLevel = 0;
 #else
@@ -652,12 +652,10 @@ void FMetalDevice::GarbageCollect()
 MTLTexturePtr FMetalDevice::CreateTexture(FMetalSurface* Surface, MTL::TextureDescriptor* Descriptor)
 {
 	MTLTexturePtr Tex = Heap.CreateTexture(Descriptor, Surface);
-#if METAL_DEBUG_OPTIONS
 	if (GMetalResourcePurgeOnDelete && !Tex->heap())
 	{
 		Tex->setPurgeableState(MTL::PurgeableStateNonVolatile);
 	}
-#endif
 	
 	return Tex;
 }
@@ -685,13 +683,12 @@ FMetalBufferPtr FMetalDevice::CreatePooledBuffer(FMetalPooledBufferArgs const& A
     FMetalBufferPtr Buffer = Heap.CreateBuffer(Args.Size, RequestedBufferOffsetAlignment, Args.Flags, FMetalCommandQueue::GetCompatibleResourceOptions((MTL::ResourceOptions)(CpuResourceOption | HazardTrackingMode | ((NS::UInteger)Args.Storage << MTL::ResourceStorageModeShift))));
 	
     check(Buffer);
-#if METAL_DEBUG_OPTIONS
+
     MTL::Buffer* MTLBuffer = Buffer->GetMTLBuffer();
 	if (GMetalResourcePurgeOnDelete && !MTLBuffer->heap())
 	{
         MTLBuffer->setPurgeableState(MTL::PurgeableStateNonVolatile);
 	}
-#endif
 	
 	return Buffer;
 }
diff --git a/Engine/Source/Runtime/Apple/MetalRHI/Private/MetalLLM.cpp b/Engine/Source/Runtime/Apple/MetalRHI/Private/MetalLLM.cpp
--- a/Engine/Source/Runtime/Apple/MetalRHI/Private/MetalLLM.cpp
+++ b/Engine/Source/Runtime/Apple/MetalRHI/Private/MetalLLM.cpp
@@ -146,16 +146,10 @@ static EPixelFormat MetalToRHIPixelFormat(MTL::PixelFormat Format)
 	return PF_MAX;
 }
 
-void MetalLLM::LogAllocTexture(FMetalDevice& Device, MTL::TextureDescriptor* Desc, MTL::Texture* Texture)
+void MetalLLM::LogAllocTexture(MTL::Texture* Texture)
 {
-	MTL::SizeAndAlign SizeAlign;
-	if (Device.SupportsFeature(EMetalFeaturesGPUCaptureManager))
-	{
-		SizeAlign = Device.GetDevice()->heapTextureSizeAndAlign(Desc);
-	}
-	
+	uint64 Size = Texture->allocatedSize();
 	void* Ptr = (void*)Texture;
-	uint64 Size = SizeAlign.size;
 	
 #if PLATFORM_IOS
 	bool bMemoryless = (Texture->storageMode() == MTL::StorageModeMemoryless);
@@ -166,90 +160,66 @@ void MetalLLM::LogAllocTexture(FMetalDevice& Device, MTL::TextureDescriptor* Des
 	}
 	INC_DWORD_STAT(STAT_MetalTextureCount);
 	
-	LLM_IF_ENABLED(FLowLevelMemTracker::Get().OnLowLevelAlloc(ELLMTracker::Platform, Ptr, Size, ELLMTag::Untagged, ELLMAllocType::System));
-	// Assign a dealloc handler to untrack the memory - but don't track the dispatch block!
+	if (Texture->usage() & MTL::TextureUsageRenderTarget)
 	{
-		LLM_SCOPED_PAUSE_TRACKING(ELLMAllocType::System);
-		
-		if (Desc->usage() & MTL::TextureUsageRenderTarget)
-		{
-			objc_setAssociatedObject((__bridge id<MTLTexture>)Texture, (void*)&MetalLLM::LogAllocTexture,
-			[[[FMetalDeallocHandler alloc] initWithBlock:^{
-				LLM_PLATFORM_SCOPE_METAL(ELLMTagMetal::RenderTargets);
-				
-				LLM_IF_ENABLED(FLowLevelMemTracker::Get().OnLowLevelFree(ELLMTracker::Platform, Ptr, ELLMAllocType::System));
-				
-#if PLATFORM_IOS
-				if (!bMemoryless)
-#endif
-				{
-					DEC_MEMORY_STAT_BY(STAT_MetalTextureMemory, Size);
-				}
-				DEC_DWORD_STAT(STAT_MetalTextureCount);
-			}] autorelease],
-			OBJC_ASSOCIATION_RETAIN);
-		}
-		else
-		{
-			objc_setAssociatedObject((__bridge id<MTLTexture>)Texture, (void*)&MetalLLM::LogAllocTexture,
-			[[[FMetalDeallocHandler alloc] initWithBlock:^{
-				LLM_PLATFORM_SCOPE_METAL(ELLMTagMetal::Textures);
-			
-				LLM_IF_ENABLED(FLowLevelMemTracker::Get().OnLowLevelFree(ELLMTracker::Platform, Ptr, ELLMAllocType::System));
-			
-#if PLATFORM_IOS
-				if (!bMemoryless)
-#endif
-				{
-					DEC_MEMORY_STAT_BY(STAT_MetalTextureMemory, Size);
-				}
-				DEC_DWORD_STAT(STAT_MetalTextureCount);
-			}] autorelease],
-			OBJC_ASSOCIATION_RETAIN);
-		}
+		LLM_PLATFORM_SCOPE_METAL(ELLMTagMetal::RenderTargets);
+		LLM_IF_ENABLED(FLowLevelMemTracker::Get().OnLowLevelAlloc(ELLMTracker::Platform, Ptr, Size, ELLMTag::Untagged, ELLMAllocType::System));
+	}
+	else
+	{
+		LLM_PLATFORM_SCOPE_METAL(ELLMTagMetal::Textures);
+		LLM_IF_ENABLED(FLowLevelMemTracker::Get().OnLowLevelAlloc(ELLMTracker::Platform, Ptr, Size, ELLMTag::Untagged, ELLMAllocType::System));
 	}
 }
 
-void MetalLLM::LogAllocBuffer(FMetalBuffer* Buffer)
+void MetalLLM::LogFreeTexture(MTL::Texture* Texture)
 {
-	void* Ptr = (void*)Buffer;
-	uint64 Size = Buffer->GetLength();
+	uint64 Size = Texture->allocatedSize();
+	void* Ptr = (void*)Texture;
 	
-	INC_MEMORY_STAT_BY(STAT_MetalBufferMemory, Size);
-	INC_DWORD_STAT(STAT_MetalBufferCount);
+#if PLATFORM_IOS
+	bool bMemoryless = (Texture->storageMode() == MTL::StorageModeMemoryless);
+	if (!bMemoryless)
+#endif
+	{
+		DEC_MEMORY_STAT_BY(STAT_MetalTextureMemory, Size);
+	}
+	DEC_DWORD_STAT(STAT_MetalTextureCount);
 	
-	LLM_IF_ENABLED(FLowLevelMemTracker::Get().OnLowLevelAlloc(ELLMTracker::Platform, Ptr, Size, ELLMTag::Untagged, ELLMAllocType::System));
-}
-
-void MetalLLM::LogAllocBuffer(FMetalBufferPtr Buffer)
-{
-	LogAllocBuffer(Buffer.Get());
+	if (Texture->usage() & MTL::TextureUsageRenderTarget)
+	{
+		LLM_PLATFORM_SCOPE_METAL(ELLMTagMetal::RenderTargets);
+		LLM_IF_ENABLED(FLowLevelMemTracker::Get().OnLowLevelFree(ELLMTracker::Platform, Ptr, ELLMAllocType::System));
+	}
+	else
+	{
+		LLM_PLATFORM_SCOPE_METAL(ELLMTagMetal::Textures);
+		LLM_IF_ENABLED(FLowLevelMemTracker::Get().OnLowLevelFree(ELLMTracker::Platform, Ptr, ELLMAllocType::System));
+	}
 }
 
 void MetalLLM::LogAllocBufferNative(MTL::Buffer* Buffer)
 {
     void* Ptr = (void*)Buffer;
     uint64 Size = Buffer->length();
-    
+
+	LLM_PLATFORM_SCOPE_METAL(ELLMTagMetal::Buffers);
     INC_MEMORY_STAT_BY(STAT_MetalBufferMemory, Size);
     INC_DWORD_STAT(STAT_MetalBufferCount);
     
 	LLM_IF_ENABLED(FLowLevelMemTracker::Get().OnLowLevelAlloc(ELLMTracker::Platform, Ptr, Size, ELLMTag::Untagged, ELLMAllocType::System));
-    // Assign a dealloc handler to untrack the memory - but don't track the dispatch block!
-    {
-        LLM_SCOPED_PAUSE_TRACKING(ELLMAllocType::System);
-        
-        objc_setAssociatedObject((__bridge id<MTLBuffer>)Buffer, (void*)&MetalLLM::LogAllocBufferNative,
-        [[[FMetalDeallocHandler alloc] initWithBlock:^{
-            LLM_PLATFORM_SCOPE_METAL(ELLMTagMetal::Buffers);
-            
-			LLM_IF_ENABLED(FLowLevelMemTracker::Get().OnLowLevelFree(ELLMTracker::Platform, Ptr, ELLMAllocType::System));
-            
-            DEC_MEMORY_STAT_BY(STAT_MetalBufferMemory, Size);
-            DEC_DWORD_STAT(STAT_MetalBufferCount);
-        }] autorelease],
-        OBJC_ASSOCIATION_RETAIN);
-    }
+}
+
+void MetalLLM::LogFreeBufferNative(MTL::Buffer* Buffer)
+{
+	void* Ptr = (void*)Buffer;
+	uint64 Size = Buffer->length();
+	
+	LLM_PLATFORM_SCOPE_METAL(ELLMTagMetal::Buffers);
+	DEC_MEMORY_STAT_BY(STAT_MetalBufferMemory, Size);
+	DEC_DWORD_STAT(STAT_MetalBufferCount);
+	
+	LLM_IF_ENABLED(FLowLevelMemTracker::Get().OnLowLevelFree(ELLMTracker::Platform, Ptr, ELLMAllocType::System));
 }
 
 void MetalLLM::LogAllocHeap(MTL::Heap* Heap)
@@ -284,3 +254,12 @@ void MetalLLM::LogAliasTexture(MTL::Texture* Texture)
 	objc_setAssociatedObject((__bridge id<MTLTexture>)Texture, (void*)&MetalLLM::LogAllocTexture, nullptr, OBJC_ASSOCIATION_RETAIN);
 }
 
+void MetalBufferStats::UpdateBufferStats(const FRHIBufferDesc& BufferDesc, int64 BufferSize, bool bAllocating)
+{
+	UE::RHICore::UpdateGlobalBufferStats(BufferDesc, BufferSize, bAllocating);
+}
+
+void MetalBufferStats::UpdateUniformBufferStats(int64 BufferSize, bool bAllocating)
+{
+	UE::RHICore::UpdateGlobalUniformBufferStats(BufferSize, bAllocating);
+}
diff --git a/Engine/Source/Runtime/Apple/MetalRHI/Private/MetalLLM.h b/Engine/Source/Runtime/Apple/MetalRHI/Private/MetalLLM.h
--- a/Engine/Source/Runtime/Apple/MetalRHI/Private/MetalLLM.h
+++ b/Engine/Source/Runtime/Apple/MetalRHI/Private/MetalLLM.h
@@ -5,6 +5,7 @@
 #include "MetalRHIPrivate.h"
 #include "HAL/LowLevelMemTracker.h"
 #include "Apple/AppleLLM.h"
+#include "RHICoreStats.h"
 
 class FMetalBuffer;
 typedef TSharedPtr<FMetalBuffer, ESPMode::ThreadSafe> FMetalBufferPtr;
@@ -50,13 +51,17 @@ class FMetalDevice;
 // These work without the LLM module
 namespace MetalLLM
 {
-	void LogAllocTexture(FMetalDevice& Device, MTL::TextureDescriptor* Desc, MTL::Texture* Texture);
-	void LogAllocBuffer(FMetalBuffer* Buffer);
-	void LogAllocBuffer(FMetalBufferPtr Buffer);
+	void LogAllocTexture(MTL::Texture* Texture);
+	void LogFreeTexture(MTL::Texture* Texture);
     void LogAllocBufferNative(MTL::Buffer* Buffer);
+	void LogFreeBufferNative(MTL::Buffer* Buffer);
 	void LogAllocHeap(MTL::Heap* Heap);
 	
 	void LogAliasTexture(MTL::Texture* Texture);
 }
 
-
+namespace MetalBufferStats
+{
+	void UpdateBufferStats(const FRHIBufferDesc& BufferDesc, int64 BufferSize, bool bAllocating);
+	void UpdateUniformBufferStats(int64 BufferSize, bool bAllocating);
+}
diff --git a/Engine/Source/Runtime/Apple/MetalRHI/Private/MetalRHI.cpp b/Engine/Source/Runtime/Apple/MetalRHI/Private/MetalRHI.cpp
--- a/Engine/Source/Runtime/Apple/MetalRHI/Private/MetalRHI.cpp
+++ b/Engine/Source/Runtime/Apple/MetalRHI/Private/MetalRHI.cpp
@@ -1565,18 +1565,9 @@ void FMetalDynamicRHI::ProcessDeferredDeleteQueue()
 				}
 				case FMetalDeferredDeleteObject::TObjectStorage::IndexOfType<FMetalBufferPtr>():
 				{
-					FMetalBufferPtr Buffer = Object.Storage.Get<FMetalBufferPtr>();
-					
+					FMetalBufferPtr Buffer = Object.Storage.Get<FMetalBufferPtr>();	
 					Buffer->MarkDeleted();
-					
-#if METAL_DEBUG_OPTIONS
-					MTL::Buffer* MTLBuffer = Buffer->GetMTLBuffer();
-					if (GMetalResourcePurgeOnDelete && !MTLBuffer->heap() &&
-						Buffer->GetOffset() == 0 && Buffer->GetLength() == MTLBuffer->length())
-					{
-						MTLBuffer->setPurgeableState(MTL::PurgeableStateEmpty);
-					}
-#endif
+				
 					break;
 				}
 				case FMetalDeferredDeleteObject::TObjectStorage::IndexOfType<MTLTexturePtr>():
@@ -1585,12 +1576,6 @@ void FMetalDynamicRHI::ProcessDeferredDeleteQueue()
 					
 					if (!Texture->buffer() && !Texture->parentTexture())
 					{
-#if METAL_DEBUG_OPTIONS
-						if (GMetalResourcePurgeOnDelete && !Texture->heap())
-						{
-							Texture->setPurgeableState(MTL::PurgeableStateEmpty);
-						}
-#endif
 						Device->GetResourceHeap().ReleaseTexture(nullptr, Texture);
 					}
 					break;
diff --git a/Engine/Source/Runtime/Apple/MetalRHI/Private/MetalRHIBuffer.cpp b/Engine/Source/Runtime/Apple/MetalRHI/Private/MetalRHIBuffer.cpp
--- a/Engine/Source/Runtime/Apple/MetalRHI/Private/MetalRHIBuffer.cpp
+++ b/Engine/Source/Runtime/Apple/MetalRHI/Private/MetalRHIBuffer.cpp
@@ -16,10 +16,6 @@
 
 #define METAL_POOL_BUFFER_BACKING 1
 
-#if !METAL_POOL_BUFFER_BACKING
-DECLARE_MEMORY_STAT(TEXT("Used Device Buffer Memory"), STAT_MetalDeviceBufferMemory, STATGROUP_MetalRHI);
-#endif
-
 #if STATS
 #define METAL_INC_DWORD_STAT_BY(Name, Size, Usage) do { if (EnumHasAnyFlags(Usage, BUF_IndexBuffer)) { INC_DWORD_STAT_BY(STAT_MetalIndex##Name, Size); } else { INC_DWORD_STAT_BY(STAT_MetalVertex##Name, Size); } } while (false)
 #else
@@ -116,22 +112,6 @@ FMetalRHIBuffer::FMetalRHIBuffer(FRHICommandListBase& RHICmdList, FMetalDevice&
 	{
 		checkf(InBufferDesc.Size <= Device.GetDevice()->maxBufferLength(), TEXT("Requested buffer size larger than supported by device."));
 		
-		// Temporary buffers less than the buffer page size - currently 4Kb - is better off going through the set*Bytes API if available.
-		// These can't be used for shader resources or UAVs if we want to use the 'Linear Texture' code path
-		
-		// TODO: Carl - Strip this code as Volatile is not used with buffer uploads
-		/*if (!bWantsView
-			&& bIsVolatile
-			&& !EnumHasAnyFlags(MetalUsage, EMetalBufferUsage::GPUOnly)
-			&& InBufferDesc.Size < MetalBufferPageSize
-			&& InBufferDesc.Size < MetalBufferBytesSize)
-		{
-            Data = new FMetalBufferData();
-            Data->InitWithSize(InBufferDesc.Size);
-			METAL_INC_DWORD_STAT_BY(MemAlloc, InBufferDesc.Size, InBufferDesc.Usage);
-		}
-		else
-		*/
 		{
 			// Static buffers will never be discarded. You can update them directly.
 			if(bIsStatic)
@@ -214,12 +194,21 @@ FMetalRHIBuffer::~FMetalRHIBuffer()
     ReleaseOwnership();
 }
 
-void FMetalRHIBuffer::AllocateBuffers()
+void FMetalRHIBuffer::AllocateBufferAtSlot(uint32_t Idx)
 {
 	uint32 AllocSize = Size;
 	
 	const bool bWantsView = EnumHasAnyFlags(GetDesc().Usage, BUF_ShaderResource | BUF_UnorderedAccess);
 	
+	// These allocations will not go into the pool.
+	uint32 RequestedBufferOffsetAlignment = BufferOffsetAlignment;
+	if(bWantsView)
+	{
+		// Buffer backed linear textures have specific align requirements
+		// We don't know upfront the pixel format that may be requested for an SRV so we can't use minimumLinearTextureAlignmentForPixelFormat:
+		RequestedBufferOffsetAlignment = BufferBackedLinearTextureOffsetAlignment;
+	}
+	
 	const EMetalBufferUsage MetalUsage = GetMetalBufferUsage(GetDesc().Usage);
 	
 	if (EnumHasAnyFlags(MetalUsage, EMetalBufferUsage::LinearTex) && !Device.SupportsFeature(EMetalFeaturesTextureBuffers))
@@ -268,69 +257,59 @@ void FMetalRHIBuffer::AllocateBuffers()
 			AllocSize = Align(AllocSize, 1024);
 		}
 	}
+	AllocSize = Align(AllocSize, RequestedBufferOffsetAlignment);
 	
-	BufferPool.SetNum(NumberOfBuffers);
+	FMetalBufferPtr Buffer = nullptr;
+
+#if METAL_POOL_BUFFER_BACKING
+	FMetalPooledBufferArgs ArgsCPU(&Device, AllocSize, GetDesc().Usage, Mode);
+	Buffer = Device.CreatePooledBuffer(ArgsCPU);
+#else
+	NS::UInteger Options = (((NS::UInteger) Mode) << MTL::ResourceStorageModeShift);
 	
-	// These allocations will not go into the pool.
-	uint32 RequestedBufferOffsetAlignment = BufferOffsetAlignment;
-	if(bWantsView)
+	METAL_GPUPROFILE(FScopedMetalCPUStats CPUStat(FString::Printf(TEXT("AllocBuffer: %llu, %llu"), AllocSize, Options)));
+	// Allocate one.
+	MTL::Buffer* BufferPtr = NS::TransferPtr(Device.GetDevice()->newBuffer(AllocSize, (MTL::ResourceOptions) Options));
+	Buffer = FMetalBufferPtr(new FMetalBuffer(BufferPtr, FMetalBuffer::FreePolicy::Owner));
+	
+	METAL_FATAL_ASSERT(Buffer, TEXT("Failed to create buffer of size %u and resource options %u"), Size, (uint32)Options);
+	
+	const bool bIsStatic = EnumHasAnyFlags(GetDesc().Usage, BUF_Static);
+	if(bIsStatic)
 	{
-		// Buffer backed linear textures have specific align requirements
-		// We don't know upfront the pixel format that may be requested for an SRV so we can't use minimumLinearTextureAlignmentForPixelFormat:
-		RequestedBufferOffsetAlignment = BufferBackedLinearTextureOffsetAlignment;
+		FString Label = FString::Printf(TEXT("Static on frame %u"), Device.GetFrameNumberRHIThread());
+		BufferPtr->setLabel(FStringToNSString(Label));
 	}
-	
-	AllocSize = Align(AllocSize, RequestedBufferOffsetAlignment);
-	for(uint32 i = 0; i < NumberOfBuffers; i++)
+	else
 	{
-		FMetalBufferPtr Buffer = nullptr;
-
-#if METAL_POOL_BUFFER_BACKING
-		FMetalPooledBufferArgs ArgsCPU(&Device, AllocSize, GetDesc().Usage, Mode);
-		Buffer = Device.CreatePooledBuffer(ArgsCPU);
-#else
-		NS::UInteger Options = (((NS::UInteger) Mode) << MTL::ResourceStorageModeShift);
-		
-		METAL_GPUPROFILE(FScopedMetalCPUStats CPUStat(FString::Printf(TEXT("AllocBuffer: %llu, %llu"), AllocSize, Options)));
-		// Allocate one.
-		MTL::Buffer* BufferPtr = NS::TransferPtr(Device.GetDevice()->newBuffer(AllocSize, (MTL::ResourceOptions) Options));
-		Buffer = FMetalBufferPtr(new FMetalBuffer(BufferPtr, FMetalBuffer::FreePolicy::Owner));
-		
-		#if STATS || ENABLE_LOW_LEVEL_MEM_TRACKER
-			MetalLLM::LogAllocBuffer(Buffer);
-		#endif
-		INC_MEMORY_STAT_BY(STAT_MetalDeviceBufferMemory, Buffer->GetLength());
-		
-		METAL_FATAL_ASSERT(Buffer, TEXT("Failed to create buffer of size %u and resource options %u"), Size, (uint32)Options);
-		
-		const bool bIsStatic = EnumHasAnyFlags(GetDesc().Usage, BUF_Static);
-		if(bIsStatic)
-		{
-			FString Label = FString::Printf(TEXT("Static on frame %u"), Device.GetFrameNumberRHIThread());
-			BufferPtr->setLabel(FStringToNSString(Label));
-		}
-		else
-		{
-			FString Label = FString::Printf(TEXT("Buffer on frame %u"), Device.GetFrameNumberRHIThread());
-			BufferPtr->setLabel(FStringToNSString(Label));
-		}
-#endif
-		BufferPool[i] = Buffer;
-		
-		check(Buffer);
-		check(AllocSize <= Buffer->GetLength());
-		check(Buffer->GetMTLBuffer()->storageMode() == Mode);
+		FString Label = FString::Printf(TEXT("Buffer on frame %u"), Device.GetFrameNumberRHIThread());
+		BufferPtr->setLabel(FStringToNSString(Label));
 	}
+#endif
+	
+	MetalBufferStats::UpdateBufferStats(GetDesc(), Buffer->GetLength(), true);
+	BufferPool[Idx] = Buffer;
+	
+	check(Buffer);
+	check(AllocSize <= Buffer->GetLength());
+	check(Buffer->GetMTLBuffer()->storageMode() == Mode);
+}
+
+void FMetalRHIBuffer::AllocateBuffers()
+{
+	BufferPool.SetNum(NumberOfBuffers);
 }
 
 void FMetalRHIBuffer::ReleaseBuffers()
 {
 	for (FMetalBufferPtr Buffer : BufferPool)
 	{
-		check(Buffer);
-		
-		METAL_INC_DWORD_STAT_BY(MemFreed, Buffer->GetLength(), GetUsage());
-		FMetalDynamicRHI::Get().DeferredDelete(Buffer);
+		if(Buffer)
+		{	
+			MetalBufferStats::UpdateBufferStats(GetDesc(), Buffer->GetLength(), false);
+			METAL_INC_DWORD_STAT_BY(MemFreed, Buffer->GetLength(), GetUsage());
+			FMetalDynamicRHI::Get().DeferredDelete(Buffer);
+		}
 	}
 } 
 
@@ -585,8 +564,6 @@ void FMetalRHIBuffer::TakeOwnership(FMetalRHIBuffer& Other)
 
 void FMetalRHIBuffer::ReleaseOwnership()
 {
-    FRHIBuffer::ReleaseOwnership();
-    
     if(TransferBuffer)
     {
         METAL_INC_DWORD_STAT_BY(MemFreed, TransferBuffer->GetLength(), GetUsage());
@@ -611,6 +588,8 @@ void FMetalRHIBuffer::ReleaseOwnership()
         AccelerationStructureHandle = nullptr;
     }
 #endif // METAL_RHI_RAYTRACING
+	
+	FRHIBuffer::ReleaseOwnership();
 }
 
 FBufferRHIRef FMetalDynamicRHI::RHICreateBuffer(FRHICommandListBase& RHICmdList, FRHIBufferDesc const& Desc, ERHIAccess /*ResourceState*/, FRHIResourceCreateInfo& CreateInfo)
diff --git a/Engine/Source/Runtime/Apple/MetalRHI/Private/MetalRHIRenderQuery.cpp b/Engine/Source/Runtime/Apple/MetalRHI/Private/MetalRHIRenderQuery.cpp
--- a/Engine/Source/Runtime/Apple/MetalRHI/Private/MetalRHIRenderQuery.cpp
+++ b/Engine/Source/Runtime/Apple/MetalRHI/Private/MetalRHIRenderQuery.cpp
@@ -101,9 +101,6 @@ FMetalQueryBuffer* FMetalQueryBufferPool::GetCurrentQueryBuffer()
 		}
 		else
 		{
-			LLM_SCOPE_METAL(ELLMTagMetal::Buffers);
-			LLM_PLATFORM_SCOPE_METAL(ELLMTagMetal::Buffers);
-
 			METAL_GPUPROFILE(FScopedMetalCPUStats CPUStat(FString::Printf(TEXT("AllocBuffer: %llu, %llu"), EQueryBufferMaxSize, MTL::ResourceStorageModeShared)));
 			
 			MTL::ResourceOptions HazardTrackingMode = MTL::ResourceHazardTrackingModeUntracked;
@@ -115,10 +112,6 @@ FMetalQueryBuffer* FMetalQueryBufferPool::GetCurrentQueryBuffer()
 			
 			Buffer = Device.GetResourceHeap().CreateBuffer(EQueryBufferMaxSize, 16, BUF_Dynamic, FMetalCommandQueue::GetCompatibleResourceOptions((MTL::ResourceOptions)(BUFFER_CACHE_MODE | HazardTrackingMode | MTL::ResourceStorageModeShared)), true);
 			FMemory::Memzero((((uint8*)Buffer->Contents())), EQueryBufferMaxSize);
-
-#if STATS || ENABLE_LOW_LEVEL_MEM_TRACKER
-			MetalLLM::LogAllocBuffer(Buffer);
-#endif // STATS || ENABLE_LOW_LEVEL_MEM_TRACKER
 		}
 
 		CurrentBuffer = new FMetalQueryBuffer(this, MoveTemp(Buffer));
diff --git a/Engine/Source/Runtime/Apple/MetalRHI/Public/MetalResources.h b/Engine/Source/Runtime/Apple/MetalRHI/Public/MetalResources.h
--- a/Engine/Source/Runtime/Apple/MetalRHI/Public/MetalResources.h
+++ b/Engine/Source/Runtime/Apple/MetalRHI/Public/MetalResources.h
@@ -298,6 +298,7 @@ public:
 	
 	bool RequiresTransferBuffer();
 	
+	void AllocateBufferAtSlot(uint32_t Idx);
 	void AllocateBuffers();
 	void ReleaseBuffers();
 	
@@ -313,6 +314,10 @@ public:
 	
 	FMetalBufferPtr GetCurrentBuffer()
 	{
+		if(BufferPool[CurrentIndex] == nullptr)
+		{
+			AllocateBufferAtSlot(CurrentIndex);
+		}
 		return BufferPool[CurrentIndex];
 	}
 	
