Metal RHI

- Removed the buffer copy as we no longer need it with the temporary allocator re-written.

diff --git a/Engine/Source/Runtime/Apple/MetalRHI/Private/MetalStateCache.cpp b/Engine/Source/Runtime/Apple/MetalRHI/Private/MetalStateCache.cpp
--- a/Engine/Source/Runtime/Apple/MetalRHI/Private/MetalStateCache.cpp
+++ b/Engine/Source/Runtime/Apple/MetalRHI/Private/MetalStateCache.cpp
@@ -332,14 +332,6 @@ void FMetalStateCache::Reset()
 					CBVTable[Frequency][i] = 0ull;
 				}
 			}
-			
-			// Free temporary allocations
-			for (FMetalBufferPtr TemporaryBuffer : TemporaryBuffers)
-			{
-				FMetalDynamicRHI::Get().DeferredDelete(TemporaryBuffer);
-			}
-			TemporaryBuffers.Reset();
-			UniformBufferVAs.Reset();
 		}
 	}
 #endif
@@ -1710,26 +1702,11 @@ void FMetalStateCache::IRBindPackedUniforms(EMetalShaderStages const Frequency,
 
 void FMetalStateCache::IRBindUniformBuffer(EMetalShaderStages const Frequency, int32 Index, FMetalUniformBuffer* UB)
 {
-	FMetalBufferPtr Buffer = GetOrCreateBackingBufferCopy(UB);
+	FMetalBufferPtr Buffer = UB->BackingBuffer;
 	CacheOrSkipResourceResidencyUpdate(Buffer->GetMTLBuffer(), Frequency, true);
 
 	CBVTable[Frequency][Index] = Buffer->GetGPUAddress();
 }
-	
-FMetalBufferPtr FMetalStateCache::GetOrCreateBackingBufferCopy(FMetalUniformBuffer* UB)
-{
-	FMetalBufferPtr* CachedCopy = UniformBufferVAs.Find(UB->BackingBuffer.Get());
-	if (CachedCopy != nullptr)
-	{
-		return *CachedCopy;
-	}
-	
-	uint8* ConstantSpace =  reinterpret_cast<uint8*>(UB->BackingBuffer->Contents());
-	FMetalBufferPtr Buffer = IRSideUploadToBuffer(ConstantSpace, UB->GetSize());
-	UniformBufferVAs.Add(UB->BackingBuffer.Get(), Buffer);
-	
-	return Buffer;
-}
 #endif
 
 void FMetalStateCache::SetShaderTexture(EMetalShaderStages const Frequency, MTL::Texture* Texture, NS::UInteger const Index, MTL::ResourceUsage const Usage)
diff --git a/Engine/Source/Runtime/Apple/MetalRHI/Private/MetalStateCache.h b/Engine/Source/Runtime/Apple/MetalRHI/Private/MetalStateCache.h
--- a/Engine/Source/Runtime/Apple/MetalRHI/Private/MetalStateCache.h
+++ b/Engine/Source/Runtime/Apple/MetalRHI/Private/MetalStateCache.h
@@ -125,8 +125,6 @@ public:
 	 */
 	FMetalBufferPtr IRSideUploadToBuffer(void const* Content, uint64 Size);
 	
-	FMetalBufferPtr GetOrCreateBackingBufferCopy(FMetalUniformBuffer* UB);
-
     template<class ShaderType, EMetalShaderStages Frequency, MTL::FunctionType FunctionType>
     void IRBindResourcesToEncoder(ShaderType Shader, FMetalCommandEncoder* Encoder);
 
@@ -310,17 +308,13 @@ private:
 		
 #if METAL_USE_METAL_SHADER_CONVERTER
     static constexpr uint32 TopLevelABNumEntry = 16;
-	TMap<FMetalBuffer*, FMetalBufferPtr>  UniformBufferVAs;
-	
-    uint64                 CBVTable[EMetalShaderStages::Num][TopLevelABNumEntry];
+    uint64 CBVTable[EMetalShaderStages::Num][TopLevelABNumEntry];
 	
 	IRRuntimeVertexBuffer VertexBufferVAs[31];
 	
 	TStaticArray<TSet<MTL::Heap*>, EMetalShaderStages::Num> HeapsUsedByStage;
 	TStaticArray<TSet<MTL::Resource*>, EMetalShaderStages::Num> RWResourcesByStage;
 	TStaticArray<TSet<MTL::Resource*>, EMetalShaderStages::Num> ROResourcesByStage;
-	
-	TArray<FMetalBufferPtr> TemporaryBuffers;
 #endif
 
 	FMetalQueryBuffer* VisibilityResults;
