Potential fix for crash when deleting PSO's due to missing lock.

diff --git a/Engine/Source/Runtime/Apple/MetalRHI/Private/MetalPipeline.cpp b/Engine/Source/Runtime/Apple/MetalRHI/Private/MetalPipeline.cpp
--- a/Engine/Source/Runtime/Apple/MetalRHI/Private/MetalPipeline.cpp
+++ b/Engine/Source/Runtime/Apple/MetalRHI/Private/MetalPipeline.cpp
@@ -339,11 +339,11 @@ public:
 	{
 		check (GMetalCacheShaderPipelines == 0);
 		{
+			FRWScopeLock Lock(PipelineMutex, SLT_Write);
 			FMetalGraphicsPipelineKey* Desc = ReverseLookup.Find(Pipeline);
+			
 			if (Desc)
-			{
-                FRWScopeLock Lock(PipelineMutex, SLT_Write);
-                
+			{   
 				Pipelines.Remove(*Desc);
 				ReverseLookup.Remove(Pipeline);
 			}
