Metal - Unify how to enable symbols in SM5 & SM6 with r.Shaders.Symbols instead of r.Shaders.ExtraData (SM6) & r.Shaders.Optimize (SM5)

#tests Test SM5 & SM6 Debugging in Xcode by only setting r.Shaders.Symbols

diff --git a/Engine/Source/Developer/Apple/MetalShaderFormat/Private/MetalCompileShaderMSC.cpp b/Engine/Source/Developer/Apple/MetalShaderFormat/Private/MetalCompileShaderMSC.cpp
--- a/Engine/Source/Developer/Apple/MetalShaderFormat/Private/MetalCompileShaderMSC.cpp
+++ b/Engine/Source/Developer/Apple/MetalShaderFormat/Private/MetalCompileShaderMSC.cpp
@@ -524,7 +524,7 @@ void FMetalCompileShaderMSC::DoCompileMetalShader(
 		
 		TArray<FString> ExtraArgs;
 		
-		if (Input.Environment.CompilerFlags.Contains(CFLAG_ExtraShaderData))
+		if (Input.Environment.CompilerFlags.Contains(CFLAG_GenerateSymbols))
 		{
 			ExtraArgs.Add(TEXT("-Zi"));
 			ExtraArgs.Add(TEXT("-Qembed_debug"));
diff --git a/Engine/Source/Runtime/Apple/MetalRHI/Private/MetalDevice.cpp b/Engine/Source/Runtime/Apple/MetalRHI/Private/MetalDevice.cpp
--- a/Engine/Source/Runtime/Apple/MetalRHI/Private/MetalDevice.cpp
+++ b/Engine/Source/Runtime/Apple/MetalRHI/Private/MetalDevice.cpp
@@ -613,12 +613,6 @@ void FMetalDevice::EnumerateFeatureSupport()
 		Features |= EMetalFeaturesValidation;
 	}
 #endif
-
-	static const auto CVar = IConsoleManager::Get().FindConsoleVariable(TEXT("r.Shaders.Optimize"));
-	if (CVar->GetInt() == 0 || FParse::Param(FCommandLine::Get(),TEXT("metalshaderdebug")))
-	{
-		Features |= EMetalFeaturesGPUTrace;
-	}
 }
 
 void FMetalDevice::EndDrawingViewport(bool bPresent)
diff --git a/Engine/Source/Runtime/Apple/MetalRHI/Private/MetalDevice.h b/Engine/Source/Runtime/Apple/MetalRHI/Private/MetalDevice.h
--- a/Engine/Source/Runtime/Apple/MetalRHI/Private/MetalDevice.h
+++ b/Engine/Source/Runtime/Apple/MetalRHI/Private/MetalDevice.h
@@ -47,38 +47,36 @@ typedef NS_OPTIONS(uint64, EMetalFeatures)
 	EMetalFeaturesMSAADepthResolve = 1 << 11,
 	/** Supports Store & Resolve in a single store action */
 	EMetalFeaturesMSAAStoreAndResolve = 1 << 12,
-	/** Supports framework GPU frame capture */
-	EMetalFeaturesGPUTrace = 1 << 13,
 	/** Supports the use of cubemap arrays */
-	EMetalFeaturesCubemapArrays = 1 << 14,
+	EMetalFeaturesCubemapArrays = 1 << 13,
 	/** Supports the specification of multiple viewports and scissor rects */
-	EMetalFeaturesMultipleViewports = 1 << 15,
+	EMetalFeaturesMultipleViewports = 1 << 14,
 	/** Supports minimum on-glass duration for drawables */
-	EMetalFeaturesPresentMinDuration = 1llu << 16llu,
+	EMetalFeaturesPresentMinDuration = 1llu << 15llu,
 	/** Supports programmatic frame capture API */
-	EMetalFeaturesGPUCaptureManager = 1llu << 17llu,
+	EMetalFeaturesGPUCaptureManager = 1llu << 16llu,
 	/** Supports efficient buffer-blits */
-	EMetalFeaturesEfficientBufferBlits = 1llu << 18llu,
+	EMetalFeaturesEfficientBufferBlits = 1llu << 17llu,
 	/** Supports any kind of buffer sub-allocation */
-	EMetalFeaturesBufferSubAllocation = 1llu << 19llu,
+	EMetalFeaturesBufferSubAllocation = 1llu << 18llu,
 	/** Supports private buffer sub-allocation */
-	EMetalFeaturesPrivateBufferSubAllocation = 1llu << 20llu,
+	EMetalFeaturesPrivateBufferSubAllocation = 1llu << 19llu,
 	/** Supports texture buffers */
-	EMetalFeaturesTextureBuffers = 1llu << 21llu,
+	EMetalFeaturesTextureBuffers = 1llu << 20llu,
 	/** Supports max. compute threads per threadgroup */
-	EMetalFeaturesMaxThreadsPerThreadgroup = 1llu << 22llu,
+	EMetalFeaturesMaxThreadsPerThreadgroup = 1llu << 21llu,
 	/** Supports parallel render encoders */
-	EMetalFeaturesParallelRenderEncoders = 1llu << 23llu,
+	EMetalFeaturesParallelRenderEncoders = 1llu << 22llu,
 	/** Supports indirect argument buffers */
-	EMetalFeaturesIABs = 1llu << 24llu,
+	EMetalFeaturesIABs = 1llu << 23llu,
 	/** Supports specifying the mutability of buffers bound to PSOs */
-	EMetalFeaturesPipelineBufferMutability = 1llu << 25llu,
+	EMetalFeaturesPipelineBufferMutability = 1llu << 24llu,
 	/** Supports tile shaders */
-	EMetalFeaturesTileShaders = 1llu << 26llu,
+	EMetalFeaturesTileShaders = 1llu << 25llu,
 	/** Unused Reserved Bit */
-	EMetalFeaturesUnusedReservedBit27 = 1llu << 27llu, // was EMetalFeaturesSeparateTessellation
+	EMetalFeaturesUnusedReservedBit27 = 1llu << 26llu, // was EMetalFeaturesSeparateTessellation
 	/** Supports indirect argument buffers Tier 2 */
-	EMetalFeaturesTier2IABs = 1llu << 28llu,
+	EMetalFeaturesTier2IABs = 1llu << 27llu,
 };
 
 /**
diff --git a/Engine/Source/Runtime/Apple/MetalRHI/Private/Shaders/Types/Templates/MetalBaseShader.h b/Engine/Source/Runtime/Apple/MetalRHI/Private/Shaders/Types/Templates/MetalBaseShader.h
--- a/Engine/Source/Runtime/Apple/MetalRHI/Private/Shaders/Types/Templates/MetalBaseShader.h
+++ b/Engine/Source/Runtime/Apple/MetalRHI/Private/Shaders/Types/Templates/MetalBaseShader.h
@@ -173,7 +173,8 @@ void TMetalBaseShader<BaseResourceType, ShaderType>::Init(TArrayView<const uint8
 	const ANSICHAR* ShaderSource = ShaderCode.FindOptionalData(EShaderOptionalDataKey::SourceCode);
 	bool bHasShaderSource = (ShaderSource && FCStringAnsi::Strlen(ShaderSource) > 0);
 
-	static bool bForceTextShaders = Device.SupportsFeature(EMetalFeaturesGPUTrace);
+	static const auto CVar = IConsoleManager::Get().FindConsoleVariable(TEXT("r.Shaders.Symbols"));
+	static bool bForceTextShaders = CVar->GetInt() == 1 || FParse::Param(FCommandLine::Get(),TEXT("metalshaderdebug"));
 	if (!bHasShaderSource)
 	{
 		int32 LZMASourceSize = 0;
