Metal RHI

- Added support for submission/interrupt thread
- Added new profiler support

diff --git a/Engine/Source/Runtime/Apple/MetalRHI/Private/MetalCommandBuffer.h b/Engine/Source/Runtime/Apple/MetalRHI/Private/MetalCommandBuffer.h
--- a/Engine/Source/Runtime/Apple/MetalRHI/Private/MetalCommandBuffer.h
+++ b/Engine/Source/Runtime/Apple/MetalRHI/Private/MetalCommandBuffer.h
@@ -5,34 +5,143 @@
 #include "MetalThirdParty.h"
 #include "MetalResources.h"
 #include "MetalShaderResources.h"
+#include "RHIBreadcrumbs.h"
+#include "MetalProfiler.h"
+#include "MetalCommandQueue.h"
+
+class FMetalCommandBufferTimer;
+struct FMetalCounterSample;
+class FMetalEventNode;
 
 class FMetalCommandBuffer
 {
 public:
-    FMetalCommandBuffer(MTL::CommandBuffer* InCommandBuffer)
+    FMetalCommandBuffer(MTL::CommandBuffer* InCommandBuffer, FMetalCommandQueue& Queue) 
+#if RHI_NEW_GPU_PROFILER
+	: EventStream(Queue.GetProfilerQueue())
+#endif
     {
         CommandBuffer = NS::RetainPtr(InCommandBuffer);
     }
     
     FORCEINLINE MTLCommandBufferPtr& GetMTLCmdBuffer() {return CommandBuffer;}
+	
+#if RHI_NEW_GPU_PROFILER == 0
+	FORCEINLINE void SetTimer(FMetalCommandBufferTimer* InTimer)
+	{
+		Timer = InTimer;
+	}
+	
+	FORCEINLINE FMetalCommandBufferTimer* GetTimer()
+	{
+		return Timer;
+	}
+#endif
+	
+#if WITH_RHI_BREADCRUMBS
+#if RHI_NEW_GPU_PROFILER
+	void BeginBreadcrumb(FRHIBreadcrumbNode* Node)
+	{
+		FMetalBreadcrumbTrackerObject TrackerObject;
+		TrackerObject.Type = EMetalBreadcrumbTrackerType_Begin;
+		TrackerObject.CmdBuffer = this;
+		TrackerObject.Node = Node;
+		BreadcrumbTrackerStream.Add(TrackerObject);
+	}
+	
+	void EndBreadcrumb(FRHIBreadcrumbNode* Node)
+	{
+		FMetalBreadcrumbTrackerObject TrackerObject;
+		TrackerObject.Type = EMetalBreadcrumbTrackerType_End;
+		TrackerObject.CmdBuffer = this;
+		TrackerObject.Node = Node;
+		BreadcrumbTrackerStream.Add(TrackerObject);
+	}
+#endif
+#endif
     
-    FORCEINLINE TSharedPtr<FMetalCommandBufferFence, ESPMode::ThreadSafe> GetCompletionFence()
-    {
-        if(CmdBufferFence == nullptr)
-        {
-            CmdBufferFence = MakeShareable(new FMetalCommandBufferFence());
-            InsertCompletionFence(CmdBufferFence);
-        }
-        assert(CmdBufferFence != nullptr);
-        return CmdBufferFence;
-    }
-    
-    FORCEINLINE void InsertCompletionFence(TSharedPtr<FMetalCommandBufferFence, ESPMode::ThreadSafe> Fence)
-    {
-        Fence->Insert(CommandBuffer);
-    }
-    
+	void AddCounterSample(FMetalCounterSamplePtr CounterSample)
+	{
+#if RHI_NEW_GPU_PROFILER == 0
+		for(FMetalEventNode* Node : ActiveEventNodes)
+		{
+			if(EventSampleCounters.Contains(Node))
+			{
+				EventSampleCounters[Node].Add(CounterSample);
+			}
+			else
+			{
+				EventSampleCounters.Add(Node, {CounterSample});
+			}
+		}
+#elif WITH_RHI_BREADCRUMBS
+		FMetalBreadcrumbTrackerObject TrackerObject;
+		TrackerObject.Type = EMetalBreadcrumbTrackerType_Encode;
+		TrackerObject.CmdBuffer = this;
+		TrackerObject.CounterSample = CounterSample;
+		BreadcrumbTrackerStream.Add(TrackerObject);
+#endif
+		
+		CounterSamples.Add(CounterSample);
+	}
+	
+#if RHI_NEW_GPU_PROFILER
+	void SetBeginWorkTimestamp(uint64_t* Timestamp)
+	{
+		BeginWorkTimestamp = Timestamp;
+	}
+	
+	void SetEndWorkTimestamp(uint64_t* Timestamp)
+	{
+		EndWorkTimestamp = Timestamp;
+	}
+
+	template <typename TEventType, typename... TArgs>
+	TEventType& EmplaceProfilerEvent(TArgs&&... Args)
+	{
+		TEventType& Data = EventStream.Emplace<TEventType>(Forward<TArgs>(Args)...);
+
+		if constexpr (std::is_same_v<UE::RHI::GPUProfiler::FEvent::FBeginWork, TEventType>)
+		{
+			// Store BeginEvents in a separate array as the CPUTimestamp field needs updating at submit time.
+			BeginEvents.Add(&Data);
+		}
+		
+		// Not sure what to do here, ask Luke!
+		return Data;
+	}
+	
+	void FlushProfilerEvents(UE::RHI::GPUProfiler::FEventStream& Destination, uint64 CPUTimestamp)
+	{
+		for (UE::RHI::GPUProfiler::FEvent::FBeginWork* BeginEvent : BeginEvents)
+		{
+			BeginEvent->CPUTimestamp = CPUTimestamp;
+		}
+
+		Destination.Append(MoveTemp(EventStream));
+	}
+#endif
+	
+	TArray<FMetalRHIRenderQuery*> TimestampQueries;
+	TArray<FMetalRHIRenderQuery*> OcclusionQueries;
+	TArray<FMetalCounterSamplePtr> CounterSamples;
+	
+#if RHI_NEW_GPU_PROFILER == 0
+	TMap<FMetalEventNode*, TArray<FMetalCounterSamplePtr>> EventSampleCounters;
+	TArray<FMetalEventNode*> ActiveEventNodes;
+#else
+	UE::RHI::GPUProfiler::FEventStream EventStream;
+	TArray<UE::RHI::GPUProfiler::FEvent::FBeginWork*, TInlineAllocator<8>> BeginEvents;
+	
+#if WITH_RHI_BREADCRUMBS
+	TArray<FMetalBreadcrumbTrackerObject> BreadcrumbTrackerStream;
+#endif
+	
+	uint64_t* BeginWorkTimestamp;
+	uint64_t* EndWorkTimestamp;
+#endif
+	
 private:
-    MTLCommandBufferPtr CommandBuffer;
-    TSharedPtr<FMetalCommandBufferFence, ESPMode::ThreadSafe> CmdBufferFence = nullptr;
+    MTLCommandBufferPtr CommandBuffer;	
+	FMetalCommandBufferTimer* Timer = nullptr;
 };
diff --git a/Engine/Source/Runtime/Apple/MetalRHI/Private/MetalCommandEncoder.cpp b/Engine/Source/Runtime/Apple/MetalRHI/Private/MetalCommandEncoder.cpp
--- a/Engine/Source/Runtime/Apple/MetalRHI/Private/MetalCommandEncoder.cpp
+++ b/Engine/Source/Runtime/Apple/MetalRHI/Private/MetalCommandEncoder.cpp
@@ -13,6 +13,7 @@
 #include "MetalProfiler.h"
 #include "MetalShaderResources.h"
 #include "MetalStateCache.h"
+#include "MetalRHIContext.h"
 
 const uint32 EncoderRingBufferSize = 1024 * 1024;
 
@@ -249,7 +250,7 @@ void FMetalCommandEncoder::StartCommandBuffer(void)
 			CommandBuffer->GetMTLCmdBuffer()->setLabel(DebugGroups.Last());
 		}
 		
-	#if ENABLE_METAL_GPUPROFILE
+	#if ENABLE_METAL_GPUPROFILE && RHI_NEW_GPU_PROFILER == 0
 		FMetalProfiler* Profiler = FMetalProfiler::GetProfiler();
 		if (Profiler)
 		{
@@ -259,7 +260,7 @@ void FMetalCommandEncoder::StartCommandBuffer(void)
 	}
 }
 	
-void FMetalCommandEncoder::EndCommandBuffer()
+void FMetalCommandEncoder::EndCommandBuffer(FMetalRHICommandContext* Context)
 {
 	check(CommandBuffer);
 	check(IsRenderCommandEncoderActive() == false
@@ -269,7 +270,7 @@ void FMetalCommandEncoder::EndCommandBuffer()
 		  && IsAccelerationStructureCommandEncoderActive() == false
 #endif // METAL_RHI_RAYTRACING
 		  );
-
+	
 	if(CommandBuffer->GetMTLCmdBuffer()->label() == nullptr && DebugGroups.Num() > 0)
 	{
 		CommandBuffer->GetMTLCmdBuffer()->setLabel(DebugGroups.Last());
@@ -287,39 +288,27 @@ void FMetalCommandEncoder::EndCommandBuffer()
 		
 		TSet<FMetalBufferPtr> NewActiveBuffers = MoveTemp(ActiveBuffers);
 		
-		FMetalCommandBufferCompletionHandler CompletionHander;
-		CompletionHander.BindLambda([&InDevice = Device, NewActiveBuffers](MTL::CommandBuffer*)
+		Context->GetContextSyncPoint()->OnCompletionCallback([&InDevice = Device, NewActiveBuffers]()
 		{
 			for (FMetalBufferPtr Buffer : NewActiveBuffers)
 			{
 				InDevice.RemoveActiveBuffer(Buffer->GetMTLBuffer(), Buffer->GetRange());
 			}
 		});
-		
-		AddCompletionHandler(CompletionHander);
 	}
 #endif
-#if ENABLE_METAL_GPUPROFILE
+	
+#if ENABLE_METAL_GPUPROFILE && RHI_NEW_GPU_PROFILER == 0
 	CommandBufferStats->End(CommandBuffer->GetMTLCmdBuffer());
 	CommandBufferStats = nullptr;
 #endif
 	
-	CommandList.FinalizeCommandBuffer(CommandBuffer, MoveTemp(CompletionHandlers));
-	CommandBuffers.Add(CommandBuffer);
+	CommandList.FinalizeCommandBuffer(CommandBuffer);
 	
 	CommandBuffer = nullptr;
 	EncoderNum = 0;
 }
 
-TArray<FMetalCommandBuffer*> FMetalCommandEncoder::Finalize()
-{
-	EndCommandBuffer();
-	
-	TArray<FMetalCommandBuffer*> Output = MoveTemp(CommandBuffers);
-	
-	return Output;
-}
-
 #pragma mark - Public Command Encoder Accessors -
 
 #if METAL_RHI_RAYTRACING
@@ -372,7 +361,7 @@ TRefCountPtr<FMetalFence> const& FMetalCommandEncoder::GetEncoderFence(void) con
 	
 #pragma mark - Public Command Encoder Mutators -
 
-void FMetalCommandEncoder::BeginRenderCommandEncoding(void)
+void FMetalCommandEncoder::BeginRenderCommandEncoding(FMetalCounterSampler* Sampler)
 {
 	check(RenderPassDesc);
 	check(CommandBuffer);
@@ -410,6 +399,11 @@ void FMetalCommandEncoder::BeginRenderCommandEncoding(void)
     // Clear Residency Cache (TODO: Move this to a separate function)
     ResourceUsage.Empty();
 
+	if(Sampler)
+	{
+		CommandBuffer->AddCounterSample(Sampler->SetupStageCounters(RenderPassDesc));
+	}
+	
     RenderCommandEncoder = NS::RetainPtr(CommandBuffer->GetMTLCmdBuffer()->renderCommandEncoder(RenderPassDesc));
 	EncoderNum++;
 
@@ -430,7 +424,7 @@ void FMetalCommandEncoder::BeginRenderCommandEncoding(void)
 	EncoderFence = CommandList.GetCommandQueue().CreateFence(Label);
 }
 
-void FMetalCommandEncoder::BeginComputeCommandEncoding(MTL::DispatchType DispatchType)
+void FMetalCommandEncoder::BeginComputeCommandEncoding(MTL::DispatchType DispatchType, FMetalCounterSampler* Sampler)
 {
 	check(CommandBuffer);
 
@@ -443,15 +437,18 @@ void FMetalCommandEncoder::BeginComputeCommandEncoding(MTL::DispatchType Dispatc
 	// Clear Residency Cache (TODO: Move this to a separate function)
     ResourceUsage.Empty();
 
-	if (DispatchType == MTL::DispatchTypeSerial)
-	{
-        ComputeCommandEncoder = NS::RetainPtr(CommandBuffer->GetMTLCmdBuffer()->computeCommandEncoder());
-	}
-	else
+	MTL::ComputePassDescriptor* ComputePassDesc = MTL::ComputePassDescriptor::alloc()->init();
+	ComputePassDesc->setDispatchType(DispatchType);
+	
+	if(Sampler)
 	{
-        ComputeCommandEncoder = NS::RetainPtr(CommandBuffer->GetMTLCmdBuffer()->computeCommandEncoder(DispatchType));
+		CommandBuffer->AddCounterSample(Sampler->SetupStageCounters(ComputePassDesc));
 	}
+	
+	ComputeCommandEncoder = NS::RetainPtr(CommandBuffer->GetMTLCmdBuffer()->computeCommandEncoder(ComputePassDesc));
     
+	ComputePassDesc->release();
+	
 	EncoderNum++;
 	
 	check(!EncoderFence);
@@ -471,7 +468,7 @@ void FMetalCommandEncoder::BeginComputeCommandEncoding(MTL::DispatchType Dispatc
 	EncoderFence = CommandList.GetCommandQueue().CreateFence(Label);
 }
 
-void FMetalCommandEncoder::BeginBlitCommandEncoding(void)
+void FMetalCommandEncoder::BeginBlitCommandEncoding(FMetalCounterSampler* Sampler)
 {
 	check(CommandBuffer);
 	check(IsRenderCommandEncoderActive() == false && IsComputeCommandEncoderActive() == false && IsBlitCommandEncoderActive() == false
@@ -486,7 +483,16 @@ void FMetalCommandEncoder::BeginBlitCommandEncoding(void)
 		CommandEncoderFence.FenceResources = MoveTemp(TransitionedResources);
 	}
 	
-    BlitCommandEncoder = NS::RetainPtr(CommandBuffer->GetMTLCmdBuffer()->blitCommandEncoder());
+	MTL::BlitPassDescriptor* BlitPassDesc = MTL::BlitPassDescriptor::alloc()->init();
+	
+	if(Sampler)
+	{
+		CommandBuffer->AddCounterSample(Sampler->SetupStageCounters(BlitPassDesc));
+	}
+	
+    BlitCommandEncoder = NS::RetainPtr(CommandBuffer->GetMTLCmdBuffer()->blitCommandEncoder(BlitPassDesc));
+	
+	BlitPassDesc->release();
 	
 	EncoderNum++;
 	
@@ -540,16 +546,6 @@ void FMetalCommandEncoder::BeginAccelerationStructureCommandEncoding(void)
 }
 #endif // METAL_RHI_RAYTRACING
 
-void FMetalCommandEncoder::SplitCommandBuffers()
-{
-	if(IsAnyCommandEncoderActive())
-	{
-		EndEncoding();
-	}
-	EndCommandBuffer();
-	StartCommandBuffer();	
-}
-
 TRefCountPtr<FMetalFence> FMetalCommandEncoder::EndEncoding(void)
 {
 	static bool bSupportsFences = Device.SupportsFeature(EMetalFeaturesFences);
@@ -708,23 +704,6 @@ TRefCountPtr<FMetalFence> FMetalCommandEncoder::EndEncoding(void)
     return Fence;
 }
 
-void FMetalCommandEncoder::InsertCommandBufferFence(TSharedPtr<FMetalCommandBufferFence, ESPMode::ThreadSafe>& Fence, FMetalCommandBufferCompletionHandler Handler)
-{
-	check(CommandBuffer);
-	
-	Fence = CommandBuffer->GetCompletionFence();
-	
-	if (Handler.IsBound())
-	{
-		AddCompletionHandler(Handler);
-	}
-}
-
-void FMetalCommandEncoder::AddCompletionHandler(FMetalCommandBufferCompletionHandler& Handler)
-{
-	CompletionHandlers.Add(Handler);
-}
-
 void FMetalCommandEncoder::UpdateFence(FMetalFence* Fence)
 {
 	check(IsRenderCommandEncoderActive() || IsComputeCommandEncoderActive() || IsBlitCommandEncoderActive()
diff --git a/Engine/Source/Runtime/Apple/MetalRHI/Private/MetalCommandEncoder.h b/Engine/Source/Runtime/Apple/MetalRHI/Private/MetalCommandEncoder.h
--- a/Engine/Source/Runtime/Apple/MetalRHI/Private/MetalCommandEncoder.h
+++ b/Engine/Source/Runtime/Apple/MetalRHI/Private/MetalCommandEncoder.h
@@ -7,14 +7,14 @@
 #include "MetalFence.h"
 #include "MetalPipeline.h"
 #include "MetalProfiler.h"
+#include "MetalCounterSampler.h"
 
 class FMetalCommandList;
 class FMetalCommandQueue;
 class FMetalGraphicsPipelineState;
-class FMetalCommandBufferFence;
 class FMetalDevice;
 struct FMetalCommandBufferStats;
-
+class FMetalRHICommandContext;
 class FMetalShaderPipeline;
 typedef TSharedPtr<FMetalShaderPipeline, ESPMode::ThreadSafe> FMetalShaderPipelinePtr;
 
@@ -173,12 +173,7 @@ public:
 	/**
 	 * Will deprecate in 5.6, splits the command buffer into chunks, required for profiling until we switch to encoder counter buffers
 	 */
-	void EndCommandBuffer();
-	
-	/**
-	 * Commit the existing command buffer if there is one & optionally waiting for completion, if there isn't a current command buffer this is a no-op.
- 	 */
-	TArray<FMetalCommandBuffer*> Finalize();
+	void EndCommandBuffer(FMetalRHICommandContext* Context);
 
 #pragma mark - Public Command Buffer Accessors -
 	
@@ -255,10 +250,10 @@ public:
 	/**
  	 * Begins encoding rendering commands into the current command buffer. No other encoder may be active
 	 */
-	void BeginRenderCommandEncoding(void);
+	void BeginRenderCommandEncoding(FMetalCounterSampler* Sampler);
 	
 	/** Begins encoding compute commands into the current command buffer. No other encoder may be active. */
-	void BeginComputeCommandEncoding(MTL::DispatchType Type = MTL::DispatchTypeSerial);
+	void BeginComputeCommandEncoding(MTL::DispatchType Type = MTL::DispatchTypeSerial, FMetalCounterSampler* Sampler = nullptr);
 
 #if METAL_RHI_RAYTRACING
 	/** Begins encoding acceleration structure commands into the current command buffer. No other encoder may be active. */
@@ -266,17 +261,11 @@ public:
 #endif // METAL_RHI_RAYTRACING
 	
 	/** Begins encoding blit commands into the current command buffer. No other encoder may be active. */
-	void BeginBlitCommandEncoding(void);
+	void BeginBlitCommandEncoding(FMetalCounterSampler* Sampler);
 	
 	/** Declare that all command generation from this encoder is complete, and detach from the MTLCommandBuffer if there is an encoder active or does nothing if there isn't. */
 	TRefCountPtr<FMetalFence> EndEncoding(void);
 	
-	/** Initialises a fence for the current command-buffer optionally adding a command-buffer completion handler to the command-buffer */
-	void InsertCommandBufferFence(TSharedPtr<FMetalCommandBufferFence, ESPMode::ThreadSafe>& Fence, FMetalCommandBufferCompletionHandler Handler);
-	
-	/** Adds a command-buffer completion handler to the command-buffer */
-	void AddCompletionHandler(FMetalCommandBufferCompletionHandler& Handler);
-	
 	/** Update the event to capture all GPU work so far enqueued by this encoder. */
 	void UpdateFence(FMetalFence* Fence);
 	
@@ -510,11 +499,6 @@ public:
 	 */
 	FMetalSubBufferRing& GetRingBuffer(void);
 	
-	/**
-	 * Splits the current command buffer
-	 */
-	void SplitCommandBuffers();
-	
 #pragma mark - Public Resource query Access -
 	
 private:
@@ -606,7 +590,6 @@ public:
     MTL::RenderPassDescriptor* RenderPassDesc = nullptr;
 	
     FMetalCommandBuffer* CommandBuffer = nullptr;
-	TArray<FMetalCommandBuffer*> CommandBuffers;
 	
     MTLRenderCommandEncoderPtr RenderCommandEncoder;
     MTLComputeCommandEncoderPtr ComputeCommandEncoder;
@@ -620,7 +603,6 @@ public:
 	FMetalCommandBufferStats* CommandBufferStats;
 #endif
 
-	TArray<FMetalCommandBufferCompletionHandler> CompletionHandlers;
 	TArray<NS::String*> DebugGroups;
     
     TSet<FMetalBufferPtr> ActiveBuffers;
diff --git a/Engine/Source/Runtime/Apple/MetalRHI/Private/MetalCommandList.cpp b/Engine/Source/Runtime/Apple/MetalRHI/Private/MetalCommandList.cpp
--- a/Engine/Source/Runtime/Apple/MetalRHI/Private/MetalCommandList.cpp
+++ b/Engine/Source/Runtime/Apple/MetalRHI/Private/MetalCommandList.cpp
@@ -205,36 +205,14 @@ void FMetalCommandList::HandleMetalCommandBufferFailure(MTL::CommandBuffer* Comp
 	}
 }
 
-void FMetalCommandList::FinalizeCommandBuffer(FMetalCommandBuffer* Buffer, TArray<FMetalCommandBufferCompletionHandler> CompletionHandlers)
+void FMetalCommandList::FinalizeCommandBuffer(FMetalCommandBuffer* Buffer)
 {
+#if RHI_NEW_GPU_PROFILER == 0
 	check(Buffer);
 	
-	FMetalCommandBufferTimer& Timer = FMetalGPUProfiler::GetFrameBufferTimer();
+	FMetalCommandBufferTimer& Timer = FMetalCommandBufferTimer::GetFrameBufferTimer();
 	Timer.Submit();
 	
-	MTL::HandlerFunction CompletionHandler = [CompletionHandlers, &InDevice = CommandQueue.GetDevice(), &FrameBufferTimer = Timer, Buffer](MTL::CommandBuffer* CompletedBuffer)
-	{
-		if (CompletedBuffer->status() == MTL::CommandBufferStatusError)
-		{
-			HandleMetalCommandBufferFailure(CompletedBuffer);
-		}
-		if (CompletionHandlers.Num())
-		{
-			for (FMetalCommandBufferCompletionHandler Handler : CompletionHandlers)
-			{
-				Handler.Execute(CompletedBuffer);
-			}
-		}
-
-		if (CompletedBuffer->status() == MTL::CommandBufferStatusCompleted)
-		{
-			FrameBufferTimer.AddTiming({CompletedBuffer->GPUStartTime(), CompletedBuffer->GPUEndTime()});
-		}
-		
-		InDevice.RemoveInflightCommandBuffer(CompletedBuffer);
-	};
-	
-	Buffer->GetMTLCmdBuffer()->addCompletedHandler(CompletionHandler);
-	
-	CommandQueue.GetDevice().AddInflightCommandBuffer(Buffer->GetMTLCmdBuffer().get());
+	Buffer->SetTimer(&Timer);
+#endif
 }
diff --git a/Engine/Source/Runtime/Apple/MetalRHI/Private/MetalCommandList.h b/Engine/Source/Runtime/Apple/MetalRHI/Private/MetalCommandList.h
--- a/Engine/Source/Runtime/Apple/MetalRHI/Private/MetalCommandList.h
+++ b/Engine/Source/Runtime/Apple/MetalRHI/Private/MetalCommandList.h
@@ -40,9 +40,8 @@ public:
 	/** 
 	 * Finalizes the command buffer ready for submission
 	 * @param Buffer The buffer to submit to the command-list.
-	 * @param CompletionHandlers The completion handlers that should be attached to this command-buffer.
 	 */
-	void FinalizeCommandBuffer(FMetalCommandBuffer* Buffer, TArray<FMetalCommandBufferCompletionHandler> CompletionHandlers);
+	void FinalizeCommandBuffer(FMetalCommandBuffer* Buffer);
 	
 #pragma mark - Public Command List Accessors -
 	
@@ -58,5 +57,8 @@ public:
 private:
 #pragma mark - Private Member Variables -
 	FMetalCommandQueue& CommandQueue;
+	
+#if RHI_NEW_GPU_PROFILER == 0
 	TSharedPtr<FMetalCommandBufferTiming, ESPMode::ThreadSafe> LastCompletedBufferTiming;
+#endif
 };
diff --git a/Engine/Source/Runtime/Apple/MetalRHI/Private/MetalCommandQueue.cpp b/Engine/Source/Runtime/Apple/MetalRHI/Private/MetalCommandQueue.cpp
--- a/Engine/Source/Runtime/Apple/MetalRHI/Private/MetalCommandQueue.cpp
+++ b/Engine/Source/Runtime/Apple/MetalRHI/Private/MetalCommandQueue.cpp
@@ -25,6 +25,14 @@ bool GMetalCommandBufferDebuggingEnabled = 0;
 
 #pragma mark - Public C++ Boilerplate -
 
+#if RHI_NEW_GPU_PROFILER
+FMetalTiming::FMetalTiming(FMetalCommandQueue& Queue) 
+	: Queue(Queue)
+	, EventStream(Queue.GetProfilerQueue())
+{
+}
+#endif
+
 FMetalCommandQueue::FMetalCommandQueue(FMetalDevice& MetalDevice, uint32 const MaxNumCommandBuffers /* = 0 */)
 	: Device(MetalDevice)
 	, RuntimeDebuggingLevel(EMetalDebugLevelOff)
@@ -87,11 +95,13 @@ FMetalCommandQueue::FMetalCommandQueue(FMetalDevice& MetalDevice, uint32 const M
 	PermittedOptions |= MTL::ResourceStorageModeMemoryless;
 #endif
 	PermittedOptions |= MTL::ResourceHazardTrackingModeTracked;
+	
+	SignalEvent.MetalEvent = Device.GetDevice()->newEvent();
 }
 
 FMetalCommandQueue::~FMetalCommandQueue(void)
 {
-	// void
+	SignalEvent.MetalEvent->release();
 }
 	
 #pragma mark - Public Command Buffer Mutators -
@@ -116,7 +126,7 @@ FMetalCommandBuffer* FMetalCommandQueue::CreateCommandBuffer(void)
     
     CmdBufferDesc->release();
                                                            
-    FMetalCommandBuffer* CommandBuffer = new FMetalCommandBuffer(CmdBuffer);
+    FMetalCommandBuffer* CommandBuffer = new FMetalCommandBuffer(CmdBuffer, *this);
     
 	INC_DWORD_STAT(STAT_MetalCommandBufferCreatedPerFrame);
 	return CommandBuffer;
@@ -212,3 +222,17 @@ void FMetalCommandQueue::InsertDebugCaptureBoundary(void)
 {
 	CommandQueue->insertDebugCaptureBoundary();
 }
+
+#if RHI_NEW_GPU_PROFILER
+UE::RHI::GPUProfiler::FQueue FMetalCommandQueue::GetProfilerQueue() const
+{
+	UE::RHI::GPUProfiler::FQueue Queue;
+	Queue.GPU = 0;
+	Queue.Index = 0;
+
+	// TODO - Carl: Multiple queues
+	Queue.Type = UE::RHI::GPUProfiler::FQueue::EType::Graphics;
+	
+	return Queue;
+}
+#endif
diff --git a/Engine/Source/Runtime/Apple/MetalRHI/Private/MetalCommandQueue.h b/Engine/Source/Runtime/Apple/MetalRHI/Private/MetalCommandQueue.h
--- a/Engine/Source/Runtime/Apple/MetalRHI/Private/MetalCommandQueue.h
+++ b/Engine/Source/Runtime/Apple/MetalRHI/Private/MetalCommandQueue.h
@@ -3,15 +3,34 @@
 #pragma once
 
 #include "MetalRHIPrivate.h"
+#include "MetalSubmission.h"
 #include "RHI.h"
 #include "RHIFeatureLevel.h"
 #include "Containers/LockFreeList.h"
 
 class FMetalCommandBuffer;
-class FMetalCommandBufferFence;
 class FMetalCommandList;
 class FMetalFence;
 class FMetalDevice;
+struct FMetalCounterSample;
+struct FMetalBreadcrumbEvent;
+class FMetalEventNode;
+
+#if RHI_NEW_GPU_PROFILER
+class FMetalTiming
+{
+public:
+	FMetalCommandQueue& Queue;
+
+	// Timer calibration data
+	uint64 GPUFrequency = 0, GPUTimestamp = 0;
+	uint64 CPUFrequency = 0, CPUTimestamp = 0;
+
+	UE::RHI::GPUProfiler::FEventStream EventStream;
+
+	FMetalTiming(FMetalCommandQueue& Queue);
+};
+#endif
 
 /**
  * FMetalCommandQueue:
@@ -74,11 +93,51 @@ public:
 
 	/** Inserts a boundary that marks the end of a frame for the debug capture tool. */
 	void InsertDebugCaptureBoundary(void);
+	
+	struct : public TQueue<FMetalPayload*, EQueueMode::Mpsc>
+	{
+		FMetalPayload* Peek()
+		{
+			if (FMetalPayload** Result = TQueue::Peek())
+			{
+				return *Result;
+			}
+
+			return nullptr;
+		}		
+	} PendingSubmission, PendingInterrupt;
 
+	FMetalPayload* PayloadToSubmit  = nullptr;
+	
+	using FPayloadArray = TArray<FMetalPayload*>;
+
+	// Batches the current payload's command lists, returning the latest fence value signaled for this queue.
+	uint64 FinalizePayload(bool bRequiresSignal, FPayloadArray& PayloadsToHandDown);
+	
+	FMetalSignalEvent& GetSignalEvent()
+	{
+		return SignalEvent;
+	}
+
+	TArray<FMetalRHIRenderQuery*> TimestampQueries;
+	TArray<FMetalRHIRenderQuery*> OcclusionQueries;
+	TArray<FMetalCounterSamplePtr> CounterSamples;
+	
+#if RHI_NEW_GPU_PROFILER == 0
+	TMap<FMetalEventNode*, TArray<FMetalCounterSamplePtr>> EventSampleCounters;
+#endif
+	
+#if RHI_NEW_GPU_PROFILER
+	UE::RHI::GPUProfiler::FQueue GetProfilerQueue() const;
+	FMetalTiming* Timing = nullptr;
+#endif
+	
 private:
 #pragma mark - Private Member Variables -
 	FMetalDevice& Device;
 	MTL::CommandQueue* CommandQueue;
 	int32 RuntimeDebuggingLevel;
 	static NS::UInteger PermittedOptions;
+	
+	FMetalSignalEvent SignalEvent;
 };
diff --git a/Engine/Source/Runtime/Apple/MetalRHI/Private/MetalCounterSampler.cpp b/Engine/Source/Runtime/Apple/MetalRHI/Private/MetalCounterSampler.cpp
new file mode 100644
--- /dev/null
+++ b/Engine/Source/Runtime/Apple/MetalRHI/Private/MetalCounterSampler.cpp
@@ -0,0 +1,234 @@
+// Copyright Epic Games, Inc. All Rights Reserved.
+
+#include "MetalCounterSampler.h"
+#include "MetalDynamicRHI.h"
+
+static TAutoConsoleVariable<bool> CVarMetalRHIInsertCounterSampleBarrier(
+	TEXT("rhi.Metal.InsertCounterSampleBarrier"),
+	true,
+	TEXT("Whether to insert a counter sampler barrier to provide the most accurate timings. (default: true)\n"),
+	ECVF_ReadOnly);
+
+void FMetalCounterSample::ResolveStageCounters(uint64_t& InStartTime, uint64_t& InEndTime)
+{
+	if(!bResolved)
+	{
+		NS::Data* Data;
+		
+		if(SampleType == EMetalCounterSampleType::RenderStage)
+		{
+			NS::Range Range = NS::Range::Make(Offset, 4);
+			Data = SampleBuffer->resolveCounterRange(Range);
+		}
+		else
+		{
+			NS::Range Range = NS::Range::Make(Offset, 2);
+			Data = SampleBuffer->resolveCounterRange(Range);
+		}
+		
+		// Convert the contents of the counter sample buffer into the standard data format.
+		check(Data);
+		if (Data == nullptr) 
+		{
+			return;
+		}
+		
+		MTL::CounterResultTimestamp* Timestamps = (MTL::CounterResultTimestamp *)(Data->mutableBytes());
+		
+		if(SampleType == EMetalCounterSampleType::RenderStage)
+		{
+			uint64_t StartTimeVertex = Timestamps[0].timestamp;
+			uint64_t EndTimeVertex = Timestamps[1].timestamp;
+			uint64_t StartTimeFragment = Timestamps[2].timestamp;
+			uint64_t EndTimeFragment = Timestamps[3].timestamp;
+			
+			StartTime = StartTimeVertex;
+			EndTime = EndTimeFragment;
+		}
+		else
+		{
+			StartTime = Timestamps[0].timestamp;
+			EndTime = Timestamps[1].timestamp;
+		}
+		
+		double OneOverNanosToCycles = 1.0 / 1000000000.0 / FPlatformTime::GetSecondsPerCycle64();
+		StartTime = FMath::TruncToInt64(StartTime * OneOverNanosToCycles);
+		EndTime = FMath::TruncToInt64(EndTime * OneOverNanosToCycles);
+		
+		bResolved = true;
+	}
+	
+	InStartTime = StartTime;
+	InEndTime = EndTime;
+}
+
+void FMetalCounterSample::ResolveBoundaryCounter(uint64_t& Time)
+{
+	if(!bResolved)
+	{
+		NS::Range Range = NS::Range::Make(Offset, 1);
+		
+		// Convert the contents of the counter sample buffer into the standard data format.
+		NS::Data* Data = SampleBuffer->resolveCounterRange(Range);
+		check(Data);
+		if (Data == nullptr) 
+		{
+			return;
+		}
+		
+		MTL::CounterResultTimestamp* Timestamps = (MTL::CounterResultTimestamp *)(Data->mutableBytes());
+		EndTime = Timestamps->timestamp;
+		
+		bResolved = true;
+	}
+	Time = EndTime;
+}
+
+TLockFreePointerListUnordered<void, PLATFORM_CACHE_LINE_SIZE> FMetalCounterSample::MemoryPool;
+
+FMetalCounterSampler::FMetalCounterSampler(FMetalDevice* InDevice, uint32_t SampleCount) : Device(InDevice)
+{
+	Size = SampleCount;
+}
+
+FMetalCounterSampler::~FMetalCounterSampler()
+{
+	SampleBuffer->release();
+}
+
+MTL::CounterSampleBuffer* FMetalCounterSampler::SwapOrAllocateBuffer(uint32_t SampleSize, uint32_t& OutOffset)
+{
+	OutOffset = Offset;
+	
+	Offset += SampleSize;
+	if(Offset > Size || !SampleBuffer)
+	{
+		// Re-add Counter Sampler on next frame
+		if(SampleBuffer)
+		{
+			FMetalDynamicRHI::Get().DeferredDelete([this, InSampleBuffer=SampleBuffer](){
+				SampleBufferFreePool.Add(InSampleBuffer);
+			});
+		}
+		
+		if(SampleBufferFreePool.Num() > 0)
+		{
+			SampleBuffer = SampleBufferFreePool.Pop();
+			Offset = 0;
+			OutOffset = 0;
+			return SampleBuffer;
+		}
+		
+		MTL::CounterSampleBufferDescriptor* BufferDesc;
+		BufferDesc = MTL::CounterSampleBufferDescriptor::alloc()->init();
+		
+		NS::Array* CounterSets = Device->GetDevice()->counterSets();
+		const MTL::CounterSet* CounterSet = (MTL::CounterSet*)CounterSets->object(0);
+		BufferDesc->setCounterSet(CounterSet);
+		BufferDesc->setStorageMode(MTL::StorageModeShared);
+		
+		Offset = 0;
+		OutOffset = 0;
+		
+		BufferDesc->setSampleCount(Size);
+		
+		NS::Error* DeviceError;
+		SampleBuffer = Device->GetDevice()->newCounterSampleBuffer(BufferDesc, &DeviceError);
+		check(SampleBuffer);
+		
+		BufferDesc->release();
+	}
+	
+	check(Offset <= Size);
+	
+	return SampleBuffer;
+}
+
+FMetalCounterSamplePtr FMetalCounterSampler::SetupStageCounters(MTL::ComputePassDescriptor* ComputePassDesc)
+{
+	uint32_t OutOffset;
+	MTL::CounterSampleBuffer* Buffer = SwapOrAllocateBuffer(2, OutOffset);
+	check(Buffer);
+	
+	FMetalCounterSamplePtr Sample = FMetalCounterSamplePtr(new FMetalCounterSample(EMetalCounterSampleType::ComputeStage, Buffer, OutOffset));
+	MTL::ComputePassSampleBufferAttachmentDescriptor* SampleDesc = ComputePassDesc->sampleBufferAttachments()->object(0);
+	
+	SampleDesc->setSampleBuffer(Buffer);
+	SampleDesc->setStartOfEncoderSampleIndex(OutOffset);
+	SampleDesc->setEndOfEncoderSampleIndex(OutOffset+1);
+	
+	return Sample;
+}
+
+FMetalCounterSamplePtr FMetalCounterSampler::SetupStageCounters(MTL::BlitPassDescriptor* BlitPassDesc)
+{
+	uint32_t OutOffset;
+	MTL::CounterSampleBuffer* Buffer = SwapOrAllocateBuffer(2, OutOffset);
+	check(Buffer);
+	
+	FMetalCounterSamplePtr Sample = FMetalCounterSamplePtr(new FMetalCounterSample(EMetalCounterSampleType::BlitStage, Buffer, OutOffset));
+	MTL::BlitPassSampleBufferAttachmentDescriptor* SampleDesc = BlitPassDesc->sampleBufferAttachments()->object(0);
+	
+	SampleDesc->setSampleBuffer(Buffer);
+	SampleDesc->setStartOfEncoderSampleIndex(OutOffset);
+	SampleDesc->setEndOfEncoderSampleIndex(OutOffset+1);
+	
+	return Sample;
+}
+
+FMetalCounterSamplePtr FMetalCounterSampler::SetupStageCounters(MTL::RenderPassDescriptor* InRenderPassDesc)
+{
+	uint32_t OutOffset;
+	MTL::CounterSampleBuffer* Buffer = SwapOrAllocateBuffer(4, OutOffset);
+	check(Buffer);
+	
+	FMetalCounterSamplePtr Sample = FMetalCounterSamplePtr(new FMetalCounterSample(EMetalCounterSampleType::RenderStage, Buffer, OutOffset));
+	MTL::RenderPassSampleBufferAttachmentDescriptor* SampleDesc = InRenderPassDesc->sampleBufferAttachments()->object(0);
+	
+	SampleDesc->setSampleBuffer(Buffer);
+	SampleDesc->setStartOfVertexSampleIndex(OutOffset);
+	SampleDesc->setEndOfVertexSampleIndex(OutOffset+1);
+	SampleDesc->setStartOfFragmentSampleIndex(OutOffset+2);
+	SampleDesc->setEndOfFragmentSampleIndex(OutOffset+3);
+	
+	return Sample;
+}
+
+FMetalCounterSamplePtr FMetalCounterSampler::SetupBoundaryCounters(MTL::RenderCommandEncoder* RenderCommandEncoder)
+{
+	uint32_t OutOffset;
+	MTL::CounterSampleBuffer* Buffer = SwapOrAllocateBuffer(1, OutOffset);
+	check(Buffer);
+	
+	FMetalCounterSamplePtr Sample = FMetalCounterSamplePtr(new FMetalCounterSample(EMetalCounterSampleType::DrawBoundary, Buffer, OutOffset));
+	
+	RenderCommandEncoder->sampleCountersInBuffer(Buffer, OutOffset++, CVarMetalRHIInsertCounterSampleBarrier.GetValueOnAnyThread());
+	
+	return Sample;
+}
+
+FMetalCounterSamplePtr FMetalCounterSampler::SetupBoundaryCounters(MTL::ComputeCommandEncoder* ComputeCommandEncoder)
+{
+	uint32_t OutOffset;
+	MTL::CounterSampleBuffer* Buffer = SwapOrAllocateBuffer(1, OutOffset);
+	check(Buffer);
+	
+	FMetalCounterSamplePtr Sample = FMetalCounterSamplePtr(new FMetalCounterSample(EMetalCounterSampleType::DispatchBoundary, Buffer, OutOffset));
+	
+	ComputeCommandEncoder->sampleCountersInBuffer(Buffer, OutOffset++, CVarMetalRHIInsertCounterSampleBarrier.GetValueOnAnyThread());
+	
+	return Sample;
+}
+
+FMetalCounterSamplePtr FMetalCounterSampler::SetupBoundaryCounters(MTL::BlitCommandEncoder* BlitCommandEncoder)
+{
+	uint32_t OutOffset;
+	MTL::CounterSampleBuffer* Buffer = SwapOrAllocateBuffer(1, OutOffset);
+	check(Buffer);
+	
+	FMetalCounterSamplePtr Sample = FMetalCounterSamplePtr(new FMetalCounterSample(EMetalCounterSampleType::BlitBoundary, Buffer, OutOffset));
+	
+	BlitCommandEncoder->sampleCountersInBuffer(Buffer, OutOffset++, CVarMetalRHIInsertCounterSampleBarrier.GetValueOnAnyThread());
+	
+	return Sample;
+}
diff --git a/Engine/Source/Runtime/Apple/MetalRHI/Private/MetalCounterSampler.h b/Engine/Source/Runtime/Apple/MetalRHI/Private/MetalCounterSampler.h
new file mode 100644
--- /dev/null
+++ b/Engine/Source/Runtime/Apple/MetalRHI/Private/MetalCounterSampler.h
@@ -0,0 +1,86 @@
+// Copyright Epic Games, Inc. All Rights Reserved.
+
+#pragma once
+
+#include "MetalThirdParty.h"
+#include "Containers/LockFreeList.h"
+
+enum class EMetalCounterSampleType
+{
+	RenderStage = 0,
+	ComputeStage,
+	BlitStage,
+	DrawBoundary,
+	DispatchBoundary,
+	BlitBoundary,
+};
+
+class FMetalDevice;
+struct FMetalCounterSample
+{
+	FMetalCounterSample(EMetalCounterSampleType Type, MTL::CounterSampleBuffer* Buffer, uint32_t InOffset) 
+		: SampleType(Type)
+		, Offset(InOffset)
+		, SampleBuffer(Buffer)
+	{}
+	
+	bool IsStageType()
+	{
+		return SampleType <= EMetalCounterSampleType::BlitStage;
+	}
+	
+	void ResolveStageCounters(uint64_t& InStartTime, uint64_t& InEndTime);
+	void ResolveBoundaryCounter(uint64_t& Time);
+	
+	void* operator new(size_t Size)
+	{
+		check(Size == sizeof(FMetalCounterSample));
+
+		void* Memory = MemoryPool.Pop();
+		if (Memory == nullptr)
+		{
+			Memory = FMemory::Malloc(sizeof(FMetalCounterSample), alignof(FMetalCounterSample));
+		}
+
+		return Memory;
+	}
+
+	void operator delete(void* Pointer)
+	{
+		MemoryPool.Push(Pointer);
+	}
+
+	static TLockFreePointerListUnordered<void, PLATFORM_CACHE_LINE_SIZE> MemoryPool;
+	
+	EMetalCounterSampleType SampleType;
+	uint32_t Offset = 0;
+	MTL::CounterSampleBuffer* SampleBuffer;
+	bool bResolved = false;
+	uint64_t StartTime, EndTime;
+};
+
+typedef TSharedPtr<FMetalCounterSample> FMetalCounterSamplePtr;
+
+class FMetalCounterSampler
+{
+public:
+	FMetalCounterSampler(FMetalDevice* InDevice, uint32_t SampleCount);
+	~FMetalCounterSampler();
+	
+	FMetalCounterSamplePtr SetupStageCounters(MTL::ComputePassDescriptor* ComputePassDesc);
+	FMetalCounterSamplePtr SetupStageCounters(MTL::BlitPassDescriptor* BlitPassDesc);
+	FMetalCounterSamplePtr SetupStageCounters(MTL::RenderPassDescriptor* InRenderPassDesc);
+	FMetalCounterSamplePtr SetupBoundaryCounters(MTL::RenderCommandEncoder* RenderCommandEncoder);
+	FMetalCounterSamplePtr SetupBoundaryCounters(MTL::ComputeCommandEncoder* ComputeCommandEncoder);
+	FMetalCounterSamplePtr SetupBoundaryCounters(MTL::BlitCommandEncoder* BlitCommandEncoder);
+	
+	MTL::CounterSampleBuffer* SwapOrAllocateBuffer(uint32_t SampleSize, uint32_t& OutOffset);
+	
+private:
+	FMetalDevice* Device;
+	MTL::CounterSampleBuffer* SampleBuffer = nullptr;
+	TArray<MTL::CounterSampleBuffer*> SampleBufferFreePool;
+	
+	uint32_t Offset = 0;
+	uint32_t Size = 0;
+};
diff --git a/Engine/Source/Runtime/Apple/MetalRHI/Private/MetalDevice.cpp b/Engine/Source/Runtime/Apple/MetalRHI/Private/MetalDevice.cpp
--- a/Engine/Source/Runtime/Apple/MetalRHI/Private/MetalDevice.cpp
+++ b/Engine/Source/Runtime/Apple/MetalRHI/Private/MetalDevice.cpp
@@ -350,12 +350,15 @@ FMetalDevice::FMetalDevice(MTL::Device* MetalDevice, uint32 InDeviceIndex)
 		
 	EnumerateFeatureSupport();
 	
-	CommandQueue = new FMetalCommandQueue(*this, GMetalCommandQueueSize);
-	check(CommandQueue);
+	for(uint32_t Idx = 0; Idx < (uint32_t)EMetalQueueType::Count; ++Idx)
+	{
+		CommandQueues.Add(new FMetalCommandQueue(*this, GMetalCommandQueueSize));
+		check(CommandQueues[Idx]);
+	}
 		
 	RuntimeDebuggingLevel = GMetalRuntimeDebugLevel;
 	
-	CaptureManager = new FMetalCaptureManager(MetalDevice, *CommandQueue);
+	CaptureManager = new FMetalCaptureManager(MetalDevice, *CommandQueues[(uint32_t)EMetalQueueType::Direct]);
 	
 	// If the separate present thread is enabled then an intermediate backbuffer is required
 	check(!GMetalSeparatePresentThread || GMetalSupportsIntermediateBackBuffer);
@@ -402,7 +405,8 @@ FMetalDevice::FMetalDevice(MTL::Device* MetalDevice, uint32 InDeviceIndex)
     BindlessDescriptorManager = new FMetalBindlessDescriptorManager(*this);
 #endif
 	
-	Heap.Init(GetCommandQueue());
+	CounterSampler = new FMetalCounterSampler(this, 4096);
+	Heap.Init(GetCommandQueue(EMetalQueueType::Direct));
 	
 	FrameSemaphore = dispatch_semaphore_create(FParse::Param(FCommandLine::Get(),TEXT("gpulockstep")) ? 1 : 3);
 }
@@ -412,11 +416,16 @@ FMetalDevice::~FMetalDevice()
 	FRHICommandListImmediate& RHICmdList = FRHICommandListImmediate::Get();
 	RHICmdList.SubmitAndBlockUntilGPUIdle();
 	
-	delete &(GetCommandQueue());
+	for(uint32_t Idx = 0; Idx < (uint32_t)EMetalQueueType::Count; ++Idx)
+	{
+		delete CommandQueues[Idx];
+	}
+
 	delete PSOManager;
     delete UniformBufferAllocator;
 	delete CaptureManager;
-	
+	delete CounterSampler;
+
     ShutdownPipelineCache();
     
 #if METAL_RHI_RAYTRACING
@@ -613,6 +622,19 @@ void FMetalDevice::EnumerateFeatureSupport()
 		Features |= EMetalFeaturesValidation;
 	}
 #endif
+	
+	// Counter Sampling Features
+	if(Device->supportsCounterSampling(MTL::CounterSamplingPointAtStageBoundary))
+	{
+		Features |= EMetalFeaturesStageCounterSampling;
+	}
+		
+	if(Device->supportsCounterSampling(MTL::CounterSamplingPointAtDrawBoundary) &&
+	   Device->supportsCounterSampling(MTL::CounterSamplingPointAtDispatchBoundary) &&
+	   Device->supportsCounterSampling(MTL::CounterSamplingPointAtBlitBoundary))
+	{
+		Features |= EMetalFeaturesBoundaryCounterSampling;
+	}
 }
 
 void FMetalDevice::EndDrawingViewport(bool bPresent)
diff --git a/Engine/Source/Runtime/Apple/MetalRHI/Private/MetalDevice.h b/Engine/Source/Runtime/Apple/MetalRHI/Private/MetalDevice.h
--- a/Engine/Source/Runtime/Apple/MetalRHI/Private/MetalDevice.h
+++ b/Engine/Source/Runtime/Apple/MetalRHI/Private/MetalDevice.h
@@ -9,12 +9,21 @@
 #include "MetalCaptureManager.h"
 #include "MetalTempAllocator.h"
 #include "MetalStateCache.h"
+#include "MetalCounterSampler.h"
 
 #if PLATFORM_IOS
 #include "IOS/IOSView.h"
 #endif
 #include "Containers/LockFreeList.h"
 
+// Defines a unique command queue type within a Metal Device (owner by the command list managers).
+// Currently only implements direct
+enum class EMetalQueueType
+{
+	Direct = 0,
+	Count,
+};
+
 /**
  * Enumeration of features which are present only on some OS/device combinations.
  * These have to be checked at runtime as well as compile time to ensure backward compatibility.
@@ -77,6 +86,10 @@ typedef NS_OPTIONS(uint64, EMetalFeatures)
 	EMetalFeaturesUnusedReservedBit27 = 1llu << 26llu, // was EMetalFeaturesSeparateTessellation
 	/** Supports indirect argument buffers Tier 2 */
 	EMetalFeaturesTier2IABs = 1llu << 27llu,
+	/** Supports counter sampling on encoder stages */
+	EMetalFeaturesStageCounterSampling = 1llu << 28llu,
+	/** Supports counter sampling on the stage boundary */
+	EMetalFeaturesBoundaryCounterSampling = 1llu << 29llu,
 };
 
 /**
@@ -132,9 +145,10 @@ public:
         return FrameNumberRHIThread;
     }
 	
-	FMetalCommandQueue& GetCommandQueue()
+	FMetalCommandQueue& GetCommandQueue(EMetalQueueType QueueType)
 	{
-		return *CommandQueue;
+		check(QueueType < EMetalQueueType::Count);
+		return *CommandQueues[(uint32_t)QueueType];
 	}
 	
 #if PLATFORM_SUPPORTS_BINDLESS_RENDERING
@@ -160,38 +174,6 @@ public:
 		return RuntimeDebuggingLevel;
 	}
 	
-	void AddInflightCommandBuffer(MTL::CommandBuffer* CmdBuffer)
-	{
-		FScopeLock Lock(&InflightCommandBufferCriticalSection);
-		InflightCommandBuffers.Add(CmdBuffer);
-	}
-	
-	void RemoveInflightCommandBuffer(MTL::CommandBuffer* CmdBuffer)
-	{
-		FScopeLock Lock(&InflightCommandBufferCriticalSection);
-		InflightCommandBuffers.Remove(CmdBuffer);
-	}
-	
-	void WaitForGPUIdle()
-	{
-		TArray<MTL::CommandBuffer*> TempInflightCommandBuffers;
-		{
-			FScopeLock Lock(&InflightCommandBufferCriticalSection);
-			TempInflightCommandBuffers = InflightCommandBuffers;
-		}
-		
-		for(MTL::CommandBuffer* CmdBuffer : TempInflightCommandBuffers)
-		{
-			// Possible for the command buffer submission to be in progress, yield until committed
-			while(CmdBuffer->status() < MTL::CommandBufferStatusCommitted)
-			{
-				FPlatformProcess::Yield();
-			}
-			
-			CmdBuffer->waitUntilCompleted();
-		}
-	}
-	
 	void IncrementFrameRHIThread()
 	{
 		FrameNumberRHIThread++;
@@ -202,6 +184,11 @@ public:
 		return FrameSemaphore;
 	}
 	
+	FMetalCounterSampler* GetCounterSampler()
+	{
+		return CounterSampler;
+	}
+	
 private:
 	FMetalDevice(MTL::Device* MetalDevice, uint32 DeviceIndex);
 	
@@ -210,7 +197,7 @@ private:
 private:
 	MTL::Device* Device;
 	
-	FMetalCommandQueue* CommandQueue;
+	TArray<FMetalCommandQueue*, TInlineAllocator<(uint32)EMetalQueueType::Count>> CommandQueues;
 	
 	/** A sempahore used to ensure that wait for previous frames to complete if more are in flight than we permit */
 	dispatch_semaphore_t FrameSemaphore;
@@ -256,9 +243,8 @@ private:
     uint32 FrameNumberRHIThread = 0;
 
 	int32 RuntimeDebuggingLevel = 0;
-
-	FCriticalSection 			   InflightCommandBufferCriticalSection;
-	TArray<MTL::CommandBuffer*>	   InflightCommandBuffers;
+	
+	FMetalCounterSampler* CounterSampler = nullptr;
 	
 #if PLATFORM_SUPPORTS_BINDLESS_RENDERING
     /** Bindless Descriptor Heaps manager. */
diff --git a/Engine/Source/Runtime/Apple/MetalRHI/Private/MetalDynamicRHI.cpp b/Engine/Source/Runtime/Apple/MetalRHI/Private/MetalDynamicRHI.cpp
--- a/Engine/Source/Runtime/Apple/MetalRHI/Private/MetalDynamicRHI.cpp
+++ b/Engine/Source/Runtime/Apple/MetalRHI/Private/MetalDynamicRHI.cpp
@@ -268,17 +268,3 @@ bool FMetalDynamicRHI::RHIMatchPrecachePSOInitializers(const FGraphicsPipelineSt
 
 	return true;
 }
-
-void FMetalDynamicRHI::AddDeferredDeleteFence(TSharedPtr<FMetalCommandBufferFence, ESPMode::ThreadSafe> Fence)
-{
-	FScopeLock Lock(&ObjectsToDeleteCS);
-	DeferredDeleteFences.Add(Fence);
-}
-
-void FMetalDynamicRHI::GatherDeferredDeleteObjects(TArray<FMetalDeferredDeleteObject>& DeferredDeleteObjects,
-											   TArray<TSharedPtr<FMetalCommandBufferFence, ESPMode::ThreadSafe>>& WaitFences)
-{
-	FScopeLock Lock(&ObjectsToDeleteCS);
-	WaitFences = MoveTemp(DeferredDeleteFences);
-	DeferredDeleteObjects = MoveTemp(ObjectsToDelete);
-}
diff --git a/Engine/Source/Runtime/Apple/MetalRHI/Private/MetalDynamicRHI.h b/Engine/Source/Runtime/Apple/MetalRHI/Private/MetalDynamicRHI.h
--- a/Engine/Source/Runtime/Apple/MetalRHI/Private/MetalDynamicRHI.h
+++ b/Engine/Source/Runtime/Apple/MetalRHI/Private/MetalDynamicRHI.h
@@ -15,6 +15,7 @@
 #include "MetalResources.h"
 #include "MetalRHIContext.h"
 #include "MetalViewport.h"
+#include "MetalSubmission.h"
 
 class FMetalDevice;
 
@@ -62,6 +63,23 @@ class FMetalDynamicRHI : public FDynamicRHI
 {
 	static inline FMetalDynamicRHI* Singleton = nullptr;
 public:
+	enum class EQueueStatus
+	{
+		None = 0,
+
+		// Work was processed through the queue.
+		Processed = 1 << 0,
+
+		// The queue has further, unprocessed work.
+		Pending = 1 << 1
+	};
+	FRIEND_ENUM_CLASS_FLAGS(EQueueStatus);
+
+	struct FProcessResult
+	{
+		EQueueStatus Status = EQueueStatus::None;
+	};
+	
 	static inline FMetalDynamicRHI& Get() { return *Singleton; }
 
 	/** Initialization constructor. */
@@ -205,6 +223,9 @@ public:
 	virtual void RHIBeginRenderQuery_TopOfPipe(FRHICommandListBase& RHICmdList, FRHIRenderQuery* RenderQuery) override;
 	virtual void RHIEndRenderQuery_TopOfPipe  (FRHICommandListBase& RHICmdList, FRHIRenderQuery* RenderQuery) override;
 	
+	virtual void RHIBeginRenderQueryBatch_TopOfPipe(FRHICommandListBase& RHICmdList, ERenderQueryType QueryType) final override;
+	virtual void RHIEndRenderQueryBatch_TopOfPipe(FRHICommandListBase& RHICmdList, ERenderQueryType QueryType) final override;
+	
 	template <typename ...Args>
 	void DeferredDelete(Args&&... InArgs)
 	{
@@ -213,10 +234,48 @@ public:
 		ObjectsToDelete.Emplace(Forward<Args>(InArgs)...);
 	}
 	
-	void AddDeferredDeleteFence(TSharedPtr<FMetalCommandBufferFence, ESPMode::ThreadSafe> Fence);
-	void GatherDeferredDeleteObjects(TArray<FMetalDeferredDeleteObject>& DeferredDeleteObjects,
-									 TArray<TSharedPtr<FMetalCommandBufferFence, ESPMode::ThreadSafe>>& WaitFences);
-	void ProcessDeferredDeleteQueue();
+	void SubmitCommands(TConstArrayView<struct FMetalFinalizedCommands*> Commands);
+	void SubmitPayloads(TArray<FMetalPayload*>&& Payloads);
+	
+	// Ensures all prior batched command lists have reached the MetalQueue
+	void FlushBatchedPayloads(FMetalCommandQueue::FPayloadArray& PayloadsToHandDown);
+
+	// Processes the interrupt queue on the calling thread, until the specified GraphEvent is signaled.
+	// If the GraphEvent is nullptr, processes the queue until no further progress is made.
+	void ProcessInterruptQueueUntil(FGraphEvent* GraphEvent);
+	
+	void ProcessPendingCommandBuffers();
+
+	// Called when the GPU has crashed. This function will not return.
+	void ProcessInterruptQueueOnGPUCrash();
+
+#if RHI_NEW_GPU_PROFILER
+	struct FMetalTimingArray : public TArray<TUniquePtr<FMetalTiming>, TInlineAllocator<GMetalMaxNumQueues>>
+	{
+		FMetalTiming* CreateNew(FMetalCommandQueue& Queue)
+		{
+			return Emplace_GetRef(MakeUnique<FMetalTiming>(Queue)).Get();
+		}
+	};
+
+	FMetalTimingArray CurrentTimingPerQueue;
+#endif
+	
+	void InitializeSubmissionPipe();
+	void ShutdownSubmissionPipe();
+	
+	// Inserts a task graph task which is executed once all previously submitted GPU work has completed (across all queues, device and adapters).
+	void EnqueueEndOfPipeTask(TUniqueFunction<void()> TaskFunc, TUniqueFunction<void(FMetalPayload&)> ModifyPayloadCallback = {});
+	FGraphEventRef EopTask;
+	
+	FCriticalSection SubmissionCS;
+	FCriticalSection InterruptCS;
+
+	FProcessResult ProcessSubmissionQueue();
+	FProcessResult ProcessInterruptQueue();
+	
+	// Enumerates all queues across all devices and active adapters
+	void ForEachQueue(TFunctionRef<void(FMetalCommandQueue&)> Callback);
 
 private:
 	FMetalDevice* Device;
@@ -225,20 +284,20 @@ private:
 	TMap<uint32, FVertexDeclarationRHIRef> VertexDeclarationCache;
 	TLockFreePointerListUnordered<FMetalRHICommandContext, PLATFORM_CACHE_LINE_SIZE> MetalCommandContextPool;
 	
-	struct FDeferredDeleteData
-	{
-		TArray<FMetalDeferredDeleteObject> DeferredDeleteObjects;
-		TArray<TSharedPtr<FMetalCommandBufferFence, ESPMode::ThreadSafe>> WaitFences;
-		bool bWaitExtraFrame = false;
-	};
-	
-	TArray<FDeferredDeleteData> DeferredDeleteQueue;
-	
 	FCriticalSection ObjectsToDeleteCS;
 	TArray<FMetalDeferredDeleteObject> ObjectsToDelete;
-	TArray<TSharedPtr<FMetalCommandBufferFence, ESPMode::ThreadSafe>> DeferredDeleteFences;
+	
+	TArray<FMetalCommandBuffer*> CmdBuffersPendingCompletion;
 	
 #if METAL_USE_METAL_SHADER_CONVERTER
     struct IRCompiler* CompilerInstance;
 #endif
+	
+	friend class FMetalThread;
+	class FMetalThread* SubmissionThread = nullptr;
+	class FMetalThread* InterruptThread = nullptr;
+	
+	TQueue<TArray<FMetalPayload*>*, EQueueMode::Mpsc> PendingPayloadsForSubmission;
 };
+
+ENUM_CLASS_FLAGS(FMetalDynamicRHI::EQueueStatus);
diff --git a/Engine/Source/Runtime/Apple/MetalRHI/Private/MetalProfiler.cpp b/Engine/Source/Runtime/Apple/MetalRHI/Private/MetalProfiler.cpp
--- a/Engine/Source/Runtime/Apple/MetalRHI/Private/MetalProfiler.cpp
+++ b/Engine/Source/Runtime/Apple/MetalRHI/Private/MetalProfiler.cpp
@@ -62,11 +62,131 @@
 int64 volatile GMetalGPUIdleTime = 0;
 int64 volatile GMetalPresentTime = 0;
 
+#if RHI_NEW_GPU_PROFILER && WITH_RHI_BREADCRUMBS
+FMetalBreadcrumbProfiler* FMetalBreadcrumbProfiler::Instance = nullptr;
+#endif
+
 void WriteString(FArchive* OutputFile, const char* String)
 {
 	OutputFile->Serialize((void*)String, sizeof(ANSICHAR)*FCStringAnsi::Strlen(String));
 }
 
+#if RHI_NEW_GPU_PROFILER == 0
+
+void FMetalCommandBufferTimer::AddTiming(FMetalCommandBufferTiming Timing)
+{
+	FScopeLock Lock(&Mutex);
+	Counter--;
+	Timings.Add(Timing);
+	
+	if(Counter == 0 && bFrameEnded)
+	{
+		RecordFrame();
+	}
+}
+
+void FMetalCommandBufferTimer::FrameEnd()
+{
+	FScopeLock Lock(&Mutex);
+	bFrameEnded = true;
+	
+	if(Counter == 0)
+	{
+		RecordFrame();
+	}
+}
+
+void FMetalCommandBufferTimer::RecordFrame()
+{
+	Timings.Sort();
+	
+	FMetalCommandBufferTiming LastBufferTiming;
+	const double CyclesPerSecond = 1.0 / FPlatformTime::GetSecondsPerCycle();
+
+	double RunningFrameTimeSeconds = 0.0;
+	uint64 FrameStartGPUCycles = 0;
+	uint64 FrameEndGPUCycles = 0;
+
+	CFTimeInterval FirstStartTime = 0.0;
+
+	// Add the timings excluding any overlapping time
+	for (const FMetalCommandBufferTiming& Timing : Timings)
+	{
+		if (FirstStartTime == 0.0)
+		{
+			FirstStartTime = Timing.StartTime;
+			LastBufferTiming = Timing;
+		}
+		
+		// Only process if the previous buffer finished before the end of this one
+		if (LastBufferTiming.EndTime < Timing.EndTime)
+		{
+			// Check if the end of the previous buffer finished before the start of this one
+			if (LastBufferTiming.EndTime > Timing.StartTime)
+			{
+				// Segment from end of last buffer to end of current
+				RunningFrameTimeSeconds += Timing.EndTime - LastBufferTiming.EndTime;
+			}
+			else
+			{
+				// Full timing of this buffer
+				RunningFrameTimeSeconds += Timing.EndTime - Timing.StartTime;
+			}
+
+			LastBufferTiming = Timing;
+		}
+	}
+	
+	FrameStartGPUCycles = FirstStartTime * CyclesPerSecond;
+	FrameEndGPUCycles = LastBufferTiming.EndTime * CyclesPerSecond;
+	
+	uint64 FrameGPUTimeCycles = uint64(CyclesPerSecond * RunningFrameTimeSeconds);
+	FPlatformAtomics::AtomicStore_Relaxed((int32*)&GGPUFrameTime, int32(FrameGPUTimeCycles));
+	
+#if STATS
+	FPlatformAtomics::AtomicStore_Relaxed(&GMetalGPUWorkTime, FrameGPUTimeCycles);
+	int64 FrameIdleTimeCycles = int64(FrameEndGPUCycles - FrameStartGPUCycles - FrameGPUTimeCycles);
+	FPlatformAtomics::AtomicStore_Relaxed(&GMetalGPUIdleTime, FrameIdleTimeCycles);
+#endif //STATS
+	
+	FMetalDynamicRHI::Get().DeferredDelete([InTimer=this](){
+		delete InTimer;
+	});
+}
+
+void FMetalCommandBufferTimer::RecordPresent(MTL::CommandBuffer* CommandBuffer)
+{
+	const CFTimeInterval GpuStartTimeSeconds = CommandBuffer->GPUStartTime();
+	const CFTimeInterval GpuEndTimeSeconds = CommandBuffer->GPUStartTime();
+	const double CyclesPerSecond = 1.0 / FPlatformTime::GetSecondsPerCycle();
+	uint64 StartTimeCycles = uint64(GpuStartTimeSeconds * CyclesPerSecond);
+	uint64 EndTimeCycles = uint64(GpuEndTimeSeconds * CyclesPerSecond);
+	int64 Time = int64(EndTimeCycles - StartTimeCycles);
+	FPlatformAtomics::AtomicStore_Relaxed(&GMetalPresentTime, Time);
+}
+// END WARNING
+
+void FMetalCommandBufferTimer::ResetFrameBufferTimings()
+{
+	if(Timer)
+	{
+		Timer->FrameEnd();
+	}
+	
+	Timer = new FMetalCommandBufferTimer();
+}
+
+FMetalCommandBufferTimer& FMetalCommandBufferTimer::GetFrameBufferTimer()
+{
+	if(!Timer)
+	{
+		Timer = new FMetalCommandBufferTimer();
+	}
+	return *Timer;
+}
+
+FMetalCommandBufferTimer* FMetalCommandBufferTimer::Timer = nullptr;
+
 void FMetalGPUTiming::PlatformStaticInitialize(void* UserData)
 {
 	// Are the static variables initialized?
@@ -92,6 +212,7 @@
 
 float FMetalEventNode::GetTiming()
 {
+	SyncPoint->Wait();
 	return FPlatformTime::ToSeconds(EndTime - StartTime);
 }
 
@@ -103,38 +224,11 @@
 	Context.StartTiming(this);
 }
 
-FMetalCommandBufferCompletionHandler FMetalEventNode::Start(void)
-{
-    FMetalCommandBufferCompletionHandler Block;
-    Block.BindLambda([this](MTL::CommandBuffer* CompletedBuffer)
-	{
-		const CFTimeInterval GpuTimeSeconds = CompletedBuffer->GPUStartTime();
-		const double CyclesPerSecond = 1.0 / FPlatformTime::GetSecondsPerCycle();
-		StartTime = GpuTimeSeconds * CyclesPerSecond;
-	});
-	
-    return Block;
-}
-
-
 void FMetalEventNode::StopTiming()
 {
 	Context.EndTiming(this);
 }
 
-FMetalCommandBufferCompletionHandler FMetalEventNode::Stop(void)
-{
-    FMetalCommandBufferCompletionHandler Block;
-    Block.BindLambda([this](MTL::CommandBuffer* CompletedBuffer)
-	{
-		// This is still used by ProfileGPU
-		const CFTimeInterval GpuTimeSeconds = CompletedBuffer->GPUEndTime();
-		const double CyclesPerSecond = 1.0 / FPlatformTime::GetSecondsPerCycle();
-		EndTime = GpuTimeSeconds * CyclesPerSecond;
-	});
-	return Block;
-}
-
 bool MetalGPUProfilerIsInSafeThread()
 {
 	return (GIsMetalInitialized && !GIsRHIInitialized) || (IsInRHIThread() || IsInActualRenderingThread());
@@ -173,39 +267,6 @@
 #endif
 }
 
-void FMetalCommandBufferTimer::AddTiming(FMetalCommandBufferTiming Timing)
-{
-	FScopeLock Lock(&Mutex);
-	Counter--;
-	Timings.Add(Timing);
-	
-	if(Counter == 0 && bFrameEnded)
-	{
-		RecordFrame();
-	}
-}
-
-void FMetalCommandBufferTimer::FrameEnd()
-{
-	FScopeLock Lock(&Mutex);
-	
-	bFrameEnded = true;
-	
-	if(Counter == 0)
-	{
-		RecordFrame();
-	}
-}
-
-void FMetalCommandBufferTimer::RecordFrame()
-{
-	Timings.Sort();
-	FMetalProfiler* Profiler = FMetalProfiler::GetProfiler();
-	Profiler->RecordFrame(*this);
-}
-
-FMetalCommandBufferTimer* FMetalGPUProfiler::Timer = nullptr;
-
 void FMetalGPUProfiler::Cleanup()
 {
 	
@@ -266,7 +327,7 @@
 	
 	if(CurrentEventNodeFrame)
 	{
-		Context.GetDevice().WaitForGPUIdle();
+		GDynamicRHI->RHIBlockUntilGPUIdle();
 		
 		CurrentEventNodeFrame->EndFrame();
 		
@@ -285,99 +346,6 @@
 	}
 }
 
-// WARNING:
-// All these recording functions MUST be called from within scheduled/completion handlers.
-// Ordering is enforced by libdispatch so calling these outside of that context WILL result in
-// incorrect values.
-void FMetalGPUProfiler::RecordFrame(FMetalCommandBufferTimer& InTimer)
-{
-	FMetalCommandBufferTiming LastBufferTiming;
-	const double CyclesPerSecond = 1.0 / FPlatformTime::GetSecondsPerCycle();
-
-	double RunningFrameTimeSeconds = 0.0;
-	uint64 FrameStartGPUCycles = 0;
-	uint64 FrameEndGPUCycles = 0;
-
-	const TArray<FMetalCommandBufferTiming>& Timings = InTimer.GetTimings();
-
-	CFTimeInterval FirstStartTime = 0.0;
-
-	// Add the timings excluding any overlapping time
-	for (const FMetalCommandBufferTiming& Timing : Timings)
-	{
-		if (FirstStartTime == 0.0)
-		{
-			FirstStartTime = Timing.StartTime;
-			LastBufferTiming = Timing;
-		}
-		
-		// Only process if the previous buffer finished before the end of this one
-		if (LastBufferTiming.EndTime < Timing.EndTime)
-		{
-			// Check if the end of the previous buffer finished before the start of this one
-			if (LastBufferTiming.EndTime > Timing.StartTime)
-			{
-				// Segment from end of last buffer to end of current
-				RunningFrameTimeSeconds += Timing.EndTime - LastBufferTiming.EndTime;
-			}
-			else
-			{
-				// Full timing of this buffer
-				RunningFrameTimeSeconds += Timing.EndTime - Timing.StartTime;
-			}
-
-			LastBufferTiming = Timing;
-		}
-	}
-	
-	FrameStartGPUCycles = FirstStartTime * CyclesPerSecond;
-	FrameEndGPUCycles = LastBufferTiming.EndTime * CyclesPerSecond;
-    
-	uint64 FrameGPUTimeCycles = uint64(CyclesPerSecond * RunningFrameTimeSeconds);
-	FPlatformAtomics::AtomicStore_Relaxed((int32*)&GGPUFrameTime, int32(FrameGPUTimeCycles));
-	
-#if STATS
-	FPlatformAtomics::AtomicStore_Relaxed(&GMetalGPUWorkTime, FrameGPUTimeCycles);
-	int64 FrameIdleTimeCycles = int64(FrameEndGPUCycles - FrameStartGPUCycles - FrameGPUTimeCycles);
-	FPlatformAtomics::AtomicStore_Relaxed(&GMetalGPUIdleTime, FrameIdleTimeCycles);
-#endif //STATS
-}
-
-void FMetalGPUProfiler::RecordPresent(MTL::CommandBuffer* CommandBuffer)
-{
-	const CFTimeInterval GpuStartTimeSeconds = CommandBuffer->GPUStartTime();
-    const CFTimeInterval GpuEndTimeSeconds = CommandBuffer->GPUStartTime();
-	const double CyclesPerSecond = 1.0 / FPlatformTime::GetSecondsPerCycle();
-	uint64 StartTimeCycles = uint64(GpuStartTimeSeconds * CyclesPerSecond);
-	uint64 EndTimeCycles = uint64(GpuEndTimeSeconds * CyclesPerSecond);
-	int64 Time = int64(EndTimeCycles - StartTimeCycles);
-	FPlatformAtomics::AtomicStore_Relaxed(&GMetalPresentTime, Time);
-}
-// END WARNING
-
-
-void FMetalGPUProfiler::ResetFrameBufferTimings()
-{
-	if(Timer)
-	{
-		Timer->FrameEnd();
-		FMetalDynamicRHI::Get().DeferredDelete([InTimer=Timer](){
-			delete InTimer;
-		});
-	}
-	
-	Timer = new FMetalCommandBufferTimer();
-}
-
-FMetalCommandBufferTimer& FMetalGPUProfiler::GetFrameBufferTimer()
-{
-	if(!Timer)
-	{
-		Timer = new FMetalCommandBufferTimer();
-	}
-	return *Timer;
-}
-
 IMetalStatsScope::~IMetalStatsScope()
 {
 	for (IMetalStatsScope* Stat : Children)
@@ -664,34 +632,42 @@
 
 void FMetalProfiler::EncodeDraw(FMetalCommandBufferStats* CmdBufStats, char const* DrawCall, uint32 RHIPrimitives, uint32 RHIVertices, uint32 RHIInstances)
 {
+#if (RHI_NEW_GPU_PROFILER == 0)
 	if (MetalGPUProfilerIsInSafeThread())
 	{
 		FMetalGPUProfiler::RegisterGPUWork(RHIPrimitives, RHIVertices);
 	}
+#endif
 }
 
 void FMetalProfiler::EncodeBlit(FMetalCommandBufferStats* CmdBufStats, char const* DrawCall)
 {
+#if (RHI_NEW_GPU_PROFILER == 0)
 	if (MetalGPUProfilerIsInSafeThread())
 	{
 		FMetalGPUProfiler::RegisterGPUWork(1, 1);
 	}
+#endif
 }
 
 void FMetalProfiler::EncodeBlit(FMetalCommandBufferStats* CmdBufStats, FString DrawCall)
 {
+#if (RHI_NEW_GPU_PROFILER == 0)
 	if (MetalGPUProfilerIsInSafeThread())
 	{
 		FMetalGPUProfiler::RegisterGPUWork(1, 1);
 	}
+#endif
 }
 
 void FMetalProfiler::EncodeDispatch(FMetalCommandBufferStats* CmdBufStats, char const* DrawCall)
 {
+#if (RHI_NEW_GPU_PROFILER == 0)
 	if (MetalGPUProfilerIsInSafeThread())
 	{
 		FMetalGPUProfiler::RegisterGPUWork(1, 1);
 	}
+#endif
 }
 
 FMetalCPUStats* FMetalProfiler::AddCPUStat(FString const& Name)
@@ -739,7 +715,7 @@
 
 void FMetalProfiler::SaveTrace()
 {
-	Context.GetDevice().WaitForGPUIdle();
+	GDynamicRHI->RHIBlockUntilGPUIdle();
 	
 	{
 		FScopeLock Lock(&Mutex);
@@ -931,3 +907,36 @@
 	TEXT("Starts or stops Metal profiler"),
 	FConsoleCommandWithWorldArgsAndOutputDeviceDelegate::CreateStatic(&HandleMetalProfileCommand)
 );
+
+#endif
+
+#if RHI_NEW_GPU_PROFILER && WITH_RHI_BREADCRUMBS
+
+void FMetalBreadcrumbProfiler::ResolveBreadcrumbTrackerStream(TArray<FMetalBreadcrumbTrackerObject>& BreadcrumbTrackerStream)
+{
+	for(FMetalBreadcrumbTrackerObject& Tracker : BreadcrumbTrackerStream)
+	{
+		switch(Tracker.Type)
+		{
+			case EMetalBreadcrumbTrackerType_Begin:
+			{
+				OnBreadcrumbBegin(Tracker.Node);
+				break;
+			}
+			case EMetalBreadcrumbTrackerType_End:
+			{	
+				OnBreadcrumbEnd(Tracker.Node);
+				break;
+			}
+			case EMetalBreadcrumbTrackerType_Encode:
+			{
+				AddSample(Tracker.CounterSample);
+				break;
+			}
+			default:
+				check(0);
+		}
+	}
+}
+
+#endif
diff --git a/Engine/Source/Runtime/Apple/MetalRHI/Private/MetalProfiler.h b/Engine/Source/Runtime/Apple/MetalRHI/Private/MetalProfiler.h
--- a/Engine/Source/Runtime/Apple/MetalRHI/Private/MetalProfiler.h
+++ b/Engine/Source/Runtime/Apple/MetalRHI/Private/MetalProfiler.h
@@ -3,7 +3,9 @@
 #pragma once
 
 #include "MetalRHIPrivate.h"
+#include "MetalCounterSampler.h"
 #include "GPUProfiler.h"
+#include "RHIBreadcrumbs.h"
 
 DECLARE_DELEGATE_OneParam(FMetalCommandBufferCompletionHandler, MTL::CommandBuffer*);
 
@@ -64,6 +66,65 @@
 DECLARE_DWORD_COUNTER_STAT_EXTERN(TEXT("Number Command Buffers Committed Per-Frame"), STAT_MetalCommandBufferCommittedPerFrame, STATGROUP_MetalRHI, );
 
 class FMetalRHICommandContext;
+class FMetalSyncPoint;
 
+/**
+ * Simple struct to hold sortable command buffer start and end timestamps.
+ */
+struct FMetalCommandBufferTiming
+{
+	CFTimeInterval StartTime;
+	CFTimeInterval EndTime;
+
+	bool operator<(const FMetalCommandBufferTiming& RHS) const
+	{
+		// Sort by start time and then by length if the commandbuffer started at the same time
+		if (this->StartTime < RHS.StartTime)
+		{
+			return true;
+		}
+		else if ((this->StartTime == RHS.StartTime) && (this->EndTime > RHS.EndTime))
+		{
+			return true;
+		}
+		return false;
+	}
+};
+
+#if RHI_NEW_GPU_PROFILER == 0
+
+class FMetalCommandBufferTimer
+{
+public:
+	void Submit()
+	{
+		FScopeLock Lock(&Mutex);
+		Counter++;
+	}
+	
+	void AddTiming(FMetalCommandBufferTiming Timing);
+	void FrameEnd();
+	void RecordFrame();
+	
+	const TArray<FMetalCommandBufferTiming>& GetTimings()
+	{
+		return Timings;
+	}
+	
+	static void ResetFrameBufferTimings();
+	static FMetalCommandBufferTimer& GetFrameBufferTimer();
+	
+	static void RecordPresent(MTL::CommandBuffer* CommandBuffer);
+	// END WARNING
+	
+private:
+	bool bFrameEnded;
+	int32_t Counter = 0;
+	TArray<FMetalCommandBufferTiming> Timings;
+	FCriticalSection Mutex;
+	
+	static FMetalCommandBufferTimer* Timer;
+};
+
 /** A single perf event node, which tracks information about a appBeginDrawEvent/appEndDrawEvent range. */
 class FMetalEventNode : public FGPUProfilerEventNode
@@ -87,13 +148,10 @@
 	 * This blocks the CPU if necessary, so can cause hitching.
 	 */
 	virtual float GetTiming() override;
-	
 	virtual void StartTiming() override;
-	
 	virtual void StopTiming() override;
 	
-    FMetalCommandBufferCompletionHandler Start(void);
-    FMetalCommandBufferCompletionHandler Stop(void);
+	FMetalSyncPoint* SyncPoint = nullptr;
 
 	bool IsRoot() const { return bRoot; }	
 	uint64 GetCycles() { return EndTime - StartTime; }
@@ -226,56 +284,7 @@
 	MTLCommandBufferPtr CmdBuffer;
 };
 
-/**
- * Simple struct to hold sortable command buffer start and end timestamps.
- */
-struct FMetalCommandBufferTiming
-{
-	CFTimeInterval StartTime;
-	CFTimeInterval EndTime;
-
-	bool operator<(const FMetalCommandBufferTiming& RHS) const
-	{
-		// Sort by start time and then by length if the commandbuffer started at the same time
-		if (this->StartTime < RHS.StartTime)
-		{
-			return true;
-		}
-		else if ((this->StartTime == RHS.StartTime) && (this->EndTime > RHS.EndTime))
-		{
-			return true;
-		}
-		return false;
-	}
-};
-
-class FMetalCommandBufferTimer
-{
-public:
-	void Submit()
-	{
-		FScopeLock Lock(&Mutex);
-		Counter++;
-	}
-	
-	void AddTiming(FMetalCommandBufferTiming Timing);
-	void FrameEnd();
-	void RecordFrame();
-	
-	const TArray<FMetalCommandBufferTiming>& GetTimings()
-	{
-		return Timings;
-	}
-	
-private:
-	bool bFrameEnded;
-	uint32_t Counter = 0;
-	TArray<FMetalCommandBufferTiming> Timings;
-	FCriticalSection Mutex;
-};
-
 class FMetalRHICommandContext;
-
 /**
  * Encapsulates GPU profiling logic and data.
  * There's only one global instance of this struct so it should only contain global data, nothing specific to a frame.
@@ -303,20 +312,8 @@
 	void BeginFrame();
 	void EndFrame();
 	
-	// WARNING:
-	// These functions MUST be called from within Metal scheduled/completion handlers
-	// since they depend on libdispatch to enforce ordering.
-	static void RecordFrame(FMetalCommandBufferTimer& Timer);
-	static void RecordPresent(MTL::CommandBuffer* CommandBuffer);
-	// END WARNING
-	
-	static void ResetFrameBufferTimings();
-	static FMetalCommandBufferTimer& GetFrameBufferTimer();
-	
 	FMetalGPUTiming TimingSupport;
 	FMetalRHICommandContext& Context;
-	
-	static FMetalCommandBufferTimer* Timer;
 };
 
 class FMetalProfiler : public FMetalGPUProfiler
@@ -393,3 +390,125 @@
 	
 	FMetalCPUStats* Stats;
 };
+
+#endif
+
+#if RHI_NEW_GPU_PROFILER && WITH_RHI_BREADCRUMBS
+
+enum class EMetalBreadcrumbTrackerType : uint8
+{
+	Begin,
+	End,
+	Encode,
+};
+
+class FMetalCommandBuffer;
+struct FMetalBreadcrumbTrackerObject
+{
+	EMetalBreadcrumbTrackerType Type;
+	FMetalCommandBuffer* CmdBuffer;
+	FMetalCounterSamplePtr CounterSample;
+	FRHIBreadcrumbNode* Node;
+};
+
+// Represents the data required for FRHIBreadcrumbNode
+// Samples can be collected across multiple counter samples and this allows us to merge the results
+struct FMetalBreadcrumbEvent
+{
+	FMetalBreadcrumbEvent(bool bRenderPass) : bWithinRenderPass(bRenderPass)
+	{}
+	
+	bool bWithinRenderPass;
+	uint64_t* TimestampTOP = nullptr;
+	uint64_t* TimestampBOP = nullptr;
+	TArray<FMetalCounterSamplePtr> Samples;
+};
+
+// Tracks the FRHIBreadcrumbNode's across encoders
+class FMetalBreadcrumbProfiler
+{
+public:
+	inline static FMetalBreadcrumbProfiler* GetInstance()
+	{
+		if(!Instance)
+		{
+			Instance = new FMetalBreadcrumbProfiler; 
+		}
+		return Instance;
+	}
+	
+	~FMetalBreadcrumbProfiler()
+	{
+	}
+
+	inline FMetalBreadcrumbEvent& GetBreadcrumbEvent(FRHIBreadcrumbNode* Breadcrumb, bool bWithinRenderPass)
+	{
+		FScopeLock Lock(&Mutex);
+		FMetalBreadcrumbEvent& Event = CreatedBreadcrumbs.FindOrAdd(Breadcrumb, FMetalBreadcrumbEvent(bWithinRenderPass));
+		Event.bWithinRenderPass |= bWithinRenderPass;
+		
+		return Event;
+	}
+	
+	inline void AddSample(FMetalCounterSamplePtr Sample)
+	{
+		FScopeLock Lock(&Mutex);
+		
+		uint64_t StartTime, EndTime;
+		Sample->ResolveStageCounters(StartTime, EndTime);
+		
+		for(FRHIBreadcrumbNode* ActiveBreadcrumb : ActiveBreadcrumbs)
+		{
+			CreatedBreadcrumbs[ActiveBreadcrumb].Samples.Add(Sample);
+		}
+	}
+	
+	inline void OnBreadcrumbBegin(FRHIBreadcrumbNode* Node)
+	{
+		FScopeLock Lock(&Mutex);
+		ActiveBreadcrumbs.Add(Node);
+	}
+	
+	inline void OnBreadcrumbEnd(FRHIBreadcrumbNode* Node)
+	{
+		FScopeLock Lock(&Mutex);
+		
+		FMetalBreadcrumbEvent& Event = CreatedBreadcrumbs[Node];
+	
+		uint64_t* Start = Event.TimestampTOP;
+		uint64_t* End = Event.TimestampBOP;
+		
+		// Ignore samples whe
+		if(Event.bWithinRenderPass || Event.Samples.Num() == 0)
+		{
+			*Start = 0; 
+			*End = 0;
+		}
+		else
+		{
+			for(FMetalCounterSamplePtr Sample : Event.Samples)
+			{
+				uint64_t StartTime, EndTime;
+				Sample->ResolveStageCounters(StartTime, EndTime);
+				
+				*Start = *Start > 0 ? FMath::Min(StartTime, *Start) : StartTime;
+				*End = *End > 0 ? FMath::Max(EndTime, *End) : EndTime;
+			}
+		}
+		
+		ActiveBreadcrumbs.Remove(Node);
+		CreatedBreadcrumbs.Remove(Node);
+	}
+	
+	// Resolves all the breakcrumbs added to the command buffer
+	void ResolveBreadcrumbTrackerStream(TArray<FMetalBreadcrumbTrackerObject>& BreadcrumbTrackerStream);
+	
+private:
+	TMap<FRHIBreadcrumbNode*, FMetalBreadcrumbEvent> CreatedBreadcrumbs;
+	TArray<FRHIBreadcrumbNode*> ActiveBreadcrumbs; 
+	
+	static FMetalBreadcrumbProfiler* Instance;
+	FCriticalSection Mutex;
+};
+
+#endif
diff --git a/Engine/Source/Runtime/Apple/MetalRHI/Private/MetalRHI.cpp b/Engine/Source/Runtime/Apple/MetalRHI/Private/MetalRHI.cpp
--- a/Engine/Source/Runtime/Apple/MetalRHI/Private/MetalRHI.cpp
+++ b/Engine/Source/Runtime/Apple/MetalRHI/Private/MetalRHI.cpp
@@ -1186,10 +1186,14 @@ FMetalDynamicRHI::FMetalDynamicRHI(ERHIFeatureLevel::Type RequestedFeatureLevel)
 #endif
 		
 	GDynamicRHI = this;
+	
+	// Start the submission and interrupt handler threads
+	InitializeSubmissionPipe();
+	
 	GIsMetalInitialized = true;
 
 	ImmediateContext.SetProfiler(nullptr);
-#if ENABLE_METAL_GPUPROFILE
+#if ENABLE_METAL_GPUPROFILE && RHI_NEW_GPU_PROFILER == 0
 	FMetalProfiler* Profiler = FMetalProfiler::CreateProfiler(ImmediateContext);
 	ImmediateContext.SetProfiler(Profiler);
 	
@@ -1216,6 +1220,9 @@ FMetalDynamicRHI::~FMetalDynamicRHI()
 {
 	check(IsInGameThread() && IsInRenderingThread());
 	
+	RHIBlockUntilGPUIdle();
+	ShutdownSubmissionPipe();
+	
 	GIsMetalInitialized = false;
 	GIsRHIInitialized = false;
 
@@ -1226,7 +1233,7 @@ FMetalDynamicRHI::~FMetalDynamicRHI()
     IRCompilerDestroy(CompilerInstance);
 #endif
 
-#if ENABLE_METAL_GPUPROFILE
+#if ENABLE_METAL_GPUPROFILE && RHI_NEW_GPU_PROFILER == 0
 	FMetalProfiler::DestroyProfiler();
 #endif
 }
@@ -1257,11 +1264,13 @@ void FMetalDynamicRHI::RHIEndFrame_RenderThread(FRHICommandListImmediate& RHICmd
 	{
 		MTL_SCOPED_AUTORELEASE_POOL;
 
-		FMetalGPUProfiler::ResetFrameBufferTimings();
 		
+#if RHI_NEW_GPU_PROFILER == 0
+		FMetalCommandBufferTimer::ResetFrameBufferTimings();
 #if ENABLE_METAL_GPUPROFILE
 		Contexts[ERHIPipeline::Graphics]->GetProfiler()->EndFrame();
 #endif
+#endif
 
 #if METAL_RHI_RAYTRACING
 		UpdateRayTracing();
@@ -1276,7 +1285,9 @@ void FMetalDynamicRHI::RHIEndFrame_RenderThread(FRHICommandListImmediate& RHICmd
 		MTL_SCOPED_AUTORELEASE_POOL;
 		
 #if ENABLE_METAL_GPUPROFILE
+#if RHI_NEW_GPU_PROFILER == 0
 		Contexts[ERHIPipeline::Graphics]->GetProfiler()->BeginFrame();
+#endif
 #endif
 	});
 }
@@ -1286,6 +1297,43 @@ void FMetalDynamicRHI::RHIEndFrame(const FRHIEndFrameArgs& Args)
 	// increment the internal frame counter
 	Device->IncrementFrameRHIThread();
     Device->GarbageCollect();
+
+#if RHI_NEW_GPU_PROFILER	
+	// Close the previous frame's timing and start a new one
+	auto Lambda = [this, OldTiming = MoveTemp(CurrentTimingPerQueue)]()
+	{
+		TArray<UE::RHI::GPUProfiler::FEventStream, TInlineAllocator<GMetalMaxNumQueues>> Streams;
+		for (auto const& Timing : OldTiming)
+		{
+			Streams.Add(MoveTemp(Timing->EventStream));
+		}
+		
+		UE::RHI::GPUProfiler::ProcessEvents(Streams);
+	};
+	
+	EnqueueEndOfPipeTask(MoveTemp(Lambda), [&](FMetalPayload& Payload)
+	{
+		// Modify the payloads the EOP task will submit to include
+		// a new timing struct and a frame boundary event.
+		Payload.Timing = CurrentTimingPerQueue.CreateNew(Payload.Queue);
+		
+		ERHIPipeline Pipeline = ERHIPipeline::Graphics;
+		
+		Payload.EventStream.Emplace<UE::RHI::GPUProfiler::FEvent::FFrameBoundary>(
+			0, Args.FrameNumber
+		#if WITH_RHI_BREADCRUMBS
+			, (Pipeline != ERHIPipeline::None) ? Args.GPUBreadcrumbs[Pipeline] : nullptr
+		#endif
+		#if STATS
+			, Args.StatsFrame
+		#endif
+		);
+	});
+#endif
+	
+	// Pump the interrupt queue to gather completed events
+	// (required if we're not using an interrupt thread).
+	ProcessInterruptQueueUntil(nullptr);
 }
 
 #if WITH_RHI_BREADCRUMBS
@@ -1314,20 +1362,57 @@ void FMetalRHICommandContext::RHIBeginBreadcrumbGPU(FRHIBreadcrumbNode* Breadcru
 	}
 
 #if ENABLE_METAL_GPUPROFILE
+#if RHI_NEW_GPU_PROFILER
+	if(!CurrentEncoder.GetCommandBuffer())
+	{
+		StartCommandBuffer();
+	}
+	
+	FMetalCommandBuffer* CmdBuffer = CurrentEncoder.GetCommandBuffer();
+	
+	// Can't process breadcrumbs if we are within a render pass
+	if(Device.SupportsFeature(EMetalFeaturesStageCounterSampling))
+	{
+		FMetalBreadcrumbEvent& Event = FMetalBreadcrumbProfiler::GetInstance()->GetBreadcrumbEvent(Breadcrumb, bWithinRenderPass);
+		Event.TimestampTOP = &CmdBuffer->EmplaceProfilerEvent<UE::RHI::GPUProfiler::FEvent::FBeginBreadcrumb>(Breadcrumb).GPUTimestampTOP;
+		*Event.TimestampTOP = 0;
+		CmdBuffer->BeginBreadcrumb(Breadcrumb);
+	}
+#else
 	if (Profiler && Profiler->IsProfilingGPU())
 	{
 		Profiler->PushEvent(GetNameStr(), FColor::White);
 	}
 #endif
+#endif
 }
 
 void FMetalRHICommandContext::RHIEndBreadcrumbGPU(FRHIBreadcrumbNode* Breadcrumb)
 {
 #if ENABLE_METAL_GPUPROFILE
+#if RHI_NEW_GPU_PROFILER
+	if(!CurrentEncoder.GetCommandBuffer())
+	{
+		StartCommandBuffer();	
+	}
+	
+	FMetalCommandBuffer* CmdBuffer = CurrentEncoder.GetCommandBuffer();
+	
+	if(Device.SupportsFeature(EMetalFeaturesStageCounterSampling))
+	{
+		FMetalBreadcrumbEvent& Event = FMetalBreadcrumbProfiler::GetInstance()->GetBreadcrumbEvent(Breadcrumb, bWithinRenderPass);
+		Event.TimestampBOP = &CmdBuffer->EmplaceProfilerEvent<UE::RHI::GPUProfiler::FEvent::FEndBreadcrumb>(Breadcrumb).GPUTimestampBOP;
+		*Event.TimestampBOP = 0;
+		
+		CmdBuffer->EndBreadcrumb(Breadcrumb);
+	}
+	
+#else
 	if (Profiler && Profiler->IsProfilingGPU())
 	{
 		Profiler->PopEvent();
 	}
+#endif
 #endif
 
 	if (ShouldEmitBreadcrumbs())
@@ -1489,8 +1574,31 @@ uint16 FMetalDynamicRHI::RHIGetPlatformTextureMaxSampleCount()
 
 void FMetalDynamicRHI::RHIBlockUntilGPUIdle()
 {
-	Device->WaitForGPUIdle();
-	ProcessDeferredDeleteQueue();
+	// Submit a new sync point to each queue
+	TArray<FMetalPayload*> Payloads;
+	Payloads.Reserve(int32(EMetalQueueType::Count));
+
+	TArray<FMetalSyncPointRef, TInlineAllocator<(uint32)EMetalQueueType::Count>> SyncPoints;
+
+	for (uint32 QueueTypeIndex = 0; QueueTypeIndex < (uint32)EMetalQueueType::Count; ++QueueTypeIndex)
+	{
+		FMetalSyncPointRef SyncPoint = FMetalSyncPoint::Create(EMetalSyncPointType::GPUAndCPU);
+
+		FMetalPayload* Payload = new FMetalPayload(Device->GetCommandQueue((EMetalQueueType)QueueTypeIndex));
+		Payload->SyncPointsToSignal.Add(SyncPoint);
+		Payload->bAlwaysSignal = true;
+
+		Payloads.Add(Payload);
+		SyncPoints.Add(SyncPoint);
+	}
+
+	SubmitPayloads(MoveTemp(Payloads));
+
+	// Block this thread until the sync points have signaled.
+	for (FMetalSyncPointRef& SyncPoint : SyncPoints)
+	{
+		SyncPoint->Wait();
+	}
 }
 
 uint32 FMetalDynamicRHI::RHIGetGPUFrameCycles(uint32 GPUIndex)
@@ -1526,35 +1634,14 @@ public:
 	TArray<FMetalCommandBuffer*> CommandBuffers;
 };
 
-static TLockFreePointerListUnordered<FMetalRHIUploadContext, PLATFORM_CACHE_LINE_SIZE> MetalUploadContextPool;
-
-void FMetalDynamicRHI::ProcessDeferredDeleteQueue()
+void FMetalDynamicRHI::RHIProcessDeleteQueue()
 {
-	uint32 Index = 0;
-	while(Index < DeferredDeleteQueue.Num())
+	FScopeLock Lock(&ObjectsToDeleteCS);
+	TArray<FMetalDeferredDeleteObject> Objects = MoveTemp(ObjectsToDelete);
+	
+	EnqueueEndOfPipeTask([this, LocalDeleteObjects = MoveTemp(Objects)]() mutable
 	{
-		FDeferredDeleteData& Data = DeferredDeleteQueue[Index];
-		
-		bool bFencesReady = true;
-		for (TSharedPtr<FMetalCommandBufferFence, ESPMode::ThreadSafe> Fence : Data.WaitFences)
-		{
-			bFencesReady &= Fence->Wait(0);
-		}
-		
-		if(!bFencesReady)
-		{
-			Index++;
-			break;
-		}
-		
-		if(Data.bWaitExtraFrame)
-		{
-			Data.bWaitExtraFrame = false;
-			Index++;
-			continue;
-		}
-		
-		for(FMetalDeferredDeleteObject& Object : Data.DeferredDeleteObjects)
+		for(FMetalDeferredDeleteObject& Object : LocalDeleteObjects)
 		{	
 			switch (Object.Storage.GetIndex())
 			{
@@ -1580,19 +1667,19 @@ void FMetalDynamicRHI::ProcessDeferredDeleteQueue()
 					}
 					break;
 				}
-#if PLATFORM_SUPPORTS_BINDLESS_RENDERING
+	#if PLATFORM_SUPPORTS_BINDLESS_RENDERING
 				case FMetalDeferredDeleteObject::TObjectStorage::IndexOfType<FRHIDescriptorHandle>():
 				{
 					FRHIDescriptorHandle Handle = Object.Storage.Get<FRHIDescriptorHandle>();
 					
 					FMetalBindlessDescriptorManager* BindlessDescriptorManager = Device->GetBindlessDescriptorManager();
 					check(BindlessDescriptorManager);
-	
+
 					BindlessDescriptorManager->FreeDescriptor(Handle);
 					
 					break;
 				}
-#endif
+	#endif
 				case FMetalDeferredDeleteObject::TObjectStorage::IndexOfType<FMetalFence*>():
 				{
 					FMetalFence* Fence = Object.Storage.Get<FMetalFence*>();
@@ -1614,102 +1701,42 @@ void FMetalDynamicRHI::ProcessDeferredDeleteQueue()
 				}
 			}
 		}
-		
-		DeferredDeleteQueue.RemoveAt(Index, EAllowShrinking::No);
-	}
+	});
 }
 
-void FMetalDynamicRHI::RHIProcessDeleteQueue()
+void FMetalDynamicRHI::EnqueueEndOfPipeTask(TUniqueFunction<void()> TaskFunc, TUniqueFunction<void(FMetalPayload&)> ModifyPayloadCallback)
 {
-	FDeferredDeleteData NewData;
-
-	if(CVarEnableMetalDeferredDeleteLatency.GetValueOnAnyThread() > 0)
+	FGraphEventArray Prereqs;
+	Prereqs.Reserve(GMetalMaxNumQueues + 1);
+	if (EopTask)
 	{
-		NewData.bWaitExtraFrame = true;
+		Prereqs.Add(EopTask);
 	}
-	GatherDeferredDeleteObjects(NewData.DeferredDeleteObjects, NewData.WaitFences);
 	
-	DeferredDeleteQueue.Add(MoveTemp(NewData));
+	TArray<FMetalPayload*> Payloads;
+	Payloads.Reserve(GMetalMaxNumQueues);
 	
-	ProcessDeferredDeleteQueue();
-}
-
-void FMetalDynamicRHI::RHIFinalizeContext(FRHIFinalizeContextArgs&& Args, TRHIPipelineArray<IRHIPlatformCommandList*>& Output)
-{
-	MTL_SCOPED_AUTORELEASE_POOL;
-	
-	FMetalRHIUploadContext* UploadContext = static_cast<FMetalRHIUploadContext*>(Args.UploadContext);
-
-	TArray<FMetalCommandBuffer*>* UploadCommandBuffers = nullptr;
-	
-	if(UploadContext)
-	{
-		UploadCommandBuffers = UploadContext->Finalize();
-		MetalUploadContextPool.Push(UploadContext);
-	}
-	
-	for(IRHIComputeContext* Context : Args.Contexts)
+	ForEachQueue([&](FMetalCommandQueue& Queue)
 	{
-		FMetalRHICommandContext* CmdContext = static_cast<FMetalRHICommandContext*>(Context);
-
-		FMetalPlatformCommandList* PlatformCmdList = new FMetalPlatformCommandList();
+		FMetalPayload* Payload = new FMetalPayload(Queue);
 		
-		if(UploadCommandBuffers)
-		{
-			PlatformCmdList->CommandBuffers.Append(*UploadCommandBuffers);
-			
-			delete UploadCommandBuffers;
-			UploadCommandBuffers = nullptr;
-		}
+		FMetalSyncPointRef SyncPoint = FMetalSyncPoint::Create(EMetalSyncPointType::GPUAndCPU);
+		Payload->SyncPointsToSignal.Emplace(SyncPoint);
+		Prereqs.Add(SyncPoint->GetGraphEvent());
 		
-		if(!CmdContext->IsInsideRenderPass())
-		{
-			PlatformCmdList->CommandBuffers.Append(CmdContext->Finalize());
-
-			CmdContext->ResetContext();
-			if(GRHISupportsParallelRHIExecute)
-			{
-				if(CmdContext != RHIGetDefaultContext())
-				{
-					MetalCommandContextPool.Push(CmdContext);
-				}
-			}
-		}
+		if (ModifyPayloadCallback)
+			ModifyPayloadCallback(*Payload);
 		
-		Output[Context->GetPipeline()] = PlatformCmdList;
-	}
+		Payloads.Add(Payload);
+	});
 	
-	check(UploadCommandBuffers == nullptr);
-}
-
-void FMetalDynamicRHI::RHISubmitCommandLists(FRHISubmitCommandListsArgs&& Args)
-{
-	for (IRHIPlatformCommandList* Ptr : Args.CommandLists)
-	{
-		FMetalPlatformCommandList* PlatformCmdList = static_cast<FMetalPlatformCommandList*>(Ptr);
-		
-		for(FMetalCommandBuffer* CommandBuffer : PlatformCmdList->CommandBuffers)
-		{
-			if(CommandBuffer)
-			{
-				AddDeferredDeleteFence(CommandBuffer->GetCompletionFence());
-				Device->GetCommandQueue().CommitCommandBuffer(CommandBuffer);
-			}
-		}
-		
-		delete PlatformCmdList;
-	}
-}
-
-IRHIUploadContext* FMetalDynamicRHI::RHIGetUploadContext()
-{
-	FMetalRHIUploadContext* Context = MetalUploadContextPool.Pop();
-	if (!Context)
-	{
-		Context = new FMetalRHIUploadContext(*Device);
-	}
+	SubmitPayloads(MoveTemp(Payloads));
 	
-	return static_cast<IRHIUploadContext*>(Context);
+	EopTask = FFunctionGraphTask::CreateAndDispatchWhenReady(
+					 MoveTemp(TaskFunc),
+					 QUICK_USE_CYCLE_STAT(FExecuteRHIThreadTask, STATGROUP_TaskGraphTasks),
+					 &Prereqs
+					 );
 }
 
 void FMetalDynamicRHI::RHIReplaceResources(FRHICommandListBase& RHICmdList, TArray<FRHIResourceReplaceInfo>&& ReplaceInfos)
@@ -1779,3 +1806,9 @@ FRHIResourceCollectionRef FMetalDynamicRHI::RHICreateResourceCollection(FRHIComm
 	return new FMetalResourceCollection(RHICmdList, InMembers);
 }
 #endif
+
+void FMetalDynamicRHI::ForEachQueue(TFunctionRef<void(FMetalCommandQueue&)> Callback)
+{
+	// TODO - Carl: Multiple Queues
+	Callback(Device->GetCommandQueue(EMetalQueueType::Direct));
+}
diff --git a/Engine/Source/Runtime/Apple/MetalRHI/Private/MetalRHIContext.cpp b/Engine/Source/Runtime/Apple/MetalRHI/Private/MetalRHIContext.cpp
--- a/Engine/Source/Runtime/Apple/MetalRHI/Private/MetalRHIContext.cpp
+++ b/Engine/Source/Runtime/Apple/MetalRHI/Private/MetalRHIContext.cpp
@@ -38,8 +38,8 @@ void METALRHI_API SafeReleaseMetalObject(NS::Object* Object)
 
 FMetalRHICommandContext::FMetalRHICommandContext(FMetalDevice& MetalDevice, class FMetalProfiler* InProfiler)
 	: Device(MetalDevice)
-	, CommandQueue(Device.GetCommandQueue())
-	, CommandList(Device.GetCommandQueue())
+	, CommandQueue(Device.GetCommandQueue(EMetalQueueType::Direct))
+	, CommandList(Device.GetCommandQueue(EMetalQueueType::Direct))
 	, CurrentEncoder(MetalDevice, CommandList)
 	, StateCache(MetalDevice, true)
 	, QueryBuffer(new FMetalQueryBufferPool(MetalDevice))
@@ -53,7 +53,6 @@ FMetalRHICommandContext::FMetalRHICommandContext(FMetalDevice& MetalDevice, clas
 FMetalRHICommandContext::~FMetalRHICommandContext()
 {
 	CurrentEncoder.Release();
-	Device.WaitForGPUIdle();
 }
 
 void FMetalRHICommandContext::ResetContext()
@@ -67,11 +66,8 @@ void FMetalRHICommandContext::ResetContext()
 	// Reallocate if necessary to ensure >= 80% usage, otherwise we're just too wasteful
 	CurrentEncoder.GetRingBuffer().Shrink();
 	
-	// Begin the render pass frame.
-	CurrentEncoder.StartCommandBuffer();
-	
 	// make sure first SetRenderTarget goes through
-	StateCache.InvalidateRenderTargets();
+	StateCache.InvalidateRenderTargets();	
 }
 
 static uint32_t MAX_COLOR_RENDER_TARGETS_PER_DESC = 8;
@@ -81,9 +77,14 @@ void FMetalRHICommandContext::BeginComputeEncoder()
 	SCOPE_CYCLE_COUNTER(STAT_MetalSwitchToComputeTime);
 	
 	check(!bWithinRenderPass);
-	check(CurrentEncoder.GetCommandBuffer());
 	check(IsInParallelRenderingThread());
 	
+	if (!CurrentEncoder.GetCommandBuffer())
+	{
+		StartCommandBuffer();
+		check(CurrentEncoder.GetCommandBuffer());
+	}
+	
 	StateCache.SetStateDirty();
 	
 	if(!CurrentEncoder.IsComputeCommandEncoderActive())
@@ -93,7 +94,9 @@ void FMetalRHICommandContext::BeginComputeEncoder()
 		{
 			CurrentEncoderFence = CurrentEncoder.EndEncoding();
 		}
-		CurrentEncoder.BeginComputeCommandEncoding(MTL::DispatchTypeSerial);
+		bool bUseStageCounterSamples = Device.SupportsFeature(EMetalFeaturesStageCounterSampling);
+		CurrentEncoder.BeginComputeCommandEncoding(MTL::DispatchTypeSerial,
+												   bUseStageCounterSamples ? Device.GetCounterSampler() : nullptr);
 	}
 	
 	if (CurrentEncoderFence)
@@ -109,6 +112,10 @@ void FMetalRHICommandContext::EndComputeEncoder()
 {
 	check(CurrentEncoder.IsComputeCommandEncoderActive());
 	
+	// If we are using breadcrumbs then end the encoding here so that our stat tracking is correct
+#if WITH_RHI_BREADCRUMBS
+	CurrentEncoderFence = CurrentEncoder.EndEncoding();
+#endif
 	StateCache.SetRenderTargetsActive(false);
 }
 
@@ -116,7 +123,12 @@ void FMetalRHICommandContext::BeginBlitEncoder()
 {
 	SCOPE_CYCLE_COUNTER(STAT_MetalSwitchToBlitTime);
 	check(!bWithinRenderPass);
-	check(CurrentEncoder.GetCommandBuffer());
+	
+	if (!CurrentEncoder.GetCommandBuffer())
+	{
+		StartCommandBuffer();
+		check(CurrentEncoder.GetCommandBuffer());
+	}
 	
 	if(!CurrentEncoder.IsBlitCommandEncoderActive())
 	{
@@ -124,7 +136,8 @@ void FMetalRHICommandContext::BeginBlitEncoder()
 		{
 			CurrentEncoderFence = CurrentEncoder.EndEncoding();
 		}
-		CurrentEncoder.BeginBlitCommandEncoding();
+		bool bUseStageCounterSamples = Device.SupportsFeature(EMetalFeaturesStageCounterSampling);
+		CurrentEncoder.BeginBlitCommandEncoding(bUseStageCounterSamples ? Device.GetCounterSampler() : nullptr);
 	}
 	
 	if (CurrentEncoderFence)
@@ -140,13 +153,17 @@ void FMetalRHICommandContext::EndBlitEncoder()
 {
 	check(CurrentEncoder.IsBlitCommandEncoderActive());
 	
+#if WITH_RHI_BREADCRUMBS
+	CurrentEncoderFence = CurrentEncoder.EndEncoding();
+#endif
+	
 	StateCache.SetRenderTargetsActive(false);
 }
 
 void FMetalRHICommandContext::RHIBeginRenderPass(const FRHIRenderPassInfo& InInfo, const TCHAR* InName)
 {	
     MTL_SCOPED_AUTORELEASE_POOL;
-	
+
 #if PLATFORM_SUPPORTS_BINDLESS_RENDERING
 	if(IsMetalBindlessEnabled())
 	{
@@ -157,14 +174,9 @@ void FMetalRHICommandContext::RHIBeginRenderPass(const FRHIRenderPassInfo& InInf
     
 	RenderPassInfo = InInfo;
 	
-	if (InInfo.NumOcclusionQueries > 0)
-	{
-		RHIBeginOcclusionQueryBatch(InInfo.NumOcclusionQueries);
-	}
-	
 	if (!CurrentEncoder.GetCommandBuffer())
 	{
-		CurrentEncoder.StartCommandBuffer();
+		StartCommandBuffer();
 		check(CurrentEncoder.GetCommandBuffer());
 	}
 	
@@ -183,7 +195,8 @@ void FMetalRHICommandContext::RHIBeginRenderPass(const FRHIRenderPassInfo& InInf
 			CurrentEncoderFence = CurrentEncoder.EndEncoding();
 		}
 		CurrentEncoder.SetRenderPassDescriptor(RenderPassDesc);
-		CurrentEncoder.BeginRenderCommandEncoding();
+		bool bUseStageCounterSamples = Device.SupportsFeature(EMetalFeaturesStageCounterSampling);
+		CurrentEncoder.BeginRenderCommandEncoding(bUseStageCounterSamples ? Device.GetCounterSampler() : nullptr);
 	}
 	
 	if (CurrentEncoderFence)
@@ -211,11 +224,6 @@ void FMetalRHICommandContext::RHIBeginRenderPass(const FRHIRenderPassInfo& InInf
 
 void FMetalRHICommandContext::RHIEndRenderPass()
 {
-	if (RenderPassInfo.NumOcclusionQueries > 0)
-	{
-		RHIEndOcclusionQueryBatch();
-	}
-    
 	check(bWithinRenderPass);
 	check(CurrentEncoder.IsRenderCommandEncoderActive());
 	
@@ -266,10 +274,12 @@ void FMetalRHICommandContext::ResolveTexture(UE::RHICore::FResolveTextureInfo In
 		Size.height = FMath::Max<uint32>(1, SourceDesc.Extent.Y >> Info.MipLevel);
 	}
 
+#if RHI_NEW_GPU_PROFILER == 0
 	if (Profiler)
 	{
 		Profiler->RegisterGPUWork();
 	}
+#endif
 
 	int32 ArraySliceBegin = Info.ArraySlice;
 	int32 ArraySliceEnd   = Info.ArraySlice + 1;
@@ -321,49 +331,6 @@ void FMetalRHICommandContext::RHICalibrateTimers(FRHITimestampCalibrationQuery*
     CalibrationQuery->GPUMicroseconds[0] = uint64(GPUTimestamp / 1000.0);
 }
 
-void FMetalDynamicRHI::RHIBeginRenderQuery_TopOfPipe(FRHICommandListBase& RHICmdList, FRHIRenderQuery* RenderQuery)
-{
-	FMetalRHIRenderQuery* Query = ResourceCast(RenderQuery);
-	Query->Begin_TopOfPipe();
-
-	FDynamicRHI::RHIBeginRenderQuery_TopOfPipe(RHICmdList, RenderQuery);
-}
-
-void FMetalDynamicRHI::RHIEndRenderQuery_TopOfPipe(FRHICommandListBase& RHICmdList, FRHIRenderQuery* RenderQuery)
-{
-	FMetalRHIRenderQuery* Query = ResourceCast(RenderQuery);
-	Query->End_TopOfPipe();
-
-	FDynamicRHI::RHIEndRenderQuery_TopOfPipe(RHICmdList, RenderQuery);
-}
-
-void FMetalRHICommandContext::RHIBeginRenderQuery(FRHIRenderQuery* QueryRHI)
-{
-    MTL_SCOPED_AUTORELEASE_POOL;
-	FMetalRHIRenderQuery* Query = ResourceCast(QueryRHI);
-    Query->Begin(this, CommandBufferFence);
-}
-
-void FMetalRHICommandContext::RHIEndRenderQuery(FRHIRenderQuery* QueryRHI)
-{
-    MTL_SCOPED_AUTORELEASE_POOL;
-    FMetalRHIRenderQuery* Query = ResourceCast(QueryRHI);
-	Query->End(this);
-}
-
-void FMetalRHICommandContext::RHIBeginOcclusionQueryBatch(uint32 NumQueriesInBatch)
-{
-    check(!CommandBufferFence.IsValid());
-	CommandBufferFence = MakeShareable(new FMetalCommandBufferFence);
-    CurrentEncoder.InsertCommandBufferFence(CommandBufferFence, FMetalCommandBufferCompletionHandler());
-}
-
-void FMetalRHICommandContext::RHIEndOcclusionQueryBatch()
-{
-	check(CommandBufferFence.IsValid());
-	CommandBufferFence.Reset();
-}
-
 void FMetalRHICommandContext::FillBuffer(MTL::Buffer* Buffer, NS::Range Range, uint8 Value)
 {
 	check(Buffer);
@@ -447,7 +414,7 @@ void FMetalRHICommandContext::CopyFromBufferToBuffer(FMetalBufferPtr SourceBuffe
 	EndBlitEncoder();
 }
 
-TArray<FMetalCommandBuffer*> FMetalRHICommandContext::Finalize()
+void FMetalRHICommandContext::Finalize(TArray<FMetalPayload*>& OutPayloads)
 {
 	GetQueryBufferPool()->ReleaseCurrentQueryBuffer();
 	
@@ -456,66 +423,128 @@ TArray<FMetalCommandBuffer*> FMetalRHICommandContext::Finalize()
 		CurrentEncoderFence = CurrentEncoder.EndEncoding();
 	}
 	
-	TArray<FMetalCommandBuffer*> CommandBuffers;
-	
 	if (CurrentEncoder.GetCommandBuffer())
 	{
-		CommandBuffers = CurrentEncoder.Finalize();
+		EndCommandBuffer();
+	}
+	
+	// Collect the context's batch of sync points to wait/signal
+	if (BatchedSyncPoints.ToWait.Num())
+	{
+		FMetalPayload* Payload = Payloads.Num()
+			? Payloads[0]
+			: GetPayload(EPhase::Wait);
+
+		Payload->SyncPointsToWait.Append(BatchedSyncPoints.ToWait);
+		BatchedSyncPoints.ToWait.Reset();
+	}
+
+	if (BatchedSyncPoints.ToSignal.Num())
+	{
+		GetPayload(EPhase::Signal)->SyncPointsToSignal.Append(BatchedSyncPoints.ToSignal);
+		BatchedSyncPoints.ToSignal.Reset();
 	}
 	
-	return CommandBuffers;
+	OutPayloads.Append(MoveTemp(Payloads));
 }
 
-void FMetalRHICommandContext::InsertCommandBufferFence(TSharedPtr<FMetalCommandBufferFence, ESPMode::ThreadSafe>& Fence, FMetalCommandBufferCompletionHandler Handler)
+void FMetalRHICommandContext::SignalSyncPoint(FMetalSyncPoint* SyncPoint)
 {
-	CurrentEncoder.InsertCommandBufferFence(Fence, Handler);
+	if(CurrentEncoder.GetCommandBuffer())
+	{
+		EndCommandBuffer();
+	}
+	
+	GetPayload(EPhase::Signal)->SyncPointsToSignal.Add(SyncPoint);
+}
+
+void FMetalRHICommandContext::WaitSyncPoint(FMetalSyncPoint* SyncPoint)
+{
+	if(CurrentEncoder.GetCommandBuffer())
+	{
+		EndCommandBuffer();
+	}
+	
+	GetPayload(EPhase::Wait)->SyncPointsToWait.Add(SyncPoint);
 }
 
 void FMetalRHICommandContext::SignalEvent(MTLEventPtr Event, uint32_t SignalCount)
 {
+	if(!CurrentEncoder.GetCommandBuffer())
+	{
+		StartCommandBuffer();
+	}
 	CurrentEncoder.SignalEvent(Event, SignalCount);
 }
 
 void FMetalRHICommandContext::WaitForEvent(MTLEventPtr Event, uint32_t SignalCount)
 {
+	if(!CurrentEncoder.GetCommandBuffer())
+	{
+		StartCommandBuffer();
+	}
 	CurrentEncoder.WaitForEvent(Event, SignalCount);
 }
 
-void FMetalRHICommandContext::StartTiming(class FMetalEventNode* EventNode)
+void FMetalRHICommandContext::StartCommandBuffer()
 {
-	FMetalCommandBufferCompletionHandler Handler;
+	check(!CurrentEncoder.GetCommandBuffer());
 	
-	bool const bHasCurrentCommandBuffer = CurrentEncoder.GetCommandBuffer();
+	// Begin the render pass frame.
+	CurrentEncoder.StartCommandBuffer();
+	
+#if RHI_NEW_GPU_PROFILER
+	auto& Event = CurrentEncoder.GetCommandBuffer()->EmplaceProfilerEvent<UE::RHI::GPUProfiler::FEvent::FBeginWork>(0);
+	Event.GPUTimestampTOP = UINT64_MAX;
+	CurrentEncoder.GetCommandBuffer()->SetBeginWorkTimestamp(&Event.GPUTimestampTOP);
+#endif
+	
+	// Add new command buffer to payload
+	GetPayload(EPhase::Execute)->CommandBuffersToExecute.Add(CurrentEncoder.GetCommandBuffer());
+}
+
+void FMetalRHICommandContext::EndCommandBuffer()
+{
+	check(CurrentEncoder.GetCommandBuffer());
+	check(!bWithinRenderPass);
 	
-	if(EventNode)
+#if RHI_NEW_GPU_PROFILER
+	auto& Event = CurrentEncoder.GetCommandBuffer()->EmplaceProfilerEvent<UE::RHI::GPUProfiler::FEvent::FEndWork>();
+	Event.GPUTimestampBOP = 0;
+	CurrentEncoder.GetCommandBuffer()->SetEndWorkTimestamp(&Event.GPUTimestampBOP);
+#endif
+	
+	if(CurrentEncoder.IsAnyCommandEncoderActive())
 	{
-		Handler = EventNode->Start();
-		
-		if (bHasCurrentCommandBuffer)
-		{
-			CurrentEncoder.AddCompletionHandler(Handler);
-		}
-		
-		if(!bWithinRenderPass)
-		{
-			CurrentEncoder.SplitCommandBuffers();
-		}
+		CurrentEncoder.EndEncoding();
 	}
+	CurrentEncoder.EndCommandBuffer(this);
+}
+
+void FMetalRHICommandContext::StartTiming(class FMetalEventNode* EventNode)
+{
+#if RHI_NEW_GPU_PROFILER == 0
+	bool const bHasCurrentCommandBuffer = CurrentEncoder.GetCommandBuffer();
 	
-	if (Handler.IsBound() && !bHasCurrentCommandBuffer)
+	if(!bHasCurrentCommandBuffer)
 	{
-		CurrentEncoder.GetCommandBuffer()->GetMTLCmdBuffer()->addScheduledHandler(
-						MTL::HandlerFunction([Handler](MTL::CommandBuffer* CommandBuffer)
-						{
-							Handler.Execute(CommandBuffer);
-						}));
+		StartCommandBuffer();
 	}
+
+	CurrentEncoder.GetCommandBuffer()->ActiveEventNodes.Add(EventNode);
+
+	EventNode->SyncPoint = GetContextSyncPoint();
+#endif
 }
 
 void FMetalRHICommandContext::EndTiming(class FMetalEventNode* EventNode)
 {
-	FMetalCommandBufferCompletionHandler Handler = EventNode->Stop();
-	CurrentEncoder.AddCompletionHandler(Handler);
+#if RHI_NEW_GPU_PROFILER == 0
+	if(CurrentEncoder.GetCommandBuffer())
+	{
+		CurrentEncoder.GetCommandBuffer()->ActiveEventNodes.Remove(EventNode);
+	}
+#endif
 }
 
 void FMetalRHICommandContext::SynchronizeResource(MTL::Resource* Resource)
@@ -546,21 +575,29 @@ void FMetalRHICommandContext::SynchronizeTexture(MTL::Texture* Texture, uint32 S
 #endif
 }
 
-void FMetalRHICommandContext::AddCompletionHandler(FMetalCommandBufferCompletionHandler& Handler)
-{
-	CurrentEncoder.AddCompletionHandler(Handler);
-}
-
 FMetalCommandBuffer* FMetalRHICommandContext::GetCurrentCommandBuffer()
 {
 	FMetalCommandBuffer* CommandBuffer = CurrentEncoder.GetCommandBuffer();
 	if(!CommandBuffer)
 	{
-		CurrentEncoder.StartCommandBuffer();
+		StartCommandBuffer();
 	}
 	return CurrentEncoder.GetCommandBuffer();
 }
 
+
+#if RHI_NEW_GPU_PROFILER
+void FMetalRHICommandContext::FlushProfilerStats()
+{
+	// Flush accumulated draw stats
+	if (StatEvent)
+	{
+		CurrentEncoder.GetCommandBuffer()->EmplaceProfilerEvent<UE::RHI::GPUProfiler::FEvent::FStats>() = StatEvent;
+		StatEvent = {};
+	}
+}
+#endif
+
 FMetalRHIUploadContext::FMetalRHIUploadContext(FMetalDevice& Device)
 {
 	UploadContext = new FMetalRHICommandContext(Device, nullptr);
@@ -578,7 +615,7 @@ FMetalRHIUploadContext::~FMetalRHIUploadContext()
 	delete WaitContext;
 }
 
-TArray<FMetalCommandBuffer*>* FMetalRHIUploadContext::Finalize()
+void FMetalRHIUploadContext::Finalize(TArray<FMetalPayload*>& OutPayloads)
 {
 	for(auto& Function : UploadFunctions)
 	{
@@ -588,19 +625,15 @@ TArray<FMetalCommandBuffer*>* FMetalRHIUploadContext::Finalize()
 	UploadSyncCounter++;
 	UploadContext->SignalEvent(UploadSyncEvent, UploadSyncCounter);
 	
-	TArray<FMetalCommandBuffer*>* CommandBuffers = new TArray<FMetalCommandBuffer*>();
-	
-	CommandBuffers->Append(UploadContext->Finalize());
+	UploadContext->Finalize(OutPayloads);
 	
 	UploadFunctions.Reset();
 	UploadContext->ResetContext();
 	
 	WaitContext->WaitForEvent(UploadSyncEvent, UploadSyncCounter);
-	CommandBuffers->Append(WaitContext->Finalize());
+	WaitContext->Finalize(OutPayloads);
 	
 	WaitContext->ResetContext();
-	
-	return CommandBuffers;
 }
 
 FMetalContextArray::FMetalContextArray(FRHIContextArray const& Contexts)
diff --git a/Engine/Source/Runtime/Apple/MetalRHI/Private/MetalRHIPrivate.h b/Engine/Source/Runtime/Apple/MetalRHI/Private/MetalRHIPrivate.h
--- a/Engine/Source/Runtime/Apple/MetalRHI/Private/MetalRHIPrivate.h
+++ b/Engine/Source/Runtime/Apple/MetalRHI/Private/MetalRHIPrivate.h
@@ -122,7 +122,7 @@ extern bool GMetalCommandBufferDebuggingEnabled;
 #define ENABLE_METAL_GPUEVENTS    (UE_BUILD_DEBUG || UE_BUILD_DEVELOPMENT || (UE_BUILD_TEST && ENABLE_METAL_GPUEVENTS_IN_TEST))
 #define ENABLE_METAL_GPUPROFILE    (ENABLE_METAL_GPUEVENTS && 1)
 
-#if ENABLE_METAL_GPUPROFILE
+#if ENABLE_METAL_GPUPROFILE && RHI_NEW_GPU_PROFILER == 0
     #define METAL_GPUPROFILE(Code) Code
 #else
     #define METAL_GPUPROFILE(Code)
diff --git a/Engine/Source/Runtime/Apple/MetalRHI/Private/MetalRHIRenderQuery.cpp b/Engine/Source/Runtime/Apple/MetalRHI/Private/MetalRHIRenderQuery.cpp
--- a/Engine/Source/Runtime/Apple/MetalRHI/Private/MetalRHIRenderQuery.cpp
+++ b/Engine/Source/Runtime/Apple/MetalRHI/Private/MetalRHIRenderQuery.cpp
@@ -7,6 +7,7 @@
 #include "MetalRHIRenderQuery.h"
 #include "MetalDevice.h"
 #include "MetalRHIPrivate.h"
+#include "MetalDynamicRHI.h"
 #include "MetalProfiler.h"
 #include "MetalLLM.h"
 #include "MetalCommandBuffer.h"
@@ -140,29 +141,95 @@ void FMetalQueryBufferPool::ReleaseQueryBuffer(FMetalBufferPtr Buffer)
 
 void FMetalQueryResult::Reset()
 {
-	CommandBufferFence.Reset();
 	bCompleted = false;
 }
 
-bool FMetalQueryResult::Wait(uint64 Millis)
+uint64 FMetalQueryResult::GetResult()
 {
-	if (!bCompleted)
+	if (IsValidRef(SourceBuffer))
 	{
-		bCompleted = CommandBufferFence->Wait(Millis);
-		return bCompleted;
+		return SourceBuffer->GetResult(Offset);
 	}
-	return true;
+	return 0;
 }
 
-uint64 FMetalQueryResult::GetResult()
+//------------------------------------------------------------------------------
+
+#pragma mark - Metal RHI Command Context Functions -
+
+void FMetalDynamicRHI::RHIBeginRenderQuery_TopOfPipe(FRHICommandListBase& RHICmdList, FRHIRenderQuery* RenderQuery)
 {
-	if (IsValidRef(SourceBuffer))
+	FMetalRHIRenderQuery* Query = ResourceCast(RenderQuery);
+	Query->Begin_TopOfPipe();
+
+	FDynamicRHI::RHIBeginRenderQuery_TopOfPipe(RHICmdList, RenderQuery);
+}
+
+void FMetalDynamicRHI::RHIEndRenderQuery_TopOfPipe(FRHICommandListBase& RHICmdList, FRHIRenderQuery* RenderQuery)
+{
+	FMetalRHIRenderQuery* Query = ResourceCast(RenderQuery);
+	auto& QueryBatchData = RHICmdList.GetQueryBatchData(Query->Type);
+
+	if (QueryBatchData[0])
 	{
-		return SourceBuffer->GetResult(Offset);
+		// This query belongs to a batch. Use the sync point we created earlier
+		Query->SyncPoint = static_cast<FMetalSyncPoint*>(QueryBatchData[0]);
 	}
-	return 0;
+	else
+	{
+		// Queries issued outside of a batch use one sync point per query.
+		Query->SyncPoint = FMetalSyncPoint::Create(EMetalSyncPointType::GPUAndCPU);
+
+		RHICmdList.EnqueueLambda([SyncPoint = Query->SyncPoint](FRHICommandListBase& RHICmdList) mutable
+		{
+			FMetalRHICommandContext& Context = FMetalRHICommandContext::Get(RHICmdList);
+			Context.BatchedSyncPoints.ToSignal.Emplace(MoveTemp(SyncPoint));
+		});
+	}
+	
+	FDynamicRHI::RHIEndRenderQuery_TopOfPipe(RHICmdList, RenderQuery);
+}
+
+void FMetalDynamicRHI::RHIBeginRenderQueryBatch_TopOfPipe(FRHICommandListBase& RHICmdList, ERenderQueryType QueryType)
+{
+	auto& QueryBatchData = RHICmdList.GetQueryBatchData(QueryType);
+	FMetalSyncPointRef QuerySyncPoint = FMetalSyncPoint::Create(EMetalSyncPointType::GPUAndCPU);
+	
+	QueryBatchData[0] = QuerySyncPoint.GetReference();
+	QuerySyncPoint->AddRef();
 }
 
+void FMetalDynamicRHI::RHIEndRenderQueryBatch_TopOfPipe(FRHICommandListBase& RHICmdList, ERenderQueryType QueryType)
+{
+	auto& QueryBatchData = RHICmdList.GetQueryBatchData(QueryType);
+	checkf(QueryBatchData[0], TEXT("A query batch for this type is not open on this command list."));
+
+	FMetalSyncPointRef SyncPoint = static_cast<FMetalSyncPoint*>(QueryBatchData[0]);
+
+	// Clear the sync point reference on the RHI command list
+	SyncPoint->Release();
+	QueryBatchData[0] = nullptr;
+
+	RHICmdList.EnqueueLambda([SyncPoint = MoveTemp(SyncPoint), QueryType](FRHICommandListBase& ExecutingCmdList)
+	{
+		FMetalRHICommandContext& Context = FMetalRHICommandContext::Get(ExecutingCmdList);
+		Context.BatchedSyncPoints.ToSignal.Add(SyncPoint);
+	});
+}
+
+void FMetalRHICommandContext::RHIBeginRenderQuery(FRHIRenderQuery* QueryRHI)
+{
+	MTL_SCOPED_AUTORELEASE_POOL;
+	FMetalRHIRenderQuery* Query = ResourceCast(QueryRHI);
+	Query->Begin(this);
+}
+
+void FMetalRHICommandContext::RHIEndRenderQuery(FRHIRenderQuery* QueryRHI)
+{
+	MTL_SCOPED_AUTORELEASE_POOL;
+	FMetalRHIRenderQuery* Query = ResourceCast(QueryRHI);
+	Query->End(this);
+}
 
 //------------------------------------------------------------------------------
 
@@ -175,7 +242,6 @@ FMetalRHIRenderQuery::FMetalRHIRenderQuery(FMetalDevice& MetalDevice, ERenderQue
 	, Buffer{}
 	, Result{0}
 	, bAvailable{false}
-	, QueryWrittenEvent(nullptr)
 {
 	// void
 }
@@ -184,12 +250,6 @@ FMetalRHIRenderQuery::~FMetalRHIRenderQuery()
 {
 	Buffer.SourceBuffer.SafeRelease();
 	Buffer.Offset = 0;
-	
-	if(QueryWrittenEvent != nullptr)
-	{
-		FPlatformProcess::ReturnSynchEventToPool(QueryWrittenEvent);
-		QueryWrittenEvent = nullptr;
-	}
 }
 
 void FMetalRHIRenderQuery::Begin_TopOfPipe()
@@ -203,20 +263,14 @@ void FMetalRHIRenderQuery::End_TopOfPipe()
 	if (Type == RQT_AbsoluteTime)
 	{
 		Buffer.Reset();
-		if(QueryWrittenEvent)
-		{
-			QueryWrittenEvent->Reset();
-		}
 	}
 	bAvailable = false;
 }
 
-void FMetalRHIRenderQuery::Begin(FMetalRHICommandContext* Context, TSharedPtr<FMetalCommandBufferFence, ESPMode::ThreadSafe> const& BatchFence)
+void FMetalRHIRenderQuery::Begin(FMetalRHICommandContext* Context)
 {
-	Buffer.CommandBufferFence.Reset();
 	Buffer.SourceBuffer.SafeRelease();
 	Buffer.Offset = 0;
-	Buffer.bBatchFence = false;
 
 	Result = 0;
 	bAvailable = false;
@@ -237,15 +291,6 @@ void FMetalRHIRenderQuery::Begin(FMetalRHICommandContext* Context, TSharedPtr<FM
 			{
 				Context->GetStateCache().SetVisibilityResultMode(MTL::VisibilityResultModeBoolean, Buffer.Offset);
 			}
-			if (BatchFence.IsValid())
-			{
-				Buffer.CommandBufferFence = BatchFence;
-				Buffer.bBatchFence = true;
-			}
-			else
-			{
-				Buffer.CommandBufferFence = MakeShareable(new FMetalCommandBufferFence);
-			}
 			break;
 		}
 		case RQT_AbsoluteTime:
@@ -267,14 +312,8 @@ void FMetalRHIRenderQuery::End(FMetalRHICommandContext* Context)
 		case RQT_Occlusion:
 		{
 			// switch back to non-occlusion rendering
-			check(Buffer.CommandBufferFence.IsValid());
 			Context->GetStateCache().SetVisibilityResultMode(MTL::VisibilityResultModeDisabled, 0);
-
-			// For unique, unbatched, queries insert the fence now
-			if (!Buffer.bBatchFence)
-			{
-				Context->InsertCommandBufferFence(Buffer.CommandBufferFence, FMetalCommandBufferCompletionHandler());
-			}
+			Context->GetCurrentCommandBuffer()->OcclusionQueries.Add(this);
 			break;
 		}
 		case RQT_AbsoluteTime:
@@ -285,57 +324,12 @@ void FMetalRHIRenderQuery::End(FMetalRHICommandContext* Context)
 			Buffer.SourceBuffer.SafeRelease();
 			Buffer.Offset = 0;
 			Buffer.bCompleted = false;
-			Buffer.bBatchFence = false;
-			TSharedPtr<FMetalCommandBufferFence, ESPMode::ThreadSafe> CommandBufferFence = MakeShareable(new FMetalCommandBufferFence);
-
 			Result = 0;
 			bAvailable = false;
-
-			if(QueryWrittenEvent == nullptr)
-			{
-				QueryWrittenEvent = FPlatformProcess::GetSynchEventFromPool(true);
-				check(QueryWrittenEvent != nullptr);
-			}
-			else
-			{
-				QueryWrittenEvent->Reset();
-			}
-
-			// Insert the fence to wait on the current command buffer
-            FMetalCommandBufferCompletionHandler Handler;
-            Handler.BindLambda([this](MTL::CommandBuffer* CmdBuffer)
-			{
-				Result = 0;
-				
-                // If there are no commands in the command buffer then this can be zero
-                // In this case GPU start time is also not correct - we need to fall back standard behaviour
-                // Only seen empty command buffers at the very end of a frame
-                if (IsRHIDeviceApple())
-                {
-                    Result = uint64(CmdBuffer->GPUEndTime() - CmdBuffer->GPUStartTime()) * 1000;
-                }
-                else
-                {
-                    Result = uint64((CmdBuffer->GPUEndTime() / 1000.0) / FPlatformTime::GetSecondsPerCycle64());
-                }
-				
-				if(Result == 0)
-				{
-					Result = (FPlatformTime::ToMilliseconds64(mach_absolute_time()) * 1000.0);
-				}
-
-				if(QueryWrittenEvent != nullptr)
-				{
-					QueryWrittenEvent->Trigger();
-				}
-				
-				this->Release();
-			});
-
-			Context->InsertCommandBufferFence(CommandBufferFence, Handler);
-			Buffer.CommandBufferFence = CommandBufferFence;
 			
-			Context->SplitCommandBuffers();
+			CommandBuffer = Context->GetCurrentCommandBuffer();
+			
+			CommandBuffer->TimestampQueries.Add(this);
             
 			break;
 		}
@@ -351,65 +345,19 @@ bool FMetalRHIRenderQuery::GetResult(uint64& OutNumPixels, bool bWait, uint32 GP
 {
 	if (!bAvailable)
 	{
-		SCOPE_CYCLE_COUNTER(STAT_RenderQueryResultTime);
-
-		bool bOK = false;
-
-		bool const bCmdBufferIncomplete = !Buffer.bCompleted;
-
-		// timer queries are used for Benchmarks which can stall a bit more
-		uint64 WaitMS = (Type == RQT_AbsoluteTime) ? 30000 : 500;
-		if (bWait)
+		if (!SyncPoint->IsComplete())
 		{
-			FRHICommandListImmediate& RHICmdList = FRHICommandListImmediate::Get();
-			
-			// RHI thread *must* be flushed at this point if the internal handles we rely upon are not yet valid.
-			if (!Buffer.CommandBufferFence.IsValid())
-			{
-				RHICmdList.ImmediateFlush(EImmediateFlushType::FlushRHIThread);
-			}
-
+			if (bWait)
 			{
-				FRenderThreadIdleScope IdleScope(ERenderThreadIdleTypes::WaitingForGPUQuery);
-				bOK = Buffer.Wait(WaitMS);
+				SyncPoint->Wait();
 			}
-
-			// Result is written in one of potentially many command buffer completion handlers
-			// But the command buffer wait above may return before the query completion handler fires
-			// We need to wait here until that has happenned, also make sure the command buffer actually completed due to timeout
-			if (bOK && (Type == RQT_AbsoluteTime) && QueryWrittenEvent != nullptr)
-			{
-				QueryWrittenEvent->Wait();
-			}
-
-			// Never wait for a failed signal again.
-			bAvailable = Buffer.bCompleted;
-		}
-		else
-		{
-			if(Buffer.CommandBufferFence.IsValid())
+			else
 			{
-				bOK = Buffer.Wait(0);
+				return false;
 			}
 		}
-
-		if (bOK == false)
-		{
-			OutNumPixels = 0;
-			UE_CLOG(bWait, LogMetal, Display, TEXT("Timed out while waiting for GPU to catch up. (%llu ms)"), WaitMS);
-			return false;
-		}
-
-		if (Type == RQT_Occlusion)
-		{
-			Result = Buffer.GetResult();
-		}
-
-		Buffer.SourceBuffer.SafeRelease();
 	}
 
-	// at this point, we are ready to read the value!
 	OutNumPixels = Result;
-
 	return true;
 }
diff --git a/Engine/Source/Runtime/Apple/MetalRHI/Private/MetalRHIRenderQuery.h b/Engine/Source/Runtime/Apple/MetalRHI/Private/MetalRHIRenderQuery.h
--- a/Engine/Source/Runtime/Apple/MetalRHI/Private/MetalRHIRenderQuery.h
+++ b/Engine/Source/Runtime/Apple/MetalRHI/Private/MetalRHIRenderQuery.h
@@ -8,6 +8,7 @@
 
 #include "MetalRHIPrivate.h"
 #include "RHIResources.h"
+#include "MetalSubmission.h"
 
 //------------------------------------------------------------------------------
 
@@ -17,7 +18,6 @@
 class FMetalRHICommandContext;
 class FMetalQueryBufferPool;
 class FMetalQueryResult;
-class FMetalCommandBufferFence;
 class FMetalDevice;
 class FMetalBuffer;
 typedef TSharedPtr<FMetalBuffer, ESPMode::ThreadSafe> FMetalBufferPtr;
@@ -85,11 +85,9 @@ public:
 	~FMetalQueryResult() = default;
 
 	void Reset();
-	bool Wait(uint64 Millis);
 	uint64 GetResult();
 
 	TRefCountPtr<FMetalQueryBuffer> SourceBuffer = nullptr;
-	TSharedPtr<FMetalCommandBufferFence, ESPMode::ThreadSafe> CommandBufferFence = nullptr;
 	uint32 Offset = 0;
 	bool bCompleted = false;
 	bool bBatchFence = false;
@@ -103,6 +101,7 @@ public:
 
 class FMetalRHIRenderQuery : public FRHIRenderQuery
 {
+	friend class FMetalDynamicRHI;
 public:
 	FMetalRHIRenderQuery(FMetalDevice& MetalDevice, ERenderQueryType InQueryType);
 	virtual ~FMetalRHIRenderQuery();
@@ -113,7 +112,7 @@ public:
 	/**
 	 * Kick off an occlusion test
 	 */
-	void Begin(FMetalRHICommandContext* Context, TSharedPtr<FMetalCommandBufferFence, ESPMode::ThreadSafe> const& BatchFence);
+	void Begin(FMetalRHICommandContext* Context);
 
 	/**
 	 * Finish up an occlusion test
@@ -140,6 +139,9 @@ private:
 	// Result availability - if not set the first call to acquire it will read the buffer & cache
 	std::atomic<bool> bAvailable;
 
-	// Timer event completion signal
-	FEvent* QueryWrittenEvent;
+	// Signaled when the result is available. Nullptr if the query has never been used.
+	FMetalSyncPointRef SyncPoint;
+	
+	// Temporary
+	FMetalCommandBuffer* CommandBuffer = nullptr;
 };
diff --git a/Engine/Source/Runtime/Apple/MetalRHI/Private/MetalSubmission.cpp b/Engine/Source/Runtime/Apple/MetalRHI/Private/MetalSubmission.cpp
new file mode 100644
--- /dev/null
+++ b/Engine/Source/Runtime/Apple/MetalRHI/Private/MetalSubmission.cpp
@@ -0,0 +1,900 @@
+// Copyright Epic Games, Inc. All Rights Reserved.
+
+#include "MetalSubmission.h"
+#include "MetalRHIPrivate.h"
+#include "MetalDynamicRHI.h"
+#include "MetalCommandBuffer.h"
+#include "HAL/Runnable.h"
+#include "HAL/RunnableThread.h"
+#include "IRenderCaptureProvider.h"
+#include "MetalProfiler.h"
+
+static TAutoConsoleVariable<int32> CVarMetalRHIUseSubmissionThread(
+	TEXT("MetalRHI.UseSubmissionThread"),
+	1,
+	TEXT("Whether to enable the Metal RHI's submission thread.\n")
+	TEXT("  0: No\n")
+	TEXT("  1: Yes\n"),
+	ECVF_ReadOnly);
+
+#define METAL_USE_INTERRUPT_THREAD 1
+#define METAL_USE_SUBMISSION_THREAD 1
+
+TLockFreePointerListUnordered<void, PLATFORM_CACHE_LINE_SIZE> FMetalSyncPoint::MemoryPool;
+
+class FMetalThread final : private FRunnable
+{
+public:
+	typedef FMetalDynamicRHI::FProcessResult(FMetalDynamicRHI::*FQueueFunc)();
+
+	FMetalThread(TCHAR const* Name, EThreadPriority Priority, FMetalDynamicRHI* RHI, FQueueFunc QueueFunc)
+		: RHI(RHI)
+		, QueueFunc(QueueFunc)
+		, Event(FPlatformProcess::GetSynchEventFromPool(true))
+		, Thread(FRunnableThread::Create(this, Name, 0, Priority))
+	{}
+
+	virtual ~FMetalThread()
+	{
+		bExit = true;
+
+		Event->Trigger();
+		
+		Thread->WaitForCompletion();
+		delete Thread;
+	}
+
+	void Kick() const
+	{
+		check(Event);
+		Event->Trigger();
+	}
+
+	void Join() const
+	{
+		Thread->WaitForCompletion();
+	}
+
+private:
+	virtual uint32 Run() override
+	{
+		while (!bExit)
+		{
+			// Process the queue until no more progress is made
+			FMetalDynamicRHI::FProcessResult Result;
+			do { Result = (RHI->*QueueFunc)(); }
+			while (EnumHasAllFlags(Result.Status, FMetalDynamicRHI::EQueueStatus::Processed));
+	
+			Event->Wait(1);
+			Event->Reset();
+		}
+
+		// Drain any remaining work in the queue
+		while (EnumHasAllFlags((RHI->*QueueFunc)().Status, FMetalDynamicRHI::EQueueStatus::Pending)) {}
+
+		return 0;
+	}
+
+	FMetalDynamicRHI* RHI;
+	FQueueFunc QueueFunc;
+	FEvent* Event;
+	bool bExit = false;
+	
+private:
+	FRunnableThread* Thread = nullptr;
+};
+
+void FMetalDynamicRHI::InitializeSubmissionPipe()
+{
+	if (FPlatformProcess::SupportsMultithreading())
+	{
+#if METAL_USE_INTERRUPT_THREAD
+		InterruptThread = new FMetalThread(TEXT("RHIInterruptThread"), TPri_Highest, this, &FMetalDynamicRHI::ProcessInterruptQueue);
+#endif
+
+#if METAL_USE_SUBMISSION_THREAD
+		bool bUseSubmissionThread = false;
+		switch (CVarMetalRHIUseSubmissionThread.GetValueOnAnyThread())
+		{
+		case 1: bUseSubmissionThread = true; break;
+		}
+		
+		if (bUseSubmissionThread)
+		{
+			SubmissionThread = new FMetalThread(TEXT("RHISubmissionThread"), TPri_Highest, this, &FMetalDynamicRHI::ProcessSubmissionQueue);
+		}
+#endif
+	}
+
+	// Initialize the timing structs in each queue, and the engine GPU profilers
+#if RHI_NEW_GPU_PROFILER
+	{
+		TArray<FMetalPayload*> Payloads;
+		TArray<UE::RHI::GPUProfiler::FQueue> ProfilerQueues;
+		
+		ForEachQueue([&](FMetalCommandQueue& Queue)
+		{			
+			FMetalPayload* Payload = Payloads.Emplace_GetRef(new FMetalPayload(Queue));
+			Payload->Timing = CurrentTimingPerQueue.CreateNew(Queue);
+
+			ProfilerQueues.Add(Queue.GetProfilerQueue());
+		});
+
+		UE::RHI::GPUProfiler::InitializeQueues(ProfilerQueues);
+	
+		SubmitPayloads(MoveTemp(Payloads));
+	}
+#endif
+}
+
+void FMetalDynamicRHI::ShutdownSubmissionPipe()
+{
+	delete SubmissionThread;
+	SubmissionThread = nullptr;
+
+	delete InterruptThread;
+	InterruptThread = nullptr;
+
+	if (EopTask)
+	{
+		ProcessInterruptQueueUntil(EopTask);
+		EopTask = nullptr;
+	}
+}
+
+struct FMetalFinalizedCommands final : public IRHIPlatformCommandList, public TArray<FMetalPayload*> 
+{};
+
+static TLockFreePointerListUnordered<FMetalRHIUploadContext, PLATFORM_CACHE_LINE_SIZE> MetalUploadContextPool;
+
+IRHIUploadContext* FMetalDynamicRHI::RHIGetUploadContext()
+{
+	FMetalRHIUploadContext* Context = MetalUploadContextPool.Pop();
+	if (!Context)
+	{
+		Context = new FMetalRHIUploadContext(*Device);
+	}
+	
+	return static_cast<IRHIUploadContext*>(Context);
+}
+
+void FMetalDynamicRHI::RHIFinalizeContext(FRHIFinalizeContextArgs&& Args, TRHIPipelineArray<IRHIPlatformCommandList*>& Output)
+{
+	MTL_SCOPED_AUTORELEASE_POOL;
+	
+	FMetalRHIUploadContext* UploadContext = static_cast<FMetalRHIUploadContext*>(Args.UploadContext);
+
+	FMetalFinalizedCommands Commands;
+	
+	if(UploadContext)
+	{
+		UploadContext->Finalize(Commands);
+		MetalUploadContextPool.Push(UploadContext);
+	}
+	
+	for(IRHIComputeContext* Context : Args.Contexts)
+	{
+		FMetalRHICommandContext* CmdContext = static_cast<FMetalRHICommandContext*>(Context);
+		
+		if(!CmdContext->IsInsideRenderPass())
+		{
+			CmdContext->Finalize(Commands);
+
+			CmdContext->ResetContext();
+			if(GRHISupportsParallelRHIExecute)
+			{
+				if(CmdContext != RHIGetDefaultContext())
+				{
+					MetalCommandContextPool.Push(CmdContext);
+				}
+			}
+		}
+		
+		Output[Context->GetPipeline()] = new FMetalFinalizedCommands(MoveTemp(Commands));
+	}
+}
+
+void FMetalDynamicRHI::RHISubmitCommandLists(FRHISubmitCommandListsArgs&& Args)
+{
+	SubmitCommands(MakeArrayView(reinterpret_cast<FMetalFinalizedCommands**>(Args.CommandLists.GetData()), Args.CommandLists.Num()));
+}
+
+void FMetalDynamicRHI::SubmitCommands(TConstArrayView<FMetalFinalizedCommands*> Commands)
+{	
+	SCOPED_NAMED_EVENT_TEXT("CommandList_Submit", FColor::Magenta);
+
+#if RHI_NEW_GPU_PROFILER
+	TArray<FMetalPayload*> AllPayloads;
+	for (FMetalFinalizedCommands* Payloads : Commands)
+	{
+	#if WITH_RHI_BREADCRUMBS
+		TSharedPtr<FRHIBreadcrumbAllocatorArray> BreadcrumbAllocators {};
+		if (Payloads->BreadcrumbAllocators.Num())
+		{
+			BreadcrumbAllocators = MakeShared<FRHIBreadcrumbAllocatorArray>(MoveTemp(Payloads->BreadcrumbAllocators));
+		}
+
+		for (FMetalPayload* Payload : *Payloads)
+		{
+			Payload->BreadcrumbRange = Payloads->BreadcrumbRange;
+			if (BreadcrumbAllocators.IsValid())
+			{
+				check(!Payload->BreadcrumbAllocators.IsValid());
+				Payload->BreadcrumbAllocators = BreadcrumbAllocators;
+			}
+		}
+	#endif
+
+		AllPayloads.Append(MoveTemp(*Payloads));
+		delete Payloads;
+	}
+
+	SubmitPayloads(MoveTemp(AllPayloads));
+
+#else
+	
+	TArray<FMetalPayload*> AllPayloads;
+	#if WITH_RHI_BREADCRUMBS
+	TArray<TSharedPtr<FRHIBreadcrumbAllocator>> BreadcrumbAllocators;
+	#endif
+
+	for (FMetalFinalizedCommands* Payloads : Commands)
+	{
+	#if WITH_RHI_BREADCRUMBS
+		for (FMetalPayload* Payload : *Payloads)
+		{
+			Payload->BreadcrumbRange = Payloads->BreadcrumbRange;
+		}
+	#endif
+
+		AllPayloads.Append(MoveTemp(static_cast<TArray<FMetalPayload*>&>(*Payloads)));
+	#if WITH_RHI_BREADCRUMBS
+		BreadcrumbAllocators.Append(MoveTemp(Payloads->BreadcrumbAllocators));
+	#endif
+		delete Payloads;
+	}
+
+	SubmitPayloads(MoveTemp(AllPayloads));
+
+	#if WITH_RHI_BREADCRUMBS
+	// Enqueue the breadcrumb allocator references for cleanup once all prior payloads have completed on the GPU.
+	DeferredDelete([Array = MoveTemp(BreadcrumbAllocators)]() {});
+	#endif
+#endif
+}
+
+void FMetalDynamicRHI::SubmitPayloads(TArray<FMetalPayload*>&& Payloads)
+{
+	if (Payloads.Num())
+	{
+		PendingPayloadsForSubmission.Enqueue(new TArray<FMetalPayload*>(MoveTemp(Payloads)));
+	}
+
+	if (SubmissionThread)
+	{
+		SubmissionThread->Kick();
+	}
+	else
+	{
+		// Since we're processing directly on the calling thread, we need to take a scope lock.
+		// Multiple engine threads might be calling Submit().
+		{
+			FScopeLock Lock(&SubmissionCS);
+
+			// Process the submission queue until no further progress is being made.
+			while (EnumHasAnyFlags(ProcessSubmissionQueue().Status, EQueueStatus::Processed)) {}
+		}
+	}
+
+	// Use this opportunity to pump the interrupt queue
+	ProcessInterruptQueueUntil(nullptr);
+}
+
+static int32 GetMaxExecuteBatchSize()
+{
+	return
+#if UE_BUILD_DEBUG
+		GRHIGlobals.IsDebugLayerEnabled ? 1 :
+#endif
+		TNumericLimits<int32>::Max();
+}
+
+FMetalDynamicRHI::FProcessResult FMetalDynamicRHI::ProcessSubmissionQueue()
+{
+	SCOPED_NAMED_EVENT_TEXT("SubmissionQueue_Process", FColor::Turquoise);
+	LLM_SCOPE_BYNAME(TEXT("RHIMisc/ProcessSubmissionQueue"));
+
+	FMetalCommandQueue::FPayloadArray PayloadsToHandDown;
+	FProcessResult Result;
+
+	auto FlushPayloads = [&PayloadsToHandDown, &Result, DynamicRHI = this]()
+	{
+		if (PayloadsToHandDown.Num() > 0)
+		{
+			Result.Status |= EQueueStatus::Processed;
+			DynamicRHI->FlushBatchedPayloads(PayloadsToHandDown);
+		}
+	};
+
+	bool bProgress;
+	bool bKickInterruptThread = false;
+
+	do
+	{
+		bProgress = false;
+		Result.Status = EQueueStatus::None;
+
+		// Push all pending payloads into the ordered per-device, per-pipe pending queues
+		{
+			TArray<FMetalPayload*>* Array;
+			while (PendingPayloadsForSubmission.Dequeue(Array))
+			{
+				for (FMetalPayload* Payload : *Array)
+				{
+					Payload->Queue.PendingSubmission.Enqueue(Payload);
+				}
+				delete Array;
+			}
+		}
+
+		//
+		// Fence values for FMetalSyncPoint are determined on the submission thread,
+		// where each queue has a monotonically incrementing fence value.
+		//
+		// We might receive work that waits on a sync point which has not yet been submitted
+		// to the queue that will signal it, so we need to delay processing of those
+		// payloads until the fence value is known.
+		//
+
+		// Process all queues (across all devices and adapters) to flush work.
+		// Any sync point waits where the fence value is unknown will be left in the 
+		// appropriate queue, to be processed the next time commands are submitted.
+		ForEachQueue([&](FMetalCommandQueue& CurrentQueue)
+		{
+			while (true)
+			{
+				TArray<FMetalCommandQueue*, TInlineAllocator<GMetalMaxNumQueues>> QueuesWithPayloads;
+				{
+					FMetalPayload* Payload = CurrentQueue.PendingSubmission.Peek();
+					if (!Payload)
+						return;
+
+					// Accumulate the list of fences to await, and their maximum values
+					while (Payload->SyncPointsToWait.Index < Payload->SyncPointsToWait.Num())
+					{
+						FMetalSyncPointRef& SyncPoint = Payload->SyncPointsToWait[Payload->SyncPointsToWait.Index];
+						if (!SyncPoint->ResolvedFence.IsSet())
+						{
+							// Need to wait on a sync point, but the fence value has not been resolved yet
+							// (no other payloads have signaled the sync point yet).
+						
+							// Skip processing this queue, and move on to the next. We will retry later when
+							// further work is submitted, which may contain the sync point we need.
+							Result.Status |= EQueueStatus::Pending;
+							return;
+						}
+
+						Payload->AddQueueFenceWait(
+							SyncPoint->ResolvedFence->Fence,
+							SyncPoint->ResolvedFence->Value
+						);
+
+						Payload->SyncPointsToWait.Index++;
+						bProgress = true;
+					}
+				
+					// All necessary sync points have been resolved.
+					Payload->SyncPointsToWait = {};
+					CurrentQueue.PendingSubmission.Pop();
+					bProgress = true;
+
+					check(!CurrentQueue.PayloadToSubmit);
+					CurrentQueue.PayloadToSubmit = Payload;
+					QueuesWithPayloads.Add(&CurrentQueue);
+					Result.Status |= EQueueStatus::Processed;
+					bKickInterruptThread = true;
+
+					for (int32 Index = 0; Index < Payload->CommandBuffersToExecute.Num(); Index++)
+					{
+						FMetalCommandBuffer* CurrentCommandBuffer = Payload->CommandBuffersToExecute[Index];
+						
+						CurrentQueue.OcclusionQueries.Append(MoveTemp(CurrentCommandBuffer->OcclusionQueries));
+						CurrentQueue.TimestampQueries.Append(MoveTemp(CurrentCommandBuffer->TimestampQueries));
+#if RHI_NEW_GPU_PROFILER == 0
+						CurrentQueue.EventSampleCounters.Append(MoveTemp(CurrentCommandBuffer->EventSampleCounters));
+#endif
+					}
+				}
+
+				// Queues with work to submit other than the current one (CurrentQueue) are performing barrier operations.
+				// Submit this work first, followed by a fence signal + enqueued wait.
+				for (FMetalCommandQueue* OtherQueue : QueuesWithPayloads)
+				{
+					if (OtherQueue != &CurrentQueue)
+					{
+						uint64 ValueSignaled = OtherQueue->FinalizePayload(true, PayloadsToHandDown);
+						CurrentQueue.PayloadToSubmit->AddQueueFenceWait(OtherQueue->GetSignalEvent(), ValueSignaled);
+					}
+					FlushPayloads();
+				}
+
+				// Now submit the original payload
+				CurrentQueue.FinalizePayload(false, PayloadsToHandDown);
+				FlushPayloads();
+			}
+		});
+	} while (bProgress);
+
+	FlushPayloads();
+
+	if (InterruptThread && bKickInterruptThread)
+	{
+		InterruptThread->Kick();
+	}
+
+	return Result;
+}
+
+uint64 FMetalCommandQueue::FinalizePayload(bool bRequiresSignal, FPayloadArray& PayloadsToHandDown)
+{
+	TRACE_CPUPROFILER_EVENT_SCOPE(ExecuteCommandList);
+	LLM_SCOPE_BYNAME(TEXT("RHIMisc/ExecuteCommandLists"));
+
+	check(PayloadToSubmit && this == &PayloadToSubmit->Queue);
+	check(PayloadToSubmit->SyncPointsToWait.Num() == 0);
+	check(!PayloadToSubmit->SignalCommandBuffer);
+	
+	// Keep the latest fence value in the submitted payload.
+	// The interrupt thread uses this to determine when work has completed.
+	PayloadToSubmit->CompletionFenceValue = ++SignalEvent.NextCompletionValue;
+	PayloadToSubmit->bAlwaysSignal |= bRequiresSignal;
+
+	// Set the fence/value pair into any sync points we need to signal.
+	for (FMetalSyncPointRef& SyncPoint : PayloadToSubmit->SyncPointsToSignal)
+	{
+		check(!SyncPoint->ResolvedFence.IsSet());
+		SyncPoint->ResolvedFence.Emplace(SignalEvent, PayloadToSubmit->CompletionFenceValue);
+	}
+
+	PayloadToSubmit->BatchedObjects.OcclusionQueries = MoveTemp(OcclusionQueries);
+	PayloadToSubmit->BatchedObjects.TimestampQueries = MoveTemp(TimestampQueries);
+	
+#if RHI_NEW_GPU_PROFILER == 0
+	PayloadToSubmit->BatchedObjects.EventSampleCounters = MoveTemp(EventSampleCounters);
+#endif
+	
+	PayloadsToHandDown.Add(PayloadToSubmit);
+	PayloadToSubmit = nullptr;
+
+	return SignalEvent.NextCompletionValue;
+}
+
+void FMetalDynamicRHI::FlushBatchedPayloads(FMetalCommandQueue::FPayloadArray& PayloadsToSubmit)
+{
+	uint32 FirstPayload = 0, LastPayload = 0;
+
+	auto Wait = [this](FMetalPayload* Payload)
+	{
+		FMetalCommandQueue& Queue = Payload->Queue;
+
+		// Wait for queue fences
+		for (auto& [LocalFence, Value] : Payload->QueueFencesToWait)
+		{
+			FMetalCommandBuffer* CurrentCommandBuffer = Queue.CreateCommandBuffer();
+			CurrentCommandBuffer->GetMTLCmdBuffer()->encodeWait(LocalFence.MetalEvent, Value);
+			Queue.CommitCommandBuffer(CurrentCommandBuffer);
+		}
+	};
+
+	auto Flush = [&]()
+	{
+		if (FirstPayload == LastPayload)
+			return;
+
+		FMetalCommandQueue& Queue = PayloadsToSubmit[FirstPayload]->Queue;
+
+		uint64 Time = FPlatformTime::Cycles64();
+
+		TArray<FMetalCommandBuffer*> CommandBuffers;
+		
+		// Accumulate the command lists from the payload
+		for (uint32 Index = FirstPayload; Index < LastPayload; ++Index)
+		{
+			FMetalPayload* Payload = PayloadsToSubmit[Index];
+			check(&Payload->Queue == &Queue);
+
+			for (FMetalCommandBuffer* CommandBuffer : Payload->CommandBuffersToExecute)
+			{
+#if RHI_NEW_GPU_PROFILER
+				CommandBuffer->FlushProfilerEvents(Payload->EventStream, Time);
+#endif
+				
+				Payload->Queue.CommitCommandBuffer(CommandBuffer);
+			}
+		}
+
+		FirstPayload = LastPayload;
+	};
+
+	auto Signal = [this](FMetalPayload* Payload)
+	{
+		FMetalCommandQueue& Queue = Payload->Queue;
+
+		// Signal the queue fence
+		if (Payload->RequiresQueueFenceSignal())
+		{
+			check(Queue.GetSignalEvent().LastSignaledValue < Payload->CompletionFenceValue);
+
+			FMetalCommandBuffer* CommandBuffer = Queue.CreateCommandBuffer();
+			
+			MTL::HandlerFunction CompletionHandler = [this](MTL::CommandBuffer* CompletedBuffer)
+			{
+				if(InterruptThread)
+				{
+					InterruptThread->Kick();
+				}
+			};
+			
+			CommandBuffer->GetMTLCmdBuffer()->addCompletedHandler(CompletionHandler);
+			CommandBuffer->GetMTLCmdBuffer()->encodeSignalEvent(Queue.GetSignalEvent().MetalEvent, Payload->CompletionFenceValue);
+			Payload->SignalCommandBuffer = CommandBuffer;
+		
+			Queue.CommitCommandBuffer(CommandBuffer);
+			Queue.GetSignalEvent().LastSignaledValue.store(Payload->CompletionFenceValue, std::memory_order_release);
+		}
+
+		// Submission of this payload is completed. Signal the submission event if one was provided.
+		if (Payload->SubmissionEvent)
+		{
+			Payload->SubmissionEvent->DispatchSubsequents();
+		}
+	};
+
+	FMetalCommandQueue* PrevQueue = nullptr;
+	for (FMetalPayload* Payload : PayloadsToSubmit)
+	{
+		if (PrevQueue != &Payload->Queue)
+		{
+			Flush();
+			PrevQueue = &Payload->Queue;
+		}
+
+		Payload->Queue.PendingInterrupt.Enqueue(Payload);
+
+		if (Payload->HasWaitWork())
+		{
+			Flush();
+			Wait(Payload);
+		}
+
+		if (Payload->HasPreExecuteWork())
+		{
+			Flush();
+			Payload->PreExecute();
+		}
+
+		LastPayload++;
+
+		if (Payload->HasSignalWork())
+		{
+			Flush();
+			Signal(Payload);
+		}
+	}
+
+	Flush();
+
+	for (FMetalPayload* Payload : PayloadsToSubmit)
+	{
+		// Only set this bool to true once we'll never touch the payload again on this thread.
+		// This is because the bool hands ownership to the interrupt thread, which might delete the payload.
+		Payload->bSubmitted = true;
+	}
+
+	PayloadsToSubmit.Reset();	
+}
+
+void FMetalPayload::AddQueueFenceWait(FMetalSignalEvent& InFence, uint64 InValue)
+{
+	for (auto& [Fence, Value] : QueueFencesToWait)
+	{
+		if (&Fence == &InFence)
+		{
+			Value = FMath::Max(Value, InValue);
+			return;
+		}
+	}
+
+	QueueFencesToWait.Add({ InFence, InValue });
+}
+
+void FMetalSyncPoint::Wait() const
+{
+	checkf(GraphEvent, TEXT("This sync point was not created with a CPU event. Cannot wait for completion on the CPU."));
+
+	if (!GraphEvent->IsComplete())
+	{
+		// Block the calling thread until the graph event is signaled by the interrupt thread.
+		SCOPED_NAMED_EVENT_TEXT("SyncPoint_Wait", FColor::Turquoise);
+		FMetalDynamicRHI::Get().ProcessInterruptQueueUntil(GraphEvent);
+	}
+
+	check(GraphEvent->IsComplete());
+}
+
+void FMetalDynamicRHI::ProcessInterruptQueueUntil(FGraphEvent* GraphEvent)
+{
+	if (InterruptThread)
+	{
+		if (GraphEvent && !GraphEvent->IsComplete())
+		{
+			InterruptThread->Kick();
+			GraphEvent->Wait();
+		}
+	}
+	else
+	{
+		// Use the current thread to process the interrupt queue until the sync point we're waiting for is signaled.
+		// If GraphEvent is nullptr, process the queue until no further progress is made (assuming we can acquire the lock), then return.
+		if (!GraphEvent || !GraphEvent->IsComplete())
+		{
+			// If we're waiting for a sync point, accumulate the idle time
+			UE::Stats::FThreadIdleStats::FScopeIdle IdleScope(GraphEvent == nullptr);
+
+		Retry:
+			if (InterruptCS.TryLock())
+			{
+				FProcessResult Result;
+				do { Result = ProcessInterruptQueue(); }
+				// If we have a sync point, keep processing until the sync point is signaled.
+				// Otherwise, process until no more progress is being made.
+				while (GraphEvent
+					? !GraphEvent->IsComplete()
+					: EnumHasAllFlags(Result.Status, EQueueStatus::Processed)
+				);
+
+				InterruptCS.Unlock();
+			}
+			else if (GraphEvent && !GraphEvent->IsComplete())
+			{
+				// Failed to get the lock. Another thread is processing the interrupt queue. Try again...
+				FPlatformProcess::SleepNoStats(0);
+				goto Retry;
+			}
+		}
+	}
+}
+
+void FMetalDynamicRHI::ProcessPendingCommandBuffers()
+{
+	double SecondsPerCycle = FPlatformTime::GetSecondsPerCycle64();
+	bool bContinueProcessing = true;
+	
+	CmdBuffersPendingCompletion.RemoveAll([SecondsPerCycle, &bContinueProcessing](FMetalCommandBuffer* CommandBuffer) 
+	{ 
+		MTLCommandBufferPtr CompletedBuffer = CommandBuffer->GetMTLCmdBuffer();
+		
+		MTL::CommandBufferStatus Status = CompletedBuffer->status();
+		if (Status == MTL::CommandBufferStatusCompleted && bContinueProcessing)
+		{
+#if RHI_NEW_GPU_PROFILER
+#if WITH_RHI_BREADCRUMBS
+			FMetalBreadcrumbProfiler::GetInstance()->ResolveBreadcrumbTrackerStream(CommandBuffer->BreadcrumbTrackerStream);
+#endif
+			
+			uint64_t* Start = CommandBuffer->BeginWorkTimestamp;
+			uint64_t* End = CommandBuffer->EndWorkTimestamp;
+			
+			if(CommandBuffer->CounterSamples.Num())
+			{
+				for(FMetalCounterSamplePtr Sample : CommandBuffer->CounterSamples)
+				{
+					uint64_t StartTime, EndTime;
+					Sample->ResolveStageCounters(StartTime, EndTime);
+					
+					*Start = *Start > 0 ? FMath::Min(StartTime, *Start) : StartTime;
+					*End = *End > 0 ? FMath::Max(EndTime, *End) : EndTime;
+				}
+			}
+			else
+			{
+				*Start = CompletedBuffer->GPUStartTime() / SecondsPerCycle;
+				*End = CompletedBuffer->GPUEndTime() / SecondsPerCycle;
+			}
+#else
+			FMetalCommandBufferTimer* Timer = CommandBuffer->GetTimer();
+			Timer->AddTiming({CompletedBuffer->GPUStartTime(), CompletedBuffer->GPUEndTime()});
+#endif
+			return true;
+		}
+		else
+		{
+			bContinueProcessing = false;
+		}
+	
+		return false; 
+	});
+}
+
+FMetalDynamicRHI::FProcessResult FMetalDynamicRHI::ProcessInterruptQueue()
+{
+	SCOPED_NAMED_EVENT_TEXT("InterruptQueue_Process", FColor::Yellow);
+	LLM_SCOPE_BYNAME(TEXT("RHIMisc/ProcessInterruptQueue"));
+
+	// Timer that clamps each tick to prevent false positive GPU timeouts
+	// when a debugger is attached and the process is broken.
+	struct FTimer
+	{
+		uint64 Elapsed;
+		uint64 Last;
+
+		FTimer()
+			: Elapsed(0)
+			, Last(FPlatformTime::Cycles64())
+		{}
+
+		void Tick()
+		{
+			static const uint64 MaxDeltaCycles = uint64(1.0 / FPlatformTime::GetSecondsPerCycle64()); // 1 second
+			uint64 Current = FPlatformTime::Cycles64();
+			Elapsed += FMath::Min(MaxDeltaCycles, Current - Last);
+			Last = Current;
+		}
+	} static Timer;
+
+	Timer.Tick();
+
+	FProcessResult Result;
+	ForEachQueue([&](FMetalCommandQueue& CurrentQueue)
+	{
+		while (FMetalPayload* Payload = CurrentQueue.PendingInterrupt.Peek())
+		{
+			if (!Payload->bSubmitted)
+				break;
+
+			// Check for GPU completion
+			FMetalSignalEvent& CurrentEvent = CurrentQueue.GetSignalEvent();
+
+			uint64 LastSignaledFenceValue = CurrentEvent.LastSignaledValue.load(std::memory_order_acquire);
+			
+			// Handle command buffer errors
+			for(FMetalCommandBuffer* CommandBuffer : Payload->CommandBuffersToExecute)
+			{	
+				MTLCommandBufferPtr CompletedBuffer = CommandBuffer->GetMTLCmdBuffer();
+				if (CompletedBuffer->status() == MTL::CommandBufferStatusError)
+				{
+					FMetalCommandList::HandleMetalCommandBufferFailure(CompletedBuffer.get());
+				}
+			}
+			
+			MTL::CommandBufferStatus Status = MTL::CommandBufferStatusCompleted;
+			
+			if(Payload->SignalCommandBuffer)
+			{
+				MTLCommandBufferPtr SignalBuffer = Payload->SignalCommandBuffer->GetMTLCmdBuffer(); 
+				Status = SignalBuffer->status();
+				if(Status == MTL::CommandBufferStatusError)
+				{
+					FMetalCommandList::HandleMetalCommandBufferFailure(SignalBuffer.get());
+				}
+			}
+			
+			// Remove Completed status check when we remove completion handlers
+			if (Status != MTL::CommandBufferStatusCompleted)
+			{
+				// Skip processing this queue and move on to the next.
+				Result.Status |= EQueueStatus::Pending;
+				break;
+			}
+			
+			// Record the command buffer timing
+			for(FMetalCommandBuffer* CommandBuffer : Payload->CommandBuffersToExecute)
+			{	
+				CmdBuffersPendingCompletion.Add(CommandBuffer);
+			}
+			
+			ProcessPendingCommandBuffers();
+			
+			// At this point, the current command list has completed on the GPU.
+			CurrentQueue.PendingInterrupt.Pop();
+			Result.Status |= EQueueStatus::Processed;
+			
+#if RHI_NEW_GPU_PROFILER
+			if (!Payload->EventStream.IsEmpty())
+			{			
+				check(CurrentQueue.Timing);
+				CurrentQueue.Timing->EventStream.Append(MoveTemp(Payload->EventStream));
+			}
+			
+			if (Payload->Timing.IsSet())
+			{
+				// Switch the new timing struct into the queue. This redirects timestamp results to separate each frame's work.
+				CurrentQueue.Timing = Payload->Timing.GetValue();
+			}
+#endif
+			// Resolve query results
+			{
+				for (FMetalRHIRenderQuery* OcclusionQuery : Payload->BatchedObjects.OcclusionQueries)
+				{
+					OcclusionQuery->Result = OcclusionQuery->Buffer.GetResult();
+				}
+				
+				for (FMetalRHIRenderQuery* TimestampQuery : Payload->BatchedObjects.TimestampQueries)
+				{
+					MTLCommandBufferPtr CmdBuffer = TimestampQuery->CommandBuffer->GetMTLCmdBuffer();
+					
+					// If there are no commands in the command buffer then this can be zero
+					// In this case GPU start time is also not correct - we need to fall back standard behaviour
+					// Only seen empty command buffers at the very end of a frame
+					
+					// Convert seconds to microseconds
+					TimestampQuery->Result = uint64(CmdBuffer->GPUEndTime()) * 1000000;
+					
+					if(TimestampQuery->Result == 0)
+					{
+						TimestampQuery->Result = (FPlatformTime::ToMilliseconds64(mach_absolute_time()) * 1000.0);
+					}
+					
+					TimestampQuery->Release();
+				}
+				
+				
+#if RHI_NEW_GPU_PROFILER == 0
+				for(auto& Pair : Payload->BatchedObjects.EventSampleCounters)
+				{
+					uint64_t& Start = Pair.Key->StartTime;
+					uint64_t& End = Pair.Key->EndTime;
+					
+					for(FMetalCounterSamplePtr Sample: Pair.Value)
+					{
+						uint64_t StartTime, EndTime;
+						Sample->ResolveStageCounters(StartTime, EndTime);
+						
+						StartTime /= 1000.0;
+						EndTime /= 1000.0;
+						
+						Start = Start > 0 ? FMath::Min(StartTime, Start) : StartTime;
+						End = End > 0 ? FMath::Max(EndTime, End) : EndTime;
+					}
+				} 
+#endif
+			}
+			
+			// Signal the CPU events of all sync points associated with this batch.
+			for (FMetalSyncPointRef& SyncPoint : Payload->SyncPointsToSignal)
+			{
+				if (SyncPoint->GraphEvent)
+				{
+					SyncPoint->GraphEvent->DispatchSubsequents();
+				}
+			}
+
+			// We're done with this payload now.
+
+			// GPU resources the payload is holding a reference to will be cleaned up here.
+			// E.g. command list allocators, which get recycled on the parent device.
+			delete Payload;
+		}
+	});
+
+	return Result;
+}
+
+FMetalPayload::FMetalPayload(FMetalCommandQueue& Queue)
+	: Queue(Queue)
+#if RHI_NEW_GPU_PROFILER
+	, EventStream(Queue.GetProfilerQueue())
+#endif
+{}
+
+FMetalPayload::~FMetalPayload()
+{
+}
+
+void FMetalPayload::PreExecute()
+{
+	if (PreExecuteCallback)
+	{
+		PreExecuteCallback(Queue);
+	}
+}
diff --git a/Engine/Source/Runtime/Apple/MetalRHI/Private/MetalSubmission.h b/Engine/Source/Runtime/Apple/MetalRHI/Private/MetalSubmission.h
new file mode 100644
--- /dev/null
+++ b/Engine/Source/Runtime/Apple/MetalRHI/Private/MetalSubmission.h
@@ -0,0 +1,242 @@
+
+// Copyright Epic Games, Inc. All Rights Reserved.
+
+#pragma once
+
+#include "Async/TaskGraphInterfaces.h"
+#include "Templates/Function.h"
+#include "Templates/RefCounting.h"
+#include "MetalProfiler.h"
+#include "RHIBreadcrumbs.h"
+#include "GPUProfiler.h"
+#include "MetalThirdParty.h"
+
+static constexpr uint32 GMetalMaxNumQueues = 1;
+
+class FMetalSyncPoint;
+class FMetalCommandQueue;
+class FMetalCommandBuffer;
+class FMetalDynamicRHI;
+class FMetalRHICommandContext;
+
+using FMetalSyncPointRef = TRefCountPtr<FMetalSyncPoint>;
+
+enum class EMetalSyncPointType
+{
+	// Sync points of this type do not include an FGraphEvent, so cannot
+	// report completion to the CPU (via either IsComplete() or Wait())
+	GPUOnly,
+
+	// Sync points of this type include an FGraphEvent. The IsComplete() and Wait() functions
+	// can be used to poll for completion from the CPU, or block the CPU, respectively.
+	GPUAndCPU,
+};
+
+// Fence type used by the device queues to manage GPU completion
+struct FMetalSignalEvent
+{
+	MTL::Event* MetalEvent;
+	uint64 NextCompletionValue = 0;
+	std::atomic<uint64> LastSignaledValue = 0;
+};
+
+// Used by FMetalSyncPoint and the submission thread to fix up signaled fence values at the end-of-pipe
+struct FMetalResolvedFence
+{
+	FMetalSignalEvent& Fence;
+	uint64 Value = 0;
+
+	FMetalResolvedFence(FMetalSignalEvent& Fence, uint64 Value)
+		: Fence(Fence)
+		, Value(Value)
+	{}
+};
+
+//
+// A sync point is a logical point on a GPU queue's timeline that can be awaited by other queues, or the CPU.
+// These are used throughout the RHI as a way to abstract the underlying Metal fences. The submission thread 
+// manages the underlying fences and signaled values, and reports completion to the relevant sync points via 
+// an FGraphEvent.
+//
+// Sync points are one-shot, meaning they represent a single timeline point, and are released after use, via ref-counting.
+// Use FMetalSyncPoint::Create() to make a new sync point and hold a reference to it via a FMetalSyncPointRef object.
+//
+class FMetalSyncPoint final : public FThreadSafeRefCountedObject
+{
+	friend FMetalDynamicRHI;
+	friend FMetalCommandQueue;
+	friend FMetalRHICommandContext;
+	
+	static TLockFreePointerListUnordered<void, PLATFORM_CACHE_LINE_SIZE> MemoryPool;
+
+	// No copying or moving
+	FMetalSyncPoint(FMetalSyncPoint const&) = delete;
+	FMetalSyncPoint(FMetalSyncPoint&&) = delete;
+
+	TOptional<FMetalResolvedFence> ResolvedFence;
+	FGraphEventRef GraphEvent;
+
+	FMetalSyncPoint(EMetalSyncPointType Type)
+	{
+		if (Type == EMetalSyncPointType::GPUAndCPU)
+		{
+			GraphEvent = FGraphEvent::CreateGraphEvent();
+		}
+	}
+
+public:
+	static FMetalSyncPointRef Create(EMetalSyncPointType Type)
+	{
+		LLM_SCOPE_BYNAME(TEXT("RHIMisc/CreateSyncPoint"));
+		return new FMetalSyncPoint(Type);
+	}
+
+	bool IsComplete() const
+	{
+		checkf(GraphEvent, TEXT("This sync point was not created with a CPU event. Cannot check completion on the CPU."));
+		return GraphEvent->IsComplete();
+	}
+
+	void Wait() const;
+	
+	void OnCompletionCallback(TUniqueFunction<void()>&& Func)
+	{
+		checkf(GraphEvent, TEXT("This sync point was not created with a CPU event."));
+
+		FGraphEventRef CompletionFence = FFunctionGraphTask::CreateAndDispatchWhenReady([Callback=MoveTemp(Func)]()
+		{
+			Callback();
+		}, TStatId(), GraphEvent, ENamedThreads::AnyThread);
+	}
+
+	FGraphEvent* GetGraphEvent() const
+	{
+		checkf(GraphEvent, TEXT("This sync point was not created with a CPU event."));
+		return GraphEvent;
+	}
+
+	EMetalSyncPointType GetType() const
+	{
+		return GraphEvent != nullptr
+			? EMetalSyncPointType::GPUAndCPU
+			: EMetalSyncPointType::GPUOnly;
+	}
+};
+
+class FMetalRHIRenderQuery;
+struct FMetalCounterSample;
+class FMetalEventNode;
+
+struct FMetalBatchedPayloadObjects
+{
+	TArray<FMetalRHIRenderQuery*> OcclusionQueries;
+	TArray<FMetalRHIRenderQuery*> TimestampQueries;
+#if RHI_NEW_GPU_PROFILER == 0
+	TMap<FMetalEventNode*, TArray<FMetalCounterSamplePtr>> EventSampleCounters;
+#endif
+
+	bool IsEmpty() const
+	{
+		return TimestampQueries.Num() == 0 && OcclusionQueries.Num() == 0
+#if RHI_NEW_GPU_PROFILER == 0
+				&& EventSampleCounters.Num() == 0
+#endif
+		;
+	}
+};
+
+// Base class to avoid 8 bytes of padding after the vtable
+struct FMetalPayloadBaseFixLayout
+{
+	virtual ~FMetalPayloadBaseFixLayout() = default;
+};
+
+class FMetalTiming;
+
+// A single unit of work (specific to a single GPU node and queue type) to be processed by the submission thread.
+struct FMetalPayload : public FMetalPayloadBaseFixLayout
+{
+	friend FMetalDynamicRHI;
+	friend FMetalCommandQueue;
+	friend FMetalRHICommandContext;
+	
+	// Constants
+	FMetalCommandQueue& Queue;
+
+	// Wait
+	struct : public TArray<FMetalSyncPointRef>
+	{
+		// Used to pause / resume iteration of the sync point array on the
+		// submission thread when we find a sync point that is unresolved.
+		int32 Index = 0;
+
+	} SyncPointsToWait;
+
+	struct FQueueFence
+	{
+		FMetalSignalEvent& Fence;
+		uint64 Value;
+	};
+	TArray<FQueueFence, TInlineAllocator<GMetalMaxNumQueues>> QueueFencesToWait;
+
+	void AddQueueFenceWait(FMetalSignalEvent& Fence, uint64 Value);
+
+	// Flags.
+	bool bAlwaysSignal = false;
+	std::atomic<bool> bSubmitted { false };
+
+	// Used by RHIRunOnQueue
+	TFunction<void(FMetalCommandQueue&)> PreExecuteCallback;
+
+	// Execute
+	TArray<FMetalCommandBuffer*> CommandBuffersToExecute;
+	FMetalCommandBuffer* SignalCommandBuffer = nullptr;
+
+	// Signal
+	TArray<FMetalSyncPointRef> SyncPointsToSignal;
+	uint64 CompletionFenceValue = 0;
+
+	FGraphEventRef SubmissionEvent;
+	TOptional<uint64> SubmissionTime;
+
+	FMetalBatchedPayloadObjects BatchedObjects;
+	
+	TMap<FMetalEventNode*, TArray<FMetalCounterSamplePtr>> EventSampleCounters;
+
+#if WITH_RHI_BREADCRUMBS
+	FRHIBreadcrumbRange BreadcrumbRange {};
+	TSharedPtr<FRHIBreadcrumbAllocatorArray> BreadcrumbAllocators {};
+#endif
+
+#if RHI_NEW_GPU_PROFILER
+	TOptional<FMetalTiming*> Timing;
+	UE::RHI::GPUProfiler::FEventStream EventStream;
+#endif
+
+	virtual ~FMetalPayload();
+
+	virtual void PreExecute();
+
+	virtual bool HasPreExecuteWork() const
+	{
+		return PreExecuteCallback != nullptr;
+	}
+
+	virtual bool RequiresQueueFenceSignal() const
+	{
+		return bAlwaysSignal || SyncPointsToSignal.Num() > 0 || HasPreExecuteWork();
+	}
+
+	virtual bool HasWaitWork() const
+	{
+		return QueueFencesToWait.Num() > 0;
+	}
+
+	virtual bool HasSignalWork() const
+	{
+		return RequiresQueueFenceSignal() || SubmissionEvent != nullptr;
+	}
+
+protected:
+	FMetalPayload(FMetalCommandQueue&);
+};
diff --git a/Engine/Source/Runtime/Apple/MetalRHI/Private/MetalTexture.cpp b/Engine/Source/Runtime/Apple/MetalRHI/Private/MetalTexture.cpp
--- a/Engine/Source/Runtime/Apple/MetalRHI/Private/MetalTexture.cpp
+++ b/Engine/Source/Runtime/Apple/MetalRHI/Private/MetalTexture.cpp
@@ -656,7 +656,6 @@ FMetalSurface::FMetalSurface(FMetalDevice& MetalDevice, FRHICommandListBase* RHI
 			// Backing buffer resource options must match the texture we are going to create from it
 			FMetalPooledBufferArgs Args(&Device, BytesPerRow * NewCreateDesc.Desc->height(), BUF_Dynamic, NewCreateDesc.Desc->storageMode(), NewCreateDesc.Desc->cpuCacheMode());
 			BackingBuffer = Device.CreatePooledBuffer(Args);
-
 			Texture = NS::TransferPtr(BackingBuffer->GetMTLBuffer()->newTexture(NewCreateDesc.Desc.get(), BackingBuffer->GetOffset(), BytesPerRow));
 		}
         else if (bTextureArrayWithAtomics)
@@ -1049,14 +1048,11 @@ void FMetalSurface::UpdateSurfaceAndDestroySourceBuffer(FMetalRHICommandContext*
 #endif
 		Context->CopyFromBufferToTexture(Source, 0, Stride, BytesPerImage, Region.size, Texture.get(), ArrayIndex, MipIndex, Region.origin, Options);
 		
-		FMetalCommandBufferCompletionHandler CompletionHandler;
-		CompletionHandler.BindLambda([TempBuffer = SourceBuffer, Size](MTL::CommandBuffer *)
+		Context->GetContextSyncPoint()->OnCompletionCallback([Size]()
 		{
 			FPlatformAtomics::InterlockedAdd(&ActiveUploads, -Size);
 		});
 		
-		Context->AddCompletionHandler(CompletionHandler);
-		
 		INC_DWORD_STAT_BY(STAT_MetalTextureMemUpdate, Size);
 	}
 	else
diff --git a/Engine/Source/Runtime/Apple/MetalRHI/Private/MetalUAV.cpp b/Engine/Source/Runtime/Apple/MetalRHI/Private/MetalUAV.cpp
--- a/Engine/Source/Runtime/Apple/MetalRHI/Private/MetalUAV.cpp
+++ b/Engine/Source/Runtime/Apple/MetalRHI/Private/MetalUAV.cpp
@@ -747,46 +747,6 @@ void FMetalUnorderedAccessView::ClearUAV(TRHICommandList_RecursiveHazardous<FMet
     }
 }
 
-bool FMetalCommandBufferFence::Wait(uint32_t TimeIntervalMs) const
-{
-    check(CmdBuffer);
-
-    bool bFinished = false;
-    
-    if(TimeIntervalMs == MAX_uint32)
-    {
-        bFinished = Condition->Wait();
-    }
-    else
-    {
-        bFinished = Condition->Wait((uint32_t)TimeIntervalMs);
-    }
-    
-    return bFinished;
-}
-
-void FMetalCommandBufferFence::Insert(MTLCommandBufferPtr CommandBuffer)
-{
-    check(CommandBuffer);
-    check(CmdBuffer.get() == nullptr);
-    
-    CmdBuffer = CommandBuffer;
-    
-    Condition->Reset();
-    
-    MTL::HandlerFunction CommandBufferCompletionHandler = [&](MTL::CommandBuffer*)
-    {
-        Condition->Trigger();
-    };
-    
-    CmdBuffer->addCompletedHandler(CommandBufferCompletionHandler);
-}
-
-void FMetalCommandBufferFence::Signal(const MTL::CommandBuffer* CommandBuffer)
-{
-    Condition->Trigger();
-}
-
 void FMetalRHICommandContext::RHICopyToStagingBuffer(FRHIBuffer* SourceBufferRHI, FRHIStagingBuffer* DestinationStagingBufferRHI, uint32 Offset, uint32 NumBytes)
 {
     MTL_SCOPED_AUTORELEASE_POOL;
@@ -814,68 +774,54 @@ void FMetalRHICommandContext::RHICopyToStagingBuffer(FRHIBuffer* SourceBufferRHI
     CopyFromBufferToBuffer(SourceBuffer->GetCurrentBuffer(), Offset, ReadbackBuffer, 0, NumBytes);
 }
 
+FMetalGPUFence::FMetalGPUFence(FName InName)
+	: FRHIGPUFence(InName)
+{
+}
+
 void FMetalGPUFence::Clear()
 {
-	Event = nullptr;
+	SyncPoint = nullptr;
 }
 
 bool FMetalGPUFence::Poll() const
 {
-	if (!Event || !Event->IsComplete())
+	bool bHasAnySyncPoint = false;
+	
+	if (!SyncPoint || !SyncPoint->IsComplete())
+	{
 		return false;
-
-	return Fence->Wait(0);
+	}
+	
+	return true;
 }
 
 void FMetalGPUFence::Wait(FRHICommandListImmediate& RHICmdList, FRHIGPUMask GPUMask) const
 {
-	if (!Event)
-		return;
-
-	// The event is signalled when the RHI thread has updated the Fence field.
-	bool bDispatched = false;
-	if (!Event->IsComplete())
+	if (SyncPoint && !SyncPoint->IsComplete())
 	{
-		RHICmdList.ImmediateFlush(EImmediateFlushType::DispatchToRHIThread);
-		bDispatched = true;
-
-		Event->Wait();
-	}
-
-	// Now that we have the correct Fence, wait on it.
-	if (!Fence->Wait(0))
-	{
-		if (!bDispatched)
-		{
-			RHICmdList.ImmediateFlush(EImmediateFlushType::DispatchToRHIThread);
-		}
-
-		Fence->Wait(MAX_uint32);
+		SyncPoint->Wait();
 	}
 }
 
-void FMetalGPUFence::Write_TopOfPipe(FRHICommandListBase& RHICmdList)
+void FMetalDynamicRHI::RHIWriteGPUFence_TopOfPipe(FRHICommandListBase& RHICmdList, FRHIGPUFence* FenceRHI)
 {
-	Event = FGraphEvent::CreateGraphEvent();
-
-	RHICmdList.EnqueueLambda([this, Event = Event](FRHICommandListBase& ExecutingCmdList)
-	{
-		MTL_SCOPED_AUTORELEASE_POOL;
+	FMetalGPUFence* Fence = ResourceCast(FenceRHI);
+	check(Fence);
 
-		FMetalRHICommandContext& Context = FMetalRHICommandContext::Get(ExecutingCmdList);
+	checkf(Fence->SyncPoint == nullptr, TEXT("The fence for the current GPU node has already been issued."));
+	Fence->SyncPoint = FMetalSyncPoint::Create(EMetalSyncPointType::GPUAndCPU);
 
-		Fence = Context.GetCurrentCommandBuffer()->GetCompletionFence();
-		check(Fence);
+	Fence->NumPendingWriteCommands.Increment();
+	RHICmdList.EnqueueLambda([Fence, SyncPoint = Fence->SyncPoint](FRHICommandListBase& CmdList)
+	{
+		FMetalRHICommandContext& Context = FMetalRHICommandContext::Get(CmdList);
+		Context.SignalSyncPoint(SyncPoint);
 
-		Event->DispatchSubsequents();
+		Fence->NumPendingWriteCommands.Decrement();
 	});
 }
 
-void FMetalDynamicRHI::RHIWriteGPUFence_TopOfPipe(FRHICommandListBase& RHICmdList, FRHIGPUFence* FenceRHI)
-{
-	ResourceCast(FenceRHI)->Write_TopOfPipe(RHICmdList);
-}
-
 void FMetalRHICommandContext::RHIWriteGPUFence(FRHIGPUFence* FenceRHI)
 {
 	checkNoEntry(); // Should never be called
diff --git a/Engine/Source/Runtime/Apple/MetalRHI/Private/MetalViewport.cpp b/Engine/Source/Runtime/Apple/MetalRHI/Private/MetalViewport.cpp
--- a/Engine/Source/Runtime/Apple/MetalRHI/Private/MetalViewport.cpp
+++ b/Engine/Source/Runtime/Apple/MetalRHI/Private/MetalViewport.cpp
@@ -496,12 +496,14 @@ void FMetalViewport::Present(FMetalCommandQueue& CommandQueue, bool bLockToVsync
 					}
 #endif
 					
+					FMetalRHICommandContext& Context = *static_cast<FMetalRHICommandContext*>(FMetalDynamicRHI::Get().RHIGetDefaultContext());
+					FMetalCommandBuffer* CurrentCommandBuffer = Context.GetCurrentCommandBuffer();
+					
 					if (DrawableTexture && (InDisplayID == 0 || !bIsInLiveResize))
 					{
-						FMetalCommandBuffer* CurrentCommandBuffer = CommandQueue.CreateCommandBuffer();
 						check(CurrentCommandBuffer);
 						
-#if ENABLE_METAL_GPUPROFILE
+#if ENABLE_METAL_GPUPROFILE && RHI_NEW_GPU_PROFILER == 0
 						FMetalProfiler* Profiler = FMetalProfiler::GetProfiler();
 						FMetalCommandBufferStats* Stats = Profiler->AllocateCommandBuffer(CurrentCommandBuffer->GetMTLCmdBuffer(), 0);
 #endif
@@ -548,7 +550,9 @@ void FMetalViewport::Present(FMetalCommandQueue& CommandQueue, bool bLockToVsync
 							dispatch_semaphore_signal(FrameSemaphore);
 							dispatch_release(FrameSemaphore);
 							
-							FMetalGPUProfiler::RecordPresent(cmd_buf);
+#if RHI_NEW_GPU_PROFILER == 0
+							FMetalCommandBufferTimer::RecordPresent(cmd_buf);
+#endif
 							LocalDrawable->release();
 #if PLATFORM_MAC
 							MainThreadCall(^{
@@ -582,8 +586,11 @@ void FMetalViewport::Present(FMetalCommandQueue& CommandQueue, bool bLockToVsync
                             }
                         }
 #endif // PLATFORM_MAC
-                        METAL_GPUPROFILE(Stats->End(CurrentCommandBuffer->GetMTLCmdBuffer()));
-                        CommandQueue.CommitCommandBuffer(CurrentCommandBuffer);
+						
+						TArray<FMetalPayload*> Payloads;
+						Context.Finalize(Payloads);
+
+						FMetalDynamicRHI::Get().SubmitPayloads(MoveTemp(Payloads));
 						
 						// Wait for the frame semaphore
 						dispatch_semaphore_wait(Device.GetFrameSemaphore(), DISPATCH_TIME_FOREVER);
@@ -708,12 +715,9 @@ void FMetalViewport::PresentImmersive(const MetalRHIVisionOS::PresentImmersivePa
 					}
 				}
 				
-				TArray<FMetalCommandBuffer*> Buffers = Context.Finalize();
-				check(Buffers.Num() > 0);
-				
 				// We need to attach the completion handler and the present signal to the final
 				// command buffer
-				FMetalCommandBuffer* FinalCommandBuffer = Buffers.Last();
+				FMetalCommandBuffer* FinalCommandBuffer = Context.GetCurrentCommandBuffer();;
 				
 #if ENABLE_METAL_GPUPROFILE
 				FMetalProfiler* Profiler = FMetalProfiler::GetProfiler();
@@ -728,7 +732,7 @@ void FMetalViewport::PresentImmersive(const MetalRHIVisionOS::PresentImmersivePa
 						dispatch_semaphore_signal(FrameSemaphore);
 						dispatch_release(FrameSemaphore);
 						
-						FMetalGPUProfiler::RecordPresent(cmd_buf);
+						FMetalCommandBufferTimer::RecordPresent(cmd_buf);
 					};
 					FinalCommandBuffer->GetMTLCmdBuffer()->addCompletedHandler(CommandBufferHandler);
 				}
@@ -738,10 +742,11 @@ void FMetalViewport::PresentImmersive(const MetalRHIVisionOS::PresentImmersivePa
 
 				METAL_GPUPROFILE(Stats->End(FinalCommandBuffer->GetMTLCmdBuffer()));			
 				
-				for(FMetalCommandBuffer* Buffer : Buffers)
-				{
-					Context.GetDevice().GetCommandQueue().CommitCommandBuffer(Buffer);
-				}
+				TArray<FMetalPayload*> Payloads;
+				Context.Finalize(Payloads);
+
+				FMetalDynamicRHI::Get().SubmitPayloads(MoveTemp(Payloads));
+				
 				Context.ResetContext();
 				
 				cp_frame_end_submission(CompositorServicesFrame);
@@ -752,7 +757,7 @@ void FMetalViewport::PresentImmersive(const MetalRHIVisionOS::PresentImmersivePa
 		}
 	}
 	
-	FMetalGPUProfiler::ResetFrameBufferTimings();
+	FMetalCommandBufferTimer::ResetFrameBufferTimings();
 }
 #endif //PLATFORM_VISIONOS
 
diff --git a/Engine/Source/Runtime/Apple/MetalRHI/Private/Shaders/Types/Templates/MetalBaseShader.h b/Engine/Source/Runtime/Apple/MetalRHI/Private/Shaders/Types/Templates/MetalBaseShader.h
--- a/Engine/Source/Runtime/Apple/MetalRHI/Private/Shaders/Types/Templates/MetalBaseShader.h
+++ b/Engine/Source/Runtime/Apple/MetalRHI/Private/Shaders/Types/Templates/MetalBaseShader.h
@@ -487,7 +487,7 @@ MTLFunctionPtr TMetalBaseShader<BaseResourceType, ShaderType>::GetCompiledFuncti
 			{
 				METAL_GPUPROFILE(FScopedMetalCPUStats CPUStat(FString::Printf(TEXT("NewFunctionAsync: %s"), *NSStringToFString(Name))));
 				METAL_GPUPROFILE(uint64 CPUStart = CPUStat.Stats ? CPUStat.Stats->CPUStartTime : 0);
-#if ENABLE_METAL_GPUPROFILE
+#if ENABLE_METAL_GPUPROFILE && RHI_NEW_GPU_PROFILER == 0
                 NS::String* nsName = Name;
                 const std::function<void(MTL::Function* pFunction, NS::Error* pError)> Handler = [Key, this, CPUStart, nsName](MTL::Function* NewFunction, NS::Error* Error){
 #else
@@ -498,7 +498,7 @@ MTLFunctionPtr TMetalBaseShader<BaseResourceType, ShaderType>::GetCompiledFuncti
 					UE_CLOG(NewFunction == nullptr, LogMetal, Fatal, TEXT("*********** Error\n%s"), *NSStringToFString(GetSourceCode()));
 
 					GetMetalCompiledShaderCache().Add(Key, Library, NS::RetainPtr(NewFunction));
-#if ENABLE_METAL_GPUPROFILE
+#if ENABLE_METAL_GPUPROFILE && RHI_NEW_GPU_PROFILER == 0
 					if (CompletionStat.Stats)
 					{
 						CompletionStat.Stats->CPUStartTime = CPUStart;
diff --git a/Engine/Source/Runtime/Apple/MetalRHI/Public/MetalRHIContext.h b/Engine/Source/Runtime/Apple/MetalRHI/Public/MetalRHIContext.h
--- a/Engine/Source/Runtime/Apple/MetalRHI/Public/MetalRHIContext.h
+++ b/Engine/Source/Runtime/Apple/MetalRHI/Public/MetalRHIContext.h
@@ -13,7 +13,6 @@
 #include "MetalRHIRenderQuery.h"
 #include "RHICore.h"
 
-class FMetalCommandBufferFence;
 class FMetalEventNode;
 
 #if PLATFORM_VISIONOS
@@ -68,9 +67,6 @@ public:
     
 	virtual void RHIBeginRenderQuery(FRHIRenderQuery* RenderQuery) final override;
 	virtual void RHIEndRenderQuery(FRHIRenderQuery* RenderQuery) final override;
-	
-	void RHIBeginOcclusionQueryBatch(uint32 NumQueriesInBatch);
-	void RHIEndOcclusionQueryBatch();
 
 	virtual void RHIDiscardRenderTargets(bool Depth, bool Stencil, uint32 ColorBitMask) final override;
 	
@@ -192,9 +188,55 @@ public:
 	bool PrepareToDraw(uint32 PrimitiveType);
 	void PrepareToDispatch();
 
-	TArray<FMetalCommandBuffer*> Finalize();
+	void Finalize(TArray<FMetalPayload*>& OutPayloads);
 	
-	void InsertCommandBufferFence(TSharedPtr<FMetalCommandBufferFence, ESPMode::ThreadSafe>& Fence, FMetalCommandBufferCompletionHandler Handler);
+	enum class EPhase
+	{
+		Wait,
+		Execute,
+		Signal
+	} CurrentPhase = EPhase::Wait;
+
+	FMetalPayload* GetPayload(EPhase Phase)
+	{
+		if (Payloads.Num() == 0 || Phase < CurrentPhase)
+		{
+			NewPayload();
+		}
+
+		CurrentPhase = Phase;
+		return Payloads.Last();
+	}
+
+	void NewPayload()
+	{
+		Payloads.Add(new FMetalPayload(Device.GetCommandQueue(EMetalQueueType::Direct)));
+	}
+	
+	FMetalSyncPoint* GetContextSyncPoint()
+	{
+		if (!ContextSyncPoint)
+		{
+			ContextSyncPoint = FMetalSyncPoint::Create(EMetalSyncPointType::GPUAndCPU);
+			BatchedSyncPoints.ToSignal.Add(ContextSyncPoint);
+		}
+
+		return ContextSyncPoint;
+	}
+	
+	// Sync points which are waited at the start / signaled at the end 
+	// of the whole batch of command lists this context recorded.
+	struct
+	{
+		TArray<FMetalSyncPointRef> ToWait;
+		TArray<FMetalSyncPointRef> ToSignal;
+	} BatchedSyncPoints;
+	
+	// Inserts a command to signal the specified sync point
+	void SignalSyncPoint(FMetalSyncPoint* SyncPoint);
+
+	// Inserts a command that blocks the GPU queue until the specified sync point is signaled.
+	void WaitSyncPoint(FMetalSyncPoint* SyncPoint);
 	
 	void StartTiming(class FMetalEventNode* EventNode);
 	void EndTiming(class FMetalEventNode* EventNode);
@@ -202,8 +244,6 @@ public:
 	void SynchronizeResource(MTL::Resource* Resource);
 	void SynchronizeTexture(MTL::Texture* Texture, uint32 Slice, uint32 Level);
 	
-	void AddCompletionHandler(FMetalCommandBufferCompletionHandler& Handler);
-	
 	/** Update the event to capture all GPU work so far enqueued by this encoder. */
 	void SignalEvent(MTLEventPtr Event, uint32_t SignalCount);
 	
@@ -264,13 +304,12 @@ public:
 		return bWithinRenderPass;
 	}
 	
-	void SplitCommandBuffers()
-	{
-		if(!bWithinRenderPass)
-		{
-			CurrentEncoder.SplitCommandBuffers();
-		}
-	}
+	void StartCommandBuffer();
+	void EndCommandBuffer();
+	
+#if RHI_NEW_GPU_PROFILER
+	void FlushProfilerStats();
+#endif
 	
 protected:
 	FMetalDevice& Device;
@@ -291,9 +330,6 @@ protected:
 	
 	MTL::RenderPassDescriptor* RenderPassDesc = nullptr;
 	
-	/** Occlusion query batch fence */
-	TSharedPtr<FMetalCommandBufferFence, ESPMode::ThreadSafe> CommandBufferFence;
-	
 	/** Profiling implementation details. */
 	class FMetalProfiler* Profiler = nullptr;
 
@@ -305,6 +341,13 @@ protected:
 
 	TArray<FRHIUniformBuffer*> GlobalUniformBuffers;
 	
+	// The array of recorded payloads the submission thread will process.
+	// These are returned when the context is finalized.
+	TArray<FMetalPayload*> Payloads;
+
+	// A sync point signaled when all payloads in this context have completed.
+	FMetalSyncPointRef ContextSyncPoint;
+	
 private:
 	void RHIClearMRT(bool bClearColor, int32 NumClearColors, const FLinearColor* ColorArray, bool bClearDepth, float Depth, bool bClearStencil, uint32 Stencil);
 };
@@ -317,9 +360,9 @@ public:
 	
 	typedef TFunction<void(FMetalRHICommandContext*)> UploadContextFunction;
 	
-	virtual TArray<FMetalCommandBuffer*>* Finalize();
+	void Finalize(TArray<FMetalPayload*>& OutPayloads);
 	
-	virtual void EnqueueFunction(UploadContextFunction Function)
+	void EnqueueFunction(UploadContextFunction Function)
 	{
 		UploadFunctions.Add(Function);
 	}
diff --git a/Engine/Source/Runtime/Apple/MetalRHI/Public/MetalResources.h b/Engine/Source/Runtime/Apple/MetalRHI/Public/MetalResources.h
--- a/Engine/Source/Runtime/Apple/MetalRHI/Public/MetalResources.h
+++ b/Engine/Source/Runtime/Apple/MetalRHI/Public/MetalResources.h
@@ -9,6 +9,7 @@
 #include "MetalRHIPrivate.h"
 #include "BoundShaderStateCache.h"
 #include "MetalShaderResources.h"
+#include "MetalSubmission.h"
 #include "ShaderCodeArchive.h"
 #include "Templates/TypeHash.h"
 
@@ -559,38 +560,19 @@ public:
 #endif // PLATFORM_SUPPORTS_BINDLESS_RENDERING
 };
 
-class FMetalCommandBufferFence
-{
-public:
-    bool Wait(uint32_t TimeInterval) const;
-    void Insert(MTLCommandBufferPtr CmdBuffer);
-    
-private:
-    void Signal(const MTL::CommandBuffer* CmdBuffer);
-    
-    FEventRef Condition { EEventMode::ManualReset };
-    MTLCommandBufferPtr CmdBuffer;
-};
-
 class FMetalGPUFence final : public FRHIGPUFence
 {
 public:
-	FMetalGPUFence(FName InName)
-		: FRHIGPUFence(InName)
-	{}
-
-	~FMetalGPUFence()
-	{}
-
-	void Clear() override;
-	bool Poll() const override;
-	void Wait(FRHICommandListImmediate& RHICmdList, FRHIGPUMask GPUMask) const override;
+	FMetalGPUFence(FName InName);
 
-	void Write_TopOfPipe(FRHICommandListBase& RHICmdList);
+	virtual void Clear() override;
+	virtual bool Poll() const override;
+	virtual void Wait(FRHICommandListImmediate& RHICmdList, FRHIGPUMask GPUMask) const override;
 
 private:
-	FGraphEventRef Event;
-    TSharedPtr<FMetalCommandBufferFence, ESPMode::ThreadSafe> Fence = nullptr;
+	FMetalSyncPointRef SyncPoint;
+	
+	friend class FMetalDynamicRHI;
 };
 
 class FMetalShaderLibrary;
diff --git a/Engine/Source/Runtime/ApplicationCore/Private/IOS/IOSAppDelegate.cpp b/Engine/Source/Runtime/ApplicationCore/Private/IOS/IOSAppDelegate.cpp
--- a/Engine/Source/Runtime/ApplicationCore/Private/IOS/IOSAppDelegate.cpp
+++ b/Engine/Source/Runtime/ApplicationCore/Private/IOS/IOSAppDelegate.cpp
@@ -37,6 +37,7 @@
 #include <AVFoundation/AVAudioSession.h>
 #import <AVFoundation/AVFoundation.h>
 #include "HAL/IConsoleManager.h"
+#include "RHICommandList.h"
 
 #if WITH_ACCESSIBILITY
 #include "IOS/Accessibility/IOSAccessibilityCache.h"
@@ -1393,6 +1394,11 @@ static void HandleUniversalLink(const FString& UniversalLink)
 #if !PLATFORM_TVOS &&!PLATFORM_VISIONOS
 	// get the interface orientation
 	
+	if(GIsSuspended)
+	{
+		return;
+	}
+	
 	NSLog(@"didRotate orientation = %d", (int)[self.Window.windowScene interfaceOrientation]);
 	
     UIInterfaceOrientation Orientation = [self.Window.windowScene interfaceOrientation];
@@ -1527,6 +1533,17 @@ FCriticalSection RenderSuspend;
     [self ToggleAudioSession:false];
 #endif
     
+	FGraphEventRef RTBlockTask = FFunctionGraphTask::CreateAndDispatchWhenReady([]()
+	{
+		FRHICommandListImmediate& RHICmdList = FRHICommandListImmediate::Get();
+		if (!RHICmdList.Bypass()) 
+		{
+			RHICmdList.SubmitAndBlockUntilGPUIdle();
+		}
+	}, TStatId(), NULL, ENamedThreads::GetRenderThread());
+
+	RTBlockTask->Wait();
+	
     RenderSuspend.TryLock();
     if (FTaskGraphInterface::IsRunning())
     {
@@ -1587,6 +1604,7 @@ FCriticalSection RenderSuspend;
 {
 	FBackgroundURLSessionHandler::HandleWillEnterForeground();
 	FEmbeddedCommunication::KeepAwake(TEXT("Background"), false);
+
 	/*
 	 Called as part of the transition from the background to the inactive state; here you can undo many of the changes made on entering the background.
 	 */
diff --git a/Engine/Source/Runtime/ApplicationCore/Private/IOS/IOSView.cpp b/Engine/Source/Runtime/ApplicationCore/Private/IOS/IOSView.cpp
--- a/Engine/Source/Runtime/ApplicationCore/Private/IOS/IOSView.cpp
+++ b/Engine/Source/Runtime/ApplicationCore/Private/IOS/IOSView.cpp
@@ -346,8 +346,12 @@ MTL::Device* GMetalDevice = nullptr;
 - (void)layoutSubviews
 {
 #if !PLATFORM_TVOS
-    auto orientation = [self.window.windowScene interfaceOrientation];
-	FIOSApplication::OrientationChanged(orientation);
+	// IOSApplication::OrientationChange calls into Slate which issues rendering commands
+	if(![IOSAppDelegate GetDelegate].bIsSuspended)
+	{
+		auto orientation = [self.window.windowScene interfaceOrientation];
+		FIOSApplication::OrientationChanged(orientation);
+	}
 #endif
 }
 
