MetalRHI

Change to fix runtime errors when building with r.Shaders.Optimize

diff --git a/Engine/Source/Developer/Apple/MetalShaderFormat/Private/MetalCompileShaderMSC.cpp b/Engine/Source/Developer/Apple/MetalShaderFormat/Private/MetalCompileShaderMSC.cpp
--- a/Engine/Source/Developer/Apple/MetalShaderFormat/Private/MetalCompileShaderMSC.cpp
+++ b/Engine/Source/Developer/Apple/MetalShaderFormat/Private/MetalCompileShaderMSC.cpp
@@ -531,6 +531,20 @@ void FMetalCompileShaderMSC::DoCompileMetalShader(
 			ExtraArgs.Add(TEXT("--ignore-line-directives"));
 		}
 		
+		if (Input.Environment.CompilerFlags.Contains(CFLAG_Debug) || Input.Environment.CompilerFlags.Contains(CFLAG_SkipOptimizationsDXC))
+		{
+			// Currently cannot enable -Od because we have unbound parameters
+			ExtraArgs.Add(TEXT("-O1"));
+		}
+		else if (Input.Environment.CompilerFlags.Contains(CFLAG_StandardOptimization))
+		{
+			ExtraArgs.Add(TEXT("-O1"));
+		}
+		else
+		{
+			ExtraArgs.Add(TEXT("-O3"));
+		}
+		
 		// Load shader source into compiler context
 		CompilerContext.LoadSource(PreprocessedShader, Input.VirtualSourceFilePath, Input.EntryPointName, Frequency, nullptr, &ExtraArgs);
 
diff --git a/Engine/Source/Developer/Apple/MetalShaderFormat/Private/MetalShaderCompiler.cpp b/Engine/Source/Developer/Apple/MetalShaderFormat/Private/MetalShaderCompiler.cpp
--- a/Engine/Source/Developer/Apple/MetalShaderFormat/Private/MetalShaderCompiler.cpp
+++ b/Engine/Source/Developer/Apple/MetalShaderFormat/Private/MetalShaderCompiler.cpp
@@ -603,7 +603,7 @@ void BuildMetalShaderOutput(
 
 	FString MetalCode = FString(USFSource);
 
-	if (ShaderInput.Environment.CompilerFlags.Contains(CFLAG_Debug))
+	if (ShaderInput.Environment.CompilerFlags.Contains(CFLAG_Debug) && !bUseMetalShaderConverter)
 	{
 		MetalCode.InsertAt(0, FString::Printf(TEXT("// %s\n"), *CCHeader.Name));
 	}
@@ -625,7 +625,7 @@ void BuildMetalShaderOutput(
 				Header.Bindings.NumSamplers, CCHeader.SamplerStates.Num(), MaxMetalSamplers, *SamplerList);
 	}
 	// TODO read from toolchain? this check isn't really doing exactly what it says
-	else if(CompileProcessAllowsRuntimeShaderCompiling(ShaderInput))
+	else if(CompileProcessAllowsRuntimeShaderCompiling(ShaderInput) && !bUseMetalShaderConverter)
 	{
 		// Write out the header and shader source code.
 		FMemoryWriter Ar(ShaderOutput.ShaderCode.GetWriteAccess(), true);
diff --git a/Engine/Source/Runtime/Apple/MetalRHI/Private/Shaders/Types/Templates/MetalBaseShader.h b/Engine/Source/Runtime/Apple/MetalRHI/Private/Shaders/Types/Templates/MetalBaseShader.h
--- a/Engine/Source/Runtime/Apple/MetalRHI/Private/Shaders/Types/Templates/MetalBaseShader.h
+++ b/Engine/Source/Runtime/Apple/MetalRHI/Private/Shaders/Types/Templates/MetalBaseShader.h
@@ -170,11 +170,20 @@ void TMetalBaseShader<BaseResourceType, ShaderType>::Init(TArrayView<const uint8
 
 	bool bOfflineCompile = (OfflineCompiledFlag > 0);
 
+	bool bUseMetalShaderConverter = EnumHasAnyFlags(Header.Bindings.Flags, EMetalBindingsFlags::UseMetalShaderConverter); 
+	
 	const ANSICHAR* ShaderSource = ShaderCode.FindOptionalData(EShaderOptionalDataKey::SourceCode);
-	bool bHasShaderSource = (ShaderSource && FCStringAnsi::Strlen(ShaderSource) > 0);
-
+	bool bHasShaderSource = false;
+	static bool bForceTextShaders = false;
+	
 	static const auto CVar = IConsoleManager::Get().FindConsoleVariable(TEXT("r.Shaders.Symbols"));
-	static bool bForceTextShaders = CVar->GetInt() == 1 || FParse::Param(FCommandLine::Get(),TEXT("metalshaderdebug"));
+	
+	if(!bUseMetalShaderConverter)
+	{
+		bForceTextShaders = CVar->GetInt() == 1 || FParse::Param(FCommandLine::Get(),TEXT("metalshaderdebug"));
+		bHasShaderSource = (ShaderSource && FCStringAnsi::Strlen(ShaderSource) > 0);
+	}
+	
 	if (!bHasShaderSource)
 	{
 		int32 LZMASourceSize = 0;
