Metal RHI

- Changed maximum query buffer size to a CVar
- Instead of failing when over allocating we now pass the occlusion query count when acquiring the buffer to ensure it fits.

diff --git a/Engine/Source/Runtime/Apple/MetalRHI/Private/MetalRHIContext.cpp b/Engine/Source/Runtime/Apple/MetalRHI/Private/MetalRHIContext.cpp
--- a/Engine/Source/Runtime/Apple/MetalRHI/Private/MetalRHIContext.cpp
+++ b/Engine/Source/Runtime/Apple/MetalRHI/Private/MetalRHIContext.cpp
@@ -182,7 +182,14 @@ void FMetalRHICommandContext::RHIBeginRenderPass(const FRHIRenderPassInfo& InInf
 	
 	StateCache.SetStateDirty();
 	StateCache.SetRenderTargetsActive(true);
-	StateCache.StartRenderPass(InInfo, QueryBuffer->GetCurrentQueryBuffer());
+	
+	FMetalQueryBuffer* VisBuffer = nullptr;
+	if(InInfo.NumOcclusionQueries > 0)
+	{
+		VisBuffer = QueryBuffer->AcquireQueryBuffer(InInfo.NumOcclusionQueries);
+	}
+		
+	StateCache.StartRenderPass(InInfo, VisBuffer);
 	
 	RenderPassDesc = StateCache.GetRenderPassDescriptor();
 	
@@ -416,8 +423,6 @@ void FMetalRHICommandContext::CopyFromBufferToBuffer(FMetalBufferPtr SourceBuffe
 
 void FMetalRHICommandContext::Finalize(TArray<FMetalPayload*>& OutPayloads)
 {
-	GetQueryBufferPool()->ReleaseCurrentQueryBuffer();
-	
 	if(CurrentEncoder.IsAnyCommandEncoderActive())
 	{
 		CurrentEncoderFence = CurrentEncoder.EndEncoding();
diff --git a/Engine/Source/Runtime/Apple/MetalRHI/Private/MetalRHIRenderQuery.cpp b/Engine/Source/Runtime/Apple/MetalRHI/Private/MetalRHIRenderQuery.cpp
--- a/Engine/Source/Runtime/Apple/MetalRHI/Private/MetalRHIRenderQuery.cpp
+++ b/Engine/Source/Runtime/Apple/MetalRHI/Private/MetalRHIRenderQuery.cpp
@@ -19,6 +19,13 @@
 
 #pragma mark - Metal RHI Private Query Buffer Resource Class -
 
+static int32 GMetalMaxQueryBufferSize = 1024 * 256;
+
+static FAutoConsoleVariableRef CVarMetalMaxQueryBufferSize(
+	TEXT("rhi.Metal.MaxQueryBufferSize"),
+	GMetalMaxQueryBufferSize,
+	TEXT("Maximum size of the the query buffer in a single context. Default = 512kb on Mac and 256k on iOS"),
+	ECVF_ReadOnly);
 
 FMetalQueryBuffer::FMetalQueryBuffer(FMetalQueryBufferPool* InPool, FMetalBufferPtr InBuffer)
 	: FRHIResource(RRT_TimestampCalibrationQuery)
@@ -46,10 +53,15 @@ FMetalQueryBuffer::~FMetalQueryBuffer()
 uint64 FMetalQueryBuffer::GetResult(uint32 Offset)
 {
 	uint64 Result = 0;
-    MTL_SCOPED_AUTORELEASE_POOL;
-    {
-		Result = *((uint64 const*)(((uint8*)Buffer->Contents()) + Offset));
+	
+	if(Buffer.IsValid())
+	{
+		MTL_SCOPED_AUTORELEASE_POOL;
+		{
+			Result = *((uint64 const*)(((uint8*)Buffer->Contents()) + Offset));
+		}
 	}
+
 	return Result;
 }
 
@@ -64,6 +76,7 @@ FMetalQueryBufferPool::FMetalQueryBufferPool(FMetalDevice& InDevice)
 	, Buffers{}
 	, Device{InDevice}
 {
+	BufferSize = GMetalMaxQueryBufferSize;
 	// void
 }
 
@@ -74,56 +87,86 @@ FMetalQueryBufferPool::~FMetalQueryBufferPool()
 
 void FMetalQueryBufferPool::Allocate(FMetalQueryResult& NewQuery)
 {
-	FMetalQueryBuffer* QB = IsValidRef(CurrentBuffer) ? CurrentBuffer.GetReference() : GetCurrentQueryBuffer();
+	check(IsValidRef(CurrentBuffer));
+	FMetalQueryBuffer* QB = CurrentBuffer.GetReference();
 
 	uint32 Offset = Align(QB->WriteOffset, EQueryBufferAlignment);
 	uint32 End = Align(Offset + EQueryResultMaxSize, EQueryBufferAlignment);
 
-	if (Align(QB->WriteOffset, EQueryBufferAlignment) + EQueryResultMaxSize <= EQueryBufferMaxSize)
+	if (Align(QB->WriteOffset, EQueryBufferAlignment) + EQueryResultMaxSize <= BufferSize)
 	{
 		NewQuery.SourceBuffer = QB;
 		NewQuery.Offset = Align(QB->WriteOffset, EQueryBufferAlignment);
 		QB->WriteOffset = Align(QB->WriteOffset, EQueryBufferAlignment) + EQueryResultMaxSize;
 	}
-	else
-	{
-		UE_LOG(LogMetal, Fatal, TEXT("No memory left in pool, check EQueryBufferMaxSize"));
-	}
 }
 
-FMetalQueryBuffer* FMetalQueryBufferPool::GetCurrentQueryBuffer()
+FMetalQueryBuffer* FMetalQueryBufferPool::AcquireQueryBuffer(uint32_t NumOcclusionQueries)
 {
-	if (!IsValidRef(CurrentBuffer) || (CurrentBuffer->Buffer->GetMTLBuffer()->storageMode() != MTL::StorageModeShared && CurrentBuffer->WriteOffset > 0))
+	uint32_t RequiredSize = NumOcclusionQueries * EQueryResultMaxSize;
+	
+	if(CurrentBuffer)
 	{
-		FMetalBufferPtr Buffer;
-		if (Buffers.Num())
+		// If we currently have a buffer and the results fit, then use it
+		if(Align(CurrentBuffer->WriteOffset, EQueryBufferAlignment) + RequiredSize <= BufferSize)
 		{
-			Buffer = Buffers.Pop();
+			return CurrentBuffer;
 		}
 		else
 		{
-			METAL_GPUPROFILE(FScopedMetalCPUStats CPUStat(FString::Printf(TEXT("AllocBuffer: %llu, %llu"), EQueryBufferMaxSize, MTL::ResourceStorageModeShared)));
-			
-			MTL::ResourceOptions HazardTrackingMode = MTL::ResourceHazardTrackingModeUntracked;
-			static bool bSupportsHeaps = Device.SupportsFeature(EMetalFeaturesHeaps);
-			if(bSupportsHeaps)
-			{
-				HazardTrackingMode = MTL::ResourceHazardTrackingModeTracked;
-			}
-			
-			Buffer = Device.GetResourceHeap().CreateBuffer(EQueryBufferMaxSize, 16, BUF_Dynamic, FMetalCommandQueue::GetCompatibleResourceOptions((MTL::ResourceOptions)(BUFFER_CACHE_MODE | HazardTrackingMode | MTL::ResourceStorageModeShared)), true);
-			FMemory::Memzero((((uint8*)Buffer->Contents())), EQueryBufferMaxSize);
+			ReleaseCurrentQueryBuffer();
 		}
-
-		CurrentBuffer = new FMetalQueryBuffer(this, MoveTemp(Buffer));
 	}
+	
+	// Need to resize if queries don't fit in our current buffer size
+	if(RequiredSize > BufferSize)
+	{
+		BufferSize = FMath::RoundUpToPowerOfTwo(RequiredSize);
+	
+		for(FMetalBufferPtr Buffer : Buffers)
+		{
+			FMetalDynamicRHI::Get().DeferredDelete(Buffer);
+		}
+		
+		Buffers.Empty();
+	}
+	
+	FMetalBufferPtr Buffer;
+	
+	if (Buffers.Num())
+	{
+		Buffer = Buffers.Pop();
+	}
+	else
+	{
+		METAL_GPUPROFILE(FScopedMetalCPUStats CPUStat(FString::Printf(TEXT("AllocBuffer: %llu, %llu"), BufferSize, MTL::ResourceStorageModeShared)));
+		
+		MTL::ResourceOptions HazardTrackingMode = MTL::ResourceHazardTrackingModeUntracked;
+		static bool bSupportsHeaps = Device.SupportsFeature(EMetalFeaturesHeaps);
+		if(bSupportsHeaps)
+		{
+			HazardTrackingMode = MTL::ResourceHazardTrackingModeTracked;
+		}
+		
+		Buffer = Device.GetResourceHeap().CreateBuffer(BufferSize, 16, BUF_Dynamic, FMetalCommandQueue::GetCompatibleResourceOptions((MTL::ResourceOptions)(BUFFER_CACHE_MODE | HazardTrackingMode | MTL::ResourceStorageModeShared)), true);
+		
+		FMemory::Memzero((((uint8*)Buffer->Contents())), BufferSize);
+	}
+	
+	check(Buffer);
+	
+	CurrentBuffer = new FMetalQueryBuffer(this, MoveTemp(Buffer));
+	return CurrentBuffer.GetReference();
+}
 
+FMetalQueryBuffer* FMetalQueryBufferPool::GetCurrentQueryBuffer()
+{
 	return CurrentBuffer.GetReference();
 }
 
 void FMetalQueryBufferPool::ReleaseCurrentQueryBuffer()
 {
-	if (IsValidRef(CurrentBuffer) && CurrentBuffer->WriteOffset > 0)
+	if (IsValidRef(CurrentBuffer) && (CurrentBuffer->WriteOffset > 0))
 	{
 		CurrentBuffer.SafeRelease();
 	}
@@ -131,7 +174,14 @@ void FMetalQueryBufferPool::ReleaseCurrentQueryBuffer()
 
 void FMetalQueryBufferPool::ReleaseQueryBuffer(FMetalBufferPtr Buffer)
 {
-	Buffers.Add(Buffer);
+	if(Buffer->GetLength() >= BufferSize)
+	{
+		Buffers.Add(Buffer);
+	}
+	else
+	{
+		FMetalDynamicRHI::Get().DeferredDelete(Buffer);
+	}
 }
 
 
diff --git a/Engine/Source/Runtime/Apple/MetalRHI/Private/MetalRHIRenderQuery.h b/Engine/Source/Runtime/Apple/MetalRHI/Private/MetalRHIRenderQuery.h
--- a/Engine/Source/Runtime/Apple/MetalRHI/Private/MetalRHIRenderQuery.h
+++ b/Engine/Source/Runtime/Apple/MetalRHI/Private/MetalRHIRenderQuery.h
@@ -52,8 +52,7 @@ public:
 	enum
 	{
 		EQueryBufferAlignment = 8,
-		EQueryResultMaxSize   = 8,
-		EQueryBufferMaxSize   = (1 << 18)
+		EQueryResultMaxSize   = 8
 	};
 
 	// Disallow a default constructor
@@ -63,10 +62,14 @@ public:
 	~FMetalQueryBufferPool();
 
 	void Allocate(FMetalQueryResult& NewQuery);
+	
+	FMetalQueryBuffer* AcquireQueryBuffer(uint32_t NumOcclusionQueries);
 	FMetalQueryBuffer* GetCurrentQueryBuffer();
 	void ReleaseCurrentQueryBuffer();
 	void ReleaseQueryBuffer(FMetalBufferPtr Buffer);
 
+	uint32_t BufferSize = 0;
+	
 	TRefCountPtr<FMetalQueryBuffer> CurrentBuffer;
 	TArray<FMetalBufferPtr> Buffers;
 	FMetalDevice& Device;
