[HWRT] Deprecate PerInstanceGeometries.
- "Lightweight SBT" for inline ray tracing is created/managed using FRHIShaderBindingTable and ERayTracingShaderBindingMode::Inline.

diff --git a/Engine/Plugins/Experimental/GPULightmass/Source/GPULightmass/Private/LightmapRenderer.cpp b/Engine/Plugins/Experimental/GPULightmass/Source/GPULightmass/Private/LightmapRenderer.cpp
--- a/Engine/Plugins/Experimental/GPULightmass/Source/GPULightmass/Private/LightmapRenderer.cpp
+++ b/Engine/Plugins/Experimental/GPULightmass/Source/GPULightmass/Private/LightmapRenderer.cpp
@@ -1324,7 +1324,6 @@ bool FSceneRenderState::SetupRayTracingScene(FRDGBuilder& GraphBuilder, FSceneUn
 				BuildParams.InstanceBuffer = InstanceBuffer.Buffer;
 				BuildParams.InstanceBufferOffset = 0;
 				BuildParams.ReferencedGeometries = RayTracingSceneInitializationData.ReferencedGeometries;
-				BuildParams.PerInstanceGeometries = RayTracingSceneInitializationData.PerInstanceGeometries;
 				BuildParams.NumInstances = RayTracingSceneInitializationData.NumNativeGPUSceneInstances + RayTracingSceneInitializationData.NumNativeCPUInstances;
 
 				RHICmdList.BuildAccelerationStructure(BuildParams);
diff --git a/Engine/Source/Runtime/Apple/MetalRHI/Private/MetalRayTracing.cpp b/Engine/Source/Runtime/Apple/MetalRHI/Private/MetalRayTracing.cpp
--- a/Engine/Source/Runtime/Apple/MetalRHI/Private/MetalRayTracing.cpp
+++ b/Engine/Source/Runtime/Apple/MetalRHI/Private/MetalRayTracing.cpp
@@ -433,21 +433,11 @@ FMetalRayTracingScene::FMetalRayTracingScene(FRayTracingSceneInitializer InIniti
 	SizeInfo = CalcRayTracingGeometrySize(InstanceDescriptor);
 
 	MutableAccelerationStructures = [[NSMutableArray<id<MTLAccelerationStructure>> new] init];
-
-	//const uint32 ParameterBufferSize = FMath::Max<uint32>(1, Initializer.NumTotalSegments) * sizeof(FHitGroupSystemRootConstants);
-	//FRHIResourceCreateInfo ParameterBufferCreateInfo(TEXT("RayTracingSceneMetadata"));
-	//PerInstanceGeometryParameterBuffer = ResourceCast(FRHICommandListExecutor::GetImmediateCommandList().CreateBuffer(ParameterBufferSize,  BUF_StructuredBuffer | BUF_ShaderResource, sizeof(FHitGroupSystemRootConstants), ERHIAccess::SRVMask, ParameterBufferCreateInfo).GetReference());
-
-	//FShaderResourceViewInitializer ViewInitializer(PerInstanceGeometryParameterBuffer);
-	//PerInstanceGeometryParameterSRV = new FMetalShaderResourceView(ViewInitializer);
 }
 
 FMetalRayTracingScene::~FMetalRayTracingScene()
 {
-	PerInstanceGeometryParameters.Empty();
 	AccelerationStructureBuffer.SafeRelease();
-    //PerInstanceGeometryParameterBuffer.SafeRelease();
-	//PerInstanceGeometryParameterSRV.SafeRelease();
 	InstanceBufferSRV.SafeRelease();
 
 	[MutableAccelerationStructures removeAllObjects];
@@ -479,72 +469,6 @@ void FMetalRayTracingScene::BindBuffer(FRHIBuffer* InBuffer, uint32 InBufferOffs
 	}
 }
 
-void FMetalRayTracingScene::BuildPerInstanceGeometryParameterBuffer()
-{
-	[MutableAccelerationStructures removeAllObjects];
-	PerInstanceGeometryParameters.Empty();
-
-	const uint32 ParameterBufferSize = FMath::Max<uint32>(1, Initializer.NumTotalSegments) * sizeof(FHitGroupSystemRootConstants);
-	//check(PerInstanceGeometryParameterBuffer->GetSize() >= ParameterBufferSize);
-
-	check(IsInRHIThread() || !IsRunningRHIInSeparateThread());
-
-	//void* MappedBuffer = PerInstanceGeometryParameterBuffer->Lock(true, RLM_WriteOnly, 0, ParameterBufferSize);
-	//FHitGroupSystemRootConstants* MappedParameters = reinterpret_cast<FHitGroupSystemRootConstants*>(MappedBuffer);
-	uint32 ParameterIndex = 0;
-
-	for (FRHIRayTracingGeometry* GeometryRHI : Initializer.PerInstanceGeometries)
-	{
-		FMetalRayTracingGeometry* Geometry = ResourceCast(GeometryRHI);
-		Geometry->SceneIndex = ParameterIndex;
-
-		MTL::AccelerationStructure* AS = Geometry->GetAccelerationStructureRead()->AccelerationStructureHandle;
-		[MutableAccelerationStructures addObject:AS];
-		InstanceBufferSRV->ReferencedResources.Add(TTuple<MTL::Resource*, MTL::ResourceUsage>((MTL::Resource*)(AS), MTL::ResourceUsageRead));
-
-		const FRayTracingGeometryInitializer& GeometryInitializer = Geometry->GetInitializer();
-
-		FMetalRHIBuffer* IndexBuffer = ResourceCast(GeometryInitializer.IndexBuffer.GetReference());
-		const uint32 IndexStride = IndexBuffer ? IndexBuffer->GetStride() : 0;
-		const uint32 IndexOffsetInBytes = GeometryInitializer.IndexBufferOffset;
-
-		for (const FRayTracingGeometrySegment& Segment : GeometryInitializer.Segments)
-		{
-			FMetalRHIBuffer* VertexBuffer = ResourceCast(Segment.VertexBuffer.GetReference());
-			checkf(VertexBuffer, TEXT("All ray tracing geometry segments must have a valid vertex buffer"));
-
-			FHitGroupSystemRootConstants SysRootConstants;
-			SysRootConstants.SetVertexAndIndexStride(Segment.VertexBufferStride, IndexStride);
-			SysRootConstants.UserData = 0;
-
-			FMetalRayTracingGeometryParameters SegmentParameters;
-			if (IndexStride)
-			{
-				SysRootConstants.IndexBufferOffsetInBytes = IndexOffsetInBytes + IndexStride * Segment.FirstPrimitive * FMetalRayTracingGeometry::IndicesPerPrimitive;
-				SegmentParameters.IndexBuffer = IndexBuffer;
-			}
-			else
-			{
-				SysRootConstants.IndexBufferOffsetInBytes = 0;
-				SegmentParameters.IndexBuffer = nullptr;
-			}
-
-			SegmentParameters.VertexBuffer = VertexBuffer;
-			SegmentParameters.RootConstantsBufferOffsetInBytes = ParameterIndex * sizeof(FHitGroupSystemRootConstants);
-			SegmentParameters.VertexBufferOffset = Segment.VertexBufferOffset;
-
-			PerInstanceGeometryParameters.Add(SegmentParameters);
-
-			check(ParameterIndex < Initializer.NumTotalSegments);
-			//MappedParameters[ParameterIndex] = SysRootConstants;
-			ParameterIndex++;
-		}
-	}
-
-	check(ParameterIndex == Initializer.NumTotalSegments);
-	//PerInstanceGeometryParameterBuffer->Unlock();
-}
-
 void FMetalRayTracingScene::BuildAccelerationStructure(
 		FMetalRHICommandContext& CommandContext,
 		FMetalRHIBuffer* InScratchBuffer, uint32 ScratchOffset,
@@ -562,8 +486,6 @@ void FMetalRayTracingScene::BuildAccelerationStructure(
 	FShaderResourceViewInitializer ViewInitializer(InstanceBuffer, InstanceOffset, 0);
 	InstanceBufferSRV = new FMetalShaderResourceView(ViewInitializer);
 
-	BuildPerInstanceGeometryParameterBuffer();
-
 	TRefCountPtr<FMetalRHIBuffer> ScratchBuffer;
 	if (InScratchBuffer == nullptr)
 	{
diff --git a/Engine/Source/Runtime/Apple/MetalRHI/Private/MetalRayTracing.h b/Engine/Source/Runtime/Apple/MetalRHI/Private/MetalRayTracing.h
--- a/Engine/Source/Runtime/Apple/MetalRHI/Private/MetalRayTracing.h
+++ b/Engine/Source/Runtime/Apple/MetalRHI/Private/MetalRayTracing.h
@@ -83,8 +83,6 @@ public:
 		FMetalRHIBuffer* ScratchBuffer, uint32 ScratchOffset,
 		FMetalRHIBuffer* InstanceBuffer, uint32 InstanceOffset);
 
-	void BuildPerInstanceGeometryParameterBuffer();
-
 	inline const FRayTracingSceneInitializer& GetInitializer() const override final { return Initializer; }
 
 	TRefCountPtr<FMetalShaderResourceView> InstanceBufferSRV;
@@ -101,13 +99,6 @@ private:
 	/** Acceleration Structure for the whole scene. */
 	TRefCountPtr<FMetalRHIBuffer> AccelerationStructureBuffer;
 
-	/** Root Constants for geometry evaluation in HitGroup/Miss (emulates D3D12 RootConstants with a global scope). */
-	TArray<FMetalRayTracingGeometryParameters> PerInstanceGeometryParameters;
-
-	// Buffer that contains per-instance index and vertex buffer binding data
-	TRefCountPtr<FMetalRHIBuffer> PerInstanceGeometryParameterBuffer;
-	TRefCountPtr<FMetalShaderResourceView> PerInstanceGeometryParameterSRV;
-
 	/** Segments descriptors  (populated when the constructor is called). */
 	NSMutableArray<id<MTLAccelerationStructure>>* MutableAccelerationStructures;
 };
diff --git a/Engine/Source/Runtime/RHI/Public/RHICommandList.h b/Engine/Source/Runtime/RHI/Public/RHICommandList.h
--- a/Engine/Source/Runtime/RHI/Public/RHICommandList.h
+++ b/Engine/Source/Runtime/RHI/Public/RHICommandList.h
@@ -3859,7 +3859,9 @@ public:
 		{
 			FRayTracingSceneBuildParams InlineParams = SceneBuildParams;
 			InlineParams.ReferencedGeometries = AllocArray(SceneBuildParams.ReferencedGeometries);
+			PRAGMA_DISABLE_DEPRECATION_WARNINGS
 			InlineParams.PerInstanceGeometries = AllocArray(SceneBuildParams.PerInstanceGeometries);
+			PRAGMA_ENABLE_DEPRECATION_WARNINGS
 
 			ALLOC_COMMAND(FRHICommandBuildAccelerationStructure)(MoveTemp(InlineParams));
 
diff --git a/Engine/Source/Runtime/RHI/Public/RHIContext.h b/Engine/Source/Runtime/RHI/Public/RHIContext.h
--- a/Engine/Source/Runtime/RHI/Public/RHIContext.h
+++ b/Engine/Source/Runtime/RHI/Public/RHIContext.h
@@ -638,10 +638,20 @@ struct FRayTracingSceneBuildParams
 	// Unique list of geometries referenced by all instances in this scene.
 	// Any referenced geometry is kept alive while the scene is alive.
 	TConstArrayView<FRHIRayTracingGeometry*> ReferencedGeometries;
-	// One entry per instance
+
+	UE_DEPRECATED(5.6, "No longer necessary.")
 	TConstArrayView<FRHIRayTracingGeometry*> PerInstanceGeometries;
 
 	EAccelerationStructureBuildMode BuildMode = EAccelerationStructureBuildMode::Build;
+
+	PRAGMA_DISABLE_DEPRECATION_WARNINGS
+	FRayTracingSceneBuildParams() = default;
+	FRayTracingSceneBuildParams(const FRayTracingSceneBuildParams&) = default;
+	FRayTracingSceneBuildParams& operator=(const FRayTracingSceneBuildParams&) = default;
+	FRayTracingSceneBuildParams(FRayTracingSceneBuildParams&&) = default;
+	FRayTracingSceneBuildParams& operator=(FRayTracingSceneBuildParams&&) = default;
+	~FRayTracingSceneBuildParams() = default;
+	PRAGMA_ENABLE_DEPRECATION_WARNINGS
 };
 
 struct FCopyBufferRegionParams
diff --git a/Engine/Source/Runtime/RHI/Public/RHIResources.h b/Engine/Source/Runtime/RHI/Public/RHIResources.h
--- a/Engine/Source/Runtime/RHI/Public/RHIResources.h
+++ b/Engine/Source/Runtime/RHI/Public/RHIResources.h
@@ -3650,8 +3650,8 @@ struct FRayTracingSceneInitializer
 	// Any referenced geometry is kept alive while the scene is alive.
 	UE_DEPRECATED(5.5, "ReferencedGeometries should be provided in FRayTracingSceneBuildParams.")
 	TArray<TRefCountPtr<FRHIRayTracingGeometry>> ReferencedGeometries;
-	// One entry per instance
-	UE_DEPRECATED(5.5, "PerInstanceGeometries should be provided in FRayTracingSceneBuildParams.")
+
+	UE_DEPRECATED(5.5, "No longer necessary.")
 	TArray<FRHIRayTracingGeometry*> PerInstanceGeometries;
 
 	// Exclusive prefix sum of `Instance.NumTransforms` for all instances in this scene. Used to emulate SV_InstanceID in hit shaders.
diff --git a/Engine/Source/Runtime/Renderer/Private/RayTracing/RayTracingInstanceBufferUtil.cpp b/Engine/Source/Runtime/Renderer/Private/RayTracing/RayTracingInstanceBufferUtil.cpp
--- a/Engine/Source/Runtime/Renderer/Private/RayTracing/RayTracingInstanceBufferUtil.cpp
+++ b/Engine/Source/Runtime/Renderer/Private/RayTracing/RayTracingInstanceBufferUtil.cpp
@@ -111,7 +111,6 @@ FRayTracingSceneInitializationData BuildRayTracingSceneInitializationData(TConst
 	Output.BaseUploadBufferOffsets.SetNumUninitialized(NumSceneInstances);
 	Output.BaseInstancePrefixSum.SetNumUninitialized(NumSceneInstances);
 	Output.InstanceGroupEntryRefs.SetNumUninitialized(NumSceneInstances);
-	Output.PerInstanceGeometries.SetNumUninitialized(NumSceneInstances);
 
 	TArray<uint32> InstanceGroups;
 
@@ -138,8 +137,6 @@ FRayTracingSceneInitializationData BuildRayTracingSceneInitializationData(TConst
 
 		checkf(InstanceDesc.GeometryRHI, TEXT("Ray tracing instance must have a valid geometry."));
 
-		Output.PerInstanceGeometries[InstanceIndex] = InstanceDesc.GeometryRHI;
-
 		Output.TotalNumSegments += InstanceDesc.GeometryRHI->GetNumSegments();
 
 		uint32 GeometryIndex = UniqueGeometries.FindOrAdd(InstanceDesc.GeometryRHI, Output.ReferencedGeometries.Num());
diff --git a/Engine/Source/Runtime/Renderer/Private/RayTracing/RayTracingScene.cpp b/Engine/Source/Runtime/Renderer/Private/RayTracing/RayTracingScene.cpp
--- a/Engine/Source/Runtime/Renderer/Private/RayTracing/RayTracingScene.cpp
+++ b/Engine/Source/Runtime/Renderer/Private/RayTracing/RayTracingScene.cpp
@@ -121,9 +121,6 @@ void FRayTracingScene::Create(FRDGBuilder& GraphBuilder, const FViewInfo& View,
 		FLayer& Layer = Layers[LayerIndex];
 
 		{
-			TArray<TRefCountPtr<FRHIRayTracingGeometry>> ReferencedGeometries;
-			TArray<FRHIRayTracingGeometry*> PerInstanceGeometries;
-
 			FRayTracingSceneInitializer Initializer;
 			Initializer.DebugName = FName(TEXT("FRayTracingScene"));
 			Initializer.MaxNumInstances = Layer.InitializationData.NumNativeGPUSceneInstances + Layer.InitializationData.NumNativeCPUInstances;
@@ -524,7 +521,6 @@ void FRayTracingScene::Build(FRDGBuilder& GraphBuilder, ERDGPassFlags ComputePas
 				BuildParams.InstanceBufferOffset = 0;
 				BuildParams.NumInstances = Layer.InitializationData.NumNativeCPUInstances + Layer.InitializationData.NumNativeGPUSceneInstances;
 				BuildParams.ReferencedGeometries = Layer.InitializationData.ReferencedGeometries;
-				BuildParams.PerInstanceGeometries = Layer.InitializationData.PerInstanceGeometries;
 
 				RHICmdList.BindAccelerationStructureMemory(Layer.RayTracingSceneRHI, PassParams->TLASBuffer->GetRHI(), 0);
 				RHICmdList.BuildAccelerationStructure(BuildParams);
diff --git a/Engine/Source/Runtime/Renderer/Private/Tests/RayTracingTestbed.cpp b/Engine/Source/Runtime/Renderer/Private/Tests/RayTracingTestbed.cpp
--- a/Engine/Source/Runtime/Renderer/Private/Tests/RayTracingTestbed.cpp
+++ b/Engine/Source/Runtime/Renderer/Private/Tests/RayTracingTestbed.cpp
@@ -314,7 +314,6 @@ bool RunRayTracingTestbed_RenderThread(const FString& Parameters)
 	BuildParams.InstanceBuffer = InstanceBuffer.Buffer;
 	BuildParams.InstanceBufferOffset = 0;
 	BuildParams.ReferencedGeometries = RayTracingSceneInitializationData.ReferencedGeometries;
-	BuildParams.PerInstanceGeometries = RayTracingSceneInitializationData.PerInstanceGeometries;
 	BuildParams.NumInstances = RayTracingSceneInitializationData.NumNativeGPUSceneInstances + RayTracingSceneInitializationData.NumNativeCPUInstances;
 
 	RHICmdList.Transition(FRHITransitionInfo(InstanceBuffer.Buffer, ERHIAccess::UAVMask, ERHIAccess::SRVCompute));
diff --git a/Engine/Source/Runtime/Renderer/Public/RayTracingInstanceBufferUtil.h b/Engine/Source/Runtime/Renderer/Public/RayTracingInstanceBufferUtil.h
--- a/Engine/Source/Runtime/Renderer/Public/RayTracingInstanceBufferUtil.h
+++ b/Engine/Source/Runtime/Renderer/Public/RayTracingInstanceBufferUtil.h
@@ -72,9 +72,19 @@ struct FRayTracingSceneInitializationData
 
 	// Unique list of geometries referenced by all instances in this scene.
 	// Any referenced geometry is kept alive while the scene is alive.
-	TArray<FRHIRayTracingGeometry*> PerInstanceGeometries;
-	// One entry per instance
 	TArray<FRHIRayTracingGeometry*> ReferencedGeometries;
+
+	UE_DEPRECATED(5.5, "No longer used as 'inline SBT' is managed using FShaderBindingTable.")
+	TArray<FRHIRayTracingGeometry*> PerInstanceGeometries;
+
+	PRAGMA_DISABLE_DEPRECATION_WARNINGS
+	FRayTracingSceneInitializationData() = default;
+	FRayTracingSceneInitializationData(const FRayTracingSceneInitializationData&) = default;
+	FRayTracingSceneInitializationData& operator=(const FRayTracingSceneInitializationData&) = default;
+	FRayTracingSceneInitializationData(FRayTracingSceneInitializationData&&) = default;
+	FRayTracingSceneInitializationData& operator=(FRayTracingSceneInitializationData&&) = default;
+	~FRayTracingSceneInitializationData() = default;
+	PRAGMA_ENABLE_DEPRECATION_WARNINGS
 };
 
 RENDERER_API FRayTracingSceneInitializationData BuildRayTracingSceneInitializationData(TConstArrayView<FRayTracingGeometryInstance> Instances);
diff --git a/Engine/Source/Runtime/VulkanRHI/Private/VulkanRayTracing.cpp b/Engine/Source/Runtime/VulkanRHI/Private/VulkanRayTracing.cpp
--- a/Engine/Source/Runtime/VulkanRHI/Private/VulkanRayTracing.cpp
+++ b/Engine/Source/Runtime/VulkanRHI/Private/VulkanRayTracing.cpp
@@ -1364,8 +1364,6 @@ void FVulkanCommandListContext::RHIBuildAccelerationStructure(const FRayTracingS
 		Scene->ReferencedGeometries.Add(ReferencedGeometry);
 	}
 
-	Scene->PerInstanceGeometries = SceneBuildParams.PerInstanceGeometries;
-
 	BuildAccelerationStructure(
 		*this,
 		*Scene,
diff --git a/Engine/Source/Runtime/VulkanRHI/Private/VulkanRayTracing.h b/Engine/Source/Runtime/VulkanRHI/Private/VulkanRayTracing.h
--- a/Engine/Source/Runtime/VulkanRHI/Private/VulkanRayTracing.h
+++ b/Engine/Source/Runtime/VulkanRHI/Private/VulkanRayTracing.h
@@ -247,8 +247,6 @@ public:
 	// Unique list of geometries referenced by all instances in this scene.
 	// Any referenced geometry is kept alive while the scene is alive.
 	TArray<TRefCountPtr<FRHIRayTracingGeometry>> ReferencedGeometries;
-	// One entry per instance
-	TArray<FRHIRayTracingGeometry*> PerInstanceGeometries;
 
 	// Native TLAS handles are owned by SRV objects in Vulkan RHI.
 	// D3D12 and other RHIs allow creating TLAS SRVs from any GPU address at any point
