Temporary keep a persistent map on metallibs to workaround driver bug
Add compatcvar r.Metal.PersistentMappingMetallib
Change DISPATCH_DATA_DESTRUCTOR_MUNMAP to empty destructor for dispatch data for metallib

diff --git a/Engine/Source/Runtime/Apple/MetalRHI/Private/MetalDynamicRHI_Shaders.cpp b/Engine/Source/Runtime/Apple/MetalRHI/Private/MetalDynamicRHI_Shaders.cpp
--- a/Engine/Source/Runtime/Apple/MetalRHI/Private/MetalDynamicRHI_Shaders.cpp
+++ b/Engine/Source/Runtime/Apple/MetalRHI/Private/MetalDynamicRHI_Shaders.cpp
@@ -13,6 +13,13 @@
 #include "Interfaces/IPluginManager.h"
 
 //------------------------------------------------------------------------------
+static bool GPersistentMappingMetallib = 1;
+static FAutoConsoleVariableRef CVarPersistentMappingMetallib(
+	TEXT("r.Metal.PersistentMappingMetallib"),
+	GPersistentMappingMetallib,
+	TEXT("Makes the metallib file mapping persistent."),
+	ECVF_Default
+);
 
 #pragma mark - Metal Dynamic RHI Shader Methods
 
@@ -151,7 +158,7 @@ FRHIShaderLibraryRef FMetalDynamicRHI::RHICreateShaderLibrary(EShaderPlatform Pl
 #endif
 	};
 
-	auto MapShaderLibraryFile = [](const FString& FilePath, TUniquePtr<FMetalShaderLibrary::FShaderLibDataOwner>& MemOwner)
+	auto MapShaderLibraryFile = [](const FString& FilePath, TUniquePtr<FMetalShaderLibrary::FShaderLibDataOwner>& MemOwner, FFileMappingFlags Flags)
 	{
 		if (FPlatformProperties::SupportsMemoryMappedFiles())
 		{
@@ -159,7 +166,7 @@ FRHIShaderLibraryRef FMetalDynamicRHI::RHICreateShaderLibrary(EShaderPlatform Pl
 			MemOwner->MappedCacheFile = Res.HasError() ? nullptr : Res.StealValue();
 			if (MemOwner->MappedCacheFile.IsValid())
 			{
-				MemOwner->MappedRegion = TUniquePtr<IMappedFileRegion>(MemOwner->MappedCacheFile->MapRegion(0, MemOwner->MappedCacheFile->GetFileSize()));
+				MemOwner->MappedRegion = TUniquePtr<IMappedFileRegion>(MemOwner->MappedCacheFile->MapRegion(0, MemOwner->MappedCacheFile->GetFileSize(), Flags));
 			}
 		}
 	};
@@ -169,7 +176,7 @@ FRHIShaderLibraryRef FMetalDynamicRHI::RHICreateShaderLibrary(EShaderPlatform Pl
 #else
 	FStaticMemoryReader* BinaryShaderAr = nullptr;
     TUniquePtr<FMetalShaderLibrary::FShaderLibDataOwner> MemOwner = MakeUnique<FMetalShaderLibrary::FShaderLibDataOwner>();
-	MapShaderLibraryFile(BinaryShaderFile, MemOwner);
+	MapShaderLibraryFile(BinaryShaderFile, MemOwner, EMappedFileFlags::ENone);
 	if (MemOwner->MappedRegion.IsValid())
 	{
 		UE_LOG(LogMetal, Display, TEXT("mmapping %s, %d bytes"), *BinaryShaderFile, MemOwner->MappedCacheFile->GetFileSize());
@@ -214,12 +221,12 @@ FRHIShaderLibraryRef FMetalDynamicRHI::RHICreateShaderLibrary(EShaderPlatform Pl
 
 				MTLLibraryPtr Library;
 				TUniquePtr<FMetalShaderLibrary::FShaderLibDataOwner> MemOwnerLibrary = MakeUnique<FMetalShaderLibrary::FShaderLibDataOwner>();
-				MapShaderLibraryFile(MetalLibraryFilePath, MemOwnerLibrary);
+				MapShaderLibraryFile(MetalLibraryFilePath, MemOwnerLibrary, GPersistentMappingMetallib ? EMappedFileFlags::EPersistentMap: EMappedFileFlags::ENone);
 				if (MemOwnerLibrary->MappedRegion.IsValid())
 				{
 					UE_LOG(LogMetal, Display, TEXT("mmapping %s, %d bytes"), *MetalLibraryFilePath, MemOwnerLibrary->MappedCacheFile->GetFileSize());
 
-					dispatch_data_t data = dispatch_data_create(MemOwnerLibrary->MappedRegion->GetMappedPtr(), MemOwnerLibrary->MappedRegion->GetMappedSize(), nil, DISPATCH_DATA_DESTRUCTOR_MUNMAP);
+					dispatch_data_t data = dispatch_data_create(MemOwnerLibrary->MappedRegion->GetMappedPtr(), MemOwnerLibrary->MappedRegion->GetMappedSize(), nil, ^{});
 					MemOwnerLibrary->Library = NS::TransferPtr(Device->GetDevice()->newLibrary(data, &Error));
 					dispatch_release(data);
 				}
@@ -230,7 +237,7 @@ FRHIShaderLibraryRef FMetalDynamicRHI::RHICreateShaderLibrary(EShaderPlatform Pl
 					if (FFileHelper::LoadFileToArray(FileData, *MetalLibraryFilePath))
 					{
 						UE_LOG(LogMetal, Display, TEXT("emulating mmapping %s, %d bytes!"), *MetalLibraryFilePath, FileData.Num());
-						dispatch_data_t data = dispatch_data_create(FileData.GetData(), FileData.Num(), nil, DISPATCH_DATA_DESTRUCTOR_MUNMAP);
+						dispatch_data_t data = dispatch_data_create(FileData.GetData(), FileData.Num(), nil, ^{});
 						MemOwnerLibrary->Library = NS::TransferPtr(Device->GetDevice()->newLibrary(data, &Error));
 						dispatch_release(data);
 					}
diff --git a/Engine/Source/Runtime/Core/Private/Apple/ApplePlatformFile.cpp b/Engine/Source/Runtime/Core/Private/Apple/ApplePlatformFile.cpp
--- a/Engine/Source/Runtime/Core/Private/Apple/ApplePlatformFile.cpp
+++ b/Engine/Source/Runtime/Core/Private/Apple/ApplePlatformFile.cpp
@@ -524,7 +524,7 @@ class FMappedFileHandle final : public IMappedFileHandle
 	FString Filename;
 	std::atomic<int32> NumOutstandingRegions;
 	int FileHandle;
-	
+	bool bKeepMap;
 public:
 	
 	FMappedFileHandle(int InFileHandle, int64 FileSize, const FString& InFilename)
@@ -535,6 +535,7 @@ public:
 #endif
 		, NumOutstandingRegions(0)
 		, FileHandle(InFileHandle)
+		, bKeepMap(false)
 	{
 	}
 
@@ -589,6 +590,11 @@ public:
 		const uint8* MapPtr = AlignedMapPtr + Offset - AlignedOffset;
 		FMappedFileRegion* Result = new FMappedFileRegion(MapPtr, AlignedMapPtr, BytesToMap, AlignedSize, Filename, Offset, this);
 
+		if (EnumHasAnyFlags(Flags.Flags, EMappedFileFlags::EPersistentMap))
+		{
+			bKeepMap = true;
+		}
+		
 		NumOutstandingRegions.fetch_add(1, std::memory_order_relaxed);
 
 		return Result;
@@ -601,9 +607,12 @@ public:
 		int32 OldNumOutstandingRegions = NumOutstandingRegions.fetch_sub(1, std::memory_order_relaxed);
 		check(OldNumOutstandingRegions > 0);	
 		
-		LLM_IF_ENABLED(FLowLevelMemTracker::Get().OnLowLevelFree(ELLMTracker::Platform, (void*)Region->AlignedPtr));
-		int Res = munmap((void*)Region->AlignedPtr, Region->AlignedSize);
-		checkf(Res == 0, TEXT("Failed to unmap, error is %d, errno is %d [params: %x, %d]"), Res, errno, MappedPtr, GetFileSize());
+		if (!bKeepMap)
+		{
+			LLM_IF_ENABLED(FLowLevelMemTracker::Get().OnLowLevelFree(ELLMTracker::Platform, (void*)Region->AlignedPtr));
+			int Res = munmap((void*)Region->AlignedPtr, Region->AlignedSize);
+			checkf(Res == 0, TEXT("Failed to unmap, error is %d, errno is %d [params: %x, %d]"), Res, errno, MappedPtr, GetFileSize());
+		}
 	}
 };
 
diff --git a/Engine/Source/Runtime/Core/Public/Async/MappedFileHandle.h b/Engine/Source/Runtime/Core/Public/Async/MappedFileHandle.h
--- a/Engine/Source/Runtime/Core/Public/Async/MappedFileHandle.h
+++ b/Engine/Source/Runtime/Core/Public/Async/MappedFileHandle.h
@@ -97,6 +97,7 @@ enum class EMappedFileFlags
 	ENone = 0,				// Do nothing
 	EPreloadHint = 1,		// Preload the data. This is only a hint and might be ignored, see IMappedFileRegion::PreloadHint
 	EFileWritable = 2,		// Make the mapped file writable. Requires OpenMappedEx to be called with OpenOptions set to EOpenReadFlags::AllowWrite.
+	EPersistentMap = 4,		// On some platforms this flag will make it so we wont unmap when IMappedFileRegion is destroyed
 							// This will create a shared mapping on Unix platform to allow any writes to be flushed to disk
 };
 
