- Fixed a leak with Metal command buffers not being released when used in submission thread.
- Removed usage of MTLCommandBufferPtr as lifetimes are finite and understood

diff --git a/Engine/Source/Runtime/Apple/MetalRHI/Private/MetalCommandBuffer.h b/Engine/Source/Runtime/Apple/MetalRHI/Private/MetalCommandBuffer.h
--- a/Engine/Source/Runtime/Apple/MetalRHI/Private/MetalCommandBuffer.h
+++ b/Engine/Source/Runtime/Apple/MetalRHI/Private/MetalCommandBuffer.h
@@ -21,10 +21,16 @@ public:
 	: EventStream(Queue.GetProfilerQueue())
 #endif
     {
-        CommandBuffer = NS::RetainPtr(InCommandBuffer);
+        CommandBuffer = InCommandBuffer;
+		CommandBuffer->retain();
     }
+	
+	~FMetalCommandBuffer()
+	{
+		CommandBuffer->release();
+	}
     
-    FORCEINLINE MTLCommandBufferPtr& GetMTLCmdBuffer() {return CommandBuffer;}
+    FORCEINLINE MTL::CommandBuffer*& GetMTLCmdBuffer() {return CommandBuffer;}
 	
 #if RHI_NEW_GPU_PROFILER == 0
 	FORCEINLINE void SetTimer(FMetalCommandBufferTimer* InTimer)
@@ -142,6 +148,6 @@ public:
 #endif
 	
 private:
-    MTLCommandBufferPtr CommandBuffer;	
+    MTL::CommandBuffer* CommandBuffer;	
 	FMetalCommandBufferTimer* Timer = nullptr;
 };
diff --git a/Engine/Source/Runtime/Apple/MetalRHI/Private/MetalCommandQueue.cpp b/Engine/Source/Runtime/Apple/MetalRHI/Private/MetalCommandQueue.cpp
--- a/Engine/Source/Runtime/Apple/MetalRHI/Private/MetalCommandQueue.cpp
+++ b/Engine/Source/Runtime/Apple/MetalRHI/Private/MetalCommandQueue.cpp
@@ -108,6 +108,8 @@ FMetalCommandQueue::~FMetalCommandQueue(void)
 
 FMetalCommandBuffer* FMetalCommandQueue::CreateCommandBuffer(void)
 {
+	MTL_SCOPED_AUTORELEASE_POOL;
+	
 #if PLATFORM_MAC
 	static bool bUnretainedRefs = FParse::Param(FCommandLine::Get(),TEXT("metalunretained"))
 	|| (!FParse::Param(FCommandLine::Get(),TEXT("metalretainrefs"))
@@ -122,12 +124,11 @@ FMetalCommandBuffer* FMetalCommandQueue::CreateCommandBuffer(void)
     CmdBufferDesc->setRetainedReferences(!bUnretainedRefs);
     CmdBufferDesc->setErrorOptions(GMetalCommandBufferDebuggingEnabled ? MTL::CommandBufferErrorOptionEncoderExecutionStatus : MTL::CommandBufferErrorOptionNone);
     
-    MTL::CommandBuffer* CmdBuffer = CommandQueue->commandBuffer(CmdBufferDesc);
-    
-    CmdBufferDesc->release();
-                                                           
+	MTL::CommandBuffer* CmdBuffer = CommandQueue->commandBuffer(CmdBufferDesc);
     FMetalCommandBuffer* CommandBuffer = new FMetalCommandBuffer(CmdBuffer, *this);
     
+	CmdBufferDesc->release();
+	
 	INC_DWORD_STAT(STAT_MetalCommandBufferCreatedPerFrame);
 	return CommandBuffer;
 }
@@ -144,26 +145,6 @@ void FMetalCommandQueue::CommitCommandBuffer(FMetalCommandBuffer* CommandBuffer)
 	{
 		CommandBuffer->GetMTLCmdBuffer()->waitUntilCompleted();
 	}
-	
-	if(IsInRHIThread() || !IsRunningRHIInSeparateThread())
-	{
-		FMetalDynamicRHI::Get().DeferredDelete([CommandBuffer]() {
-			delete CommandBuffer;
-		});
-	}
-	else
-	{
-		// Deferred deletes need to run on the RHI thread
-		FFunctionGraphTask::CreateAndDispatchWhenReady(
-		   [CommandBuffer]()
-		   {
-			   FMetalDynamicRHI::Get().DeferredDelete([CommandBuffer]() {
-				   delete CommandBuffer;
-			   });
-		   },
-		   QUICK_USE_CYCLE_STAT(FExecuteRHIThreadTask, STATGROUP_TaskGraphTasks), nullptr, ENamedThreads::RHIThread);	
-	}
-		
 }
 
 FMetalFence* FMetalCommandQueue::CreateFence(NS::String* Label) const
diff --git a/Engine/Source/Runtime/Apple/MetalRHI/Private/MetalProfiler.cpp b/Engine/Source/Runtime/Apple/MetalRHI/Private/MetalProfiler.cpp
--- a/Engine/Source/Runtime/Apple/MetalRHI/Private/MetalProfiler.cpp
+++ b/Engine/Source/Runtime/Apple/MetalRHI/Private/MetalProfiler.cpp
@@ -393,11 +393,11 @@ FString IMetalStatsScope::GetJSONRepresentation(uint32 Pid)
 	return JSONOutput;
 }
 
-FMetalCommandBufferStats::FMetalCommandBufferStats(MTLCommandBufferPtr CommandBuffer, uint64 InGPUThreadIndex)
+FMetalCommandBufferStats::FMetalCommandBufferStats(MTL::CommandBuffer* CommandBuffer, uint64 InGPUThreadIndex)
 {
 	CmdBuffer = CommandBuffer;
 	
-	Name = FString::Printf(TEXT("CommandBuffer: %p %s"), CmdBuffer.get(), *NSStringToFString(CmdBuffer->label()));
+	Name = FString::Printf(TEXT("CommandBuffer: %p %s"), CmdBuffer, *NSStringToFString(CmdBuffer->label()));
 	
 	CPUThreadIndex = FPlatformTLS::GetCurrentThreadId();
 	GPUThreadIndex = InGPUThreadIndex;
@@ -409,7 +409,7 @@ FMetalCommandBufferStats::~FMetalCommandBufferStats()
 {
 }
 
-void FMetalCommandBufferStats::Start(MTLCommandBufferPtr& CommandBuffer)
+void FMetalCommandBufferStats::Start(MTL::CommandBuffer*& CommandBuffer)
 {
 	CPUStartTime = FPlatformTime::ToMilliseconds64(mach_absolute_time()) * 1000.0;
 	CPUEndTime = 0;
@@ -418,7 +418,7 @@ void FMetalCommandBufferStats::Start(MTLCommandBufferPtr& CommandBuffer)
 	GPUEndTime = 0;
 }
 
-void FMetalCommandBufferStats::End(MTLCommandBufferPtr& CommandBuffer)
+void FMetalCommandBufferStats::End(MTL::CommandBuffer*& CommandBuffer)
 {
 	check(CommandBuffer == CmdBuffer);
 	
@@ -474,10 +474,10 @@ FMetalDisplayStats::~FMetalDisplayStats()
 {
 }
 
-void FMetalDisplayStats::Start(MTLCommandBufferPtr& Buffer)
+void FMetalDisplayStats::Start(MTL::CommandBuffer*& Buffer)
 {
 }
-void FMetalDisplayStats::End(MTLCommandBufferPtr& Buffer)
+void FMetalDisplayStats::End(MTL::CommandBuffer*& Buffer)
 {
 }
 
@@ -510,11 +510,11 @@ void FMetalCPUStats::End(void)
 	CPUEndTime = FPlatformTime::ToMilliseconds64(mach_absolute_time()) * 1000.0;
 }
 
-void FMetalCPUStats::Start(MTLCommandBufferPtr& Buffer)
+void FMetalCPUStats::Start(MTL::CommandBuffer*& Buffer)
 {
 	
 }
-void FMetalCPUStats::End(MTLCommandBufferPtr& Buffer)
+void FMetalCPUStats::End(MTL::CommandBuffer*& Buffer)
 {
 	
 }
@@ -685,7 +685,7 @@ FMetalCPUStats* FMetalProfiler::AddCPUStat(FString const& Name)
 	}
 }
 
-FMetalCommandBufferStats* FMetalProfiler::AllocateCommandBuffer(MTLCommandBufferPtr Buffer, uint64 GPUThreadIndex)
+FMetalCommandBufferStats* FMetalProfiler::AllocateCommandBuffer(MTL::CommandBuffer* Buffer, uint64 GPUThreadIndex)
 {
 	return new FMetalCommandBufferStats(Buffer, GPUThreadIndex);
 }
diff --git a/Engine/Source/Runtime/Apple/MetalRHI/Private/MetalProfiler.h b/Engine/Source/Runtime/Apple/MetalRHI/Private/MetalProfiler.h
--- a/Engine/Source/Runtime/Apple/MetalRHI/Private/MetalProfiler.h
+++ b/Engine/Source/Runtime/Apple/MetalRHI/Private/MetalProfiler.h
@@ -240,8 +240,8 @@ struct IMetalStatsScope
 	
 	virtual ~IMetalStatsScope();
 	
-	virtual void Start(MTLCommandBufferPtr& CommandBuffer) = 0;
-	virtual void End(MTLCommandBufferPtr& CommandBuffer) = 0;
+	virtual void Start(MTL::CommandBuffer*& CommandBuffer) = 0;
+	virtual void End(MTL::CommandBuffer*& CommandBuffer) = 0;
 	
 	FString GetJSONRepresentation(uint32 PID);
 };
@@ -254,8 +254,8 @@ struct FMetalCPUStats : public IMetalStatsScope
 	void Start(void);
 	void End(void);
 	
-	virtual void Start(MTLCommandBufferPtr& CommandBuffer) final override;
-	virtual void End(MTLCommandBufferPtr& CommandBuffer) final override;
+	virtual void Start(MTL::CommandBuffer*& CommandBuffer) final override;
+	virtual void End(MTL::CommandBuffer*& CommandBuffer) final override;
 };
 
 struct FMetalDisplayStats : public IMetalStatsScope
@@ -263,8 +263,8 @@ struct FMetalDisplayStats : public IMetalStatsScope
 	FMetalDisplayStats(uint32 DisplayID, double OutputSeconds, double Duration);
 	virtual ~FMetalDisplayStats();
 	
-	virtual void Start(MTLCommandBufferPtr& CommandBuffer) final override;
-	virtual void End(MTLCommandBufferPtr& CommandBuffer) final override;
+	virtual void Start(MTL::CommandBuffer*& CommandBuffer) final override;
+	virtual void End(MTL::CommandBuffer*& CommandBuffer) final override;
 };
 
 enum EMTLFenceType
@@ -275,13 +275,13 @@ enum EMTLFenceType
 
 struct FMetalCommandBufferStats : public IMetalStatsScope
 {
-	FMetalCommandBufferStats(MTLCommandBufferPtr CommandBuffer, uint64 GPUThreadIndex);
+	FMetalCommandBufferStats(MTL::CommandBuffer* CommandBuffer, uint64 GPUThreadIndex);
 	virtual ~FMetalCommandBufferStats();
 	
-	virtual void Start(MTLCommandBufferPtr& CommandBuffer) final override;
-	virtual void End(MTLCommandBufferPtr& CommandBuffer) final override;
+	virtual void Start(MTL::CommandBuffer*& CommandBuffer) final override;
+	virtual void End(MTL::CommandBuffer*& CommandBuffer) final override;
 
-	MTLCommandBufferPtr CmdBuffer;
+	MTL::CommandBuffer* CmdBuffer;
 };
 
 class FMetalRHICommandContext;
@@ -342,7 +342,7 @@ public:
 	void EncodeDispatch(FMetalCommandBufferStats* CmdBufStats, char const* DrawCall);
 	
 	FMetalCPUStats* AddCPUStat(FString const& Name);
-	FMetalCommandBufferStats* AllocateCommandBuffer(MTLCommandBufferPtr CommandBuffer, uint64 GPUThreadIndex);
+	FMetalCommandBufferStats* AllocateCommandBuffer(MTL::CommandBuffer* CommandBuffer, uint64 GPUThreadIndex);
 	void AddCommandBuffer(FMetalCommandBufferStats* CommandBuffer);
 	virtual void PushEvent(const TCHAR* Name, FColor Color) final override;
 	virtual void PopEvent() final override;
diff --git a/Engine/Source/Runtime/Apple/MetalRHI/Private/MetalSubmission.cpp b/Engine/Source/Runtime/Apple/MetalRHI/Private/MetalSubmission.cpp
--- a/Engine/Source/Runtime/Apple/MetalRHI/Private/MetalSubmission.cpp
+++ b/Engine/Source/Runtime/Apple/MetalRHI/Private/MetalSubmission.cpp
@@ -483,6 +483,10 @@ void FMetalDynamicRHI::FlushBatchedPayloads(FMetalCommandQueue::FPayloadArray& P
 			FMetalCommandBuffer* CurrentCommandBuffer = Queue.CreateCommandBuffer();
 			CurrentCommandBuffer->GetMTLCmdBuffer()->encodeWait(LocalFence.MetalEvent, Value);
 			Queue.CommitCommandBuffer(CurrentCommandBuffer);
+			
+			DeferredDelete([CurrentCommandBuffer]() {
+				delete CurrentCommandBuffer;
+			});
 		}
 	};
 
@@ -672,7 +676,7 @@ void FMetalDynamicRHI::ProcessPendingCommandBuffers()
 	
 	CmdBuffersPendingCompletion.RemoveAll([SecondsPerCycle, &bContinueProcessing](FMetalCommandBuffer* CommandBuffer) 
 	{ 
-		MTLCommandBufferPtr CompletedBuffer = CommandBuffer->GetMTLCmdBuffer();
+		MTL::CommandBuffer* CompletedBuffer = CommandBuffer->GetMTLCmdBuffer();
 		
 		MTL::CommandBufferStatus Status = CompletedBuffer->status();
 		if (Status == MTL::CommandBufferStatusCompleted && bContinueProcessing)
@@ -705,6 +709,9 @@ void FMetalDynamicRHI::ProcessPendingCommandBuffers()
 			FMetalCommandBufferTimer* Timer = CommandBuffer->GetTimer();
 			Timer->AddTiming({CompletedBuffer->GPUStartTime(), CompletedBuffer->GPUEndTime()});
 #endif
+			
+			delete CommandBuffer;
+			
 			return true;
 		}
 		else
@@ -760,10 +767,10 @@ FMetalDynamicRHI::FProcessResult FMetalDynamicRHI::ProcessInterruptQueue()
 			// Handle command buffer errors
 			for(FMetalCommandBuffer* CommandBuffer : Payload->CommandBuffersToExecute)
 			{	
-				MTLCommandBufferPtr CompletedBuffer = CommandBuffer->GetMTLCmdBuffer();
+				MTL::CommandBuffer* CompletedBuffer = CommandBuffer->GetMTLCmdBuffer();
 				if (CompletedBuffer->status() == MTL::CommandBufferStatusError)
 				{
-					FMetalCommandList::HandleMetalCommandBufferFailure(CompletedBuffer.get());
+					FMetalCommandList::HandleMetalCommandBufferFailure(CompletedBuffer);
 				}
 			}
 			
@@ -771,11 +778,11 @@ FMetalDynamicRHI::FProcessResult FMetalDynamicRHI::ProcessInterruptQueue()
 			
 			if(Payload->SignalCommandBuffer)
 			{
-				MTLCommandBufferPtr SignalBuffer = Payload->SignalCommandBuffer->GetMTLCmdBuffer(); 
+				MTL::CommandBuffer* SignalBuffer = Payload->SignalCommandBuffer->GetMTLCmdBuffer(); 
 				Status = SignalBuffer->status();
 				if(Status == MTL::CommandBufferStatusError)
 				{
-					FMetalCommandList::HandleMetalCommandBufferFailure(SignalBuffer.get());
+					FMetalCommandList::HandleMetalCommandBufferFailure(SignalBuffer);
 				}
 			}
 			
@@ -786,19 +793,12 @@ FMetalDynamicRHI::FProcessResult FMetalDynamicRHI::ProcessInterruptQueue()
 				Result.Status |= EQueueStatus::Pending;
 				break;
 			}
-			
-			// Record the command buffer timing
-			for(FMetalCommandBuffer* CommandBuffer : Payload->CommandBuffersToExecute)
-			{	
-				CmdBuffersPendingCompletion.Add(CommandBuffer);
+
+			if(Payload->SignalCommandBuffer)
+			{
+				delete Payload->SignalCommandBuffer;
 			}
 			
-			ProcessPendingCommandBuffers();
-			
-			// At this point, the current command list has completed on the GPU.
-			CurrentQueue.PendingInterrupt.Pop();
-			Result.Status |= EQueueStatus::Processed;
-			
 #if RHI_NEW_GPU_PROFILER
 			if (!Payload->EventStream.IsEmpty())
 			{			
@@ -821,7 +821,7 @@ FMetalDynamicRHI::FProcessResult FMetalDynamicRHI::ProcessInterruptQueue()
 				
 				for (FMetalRHIRenderQuery* TimestampQuery : Payload->BatchedObjects.TimestampQueries)
 				{
-					MTLCommandBufferPtr CmdBuffer = TimestampQuery->CommandBuffer->GetMTLCmdBuffer();
+					MTL::CommandBuffer* CmdBuffer = TimestampQuery->CommandBuffer->GetMTLCmdBuffer();
 					
 					// If there are no commands in the command buffer then this can be zero
 					// In this case GPU start time is also not correct - we need to fall back standard behaviour
@@ -870,7 +870,17 @@ FMetalDynamicRHI::FProcessResult FMetalDynamicRHI::ProcessInterruptQueue()
 			}
 
 			// We're done with this payload now.
-
+			for(FMetalCommandBuffer* CommandBuffer : Payload->CommandBuffersToExecute)
+			{	
+				CmdBuffersPendingCompletion.Add(CommandBuffer);
+			}
+			
+			ProcessPendingCommandBuffers();
+			
+			// At this point, the current command list has completed on the GPU.
+			CurrentQueue.PendingInterrupt.Pop();
+			Result.Status |= EQueueStatus::Processed;
+			
 			// GPU resources the payload is holding a reference to will be cleaned up here.
 			// E.g. command list allocators, which get recycled on the parent device.
 			delete Payload;
