Move Metal Library creation to CreateShader when it is first requested
Modify the Persistent Metallib workaround to be more isolated

diff --git a/Engine/Source/Runtime/Apple/MetalRHI/Private/MetalDynamicRHI_Shaders.cpp b/Engine/Source/Runtime/Apple/MetalRHI/Private/MetalDynamicRHI_Shaders.cpp
--- a/Engine/Source/Runtime/Apple/MetalRHI/Private/MetalDynamicRHI_Shaders.cpp
+++ b/Engine/Source/Runtime/Apple/MetalRHI/Private/MetalDynamicRHI_Shaders.cpp
@@ -13,13 +13,6 @@
 #include "Interfaces/IPluginManager.h"
 
 //------------------------------------------------------------------------------
-static bool GPersistentMappingMetallib = 1;
-static FAutoConsoleVariableRef CVarPersistentMappingMetallib(
-	TEXT("r.Metal.PersistentMappingMetallib"),
-	GPersistentMappingMetallib,
-	TEXT("Makes the metallib file mapping persistent."),
-	ECVF_Default
-);
 
 #pragma mark - Metal Dynamic RHI Shader Methods
 
@@ -158,7 +151,7 @@ FRHIShaderLibraryRef FMetalDynamicRHI::RHICreateShaderLibrary(EShaderPlatform Pl
 #endif
 	};
 
-	auto MapShaderLibraryFile = [](const FString& FilePath, TUniquePtr<FMetalShaderLibrary::FShaderLibDataOwner>& MemOwner, FFileMappingFlags Flags)
+	auto MapShaderLibraryFile = [](const FString& FilePath, FMetalShaderLibrary::FShaderLibDataOwner* MemOwner)
 	{
 		if (FPlatformProperties::SupportsMemoryMappedFiles())
 		{
@@ -166,7 +159,7 @@ FRHIShaderLibraryRef FMetalDynamicRHI::RHICreateShaderLibrary(EShaderPlatform Pl
 			MemOwner->MappedCacheFile = Res.HasError() ? nullptr : Res.StealValue();
 			if (MemOwner->MappedCacheFile.IsValid())
 			{
-				MemOwner->MappedRegion = TUniquePtr<IMappedFileRegion>(MemOwner->MappedCacheFile->MapRegion(0, MemOwner->MappedCacheFile->GetFileSize(), Flags));
+				MemOwner->MappedRegion = TUniquePtr<IMappedFileRegion>(MemOwner->MappedCacheFile->MapRegion(0, MemOwner->MappedCacheFile->GetFileSize()));
 			}
 		}
 	};
@@ -176,7 +169,7 @@ FRHIShaderLibraryRef FMetalDynamicRHI::RHICreateShaderLibrary(EShaderPlatform Pl
 #else
 	FStaticMemoryReader* BinaryShaderAr = nullptr;
     TUniquePtr<FMetalShaderLibrary::FShaderLibDataOwner> MemOwner = MakeUnique<FMetalShaderLibrary::FShaderLibDataOwner>();
-	MapShaderLibraryFile(BinaryShaderFile, MemOwner, EMappedFileFlags::ENone);
+	MapShaderLibraryFile(BinaryShaderFile, MemOwner.Get());
 	if (MemOwner->MappedRegion.IsValid())
 	{
 		UE_LOG(LogMetal, Display, TEXT("mmapping %s, %d bytes"), *BinaryShaderFile, MemOwner->MappedCacheFile->GetFileSize());
@@ -210,48 +203,19 @@ FRHIShaderLibraryRef FMetalDynamicRHI::RHICreateShaderLibrary(EShaderPlatform Pl
         {
             check(((SerializedShaders.GetNumShaders() + Header.NumShadersPerLibrary - 1) / Header.NumShadersPerLibrary) == Header.NumLibraries);
 
-			TArray<TUniquePtr<FMetalShaderLibrary::FShaderLibDataOwner>> LibraryMemOwner;
+			TArray<TUniquePtr<FMetalShaderLibrary::FMetalLibDataOwner>> LibraryMemOwner;
 
 			LibraryMemOwner.Empty(Header.NumLibraries);
 
             for (uint32 i = 0; i < Header.NumLibraries; i++)
             {
 				FString MetalLibraryFilePath = (FilePath / LibName) + FString::Printf(TEXT(".%d.metallib"), i);
-				NS::Error* Error;
-
-				MTLLibraryPtr Library;
-				TUniquePtr<FMetalShaderLibrary::FShaderLibDataOwner> MemOwnerLibrary = MakeUnique<FMetalShaderLibrary::FShaderLibDataOwner>();
-				MapShaderLibraryFile(MetalLibraryFilePath, MemOwnerLibrary, GPersistentMappingMetallib ? EMappedFileFlags::EPersistentMap: EMappedFileFlags::ENone);
-				if (MemOwnerLibrary->MappedRegion.IsValid())
-				{
-					UE_LOG(LogMetal, Display, TEXT("mmapping %s, %d bytes"), *MetalLibraryFilePath, MemOwnerLibrary->MappedCacheFile->GetFileSize());
-
-					dispatch_data_t data = dispatch_data_create(MemOwnerLibrary->MappedRegion->GetMappedPtr(), MemOwnerLibrary->MappedRegion->GetMappedSize(), nil, ^{});
-					MemOwnerLibrary->Library = NS::TransferPtr(Device->GetDevice()->newLibrary(data, &Error));
-					dispatch_release(data);
-				}
-
-				if (MemOwnerLibrary->Library.get() == nullptr)
-				{
-					TArray<uint8>& FileData = MemOwnerLibrary->Mem;
-					if (FFileHelper::LoadFileToArray(FileData, *MetalLibraryFilePath))
-					{
-						UE_LOG(LogMetal, Display, TEXT("emulating mmapping %s, %d bytes!"), *MetalLibraryFilePath, FileData.Num());
-						dispatch_data_t data = dispatch_data_create(FileData.GetData(), FileData.Num(), nil, ^{});
-						MemOwnerLibrary->Library = NS::TransferPtr(Device->GetDevice()->newLibrary(data, &Error));
-						dispatch_release(data);
-					}
-				}
-
-				if ((MemOwnerLibrary != nullptr) && (MemOwnerLibrary->Library.get() != nullptr))
-				{
-					LibraryMemOwner.Add(MoveTemp(MemOwnerLibrary));
-				}
-				else
-				{
-					UE_LOG(LogMetal, Display, TEXT("Failed to create library: %s"), *MetalLibraryFilePath);
-					return nullptr;
-				}
+
+				TUniquePtr<FMetalShaderLibrary::FMetalLibDataOwner> MemOwnerLibrary = MakeUnique<FMetalShaderLibrary::FMetalLibDataOwner>();
+				MemOwnerLibrary->MetalLibraryFilePath = MetalLibraryFilePath;
+				MapShaderLibraryFile(MetalLibraryFilePath, MemOwnerLibrary.Get());
+				
+				LibraryMemOwner.Add(MoveTemp(MemOwnerLibrary));
             }
 
             FMetalShaderLibrary* MtlLib = new FMetalShaderLibrary(*Device, Platform, Name, BinaryShaderFile, Header, MoveTemp(SerializedShaders), MoveTemp(ShaderCode), MoveTemp(LibraryMemOwner)
diff --git a/Engine/Source/Runtime/Apple/MetalRHI/Private/Shaders/MetalShaderLibrary.cpp b/Engine/Source/Runtime/Apple/MetalRHI/Private/Shaders/MetalShaderLibrary.cpp
--- a/Engine/Source/Runtime/Apple/MetalRHI/Private/Shaders/MetalShaderLibrary.cpp
+++ b/Engine/Source/Runtime/Apple/MetalRHI/Private/Shaders/MetalShaderLibrary.cpp
@@ -12,6 +12,13 @@
 #endif // !UE_BUILD_SHIPPING
 #include "MetalShaderTypes.h"
 
+static bool GPersistentMappingMetallib = 1;
+static FAutoConsoleVariableRef CVarPersistentMappingMetallib(
+	TEXT("r.Metal.PersistentMappingMetallib"),
+	GPersistentMappingMetallib,
+	TEXT("Makes the metallib file mapping persistent."),
+	ECVF_Default
+);
 
 //------------------------------------------------------------------------------
 
@@ -44,7 +51,15 @@ TMap<FString, FRHIShaderLibrary*> FMetalShaderLibrary::LoadedShaderLibraryMap;
 //------------------------------------------------------------------------------
 
 #pragma mark - Metal Shader Library Class
-
+FMetalShaderLibrary::FMetalLibDataOwner::~FMetalLibDataOwner()
+{
+	// Intentionally leak the metallib file handle to workaround driver bug
+	if (GPersistentMappingMetallib > 0)
+	{
+		MappedCacheFile.Release();
+		MappedRegion.Release();
+	}
+}
 
 FMetalShaderLibrary::FMetalShaderLibrary(FMetalDevice& MetalDevice,
 										 EShaderPlatform Platform,
@@ -53,7 +68,7 @@ FMetalShaderLibrary::FMetalShaderLibrary(FMetalDevice& MetalDevice,
 										 const FMetalShaderLibraryHeader& InHeader,
 										 FSerializedShaderArchive&& InSerializedShaders,
 										 FShaderCodeArrayType&& InShaderCode,
-										 TArray<TUniquePtr<FShaderLibDataOwner>>&& InLibraryMemOwner)
+										 TArray<TUniquePtr<FMetalLibDataOwner>>&& InLibraryMemOwner)
 	: FRHIShaderLibrary(Platform, Name)
 	, Device(MetalDevice)
 	, ShaderLibraryFilename(InShaderLibraryFilename)
@@ -138,6 +153,33 @@ TRefCountPtr<FRHIShader> FMetalShaderLibrary::CreateShader(int32 Index, bool bRe
 	const TArrayView<const uint8> Code = MakeArrayView(ShaderCode.GetData() + ShaderEntry.Offset, ShaderEntry.Size);
 	const int32 LibraryIndex = Index / Header.NumShadersPerLibrary;
 
+	{
+		TUniquePtr<FMetalLibDataOwner>& LibraryMemOwnerRef = LibraryMemOwner[LibraryIndex];
+		UE::TScopeLock Lock(LibraryMemOwnerRef->LibraryLock);
+		if (LibraryMemOwnerRef->Library.get() == nullptr)
+		{
+			NS::Error* Error;
+			if (LibraryMemOwnerRef->MappedRegion.IsValid())
+			{
+				UE_LOG(LogMetal, Display, TEXT("mmapping %s, %d bytes"), *LibraryMemOwnerRef->MetalLibraryFilePath, LibraryMemOwnerRef->MappedCacheFile->GetFileSize());
+				dispatch_data_t data = dispatch_data_create(LibraryMemOwnerRef->MappedRegion->GetMappedPtr(), LibraryMemOwnerRef->MappedRegion->GetMappedSize(), nil, ^ {});
+				LibraryMemOwnerRef->Library = NS::TransferPtr(Device.GetDevice()->newLibrary(data, &Error));
+				dispatch_release(data);
+			}
+			else
+			{
+				TArray<uint8>& FileData = LibraryMemOwnerRef->Mem;
+				if (FFileHelper::LoadFileToArray(FileData, *LibraryMemOwnerRef->MetalLibraryFilePath))
+				{
+					UE_LOG(LogMetal, Display, TEXT("emulating mmapping %s, %d bytes!"), *LibraryMemOwnerRef->MetalLibraryFilePath, FileData.Num());
+					dispatch_data_t data = dispatch_data_create(FileData.GetData(), FileData.Num(), nil, ^ {});
+					LibraryMemOwnerRef->Library = NS::TransferPtr(Device.GetDevice()->newLibrary(data, &Error));
+					dispatch_release(data);
+				}
+			}
+		}
+	}
+
 	TRefCountPtr<FRHIShader> Shader;
 	switch (ShaderEntry.Frequency)
 	{
diff --git a/Engine/Source/Runtime/Apple/MetalRHI/Private/Shaders/MetalShaderLibrary.h b/Engine/Source/Runtime/Apple/MetalRHI/Private/Shaders/MetalShaderLibrary.h
--- a/Engine/Source/Runtime/Apple/MetalRHI/Private/Shaders/MetalShaderLibrary.h
+++ b/Engine/Source/Runtime/Apple/MetalRHI/Private/Shaders/MetalShaderLibrary.h
@@ -23,7 +23,15 @@ public:
 		TArray<uint8> Mem;
 		TUniquePtr<IMappedFileHandle> MappedCacheFile;
 		TUniquePtr<IMappedFileRegion> MappedRegion;
+	};
+
+	struct FMetalLibDataOwner : public FShaderLibDataOwner
+	{
 		MTLLibraryPtr Library;
+		FString MetalLibraryFilePath;
+		UE::FMutex LibraryLock;
+
+		~FMetalLibDataOwner();
 	};
 #if !USE_MMAPPED_SHADERARCHIVE
     using FShaderCodeArrayType = TArray<uint8>;
@@ -40,7 +48,7 @@ public:
         const FMetalShaderLibraryHeader& InHeader,
         FSerializedShaderArchive&& InSerializedShaders,
         FShaderCodeArrayType&& InShaderCode,
-		TArray<TUniquePtr<FShaderLibDataOwner>>&& InLibraryMemOwner,
+		TArray<TUniquePtr<FMetalLibDataOwner>>&& InLibraryMemOwner,
         TUniquePtr<FShaderLibDataOwner>&& InMemOwner)
         :
         FMetalShaderLibrary(Device, Platform, Name, InShaderLibraryFilename, InHeader, MoveTemp(InSerializedShaders), MoveTemp(InShaderCode), MoveTemp(InLibraryMemOwner))
@@ -59,7 +67,7 @@ public:
 						const FMetalShaderLibraryHeader& InHeader,
 						FSerializedShaderArchive&& InSerializedShaders,
 						FShaderCodeArrayType&& InShaderCode,
-						TArray<TUniquePtr<FShaderLibDataOwner>>&& InLibraryMemOwner);
+						TArray<TUniquePtr<FMetalLibDataOwner>>&& InLibraryMemOwner);
 
 	virtual ~FMetalShaderLibrary();
 
@@ -89,7 +97,7 @@ private:
 	FMetalShaderLibraryHeader Header;
 	FSerializedShaderArchive SerializedShaders;
 	FShaderCodeArrayType ShaderCode;
-	TArray<TUniquePtr<FShaderLibDataOwner>> LibraryMemOwner;
+	TArray<TUniquePtr<FMetalLibDataOwner>> LibraryMemOwner;
 #if !UE_BUILD_SHIPPING
 	class FMetalShaderDebugZipFile* DebugFile;
 #endif
diff --git a/Engine/Source/Runtime/Core/Private/Apple/ApplePlatformFile.cpp b/Engine/Source/Runtime/Core/Private/Apple/ApplePlatformFile.cpp
--- a/Engine/Source/Runtime/Core/Private/Apple/ApplePlatformFile.cpp
+++ b/Engine/Source/Runtime/Core/Private/Apple/ApplePlatformFile.cpp
@@ -524,7 +524,7 @@ class FMappedFileHandle final : public IMappedFileHandle
 	FString Filename;
 	std::atomic<int32> NumOutstandingRegions;
 	int FileHandle;
-	bool bKeepMap;
+	
 public:
 	
 	FMappedFileHandle(int InFileHandle, int64 FileSize, const FString& InFilename)
@@ -535,7 +535,6 @@ public:
 #endif
 		, NumOutstandingRegions(0)
 		, FileHandle(InFileHandle)
-		, bKeepMap(false)
 	{
 	}
 
@@ -590,11 +589,6 @@ public:
 		const uint8* MapPtr = AlignedMapPtr + Offset - AlignedOffset;
 		FMappedFileRegion* Result = new FMappedFileRegion(MapPtr, AlignedMapPtr, BytesToMap, AlignedSize, Filename, Offset, this);
 
-		if (EnumHasAnyFlags(Flags.Flags, EMappedFileFlags::EPersistentMap))
-		{
-			bKeepMap = true;
-		}
-		
 		NumOutstandingRegions.fetch_add(1, std::memory_order_relaxed);
 
 		return Result;
@@ -607,12 +601,9 @@ public:
 		int32 OldNumOutstandingRegions = NumOutstandingRegions.fetch_sub(1, std::memory_order_relaxed);
 		check(OldNumOutstandingRegions > 0);	
 		
-		if (!bKeepMap)
-		{
-			LLM_IF_ENABLED(FLowLevelMemTracker::Get().OnLowLevelFree(ELLMTracker::Platform, (void*)Region->AlignedPtr));
-			int Res = munmap((void*)Region->AlignedPtr, Region->AlignedSize);
-			checkf(Res == 0, TEXT("Failed to unmap, error is %d, errno is %d [params: %x, %d]"), Res, errno, MappedPtr, GetFileSize());
-		}
+		LLM_IF_ENABLED(FLowLevelMemTracker::Get().OnLowLevelFree(ELLMTracker::Platform, (void*)Region->AlignedPtr));
+		int Res = munmap((void*)Region->AlignedPtr, Region->AlignedSize);
+		checkf(Res == 0, TEXT("Failed to unmap, error is %d, errno is %d [params: %x, %d]"), Res, errno, MappedPtr, GetFileSize());
 	}
 };
 
diff --git a/Engine/Source/Runtime/Core/Public/Async/MappedFileHandle.h b/Engine/Source/Runtime/Core/Public/Async/MappedFileHandle.h
--- a/Engine/Source/Runtime/Core/Public/Async/MappedFileHandle.h
+++ b/Engine/Source/Runtime/Core/Public/Async/MappedFileHandle.h
@@ -97,7 +97,6 @@ enum class EMappedFileFlags
 	ENone = 0,				// Do nothing
 	EPreloadHint = 1,		// Preload the data. This is only a hint and might be ignored, see IMappedFileRegion::PreloadHint
 	EFileWritable = 2,		// Make the mapped file writable. Requires OpenMappedEx to be called with OpenOptions set to EOpenReadFlags::AllowWrite.
-	EPersistentMap = 4,		// On some platforms this flag will make it so we wont unmap when IMappedFileRegion is destroyed
 							// This will create a shared mapping on Unix platform to allow any writes to be flushed to disk
 };
 
