Fix for compile error on VisionOS

diff --git a/Engine/Source/Runtime/Apple/MetalRHI/Private/MetalViewport.cpp b/Engine/Source/Runtime/Apple/MetalRHI/Private/MetalViewport.cpp
--- a/Engine/Source/Runtime/Apple/MetalRHI/Private/MetalViewport.cpp
+++ b/Engine/Source/Runtime/Apple/MetalRHI/Private/MetalViewport.cpp
@@ -737,7 +737,7 @@ void FMetalViewport::PresentImmersive(const MetalRHIVisionOS::PresentImmersivePa
 					FinalCommandBuffer->GetMTLCmdBuffer()->addCompletedHandler(CommandBufferHandler);
 				}
 				
-				cp_drawable_encode_present(VisionOSParams.SwiftDrawable, (__bridge id<MTLCommandBuffer>)FinalCommandBuffer->GetMTLCmdBuffer().get());
+				cp_drawable_encode_present(VisionOSParams.SwiftDrawable, (__bridge id<MTLCommandBuffer>)FinalCommandBuffer->GetMTLCmdBuffer());
 				cp_frame_t CompositorServicesFrame = VisionOSParams.SwiftFrame;
 
 				METAL_GPUPROFILE(Stats->End(FinalCommandBuffer->GetMTLCmdBuffer()));			
