iOS fix audio interruption during foreground phone calls.

When iOS is iterrupted by a phone call, user can take the call while keeping the application in the foreground. During this time, the Audio mixer still runs, however, AudioUnit/AUGraph gets paused and stops dispatching callbacks. During these interruptions spawn a Null audio device to continue consuming the packets similarly to the necessary approach on XAudio2. To achieve this, add audio interruption delegate that the AudioUnit platform can listen to and adjust its behaviour accordingly.

Additionally, in AudioMixerNullDevice.cpp, ensure that  CallbackThread is initialized before the Null device thread is awakened so that it isn't dereferened on that thread prior to initialization.

diff --git a/Engine/Source/Runtime/Apple/AudioMixerAudioUnit/Private/AudioMixerPlatformAudioUnit.cpp b/Engine/Source/Runtime/Apple/AudioMixerAudioUnit/Private/AudioMixerPlatformAudioUnit.cpp
--- a/Engine/Source/Runtime/Apple/AudioMixerAudioUnit/Private/AudioMixerPlatformAudioUnit.cpp
+++ b/Engine/Source/Runtime/Apple/AudioMixerAudioUnit/Private/AudioMixerPlatformAudioUnit.cpp
@@ -302,7 +302,8 @@ namespace Audio
         
         FCoreDelegates::ApplicationWillDeactivateDelegate.AddRaw(this, &FMixerPlatformAudioUnit::SuspendContext);
         FCoreDelegates::ApplicationHasReactivatedDelegate.AddRaw(this, &FMixerPlatformAudioUnit::ResumeContext);
-		
+		FCoreDelegates::AudioInterruptionDelegate.AddRaw(this, &FMixerPlatformAudioUnit::InterruptContext);
+
 		return true;
 	}
 	
@@ -315,7 +316,8 @@ namespace Audio
 		
 		AudioStreamInfo.StreamState = EAudioOutputStreamState::Closed;
         
-        FCoreDelegates::ApplicationWillDeactivateDelegate.RemoveAll(this);
+		FCoreDelegates::AudioInterruptionDelegate.RemoveAll(this);
+		FCoreDelegates::ApplicationWillDeactivateDelegate.RemoveAll(this);
         FCoreDelegates::ApplicationHasReactivatedDelegate.RemoveAll(this);
 		
 		return true;
@@ -380,8 +382,7 @@ namespace Audio
 	
 	void FMixerPlatformAudioUnit::SubmitBuffer(const uint8* Buffer)
 	{
-		
-		if(!Buffer)
+		if (!Buffer || bIsUsingNullDevice)
 		{
 			return;
 		}
@@ -527,6 +528,18 @@ namespace Audio
 		}
 	}
 	
+	void FMixerPlatformAudioUnit::InterruptContext(bool Interrupted)
+	{
+		if (Interrupted)
+		{
+			StartRunningNullDevice();
+		}
+		else
+		{
+			StopRunningNullDevice();
+		}
+	}
+
 	bool FMixerPlatformAudioUnit::PerformCallback(AudioBufferList* OutputBufferData)
 	{
 		bInCallback = true;
diff --git a/Engine/Source/Runtime/Apple/AudioMixerAudioUnit/Private/AudioMixerPlatformAudioUnit.h b/Engine/Source/Runtime/Apple/AudioMixerAudioUnit/Private/AudioMixerPlatformAudioUnit.h
--- a/Engine/Source/Runtime/Apple/AudioMixerAudioUnit/Private/AudioMixerPlatformAudioUnit.h
+++ b/Engine/Source/Runtime/Apple/AudioMixerAudioUnit/Private/AudioMixerPlatformAudioUnit.h
@@ -86,6 +86,8 @@ namespace Audio
 		mutable bool bInternalPlatformSettingsInitialized{ false };
 		mutable FAudioPlatformSettings InternalPlatformSettings;
 		
+		void InterruptContext(bool Interrupted);
+
 		bool PerformCallback(AudioBufferList* OutputBufferData);
 		void HandleError(const TCHAR* InLogOutput, bool bTeardown = true);
 		static OSStatus AudioRenderCallback(void* RefCon, AudioUnitRenderActionFlags* ActionFlags,
diff --git a/Engine/Source/Runtime/ApplicationCore/Private/IOS/IOSAppDelegate.cpp b/Engine/Source/Runtime/ApplicationCore/Private/IOS/IOSAppDelegate.cpp
--- a/Engine/Source/Runtime/ApplicationCore/Private/IOS/IOSAppDelegate.cpp
+++ b/Engine/Source/Runtime/ApplicationCore/Private/IOS/IOSAppDelegate.cpp
@@ -565,9 +565,11 @@ static IOSAppDelegate* CachedDelegate = nil;
 		{
 			case AVAudioSessionInterruptionTypeBegan:
 				// Notification that our audio was stopped by the OS. No action needed.
+				FCoreDelegates::AudioInterruptionDelegate.Broadcast(true);
 				break;
 
 			case AVAudioSessionInterruptionTypeEnded:
+				FCoreDelegates::AudioInterruptionDelegate.Broadcast(false);
 				NSNumber * interruptionOption = [[notification userInfo] objectForKey:AVAudioSessionInterruptionOptionKey];
 				if (interruptionOption != nil && interruptionOption.unsignedIntegerValue > 0)
 				{
diff --git a/Engine/Source/Runtime/AudioMixerCore/Private/AudioMixerNullDevice.cpp b/Engine/Source/Runtime/AudioMixerCore/Private/AudioMixerNullDevice.cpp
--- a/Engine/Source/Runtime/AudioMixerCore/Private/AudioMixerNullDevice.cpp
+++ b/Engine/Source/Runtime/AudioMixerCore/Private/AudioMixerNullDevice.cpp
@@ -91,6 +91,8 @@ namespace Audio
 		// Make sure we're in a waitable start on startup.
 		SleepEvent->Reset();
 
+		CallbackThread.Reset(FRunnableThread::Create(this, TEXT("AudioMixerNullCallbackThread"), 0, ThreadPriority, FPlatformAffinity::GetAudioRenderThreadMask()));
+
 		// If we are marked to pause on startup, make sure the event is in a waitable state.
 		if (bStartPaused)
 		{
@@ -99,9 +101,7 @@ namespace Audio
 		else
 		{
 			WakeupEvent->Trigger();
-		}
-		
-		CallbackThread.Reset(FRunnableThread::Create(this, TEXT("AudioMixerNullCallbackThread"), 0, ThreadPriority, FPlatformAffinity::GetAudioRenderThreadMask()));
+		}		
 	}
 		
 	void FMixerNullCallback::Stop()
diff --git a/Engine/Source/Runtime/Core/Private/Misc/CoreDelegates.cpp b/Engine/Source/Runtime/Core/Private/Misc/CoreDelegates.cpp
--- a/Engine/Source/Runtime/Core/Private/Misc/CoreDelegates.cpp
+++ b/Engine/Source/Runtime/Core/Private/Misc/CoreDelegates.cpp
@@ -166,6 +166,7 @@ TMulticastDelegate<void(const TArray<FString>&)> FCoreDelegates::ApplicationRece
 
 TMulticastDelegate<void(bool)> FCoreDelegates::UserMusicInterruptDelegate;
 TMulticastDelegate<void(bool)> FCoreDelegates::AudioRouteChangedDelegate;
+TMulticastDelegate<void(bool)> FCoreDelegates::AudioInterruptionDelegate;
 TMulticastDelegate<void(bool, int)> FCoreDelegates::AudioMuteDelegate;
 TMulticastDelegate<void()> FCoreDelegates::ApplicationRequestAudioState;
 
diff --git a/Engine/Source/Runtime/Core/Public/Misc/CoreDelegates.h b/Engine/Source/Runtime/Core/Public/Misc/CoreDelegates.h
--- a/Engine/Source/Runtime/Core/Public/Misc/CoreDelegates.h
+++ b/Engine/Source/Runtime/Core/Public/Misc/CoreDelegates.h
@@ -354,6 +354,9 @@ public:
 	// For instance, when the headphones are plugged in or removed
 	static CORE_API TMulticastDelegate<void(bool)> AudioRouteChangedDelegate;
 
+	// [iOS only] Called when the audio session is interrupted
+	static CORE_API TMulticastDelegate<void(bool)> AudioInterruptionDelegate;
+	
 	// Generally, events triggering UserMusicInterruptDelegate or AudioMuteDelegate happen only
 	// when a change occurs. When a system comes online needing the current audio state but the
 	// event has already been broadcast, calling ApplicationRequestAudioState will force the
