Further iOS audio interruptions. Handle case where audio is interrupted by foreground phone call and user subsequently pulls down the notification drawer fully during that phone call.

diff --git a/Engine/Source/Runtime/Apple/AudioMixerAudioUnit/Private/AudioMixerPlatformAudioUnit.cpp b/Engine/Source/Runtime/Apple/AudioMixerAudioUnit/Private/AudioMixerPlatformAudioUnit.cpp
--- a/Engine/Source/Runtime/Apple/AudioMixerAudioUnit/Private/AudioMixerPlatformAudioUnit.cpp
+++ b/Engine/Source/Runtime/Apple/AudioMixerAudioUnit/Private/AudioMixerPlatformAudioUnit.cpp
@@ -464,6 +464,11 @@ namespace Audio
 		if (SuspendCounter > 0)
 		{
 			FPlatformAtomics::InterlockedDecrement(&SuspendCounter);
+			
+			if (bIsUsingNullDevice)
+			{
+				StopRunningNullDevice();
+			}
             
             if (AudioUnitGraph != NULL)
             {
