Limit rhi.Metal.MaxQueryBufferSize to 65528 on A12 or below devices.

diff --git a/Engine/Source/Runtime/Apple/MetalRHI/Private/MetalRHIRenderQuery.cpp b/Engine/Source/Runtime/Apple/MetalRHI/Private/MetalRHIRenderQuery.cpp
--- a/Engine/Source/Runtime/Apple/MetalRHI/Private/MetalRHIRenderQuery.cpp
+++ b/Engine/Source/Runtime/Apple/MetalRHI/Private/MetalRHIRenderQuery.cpp
@@ -77,7 +77,11 @@ FMetalQueryBufferPool::FMetalQueryBufferPool(FMetalDevice& InDevice)
 	, Device{InDevice}
 {
 	BufferSize = GMetalMaxQueryBufferSize;
-	// void
+	if (!Device.GetDevice()->supportsFamily(MTL::GPUFamilyApple6))
+	{
+		// On A12 and below devices the offset in setVisibilityResultMode needs to be <= 65528
+		BufferSize = FMath::Min(65528, GMetalMaxQueryBufferSize);
+	}
 }
 
 FMetalQueryBufferPool::~FMetalQueryBufferPool()
