MetalRHI

- Removed triple buffering in MetalRHIBuffer, this code breaks whenever there are more that 3 locks per frame.

diff --git a/Engine/Source/Runtime/Apple/MetalRHI/Private/MetalRHIBuffer.cpp b/Engine/Source/Runtime/Apple/MetalRHI/Private/MetalRHIBuffer.cpp
--- a/Engine/Source/Runtime/Apple/MetalRHI/Private/MetalRHIBuffer.cpp
+++ b/Engine/Source/Runtime/Apple/MetalRHI/Private/MetalRHIBuffer.cpp
@@ -114,29 +114,14 @@ FMetalRHIBuffer::FMetalRHIBuffer(FRHICommandListBase& RHICmdList, FMetalDevice&
 	if (CreateDesc.Size)
 	{
 		checkf(CreateDesc.Size <= Device.GetDevice()->maxBufferLength(), TEXT("Requested buffer size larger than supported by device."));
-		
-		{
-			// Static buffers will never be discarded. You can update them directly.
-			if(bIsStatic)
-			{
-				NumberOfBuffers = 1;
-			}
-			else
-			{
-				check(bIsDynamic || bIsVolatile);
-				NumberOfBuffers = 3;
-			}
-			
-			check(NumberOfBuffers > 0);
 			
 #if PLATFORM_MAC
-			// Buffer can be blit encoder copied on lock/unlock, we need to know that the buffer size is large enough for copy operations that are in multiples of
-			// 4 bytes on macOS, iOS can be 1 byte.  Update size to know we have at least this much buffer memory, it will be larger in the end.
-			Size = Align(CreateDesc.Size, 4);
+		// Buffer can be blit encoder copied on lock/unlock, we need to know that the buffer size is large enough for copy operations that are in multiples of
+		// 4 bytes on macOS, iOS can be 1 byte.  Update size to know we have at least this much buffer memory, it will be larger in the end.
+		Size = Align(CreateDesc.Size, 4);
 #endif
-			
-			AllocateBuffers();
-		}
+		
+		AllocateBuffer();
 	}
 
 	if (InResourceArray && CreateDesc.Size > 0)
@@ -197,8 +182,19 @@ FMetalRHIBuffer::~FMetalRHIBuffer()
     ReleaseOwnership();
 }
 
-void FMetalRHIBuffer::AllocateBufferAtSlot(uint32_t Idx)
+void FMetalRHIBuffer::SwitchBuffer()
+{
+	AllocateBuffer();
+	UpdateLinkedViews();
+}
+
+void FMetalRHIBuffer::AllocateBuffer()
 {
+	if(CurrentBuffer)
+	{
+		ReleaseBuffer();
+	}
+	
 	uint32 AllocSize = Size;
 	
 	const bool bWantsView = EnumHasAnyFlags(GetDesc().Usage, BUF_ShaderResource | BUF_UnorderedAccess);
@@ -291,28 +287,20 @@ void FMetalRHIBuffer::AllocateBufferAtSlot(uint32_t Idx)
 #endif
 	
 	MetalBufferStats::UpdateBufferStats(GetDesc(), Buffer->GetLength(), true);
-	BufferPool[Idx] = Buffer;
+	CurrentBuffer = Buffer;
 	
 	check(Buffer);
 	check(AllocSize <= Buffer->GetLength());
 	check(Buffer->GetMTLBuffer()->storageMode() == Mode);
 }
 
-void FMetalRHIBuffer::AllocateBuffers()
+void FMetalRHIBuffer::ReleaseBuffer()
 {
-	BufferPool.SetNum(NumberOfBuffers);
-}
-
-void FMetalRHIBuffer::ReleaseBuffers()
-{
-	for (FMetalBufferPtr Buffer : BufferPool)
-	{
-		if(Buffer)
-		{	
-			MetalBufferStats::UpdateBufferStats(GetDesc(), Buffer->GetLength(), false);
-			METAL_INC_DWORD_STAT_BY(MemFreed, Buffer->GetLength(), GetUsage());
-			FMetalDynamicRHI::Get().DeferredDelete(Buffer);
-		}
+	if(CurrentBuffer)
+	{	
+		MetalBufferStats::UpdateBufferStats(GetDesc(), CurrentBuffer->GetLength(), false);
+		METAL_INC_DWORD_STAT_BY(MemFreed, CurrentBuffer->GetLength(), GetUsage());
+		FMetalDynamicRHI::Get().DeferredDelete(CurrentBuffer);
 	}
 } 
 
@@ -361,27 +349,15 @@ void* FMetalRHIBuffer::Lock(bool bIsOnRHIThread, EResourceLockMode InLockMode, u
 	
 	if(bWriteLock)
 	{
-		if(bIsStatic)
-		{
-			// Static buffers do not discard. They just return the buffer or a transfer buffer.
-			// You are not supposed to lock more than once a frame.
-		}
-		else
+		const bool bUseTransferBuffer = RequiresTransferBuffer();
+		
+		if(!bIsStatic || bUseTransferBuffer)
 		{
-			check(bIsDynamic || bIsVolatile);
-			// cycle to next allocation
-			AdvanceBackingIndex();
+			SwitchBuffer();
 		}
-
-		// Use transfer buffer for writing into 'Static' buffers as they could be in use by GPU atm
-		// Initialization of 'Static' buffers still uses direct copy when possible
-		const bool bUseTransferBuffer = RequiresTransferBuffer();
+		
 		if(bUseTransferBuffer)
-		{
-			// Re-allocate the buffer in case of multiple writes in a single frame
-			ReleaseBuffers();
-			AllocateBuffers();
-			
+		{			
 			if(InTransferBuffer)
 			{
 				TransferBuffer = InTransferBuffer;
@@ -389,7 +365,6 @@ void* FMetalRHIBuffer::Lock(bool bIsOnRHIThread, EResourceLockMode InLockMode, u
 			else
 			{
 				TransferBuffer = Device.GetTransferAllocator()->Allocate(Len);
-				
 				ReturnPointer = (uint8*) TransferBuffer->Contents();
 				check(ReturnPointer != nullptr);
 			}
@@ -512,8 +487,6 @@ void FMetalRHIBuffer::Unlock(FRHICommandListBase& RHICmdList)
 				// shared buffers are always mapped so nothing happens
 				check(Mode == MTL::StorageModeShared);
 			}
-
-			UpdateLinkedViews();
 		}
 		else
 		{
@@ -544,10 +517,7 @@ void FMetalRHIBuffer::TakeOwnership(FMetalRHIBuffer& Other)
     FRHIBuffer::TakeOwnership(Other);
 
     TransferBuffer = Other.TransferBuffer;
-    BufferPool = Other.BufferPool;
     Data = Other.Data;
-    CurrentIndex = Other.CurrentIndex;
-    NumberOfBuffers = Other.NumberOfBuffers;
     CurrentLockMode = Other.CurrentLockMode;
     LockOffset = Other.LockOffset;
     LockSize = Other.LockSize;
@@ -555,10 +525,7 @@ void FMetalRHIBuffer::TakeOwnership(FMetalRHIBuffer& Other)
     Mode = Other.Mode;
     
     Other.TransferBuffer = nullptr;
-    Other.BufferPool.SetNum(0);
     Other.Data = nullptr;
-    Other.CurrentIndex = 0;
-    Other.NumberOfBuffers = 0;
     Other.CurrentLockMode = RLM_Num;
     Other.LockOffset = 0;
     Other.LockSize = 0;
@@ -573,7 +540,7 @@ void FMetalRHIBuffer::ReleaseOwnership()
 		FMetalDynamicRHI::Get().DeferredDelete(TransferBuffer);
     }
     
-	ReleaseBuffers();
+	ReleaseBuffer();
 
     if (Data)
     {
diff --git a/Engine/Source/Runtime/Apple/MetalRHI/Private/MetalTexture.cpp b/Engine/Source/Runtime/Apple/MetalRHI/Private/MetalTexture.cpp
--- a/Engine/Source/Runtime/Apple/MetalRHI/Private/MetalTexture.cpp
+++ b/Engine/Source/Runtime/Apple/MetalRHI/Private/MetalTexture.cpp
@@ -1583,7 +1583,9 @@ static FMetalBufferPtr Internal_CreateBufferAndCopyTexture2DUpdateRegionData(FMe
 		const uint32 SourcePitch = InOutSourcePitch;
 		const uint32 StagingPitch = static_cast<size_t>(WidthInBlocks) * FormatInfo.BlockBytes;
 
-		const uint32 BufferSize = UpdateRegion.Height * InOutSourcePitch;
+		const uint32 BufferSize = HeightInBlocks * StagingPitch;
+		
+		check(BufferSize >= StagingPitch*HeightInBlocks);
 		OutBuffer = Device.CreatePooledBuffer(FMetalPooledBufferArgs(&Device, BufferSize, BUF_Static, MTL::StorageModeShared));
 
 		uint8* pDestRow = (uint8*)OutBuffer->Contents();
diff --git a/Engine/Source/Runtime/Apple/MetalRHI/Private/MetalTransitionData.cpp b/Engine/Source/Runtime/Apple/MetalRHI/Private/MetalTransitionData.cpp
--- a/Engine/Source/Runtime/Apple/MetalRHI/Private/MetalTransitionData.cpp
+++ b/Engine/Source/Runtime/Apple/MetalRHI/Private/MetalTransitionData.cpp
@@ -82,7 +82,7 @@ void FMetalTransitionData::EndResourceTransitions(FMetalCommandEncoder& CurrentE
 			case FRHITransitionInfo::EType::Buffer:
 			{
 				auto Resource = ResourceCast(Info.Buffer);
-				if (Resource->GetCurrentBufferOrNil())
+				if (Resource->GetCurrentBuffer())
 				{
 					CurrentEncoder.TransitionResources(Resource->GetCurrentBuffer()->GetMTLBuffer());
 				}
diff --git a/Engine/Source/Runtime/Apple/MetalRHI/Public/MetalResources.h b/Engine/Source/Runtime/Apple/MetalRHI/Public/MetalResources.h
--- a/Engine/Source/Runtime/Apple/MetalRHI/Public/MetalResources.h
+++ b/Engine/Source/Runtime/Apple/MetalRHI/Public/MetalResources.h
@@ -291,17 +291,14 @@ public:
 class FMetalRHIBuffer final : public FRHIBuffer, public FMetalViewableResource
 {
 public:
-	// Matches other RHIs
-	static constexpr const uint32 MetalMaxNumBufferedFrames = 4;
-	
 	FMetalRHIBuffer(FRHICommandListBase& RHICmdList, FMetalDevice& MetalDevice, const FRHIBufferCreateDesc& CreateDesc, FResourceArrayUploadInterface* InResourceArray);
 	virtual ~FMetalRHIBuffer();
 	
 	bool RequiresTransferBuffer();
 	
-	void AllocateBufferAtSlot(uint32_t Idx);
-	void AllocateBuffers();
-	void ReleaseBuffers();
+	void AllocateBuffer();
+	void ReleaseBuffer();
+	void SwitchBuffer();
 	
 	/**
 	 * Prepare a CPU accessible buffer for uploading to GPU memory
@@ -315,25 +312,7 @@ public:
 	
 	FMetalBufferPtr GetCurrentBuffer()
 	{
-		if(BufferPool[CurrentIndex] == nullptr)
-		{
-			AllocateBufferAtSlot(CurrentIndex);
-		}
-		return BufferPool[CurrentIndex];
-	}
-	
-	FMetalBufferPtr GetCurrentBufferOrNil()
-	{
-		if (NumberOfBuffers > 0)
-		{
-			return GetCurrentBuffer();
-		}
-		return nullptr;
-	}
-	
-	void AdvanceBackingIndex()
-	{
-		CurrentIndex = (CurrentIndex + 1) % NumberOfBuffers;
+		return CurrentBuffer;
 	}
 	
 #if METAL_RHI_RAYTRACING
@@ -357,15 +336,11 @@ public:
 	// A temporary shared/CPU accessible buffer for upload/download
 	FMetalBufferPtr TransferBuffer = nullptr;
 	
-	TArray<FMetalBufferPtr> BufferPool;
+	FMetalBufferPtr CurrentBuffer;
 	
 	/** Buffer for small buffers < 4Kb to avoid heap fragmentation. */
 	FMetalBufferData* Data = nullptr;
 	
-	// The active buffer.
-	uint8 CurrentIndex = 0;
-	// How many buffers are actually allocated
-	uint8 NumberOfBuffers = 0;
 	// Current lock mode. RLM_Num indicates this buffer is not locked.
 	uint16 CurrentLockMode = RLM_Num;
 	
@@ -390,7 +365,6 @@ public:
 	}
 
 	static_assert((1 << 16) > RLM_Num, "Lock mode does not fit in bitfield");
-	static_assert((1 << 8) > MetalMaxNumBufferedFrames, "Buffer count does not fit in bitfield");
 	
 private:
 	// Allocate the CPU accessible buffer for data transfer.
