Limit rhi.Metal.MaxQueryBufferSize to 65528 also on A13

diff --git a/Engine/Source/Runtime/Apple/MetalRHI/Private/MetalRHIRenderQuery.cpp b/Engine/Source/Runtime/Apple/MetalRHI/Private/MetalRHIRenderQuery.cpp
--- a/Engine/Source/Runtime/Apple/MetalRHI/Private/MetalRHIRenderQuery.cpp
+++ b/Engine/Source/Runtime/Apple/MetalRHI/Private/MetalRHIRenderQuery.cpp
@@ -77,9 +77,9 @@ FMetalQueryBufferPool::FMetalQueryBufferPool(FMetalDevice& InDevice)
 	, Device{InDevice}
 {
 	BufferSize = GMetalMaxQueryBufferSize;
-	if (!Device.GetDevice()->supportsFamily(MTL::GPUFamilyApple6))
+	if (!Device.GetDevice()->supportsFamily(MTL::GPUFamilyApple7))
 	{
-		// On A12 and below devices the offset in setVisibilityResultMode needs to be <= 65528
+		// On A13 and below devices the offset in setVisibilityResultMode needs to be <= 65528
 		BufferSize = FMath::Min(65528, GMetalMaxQueryBufferSize);
 	}
 }
