Metal - Allow forcing ES3_1 on iOS with -es31

#tests Local testing with and without -es31 on iPhone

diff --git a/Engine/Source/Runtime/Apple/MetalRHI/Private/MetalRHI.cpp b/Engine/Source/Runtime/Apple/MetalRHI/Private/MetalRHI.cpp
--- a/Engine/Source/Runtime/Apple/MetalRHI/Private/MetalRHI.cpp
+++ b/Engine/Source/Runtime/Apple/MetalRHI/Private/MetalRHI.cpp
@@ -286,9 +286,10 @@ FMetalDynamicRHI::FMetalDynamicRHI(ERHIFeatureLevel::Type RequestedFeatureLevel)
     GConfig->GetBool(TEXT("/Script/IOSRuntimeSettings.IOSRuntimeSettings"), TEXT("bSupportsMetalMRT"), bProjectSupportsMRTs, GEngineIni);
 
 	bool const bRequestedMetalMRT = ((RequestedFeatureLevel >= ERHIFeatureLevel::SM5) || (!bRequestedFeatureLevel && FParse::Param(FCommandLine::Get(),TEXT("metalmrt"))));
+	bool const bForceES3_1 = FParse::Param(FCommandLine::Get(),TEXT("es31"));
 
     // Only allow SM5 MRT on A9 or above devices
-    if (bProjectSupportsMRTs && bRequestedMetalMRT && !bIsA8FeatureSet)
+    if (bProjectSupportsMRTs && bRequestedMetalMRT && !bIsA8FeatureSet && !bForceES3_1)
     {
 #if PLATFORM_TVOS
 		ValidateTargetedRHIFeatureLevelExists(SP_METAL_MRT);
@@ -301,7 +302,7 @@ FMetalDynamicRHI::FMetalDynamicRHI(ERHIFeatureLevel::Type RequestedFeatureLevel)
     }
     else
 	{
-		if (bRequestedMetalMRT)
+		if (bRequestedMetalMRT && !bForceES3_1)
 		{
 			UE_LOG(LogMetal, Warning, TEXT("Metal MRT support requires an iOS or tvOS device with an A8 processor or later. Falling back to Metal ES 3.1."));
 		}
