CrashReporter Screenshot for iOS

#rnx

diff --git a/Engine/Source/Runtime/Apple/MetalRHI/Private/MetalDynamicRHI.h b/Engine/Source/Runtime/Apple/MetalRHI/Private/MetalDynamicRHI.h
--- a/Engine/Source/Runtime/Apple/MetalRHI/Private/MetalDynamicRHI.h
+++ b/Engine/Source/Runtime/Apple/MetalRHI/Private/MetalDynamicRHI.h
@@ -154,6 +154,9 @@ public:
 	virtual void RHIBindDebugLabelName(FRHICommandListBase& RHICmdList, FRHITexture* Texture, const TCHAR* Name) final override;
 	virtual void RHIReadSurfaceData(FRHITexture* Texture,FIntRect Rect,TArray<FColor>& OutData,FReadSurfaceDataFlags InFlags) final override;
 	virtual void RHIReadSurfaceData(FRHITexture* TextureRHI, FIntRect InRect, TArray<FLinearColor>& OutData, FReadSurfaceDataFlags InFlags) final override;
+#if PLATFORM_IOS
+	virtual void RHIReadSurfaceDataDirect(FRHITexture* TextureRHI, FIntRect Rect, TArray<FColor>& OutData) final override;
+#endif // PLATFORM_IOS
 	virtual void RHIMapStagingSurface(FRHITexture* Texture, FRHIGPUFence* Fence, void*& OutData, int32& OutWidth, int32& OutHeight, uint32 GPUIndex = 0) final override;
 	virtual void RHIUnmapStagingSurface(FRHITexture* Texture, uint32 GPUIndex = 0) final override;
 	virtual void RHIReadSurfaceFloatData(FRHITexture* Texture,FIntRect Rect,TArray<FFloat16Color>& OutData,ECubeFace CubeFace,int32 ArrayIndex,int32 MipIndex) final override;
diff --git a/Engine/Source/Runtime/Apple/MetalRHI/Private/MetalRenderTarget.cpp b/Engine/Source/Runtime/Apple/MetalRHI/Private/MetalRenderTarget.cpp
--- a/Engine/Source/Runtime/Apple/MetalRHI/Private/MetalRenderTarget.cpp
+++ b/Engine/Source/Runtime/Apple/MetalRHI/Private/MetalRenderTarget.cpp
@@ -278,6 +278,41 @@ void FMetalDynamicRHI::RHIReadSurfaceData(FRHITexture* TextureRHI, FIntRect Rect
 	}
 }
 
+#if PLATFORM_IOS
+void FMetalDynamicRHI::RHIReadSurfaceDataDirect(FRHITexture* TextureRHI, FIntRect Rect, TArray<FColor>& OutData)
+{   
+	if (!ensure(TextureRHI)) return;
+
+	FMetalSurface* Surface = GetMetalSurfaceFromRHITexture(TextureRHI);
+
+	// allocate output space
+	const uint32 SizeX = Rect.Width();
+	const uint32 SizeY = Rect.Height();
+	OutData.SetNumUninitialized(SizeX * SizeY);
+
+	MTLTexturePtr Texture = Surface->Texture;
+    if(!Texture && EnumHasAnyFlags(Surface->GetDesc().Flags, TexCreate_Presentable))
+    {
+        Texture = Surface->GetCurrentTexture();
+    }
+	
+    if(!Texture)
+    {
+		OutData.Empty();
+
+		UE_LOG(LogRHI, Error, TEXT("Trying to read from an uninitialised texture."));
+		
+        return;
+    }
+
+	const uint32 Stride = GPixelFormats[Surface->GetDesc().Format].BlockBytes * SizeX;
+	const uint32 BytesPerImage = Stride * SizeY;
+
+	MTL::Region Region(Rect.Min.X, Rect.Min.Y, SizeX, SizeY);
+	Texture->getBytes(OutData.GetData(), Stride, BytesPerImage, Region, 0, 0);
+}
+#endif // PLATFORM_IOS
+
 void FMetalDynamicRHI::RHIMapStagingSurface(FRHITexture* TextureRHI, FRHIGPUFence* FenceRHI, void*& OutData, int32& OutWidth, int32& OutHeight, uint32 GPUIndex)
 {
     MTL_SCOPED_AUTORELEASE_POOL;
diff --git a/Engine/Source/Runtime/Core/Private/IOS/IOSPlatformCrashContext.cpp b/Engine/Source/Runtime/Core/Private/IOS/IOSPlatformCrashContext.cpp
--- a/Engine/Source/Runtime/Core/Private/IOS/IOSPlatformCrashContext.cpp
+++ b/Engine/Source/Runtime/Core/Private/IOS/IOSPlatformCrashContext.cpp
@@ -26,6 +26,7 @@ PLCrashReporter* FIOSApplicationInfo::CrashReporter;
 #endif
 FIOSMallocCrashHandler* FIOSApplicationInfo::CrashMalloc;
 
+TFunction<void(char const* const InfoFolder)> FIOSCrashContext::GenerateScreenshotFunc = nullptr;
 
 FIOSCrashContext::FIOSCrashContext(ECrashContextType InType, const TCHAR* InErrorMessage)
 	: FApplePlatformCrashContext(InType, InErrorMessage)
@@ -120,6 +121,10 @@ void FIOSCrashContext::ConvertMinidump(char const* OutputPath, char const* Input
 #endif
 }
 
+void FIOSCrashContext::SetGenerateScreenshotFunc(const TFunction<void (const char *const)>& InGenerateScreenshotFunc)
+{
+    FIOSCrashContext::GenerateScreenshotFunc = InGenerateScreenshotFunc;
+}
 
 void FIOSCrashContext::GenerateInfoInFolder(char const* const InfoFolder, bool bIsEnsure) const
 {
@@ -241,6 +246,7 @@ void FIOSCrashContext::GenerateCrashInfo() const
         
     const bool bIsEnsure = false;
     GenerateInfoInFolder(CrashInfoFolder, bIsEnsure);
+    if (FIOSCrashContext::GenerateScreenshotFunc) FIOSCrashContext::GenerateScreenshotFunc(CrashInfoFolder);
         
     // for IOS we will need to send the report on the next run
     
diff --git a/Engine/Source/Runtime/Core/Public/IOS/IOSPlatformCrashContext.h b/Engine/Source/Runtime/Core/Public/IOS/IOSPlatformCrashContext.h
--- a/Engine/Source/Runtime/Core/Public/IOS/IOSPlatformCrashContext.h
+++ b/Engine/Source/Runtime/Core/Public/IOS/IOSPlatformCrashContext.h
@@ -87,10 +87,14 @@ struct CORE_API FIOSCrashContext : public FApplePlatformCrashContext
 
 	/** Converts the PLCrashReporter minidump */
 	static void ConvertMinidump(char const* OutputPath, char const* InputPath);
+    
+    static void SetGenerateScreenshotFunc(const TFunction<void(char const* const InfoFolder)>& InGenerateScreenshotFunc);
 	
 private:
 	/** Serializes platform specific properties to the buffer. */
 	virtual void AddPlatformSpecificProperties() const;
+
+    static TFunction<void(char const* const InfoFolder)> GenerateScreenshotFunc;
 };
 
 //Single global FIOSApplicationInfo to store our App's information
diff --git a/Engine/Source/Runtime/RHI/Public/DynamicRHI.h b/Engine/Source/Runtime/RHI/Public/DynamicRHI.h
--- a/Engine/Source/Runtime/RHI/Public/DynamicRHI.h
+++ b/Engine/Source/Runtime/RHI/Public/DynamicRHI.h
@@ -627,6 +627,15 @@ public:
 	// Default fallback; will not work for non-8-bit surfaces and it's extremely slow.
 	RHI_API virtual void RHIReadSurfaceData(FRHITexture* Texture, FIntRect Rect, TArray<FLinearColor>& OutData, FReadSurfaceDataFlags InFlags);
 
+	/** 
+	* Implemented by the platform to read directly from the texture. 
+	* Useful for getting the last back buffer content after a crash.
+	* As the RHIReadSurfaceData implementation will copy the texture's content to a new temporary buffer, that won't work for an unknown crash state.
+	*/
+#if PLATFORM_IOS
+	virtual void RHIReadSurfaceDataDirect(FRHITexture* Texture, FIntRect Rect, TArray<FColor>& OutData) {}
+#endif // PLATFORM_IOS
+
 	/** Watch out for OutData to be 0 (can happen on DXGI_ERROR_DEVICE_REMOVED), don't call RHIUnmapStagingSurface in that case. */
 	// FlushType: Flush Immediate (seems wrong)
 	virtual void RHIMapStagingSurface(FRHITexture* Texture, FRHIGPUFence* Fence, void*& OutData, int32& OutWidth, int32& OutHeight, uint32 GPUIndex = 0) = 0;
