Fix shadowed variable in Metal RHI

#rnx

diff --git a/Engine/Source/Runtime/Apple/MetalRHI/Private/MetalRHIBuffer.cpp b/Engine/Source/Runtime/Apple/MetalRHI/Private/MetalRHIBuffer.cpp
--- a/Engine/Source/Runtime/Apple/MetalRHI/Private/MetalRHIBuffer.cpp
+++ b/Engine/Source/Runtime/Apple/MetalRHI/Private/MetalRHIBuffer.cpp
@@ -453,9 +453,9 @@ void FMetalRHIBuffer::Unlock(FRHICommandListBase& RHICmdList)
 {
 	if (!Data)
 	{
-		FMetalBufferPtr CurrentBuffer = GetCurrentBuffer();
+		FMetalBufferPtr CurrentBufferToUnlock = GetCurrentBuffer();
 		
-		check(CurrentBuffer);
+		check(CurrentBufferToUnlock);
 		check(LockSize > 0);
 		const bool bWriteLock = CurrentLockMode == RLM_WriteOnly;
 		const bool bIsStatic = EnumHasAnyFlags(GetUsage(), BUF_Static);
@@ -463,7 +463,7 @@ void FMetalRHIBuffer::Unlock(FRHICommandListBase& RHICmdList)
 		if (bWriteLock)
 		{
 			check(LockOffset == 0);
-            check(LockSize <= CurrentBuffer->GetLength());
+            check(LockSize <= CurrentBufferToUnlock->GetLength());
 
 			// Use transfer buffer for writing into 'Static' buffers as they could be in use by GPU atm
 			// Initialization of 'Static' buffers still uses direct copy when possible
@@ -473,9 +473,7 @@ void FMetalRHIBuffer::Unlock(FRHICommandListBase& RHICmdList)
 			{
 				FMetalRHIUploadContext& UploadContext = static_cast<FMetalRHIUploadContext&>(RHICmdList.GetUploadContext());
 				
-				CurrentBuffer = GetCurrentBuffer();
-				
-				UploadContext.EnqueueFunction([&InDevice=Device, Size=LockSize, Dest=CurrentBuffer, InTransferBuffer=TransferBuffer](FMetalRHICommandContext* Context)
+				UploadContext.EnqueueFunction([&InDevice=Device, Size=LockSize, Dest=CurrentBufferToUnlock, InTransferBuffer=TransferBuffer](FMetalRHICommandContext* Context)
 				{
 					Context->CopyFromBufferToBuffer(InTransferBuffer, 0, Dest, 0, Size);
 					FMetalDynamicRHI::Get().DeferredDelete(InTransferBuffer);
@@ -486,7 +484,7 @@ void FMetalRHIBuffer::Unlock(FRHICommandListBase& RHICmdList)
 #if PLATFORM_MAC
 			else if (Mode == MTL::StorageModeManaged)
 			{
-				CurrentBuffer->GetMTLBuffer()->didModifyRange(NS::Range(LockOffset + CurrentBuffer->GetOffset(), LockSize));
+				CurrentBufferToUnlock->GetMTLBuffer()->didModifyRange(NS::Range(LockOffset + CurrentBufferToUnlock->GetOffset(), LockSize));
 			}
 #endif //PLATFORM_MAC
 			else
