Moving duplicated RHI code to the new RHICoreTexture.h
- Moved the texture creationg LLM and Trace macro combinations into RHICORE_SCOPE_CREATE_TEXTURE
- Deduplicating the implemenations for calculating array indices for locking into UE::RHICore::GetLockArrayIndex

diff --git a/Engine/Source/Runtime/Apple/MetalRHI/Private/MetalTexture.cpp b/Engine/Source/Runtime/Apple/MetalRHI/Private/MetalTexture.cpp
--- a/Engine/Source/Runtime/Apple/MetalRHI/Private/MetalTexture.cpp
+++ b/Engine/Source/Runtime/Apple/MetalRHI/Private/MetalTexture.cpp
@@ -17,6 +17,7 @@
 #include "MetalBindlessDescriptors.h"
 #include "HAL/LowLevelMemStats.h"
 #include "ProfilingDebugging/AssetMetadataTrace.h"
+#include "RHICoreTexture.h"
 #include <CoreVideo/CVMetalTexture.h>
 
 volatile int64 FMetalSurface::ActiveUploads = 0;
@@ -918,23 +919,6 @@ void FMetalSurface::SafeRelease(MTLTexturePtr InTexture)
 	}
 }
 
-static uint32 GetMetalResourceIndex(const FRHILockTextureArgs& Arguments)
-{
-	const FRHITextureDesc& Desc = Arguments.Texture->GetDesc();
-
-	if (Desc.IsTextureCube())
-	{
-		return Arguments.FaceIndex + Arguments.ArrayIndex * 6;
-	}
-
-	if (Desc.IsTextureArray())
-	{
-		return Arguments.ArrayIndex;
-	}
-
-	return 0;
-}
-
 MTL::Buffer* FMetalSurface::AllocSurface(const FRHILockTextureArgs& Arguments, uint32 MipBytes, uint32 DestStride)
 {
 	check(IsInRenderingThread());
@@ -1084,7 +1068,7 @@ void FMetalSurface::UpdateSurfaceAndDestroySourceBuffer(FMetalRHICommandContext*
 
 FRHILockTextureResult FMetalSurface::Lock(const FRHILockTextureArgs& Arguments, bool bSingleLayer)
 {
-	const uint32 ResourceIndex = GetMetalResourceIndex(Arguments);
+	const uint32 ResourceIndex = UE::RHICore::GetLockArrayIndex(GetDesc(), Arguments);
 
 	FRHILockTextureResult Result;
 
@@ -1224,7 +1208,7 @@ void FMetalSurface::Unlock(const FRHILockTextureArgs& Arguments)
 	FRHILockTracker::FTextureLockParams Params = GRHILockTracker.Unlock(Arguments);
 	MTL::Buffer* SourceData = reinterpret_cast<MTL::Buffer*>(Params.Data);
 
-	const uint32 ResourceIndex = GetMetalResourceIndex(Arguments);
+	const uint32 ResourceIndex = UE::RHICore::GetLockArrayIndex(GetDesc(), Arguments);
 
 	FRHICommandListImmediate& RHICmdList = FRHICommandListImmediate::Get();
 	RHICmdList.EnqueueLambda(
@@ -1482,7 +1466,7 @@ FRHILockTextureResult FMetalDynamicRHI::RHILockTexture(FRHICommandListImmediate&
 	MTL_SCOPED_AUTORELEASE_POOL;
 
 	FMetalSurface* Texture = ResourceCast(Arguments.Texture);
-	const uint32 ResourceIndex = GetMetalResourceIndex(Arguments);
+	const uint32 ResourceIndex = UE::RHICore::GetLockArrayIndex(Texture->GetDesc(), Arguments);
 
 	if (Texture->GetDesc().Dimension == ETextureDimension::Texture2D)
 	{
diff --git a/Engine/Source/Runtime/D3D12RHI/Private/D3D12Texture.cpp b/Engine/Source/Runtime/D3D12RHI/Private/D3D12Texture.cpp
--- a/Engine/Source/Runtime/D3D12RHI/Private/D3D12Texture.cpp
+++ b/Engine/Source/Runtime/D3D12RHI/Private/D3D12Texture.cpp
@@ -12,6 +12,7 @@
 #include "HAL/LowLevelMemStats.h"
 #include "ProfilingDebugging/AssetMetadataTrace.h"
 #include "RHICoreStats.h"
+#include "RHICoreTexture.h"
 
 int64 FD3D12GlobalStats::GDedicatedVideoMemory = 0;
 int64 FD3D12GlobalStats::GDedicatedSystemMemory = 0;
@@ -1708,22 +1709,6 @@ void FD3D12Texture::AliasResources(FD3D12Texture* Texture)
 	}
 }
 
-static uint32 GetD3D12ArrayIndex(const FRHILockTextureArgs& Arguments, const FRHITextureDesc& Desc)
-{
-	if (Desc.IsTextureCube())
-	{
-		const uint32 D3DFace = GetD3D12CubeFace(static_cast<ECubeFace>(Arguments.FaceIndex));
-		return D3DFace + Arguments.ArrayIndex * 6;
-	}
-
-	if (Desc.IsTextureArray())
-	{
-		return Arguments.ArrayIndex;
-	}
-
-	return 0;
-}
-
 FRHILockTextureResult FD3D12Texture::Lock(FRHICommandListBase& RHICmdList, const FRHILockTextureArgs& Arguments)
 {
 	SCOPE_CYCLE_COUNTER(STAT_D3D12LockTextureTime);
@@ -1735,7 +1720,7 @@ FRHILockTextureResult FD3D12Texture::Lock(FRHICommandListBase& RHICmdList, const
 	FD3D12Adapter* Adapter = Device->GetParentAdapter();
 
 	const FRHITextureDesc& Desc = GetDesc();
-	const uint32 ArrayIndex = GetD3D12ArrayIndex(Arguments, Desc);
+	const uint32 ArrayIndex = UE::RHICore::GetLockArrayIndex(Desc, Arguments);
 
 	// Calculate the subresource index corresponding to the specified mip-map.
 	const uint32 Subresource = CalcSubresource(Arguments.MipIndex, ArrayIndex, this->GetNumMips());
@@ -2021,7 +2006,7 @@ void FD3D12Texture::Unlock(FRHICommandListBase& RHICmdList, const FRHILockTextur
 	check(IsHeadLink());
 
 	const FRHITextureDesc& Desc = GetDesc();
-	const uint32 ArrayIndex = GetD3D12ArrayIndex(Arguments, Desc);
+	const uint32 ArrayIndex = UE::RHICore::GetLockArrayIndex(Desc, Arguments);
 
 	// Calculate the subresource index corresponding to the specified mip-map.
 	const uint32 Subresource = CalcSubresource(Arguments.MipIndex, ArrayIndex, Desc.NumMips);
diff --git a/Engine/Source/Runtime/OpenGLDrv/Private/OpenGLTexture.cpp b/Engine/Source/Runtime/OpenGLDrv/Private/OpenGLTexture.cpp
--- a/Engine/Source/Runtime/OpenGLDrv/Private/OpenGLTexture.cpp
+++ b/Engine/Source/Runtime/OpenGLDrv/Private/OpenGLTexture.cpp
@@ -16,6 +16,7 @@
 #include "ProfilingDebugging/AssetMetadataTrace.h"
 #include "Engine/Texture.h"
 #include "RHICoreStats.h"
+#include "RHICoreTexture.h"
 
 #if PLATFORM_ANDROID
 #include "ThirdParty/Android/detex/AndroidETC.h"
@@ -1979,30 +1980,13 @@ FTextureRHIRef FOpenGLDynamicRHI::RHICreateAliasedTexture(FTextureRHIRef& Source
 	return new FOpenGLTexture(*ResourceCast(SourceTexture), *Name, FOpenGLTexture::AliasResource);
 }
 
-static uint32 GetGLArrayIndex(const FRHILockTextureArgs& Arguments)
-{
-	const FRHITextureDesc& Desc = Arguments.Texture->GetDesc();
-
-	if (Desc.IsTextureCube())
-	{
-		return Arguments.FaceIndex + Arguments.ArrayIndex * 6;
-	}
-
-	if (Desc.IsTextureArray())
-	{
-		return Arguments.ArrayIndex;
-	}
-
-	return 0;
-}
-
 FRHILockTextureResult FOpenGLDynamicRHI::RHILockTexture(FRHICommandListImmediate& RHICmdList, const FRHILockTextureArgs& Arguments)
 {
 	check(IsInRenderingThread());
 	check(Arguments.LockMode != EResourceLockMode::RLM_WriteOnly_NoOverwrite);
 
 	FOpenGLTexture* Texture = ResourceCast(Arguments.Texture);
-	const uint32 ArrayIndex = GetGLArrayIndex(Arguments);
+	const uint32 ArrayIndex = UE::RHICore::GetLockArrayIndex(Texture->GetDesc(), Arguments);
 
 	FRHILockTextureResult Result;
 
@@ -2033,7 +2017,7 @@ void FOpenGLDynamicRHI::RHIUnlockTexture(FRHICommandListImmediate& RHICmdList, c
 	check(IsInRenderingThread());
 
 	FOpenGLTexture* Texture = ResourceCast(Arguments.Texture);
-	const uint32 ArrayIndex = GetGLArrayIndex(Arguments);
+	const uint32 ArrayIndex = UE::RHICore::GetLockArrayIndex(Texture->GetDesc(), Arguments);
 
 	FTextureLockTracker::FLockParams Params = GLLockTracker.Unlock(Texture, ArrayIndex, Arguments.MipIndex);
 	if (Params.LockMode == RLM_ReadOnly)
diff --git a/Engine/Source/Runtime/RHICore/Internal/RHICoreTexture.h b/Engine/Source/Runtime/RHICore/Internal/RHICoreTexture.h
new file mode 100644
--- /dev/null
+++ b/Engine/Source/Runtime/RHICore/Internal/RHICoreTexture.h
@@ -0,0 +1,29 @@
+// Copyright Epic Games, Inc. All Rights Reserved.
+
+#pragma once
+
+#include "DynamicRHI.h"
+#include "RHIResources.h"
+
+namespace UE::RHICore
+{
+	inline uint32 GetCombinedArrayIndex(const FRHITextureDesc& Desc, uint32 FaceIndex, uint32 ArrayIndex)
+	{
+		if (Desc.IsTextureCube())
+		{
+			return FaceIndex + ArrayIndex * 6;
+		}
+
+		if (Desc.IsTextureArray())
+		{
+			return ArrayIndex;
+		}
+
+		return 0;
+	}
+
+	inline uint32 GetLockArrayIndex(const FRHITextureDesc& Desc, const FRHILockTextureArgs& Arguments)
+	{
+		return GetCombinedArrayIndex(Desc, Arguments.FaceIndex, Arguments.ArrayIndex);
+	}
+}
diff --git a/Engine/Source/Runtime/VulkanRHI/Private/VulkanTexture.cpp b/Engine/Source/Runtime/VulkanRHI/Private/VulkanTexture.cpp
--- a/Engine/Source/Runtime/VulkanRHI/Private/VulkanTexture.cpp
+++ b/Engine/Source/Runtime/VulkanRHI/Private/VulkanTexture.cpp
@@ -16,7 +16,7 @@
 #include "RHILockTracker.h"
 #include "HAL/LowLevelMemStats.h"
 #include "ProfilingDebugging/AssetMetadataTrace.h"
-
+#include "RHICoreTexture.h"
 
 // This is a workaround for issues with AFBC on Mali GPUs before the G710
 int32 GVulkanDepthStencilForceStorageBit = 0;
@@ -1214,21 +1214,6 @@ FRHILockTextureResult FVulkanDynamicRHI::RHILockTexture(FRHICommandListImmediate
 	return Result;
 }
 
-static uint32 GetVulkanArrayIndex(const FRHILockTextureArgs& Arguments, const FRHITextureDesc& Desc)
-{
-	if (Desc.IsTextureCube())
-	{
-		return Arguments.FaceIndex + Arguments.ArrayIndex * 6;
-	}
-
-	if (Desc.IsTextureArray())
-	{
-		return Arguments.ArrayIndex;
-	}
-
-	return 0;
-}
-
 void FVulkanDynamicRHI::RHIUnlockTexture(FRHICommandListImmediate& RHICmdList, const FRHILockTextureArgs& Arguments)
 {
 	LLM_SCOPE_VULKAN(ELLMTagVulkan::VulkanTextures);
@@ -1244,7 +1229,7 @@ void FVulkanDynamicRHI::RHIUnlockTexture(FRHICommandListImmediate& RHICmdList, c
 	FVulkanTexture* Texture = ResourceCast(Arguments.Texture);
 	const FRHITextureDesc& Desc = Texture->GetDesc();
 
-	const uint32 ArrayIndex = GetVulkanArrayIndex(Arguments, Desc);
+	const uint32 ArrayIndex = UE::RHICore::GetLockArrayIndex(Desc, Arguments);
 
 	uint32 MipWidth  = FMath::Max<uint32>(Desc.Extent.X >> Arguments.MipIndex, 0);
 	uint32 MipHeight = FMath::Max<uint32>(Desc.Extent.Y >> Arguments.MipIndex, 0);
diff --git a/Engine/Source/Runtime/Windows/D3D11RHI/Private/D3D11Texture.cpp b/Engine/Source/Runtime/Windows/D3D11RHI/Private/D3D11Texture.cpp
--- a/Engine/Source/Runtime/Windows/D3D11RHI/Private/D3D11Texture.cpp
+++ b/Engine/Source/Runtime/Windows/D3D11RHI/Private/D3D11Texture.cpp
@@ -21,6 +21,7 @@
 #include "ProfilingDebugging/MemoryTrace.h"
 #include "ProfilingDebugging/AssetMetadataTrace.h"
 #include "RHICoreStats.h"
+#include "RHICoreTexture.h"
 #include "RHIUtilities.h"
 #include "RHITextureUtils.h"
 
@@ -991,22 +992,6 @@ FTextureRHIRef FD3D11DynamicRHI::AsyncReallocateTexture2D_RenderThread(
 	return NewTexture2D;
 }
 
-static uint32 GetD3D11ArrayIndex(const FRHILockTextureArgs& Arguments, const FRHITextureDesc& Desc)
-{
-	if (Desc.IsTextureCube())
-	{
-		const uint32 D3DFace = GetD3D11CubeFace(static_cast<ECubeFace>(Arguments.FaceIndex));
-		return D3DFace + Arguments.ArrayIndex * 6;
-	}
-
-	if (Desc.IsTextureArray())
-	{
-		return Arguments.ArrayIndex;
-	}
-
-	return 0;
-}
-
 FRHILockTextureResult FD3D11Texture::Lock(FD3D11DynamicRHI* D3DRHI, const FRHILockTextureArgs& Arguments, bool bForceLockDeferred)
 {
 	check(!IsTexture3D()); // Only 2D texture locks are implemented
@@ -1015,7 +1000,7 @@ FRHILockTextureResult FD3D11Texture::Lock(FD3D11DynamicRHI* D3DRHI, const FRHILo
 
 	const FRHITextureDesc& Desc = this->GetDesc();
 
-	const uint32 ArrayIndex = GetD3D11ArrayIndex(Arguments, Desc);
+	const uint32 ArrayIndex = UE::RHICore::GetLockArrayIndex(Desc, Arguments);
 
 	// Calculate the subresource index corresponding to the specified mip-map.
 	const uint32 Subresource = D3D11CalcSubresource(Arguments.MipIndex, ArrayIndex, Desc.NumMips);
@@ -1133,7 +1118,7 @@ void FD3D11Texture::Unlock(FD3D11DynamicRHI* D3DRHI, const FRHILockTextureArgs&
 	SCOPE_CYCLE_COUNTER(STAT_D3D11UnlockTextureTime);
 
 	const FRHITextureDesc& Desc = this->GetDesc();
-	const uint32 ArrayIndex = GetD3D11ArrayIndex(Arguments, Desc);
+	const uint32 ArrayIndex = UE::RHICore::GetLockArrayIndex(Desc, Arguments);
 
 	// Calculate the subresource index corresponding to the specified mip-map.
 	const uint32 Subresource = D3D11CalcSubresource(Arguments.MipIndex, ArrayIndex, Desc.NumMips);
