Metal - Refactor ShaderPlatform naming to match other existing ShaderPlatform and remove confusion
SP_METAL to SP_METAL_ES3_1_IOS
SP_METAL_MRT to SP_METAL_SM5_IOS
SP_METAL_MACES3_1 to SP_METAL_ES3_1
SP_METAL_TVOS to SP_METAL_ES3_1_TVOS
SP_METAL_MRT_TVOS to SP_METAL_SM5_TVOS
Removal of SP_METAL_MRT_MAC

#tests Preflight and tested iOS Build

diff --git a/Engine/Config/IOS/DataDrivenPlatformInfo.ini b/Engine/Config/IOS/DataDrivenPlatformInfo.ini
--- a/Engine/Config/IOS/DataDrivenPlatformInfo.ini
+++ b/Engine/Config/IOS/DataDrivenPlatformInfo.ini
@@ -26,8 +26,8 @@ PlatformGroupName=Mobile
 
 [PreviewPlatform IOSMetal]
 PlatformName=IOS
-ShaderFormat=SF_METAL
-ShaderPlatform=METAL
+ShaderFormat=SF_METAL_ES3_1_IOS
+ShaderPlatform=METAL_ES3_1_IOS
 ActiveIconName=LevelEditor.PreviewMode.iOS.Enabled
 InactiveIconName=LevelEditor.PreviewMode.iOS.Disabled
 MenuTooltip=NSLOCTEXT("PreviewPlatform", "PreviewMenuTooltip_iOS", "Mobile preview using iOS material quality settings.")
@@ -37,8 +37,8 @@ FriendlyName=NSLOCTEXT("PreviewPlatform", "PreviewFriendlyName_GenericIOS", "Gen
 
 [PreviewPlatform IOSMetalSM5]
 PlatformName=IOS
-ShaderFormat=SF_METAL_MRT
-ShaderPlatform=METAL_MRT
+ShaderFormat=SF_METAL_SM5_IOS
+ShaderPlatform=METAL_SM5_IOS
 EnabledCVar=ini:Engine:IOS:/Script/IOSRuntimeSettings.IOSRuntimeSettings:bSupportsMetalMRT
 ActiveIconName=LevelEditor.PreviewMode.iOSSM5.Enabled
 InactiveIconName=LevelEditor.PreviewMode.iOSSM5.Disabled
@@ -46,10 +46,10 @@ MenuTooltip=NSLOCTEXT("PreviewPlatform", "PreviewMenuTooltip_iOSSM5", "Mobile pr
 IconText=NSLOCTEXT("PreviewPlatform", "PreviewIconText_iOSSM5", "iOS SM5")
 DeviceProfileName=IOS
 
-[ShaderPlatform METAL]
+[ShaderPlatform METAL_ES3_1_IOS]
 Language=Metal
 MaxFeatureLevel=ES3_1
-ShaderFormat=SF_METAL
+ShaderFormat=SF_METAL_ES3_1_IOS
 bIsMobile=true
 bTargetsTiledGPU=true
 bSupports4ComponentUAVReadWrite=true
@@ -60,10 +60,10 @@ bSupportsIndependentSamplers=true
 bSupportsComputeFramework = true
 FriendlyName=LOCTEXT("FriendlyShaderPlatformName_IOSMobile", "IOS Metal Mobile")
 
-[ShaderPlatform METAL_MRT]
+[ShaderPlatform METAL_SM5_IOS]
 Language=Metal
 MaxFeatureLevel=SM5
-ShaderFormat=SF_METAL_MRT
+ShaderFormat=SF_METAL_SM5_IOS
 bIsMetalMRT=true
 bTargetsTiledGPU=true
 bSupportsIndexBufferUAVs=true
diff --git a/Engine/Config/Mac/DataDrivenPlatformInfo.ini b/Engine/Config/Mac/DataDrivenPlatformInfo.ini
--- a/Engine/Config/Mac/DataDrivenPlatformInfo.ini
+++ b/Engine/Config/Mac/DataDrivenPlatformInfo.ini
@@ -107,24 +107,10 @@ MaxMeshShaderThreadGroupSize=128
 bSupportsRealTypes=RuntimeGuaranteed
 FriendlyName=LOCTEXT("FriendlyShaderPlatformName_MacSM6", "Mac Metal SM6")
 
-[ShaderPlatform METAL_MRT_MAC]
-Language=Metal
-MaxFeatureLevel=SM5
-ShaderFormat=SF_METAL_MRT_MAC
-bIsPC=true
-bIsMetalMRT=true
-bSupportsIndexBufferUAVs=true
-bSupports4ComponentUAVReadWrite=true
-bSupportsManualVertexFetch=true
-bSupportsVolumeTextureAtomics=false
-bSupportsDualSourceBlending=true
-bSupportsIndependentSamplers=true
-FriendlyName=LOCTEXT("FriendlyShaderPlatformName_MacSM5_MRT", "Mac Metal SM5 MRT")
-
-[ShaderPlatform METAL_MACES3_1]
+[ShaderPlatform METAL_ES3_1]
 Language=Metal
 MaxFeatureLevel=ES3_1
-ShaderFormat=SF_METAL_MACES3_1
+ShaderFormat=SF_METAL_ES3_1
 bIsPC=true
 bIsMobile=true
 bSupports4ComponentUAVReadWrite=true
diff --git a/Engine/Config/TVOS/DataDrivenPlatformInfo.ini b/Engine/Config/TVOS/DataDrivenPlatformInfo.ini
--- a/Engine/Config/TVOS/DataDrivenPlatformInfo.ini
+++ b/Engine/Config/TVOS/DataDrivenPlatformInfo.ini
@@ -16,10 +16,10 @@ bUsesHostCompiler=false
 bUATClosesAfterLaunch=true
 PlatformGroupName=Mobile
 
-[ShaderPlatform METAL_MRT_TVOS]
+[ShaderPlatform METAL_SM5_TVOS]
 Language=Metal
 MaxFeatureLevel=SM5
-ShaderFormat=SF_METAL_MRT_TVOS
+ShaderFormat=SF_METAL_SM5_TVOS
 bIsMetalMRT=true
 bTargetsTiledGPU=true
 bSupportsIndexBufferUAVs=true
@@ -29,14 +29,14 @@ bSupportsDualSourceBlending=true
 bSupportsIndependentSamplers=true
 FriendlyName=LOCTEXT("FriendlyShaderPlatformName_TVOS_SM5", "TVOS Metal SM5")
 
-[ShaderPlatform METAL_TVOS]
+[ShaderPlatform METAL_ES3_1_TVOS]
 Language=Metal
 MaxFeatureLevel=ES3_1
-ShaderFormat=SF_METAL_TVOS
+ShaderFormat=SF_METAL_ES3_1_TVOS
 bIsMobile=true
 bTargetsTiledGPU=true
 bSupports4ComponentUAVReadWrite=true
 bSupportsManualVertexFetch=false
 bSupportsShaderPipelines = false
 bSupportsIndependentSamplers=true
-FriendlyName=LOCTEXT("FriendlyShaderPlatformName_TVOS_Mobile", "TVOS Metal Mobile")
\ No newline at end of file
+FriendlyName=LOCTEXT("FriendlyShaderPlatformName_TVOS_Mobile", "TVOS Metal Mobile")
diff --git a/Engine/Plugins/FX/Niagara/Source/NiagaraEditor/Private/ViewModels/NiagaraScriptStatsViewModel.cpp b/Engine/Plugins/FX/Niagara/Source/NiagaraEditor/Private/ViewModels/NiagaraScriptStatsViewModel.cpp
--- a/Engine/Plugins/FX/Niagara/Source/NiagaraEditor/Private/ViewModels/NiagaraScriptStatsViewModel.cpp
+++ b/Engine/Plugins/FX/Niagara/Source/NiagaraEditor/Private/ViewModels/NiagaraScriptStatsViewModel.cpp
@@ -859,8 +859,8 @@ void FNiagaraScriptStatsViewModel::BuildShaderPlatformDetails()
 	StaticShaderPlatforms.Add(SP_VULKAN_SM5_ANDROID);
 	StaticShaderPlatforms.Add(SP_METAL_SM5);
     StaticShaderPlatforms.Add(SP_METAL_SM6);
-	StaticShaderPlatforms.Add(SP_METAL);
-	StaticShaderPlatforms.Add(SP_METAL_MRT);
+	StaticShaderPlatforms.Add(SP_METAL_ES3_1_IOS);
+	StaticShaderPlatforms.Add(SP_METAL_SM5_IOS);
 
 	TArray<EShaderPlatform, TInlineAllocator<SP_StaticPlatform_Last + 1>> ValidShaderPlatforms;
 	for (EShaderPlatform StaticShaderPlatform : StaticShaderPlatforms)
diff --git a/Engine/Shaders/Private/CapsuleShadowShaders.usf b/Engine/Shaders/Private/CapsuleShadowShaders.usf
--- a/Engine/Shaders/Private/CapsuleShadowShaders.usf
+++ b/Engine/Shaders/Private/CapsuleShadowShaders.usf
@@ -1111,7 +1111,7 @@ void CapsuleShadowingCS(
 		// The third param of this function is used only by mobile paths which do not support bent normals.
 		ApplySelfShadowingIntensityForDeferred(ScreenUV, Visibility, float2(0,0));
 	#endif
-	#if METAL_PROFILE 
+	#if METAL_ES3_1_PROFILE 
 		// clamp max depth to avoid #inf
 		SceneDepth = min(SceneDepth, 65500.0f);
 	#endif
diff --git a/Engine/Shaders/Private/Common.ush b/Engine/Shaders/Private/Common.ush
--- a/Engine/Shaders/Private/Common.ush
+++ b/Engine/Shaders/Private/Common.ush
@@ -487,7 +487,7 @@ MaterialFloat4 Texture2DArrayGatherAlpha(Texture2DArray Tex, SamplerState Sample
 // to workaround we always sample top mip level
 MaterialFloat4 Texture1DSample_Decal(Texture1D Tex, SamplerState Sampler, float UV)
 {
-#if METAL_PROFILE || COMPILER_GLSL_ES3_1	
+#if METAL_ES3_1_PROFILE || COMPILER_GLSL_ES3_1	
 	return Texture1DSampleLevel(Tex, Sampler, UV, 0);
 #else
 	return Texture1DSample(Tex, Sampler, UV);
@@ -495,7 +495,7 @@ MaterialFloat4 Texture1DSample_Decal(Texture1D Tex, SamplerState Sampler, float
 }
 MaterialFloat4 Texture2DSample_Decal(Texture2D Tex, SamplerState Sampler, float2 UV)
 {
-#if METAL_PROFILE || COMPILER_GLSL_ES3_1
+#if METAL_ES3_1_PROFILE || COMPILER_GLSL_ES3_1
 	return Texture2DSampleLevel(Tex, Sampler, UV, 0);
 #else
 	return Texture2DSample(Tex, Sampler, UV);
@@ -503,7 +503,7 @@ MaterialFloat4 Texture2DSample_Decal(Texture2D Tex, SamplerState Sampler, float2
 }
 MaterialFloat4 Texture3DSample_Decal(Texture3D Tex, SamplerState Sampler, float3 UV)
 {
-#if METAL_PROFILE || COMPILER_GLSL_ES3_1
+#if METAL_ES3_1_PROFILE || COMPILER_GLSL_ES3_1
 	return Texture3DSampleLevel(Tex, Sampler, UV, 0);
 #else
 	return Texture3DSample(Tex, Sampler, UV);
@@ -511,7 +511,7 @@ MaterialFloat4 Texture3DSample_Decal(Texture3D Tex, SamplerState Sampler, float3
 }
 MaterialFloat4 Texture2DArraySample_Decal(Texture2DArray Tex, SamplerState Sampler, float3 UV)
 {
-#if METAL_PROFILE || COMPILER_GLSL_ES3_1
+#if METAL_ES3_1_PROFILE || COMPILER_GLSL_ES3_1
 	return Texture2DArraySampleLevel(Tex, Sampler, UV, 0);
 #else
 	return Texture2DArraySample(Tex, Sampler, UV);
@@ -519,7 +519,7 @@ MaterialFloat4 Texture2DArraySample_Decal(Texture2DArray Tex, SamplerState Sampl
 }
 MaterialFloat4 TextureCubeSample_Decal(TextureCube Tex, SamplerState Sampler, float3 UV)
 {
-#if METAL_PROFILE || COMPILER_GLSL_ES3_1
+#if METAL_ES3_1_PROFILE || COMPILER_GLSL_ES3_1
 	return TextureCubeSampleLevel(Tex, Sampler, UV, 0);
 #else
 	return TextureCubeSample(Tex, Sampler, UV);
@@ -527,7 +527,7 @@ MaterialFloat4 TextureCubeSample_Decal(TextureCube Tex, SamplerState Sampler, fl
 }
 MaterialFloat4 TextureExternalSample_Decal(TextureExternal Tex, SamplerState Sampler, float2 UV)
 {
-#if METAL_PROFILE || COMPILER_GLSL_ES3_1
+#if METAL_ES3_1_PROFILE || COMPILER_GLSL_ES3_1
 	return TextureExternalSampleLevel(Tex, Sampler, UV, 0);
 #else
 	return TextureExternalSample(Tex, Sampler, UV);
diff --git a/Engine/Shaders/Private/DeferredShadingCommon.ush b/Engine/Shaders/Private/DeferredShadingCommon.ush
--- a/Engine/Shaders/Private/DeferredShadingCommon.ush
+++ b/Engine/Shaders/Private/DeferredShadingCommon.ush
@@ -757,7 +757,7 @@ FGBufferData MobileDecodeGBuffer(half4 InGBufferA, half4 InGBufferB, half4 InGBu
 
 #define	FRAMEBUFFER_FETCH_ACCESS (IS_MOBILE_DEFERREDSHADING_SUBPASS && PIXELSHADER && ALLOW_FRAMEBUFFER_FETCH)
 
-#if (METAL_PROFILE && FRAMEBUFFER_FETCH_ACCESS)
+#if (METAL_ES3_1_PROFILE && FRAMEBUFFER_FETCH_ACCESS)
 #include "/Engine/Public/Platform/Metal/MetalSubpassSupport.ush"
 #elif (VULKAN_PROFILE && FRAMEBUFFER_FETCH_ACCESS)
 #include "/Engine/Public/Platform/Vulkan/VulkanSubpassSupport.ush"
@@ -774,7 +774,7 @@ void MobileFetchGBuffer(in float2 UV, out half4 GBufferA, out half4 GBufferB, ou
 	#if MOBILE_EXTENDED_GBUFFER
 		GBufferD = VulkanSubpassFetch4();
 	#endif
-#elif (METAL_PROFILE && !MAC && FRAMEBUFFER_FETCH_ACCESS)
+#elif (METAL_ES3_1_PROFILE && !MAC && FRAMEBUFFER_FETCH_ACCESS)
 	GBufferA = SubpassFetchRGBA_1(); 
 	GBufferB = SubpassFetchRGBA_2(); 
 	GBufferC = SubpassFetchRGBA_3(); 
@@ -1225,4 +1225,4 @@ half3 AOMultiBounce(half3 BaseColor, half AO)
 	return max(AO, ((AO * a + b) * AO + c) * AO);
 }
 
-#define LIGHTING_CHANNELS_TEXTURE_DISTANCE_FIELD_REPRESENTATION_BIT 3
\ No newline at end of file
+#define LIGHTING_CHANNELS_TEXTURE_DISTANCE_FIELD_REPRESENTATION_BIT 3
diff --git a/Engine/Shaders/Private/DistanceFieldAOShared.ush b/Engine/Shaders/Private/DistanceFieldAOShared.ush
--- a/Engine/Shaders/Private/DistanceFieldAOShared.ush
+++ b/Engine/Shaders/Private/DistanceFieldAOShared.ush
@@ -45,7 +45,7 @@ SamplerState DistanceFieldNormalSampler;
 
 float4 EncodeDownsampledGBuffer(in float3 WorldNormal, float SceneDepth)
 {
-	#if METAL_PROFILE 
+	#if METAL_ES3_1_PROFILE 
 		// clamp max depth to avoid #inf
 		SceneDepth = min(SceneDepth, 65500.0f);
 	#endif
@@ -228,4 +228,4 @@ StructuredBuffer<uint> CulledTilesStartOffsetArray;
 
 #ifndef CULLED_TILE_DATA_STRIDE
 #define CULLED_TILE_DATA_STRIDE 1
-#endif
\ No newline at end of file
+#endif
diff --git a/Engine/Shaders/Private/DistanceFieldShadowing.usf b/Engine/Shaders/Private/DistanceFieldShadowing.usf
--- a/Engine/Shaders/Private/DistanceFieldShadowing.usf
+++ b/Engine/Shaders/Private/DistanceFieldShadowing.usf
@@ -803,7 +803,7 @@ void DistanceFieldShadowingCS(
 	Result = bShouldComputeShadowing ? NumIntersectingObjects / 256.0f : 0.0f;
 #endif
 
-	#if METAL_PROFILE 
+	#if METAL_ES3_1_PROFILE 
 		// clamp max depth to avoid #inf
 		SceneDepth = min(SceneDepth, 65500.0f);
 	#endif
diff --git a/Engine/Shaders/Private/GeometryCollectionVertexFactory.ush b/Engine/Shaders/Private/GeometryCollectionVertexFactory.ush
--- a/Engine/Shaders/Private/GeometryCollectionVertexFactory.ush
+++ b/Engine/Shaders/Private/GeometryCollectionVertexFactory.ush
@@ -51,7 +51,7 @@ struct FVertexFactoryInput
 	float4	Position	: ATTRIBUTE0;
 
 #if !MANUAL_VERTEX_FETCH
-	#if METAL_PROFILE
+	#if METAL_ES3_1_PROFILE
 		float3	TangentX	: ATTRIBUTE1;
 		// TangentZ.w contains sign of tangent basis determinant
 		float4	TangentZ	: ATTRIBUTE2;
diff --git a/Engine/Shaders/Private/HairStrands/HairCardsVertexFactory.ush b/Engine/Shaders/Private/HairStrands/HairCardsVertexFactory.ush
--- a/Engine/Shaders/Private/HairStrands/HairCardsVertexFactory.ush
+++ b/Engine/Shaders/Private/HairStrands/HairCardsVertexFactory.ush
@@ -41,7 +41,7 @@ struct FVertexFactoryInput
 	float4 PrevPosition  : ATTRIBUTE5;
 
 	// TangentZ.w contains sign of tangent basis determinant
-  #if METAL_PROFILE
+  #if METAL_ES3_1_PROFILE
 	float3	TangentX : ATTRIBUTE1;
 	float4	TangentZ : ATTRIBUTE2;
   #else
diff --git a/Engine/Shaders/Private/LocalVertexFactory.ush b/Engine/Shaders/Private/LocalVertexFactory.ush
--- a/Engine/Shaders/Private/LocalVertexFactory.ush
+++ b/Engine/Shaders/Private/LocalVertexFactory.ush
@@ -52,7 +52,7 @@ struct FVertexFactoryInput
 	float4	Position	: ATTRIBUTE0;
 
 #if !MANUAL_VERTEX_FETCH
-	#if METAL_PROFILE
+	#if METAL_ES3_1_PROFILE
 		float3	TangentX	: ATTRIBUTE1;
 		// TangentZ.w contains sign of tangent basis determinant
 		float4	TangentZ	: ATTRIBUTE2;
diff --git a/Engine/Shaders/Private/MaterialTemplate.ush b/Engine/Shaders/Private/MaterialTemplate.ush
--- a/Engine/Shaders/Private/MaterialTemplate.ush
+++ b/Engine/Shaders/Private/MaterialTemplate.ush
@@ -2991,7 +2991,7 @@ MaterialFloat2 GetLightmapUVs_DDY(FMaterialPixelParameters Parameters)
 
 //The post-process material needs to decode the scene color since it's encoded at PreTonemapMSAA if MSAA enabled on MetalMobilePlatorm 
 //The POST_PROCESS_MATERIAL_BEFORE_TONEMAP is 1 for both BL_SceneColorBeforeDOF and BL_SceneColorAfterDOF post-process materials
-#if FEATURE_LEVEL <= FEATURE_LEVEL_ES3_1 && POST_PROCESS_MATERIAL && POST_PROCESS_MATERIAL_BEFORE_TONEMAP && METAL_PROFILE
+#if FEATURE_LEVEL <= FEATURE_LEVEL_ES3_1 && POST_PROCESS_MATERIAL && POST_PROCESS_MATERIAL_BEFORE_TONEMAP && METAL_ES3_1_PROFILE
 uint bMetalMSAAHDRDecode;
 #endif
 
@@ -3061,7 +3061,7 @@ MaterialFloat4 MobileSceneTextureLookup(inout FMaterialPixelParameters Parameter
 	case PPI_PostProcessInput0:
 	{
 		MaterialFloat4 Input0 = Texture2DSample(PostProcessInput_0_Texture, PostProcessInput_0_SharedSampler, UV);
-		#if POST_PROCESS_MATERIAL_BEFORE_TONEMAP && METAL_PROFILE
+		#if POST_PROCESS_MATERIAL_BEFORE_TONEMAP && METAL_ES3_1_PROFILE
 		// Decode the input color since the color is encoded for MSAA 
 		// The decode instructions might be able to skip with dynamic branch
 		if (bMetalMSAAHDRDecode)
diff --git a/Engine/Shaders/Private/MobileBasePassPixelShader.usf b/Engine/Shaders/Private/MobileBasePassPixelShader.usf
--- a/Engine/Shaders/Private/MobileBasePassPixelShader.usf
+++ b/Engine/Shaders/Private/MobileBasePassPixelShader.usf
@@ -178,7 +178,7 @@ half3 FrameBufferBlendOp(half4 Source)
 #if MOBILE_TRANSLUCENT_COLOR_TRANSMITTANCE_PROGRAMMABLE_BLENDING
 	#if VULKAN_PROFILE
 		#define FramebufferFetchColor0 VulkanSubpassFetch0
-	#elif METAL_PROFILE && !MAC
+	#elif METAL_ES3_1_PROFILE && !MAC
 		#define FramebufferFetchColor0 SubpassFetchRGBA_0
 	#elif COMPILER_GLSL_ES3_1
 		#define FramebufferFetchColor0 FramebufferFetchES2
diff --git a/Engine/Shaders/Private/PostProcessHMD.usf b/Engine/Shaders/Private/PostProcessHMD.usf
--- a/Engine/Shaders/Private/PostProcessHMD.usf
+++ b/Engine/Shaders/Private/PostProcessHMD.usf
@@ -69,7 +69,7 @@ void MainPS(
 {
 	// Only testing 1 channel for optimization
 #if PIXELSHADER // Fix gles error since VS and PS are in the same file
-#if METAL_PROFILE || METAL_MRT_PROFILE || METAL_SM5_PROFILE || METAL_SM6_PROFILE // Fix Mac issue with bvec2
+#if METAL_ES3_1_PROFILE || METAL_SM5_IOS_TVOS_PROFILE || METAL_SM5_PROFILE || METAL_SM6_PROFILE // Fix Mac issue with bvec2
 	const float2 OOBCoords = TexCoord1.zw - saturate(TexCoord1.zw);
 	if( dot(OOBCoords, OOBCoords) != 0.0 )
 #else
diff --git a/Engine/Shaders/Private/PostProcessMobile.usf b/Engine/Shaders/Private/PostProcessMobile.usf
--- a/Engine/Shaders/Private/PostProcessMobile.usf
+++ b/Engine/Shaders/Private/PostProcessMobile.usf
@@ -245,7 +245,7 @@ void PreTonemapMSAA_Mobile(
 	out HALF4_TYPE OutColor : SV_Target0
 	)
 {
-#if (METAL_PROFILE && !MAC)
+#if (METAL_ES3_1_PROFILE && !MAC)
 	// On-chip pre-tonemap before MSAA resolve.
 	OutColor = SubpassFetchRGBA_0();
 	OutColor.rgb *= rcp(OutColor.r*0.299 + OutColor.g*0.587 + OutColor.b*0.114 + 1.0);
@@ -273,7 +273,7 @@ void MSAADecodeAndCopyRectPS(
 	out HALF4_TYPE OutColor : SV_Target0
 	)
 {
-#if (METAL_PROFILE && !MAC)
+#if (METAL_ES3_1_PROFILE && !MAC)
 	float2 UV = UVAndScreenPos.xy;
 	OutColor = Texture2DSample(InputTexture, InputSampler, UV);
 	OutColor.rgb *= rcp(OutColor.r*(-0.299) + OutColor.g*(-0.587) + OutColor.b*(-0.114) + 1.0);
diff --git a/Engine/Shaders/Private/SceneTexturesCommon.ush b/Engine/Shaders/Private/SceneTexturesCommon.ush
--- a/Engine/Shaders/Private/SceneTexturesCommon.ush
+++ b/Engine/Shaders/Private/SceneTexturesCommon.ush
@@ -219,7 +219,7 @@ float LookupDeviceZ( float2 ScreenUV, float ArrayIndex )
 	return DepthbufferFetchES2();
 #elif MOBILE_DEPTHFECTH && VULKAN_PROFILE
 	return VulkanSubpassDepthFetch();
-#elif MOBILE_DEPTHFECTH && (METAL_PROFILE && !MAC) && USE_SCENE_DEPTH_AUX
+#elif MOBILE_DEPTHFECTH && (METAL_ES3_1_PROFILE && !MAC) && USE_SCENE_DEPTH_AUX
 	return DepthbufferFetchES2();
 #elif (USE_SCENE_DEPTH_AUX && !MOBILE_DEFERRED_SHADING)
 	// SceneDepth texture is discarded after BasePass (with forward shading) 
diff --git a/Engine/Shaders/Private/ScreenSpaceShadows.usf b/Engine/Shaders/Private/ScreenSpaceShadows.usf
--- a/Engine/Shaders/Private/ScreenSpaceShadows.usf
+++ b/Engine/Shaders/Private/ScreenSpaceShadows.usf
@@ -144,7 +144,7 @@ void ScreenSpaceShadowsCS(
 
 	float Result = HitDistance > 0.0 ? (1.0f - GetContactShadowOcclusion(HitUV)) : 1.0f;
 
-	#if METAL_PROFILE 
+	#if METAL_ES3_1_PROFILE 
 		// clamp max depth to avoid #inf
 		SceneDepth = min(SceneDepth, 65500.0f);
 	#endif
diff --git a/Engine/Shaders/Private/SingleLayerWaterCommon.ush b/Engine/Shaders/Private/SingleLayerWaterCommon.ush
--- a/Engine/Shaders/Private/SingleLayerWaterCommon.ush
+++ b/Engine/Shaders/Private/SingleLayerWaterCommon.ush
@@ -22,11 +22,11 @@
 #if MATERIAL_SHADINGMODEL_SINGLELAYERWATER
 	#if SUPPORTS_INDEPENDENT_SAMPLERS
 		#define SingleLayerWaterSceneColorSampler View.SharedBilinearClampedSampler
-		#if METAL_PROFILE
+		#if METAL_ES3_1_PROFILE
 			#define SingleLayerWaterSceneDepthSampler View.SharedPointClampedSampler
 		#else // TODO: Dynamically use a point sampler instead if the depth texture was not downsampled.
 			#define SingleLayerWaterSceneDepthSampler View.SharedBilinearClampedSampler
-		#endif // METAL_PROFILE
+		#endif // METAL_ES3_1_PROFILE
 	#else
 		#define SingleLayerWaterSceneColorSampler SingleLayerWater.SceneColorWithoutSingleLayerWaterSampler
 		#define SingleLayerWaterSceneDepthSampler SingleLayerWater.SceneDepthWithoutSingleLayerWaterSampler
@@ -52,7 +52,7 @@ bool WaterRefract(float3 V, float3 N, float Eta, out float3 R)
 
 float WaterSampleSceneDepthWithoutWater(Texture2D SceneDepthWithoutWaterTexture, SamplerState SceneDepthWithoutWaterSampler, float2 UV, float2 SceneDepthWithoutWaterTextureSize, float2 SceneDepthWithoutWaterTextureInvSize)
 {
-#if METAL_PROFILE
+#if METAL_ES3_1_PROFILE
 	// On Metal we cannot do hardware bilinear sampling of float32 textures, so we need to do it manually.
 	FBilinearSampleInfos SampleInfos = GetBilinearSampleLevelInfos(UV, SceneDepthWithoutWaterTextureSize, SceneDepthWithoutWaterTextureInvSize);
 	float4 PixelSceneDepth4 = SceneDepthWithoutWaterTexture.Gather(SceneDepthWithoutWaterSampler, UV).wzxy;
@@ -60,7 +60,7 @@ float WaterSampleSceneDepthWithoutWater(Texture2D SceneDepthWithoutWaterTexture,
 #else
 	// A linear sampler should be passed in here if downsampling was used.
 	return SceneDepthWithoutWaterTexture.SampleLevel(SceneDepthWithoutWaterSampler, UV, 0).x;
-#endif // METAL_PROFILE
+#endif // METAL_ES3_1_PROFILE
 }
 
 #define SLW_MASK_IS_WATER_BITMASK (128.0f / 255.0f) // Unorm8 mask with 8th bit set
diff --git a/Engine/Shaders/Public/Platform.ush b/Engine/Shaders/Public/Platform.ush
--- a/Engine/Shaders/Public/Platform.ush
+++ b/Engine/Shaders/Public/Platform.ush
@@ -173,12 +173,22 @@
 #define ES3_1_PROFILE 0
 #endif
 
+#ifndef METAL_ES3_1_PROFILE
+#define METAL_ES3_1_PROFILE 0
+#endif
+
+// Deprecated, use METAL_ES3_1_PROFILE instead
 #ifndef METAL_PROFILE
-#define METAL_PROFILE 0
+#define METAL_PROFILE METAL_ES3_1_PROFILE
+#endif
+
+#ifndef METAL_SM5_IOS_TVOS_PROFILE
+#define METAL_SM5_IOS_TVOS_PROFILE 0
 #endif
 
+// Deprecated, use METAL_SM5_IOS_TVOS_PROFILE instead
 #ifndef METAL_MRT_PROFILE
-#define METAL_MRT_PROFILE 0
+#define METAL_MRT_PROFILE METAL_SM5_IOS_TVOS_PROFILE
 #endif
 
 #ifndef METAL_SM5_PROFILE
@@ -321,7 +331,7 @@
 #define FORCE_FLOATS 0
 #endif
 
-#if COMPILER_DXC && !FORCE_FLOATS && (ES3_1_PROFILE || METAL_PROFILE || VULKAN_PROFILE_SM5)
+#if COMPILER_DXC && !FORCE_FLOATS && (ES3_1_PROFILE || METAL_ES3_1_PROFILE || VULKAN_PROFILE_SM5)
 	#pragma dxc diagnostic ignored "-Wconversion"
 #endif
 
@@ -334,7 +344,7 @@
 	#define half3x3 min16float3x3
 	#define half3x4 min16float3x4
 	#define half4x4 min16float4x4
-#elif ((!(ES3_1_PROFILE || METAL_PROFILE || VULKAN_PROFILE_SM5)) && !PLATFORM_SUPPORTS_REAL_TYPES) || FORCE_FLOATS
+#elif ((!(ES3_1_PROFILE || METAL_ES3_1_PROFILE || VULKAN_PROFILE_SM5)) && !PLATFORM_SUPPORTS_REAL_TYPES) || FORCE_FLOATS
 	// Always use floats when using the ES3/METAL compiler, because platforms not optimized for lower precision, 
 	// And we don't want potential side effects on other platforms
 	#define half float
@@ -400,16 +410,16 @@
 	#define FEATURE_LEVEL FEATURE_LEVEL_SM6
 	#define STENCIL_COMPONENT_SWIZZLE .x
 
-#elif METAL_PROFILE
+#elif METAL_ES3_1_PROFILE
 	#define FEATURE_LEVEL  FEATURE_LEVEL_ES3_1
-	// @todo metal: remove this and make sure all uses handle METAL_PROFILE
+	// @todo metal: remove this and make sure all uses handle METAL_ES3_1_PROFILE
 	#undef ES3_1_PROFILE
 	#define ES3_1_PROFILE 1
 	#define FCOLOR_COMPONENT_SWIZZLE .rgba
 	#define FMANUALFETCH_COLOR_COMPONENT_SWIZZLE .bgra
 	#define STENCIL_COMPONENT_SWIZZLE .x
 
-#elif METAL_MRT_PROFILE
+#elif METAL_SM5_IOS_TVOS_PROFILE
 	#define FEATURE_LEVEL FEATURE_LEVEL_SM5
 	#define FCOLOR_COMPONENT_SWIZZLE .rgba
 	#define FMANUALFETCH_COLOR_COMPONENT_SWIZZLE .bgra
@@ -487,7 +497,7 @@
 	#define UB_INT(Dim)         int##Dim
 	#define UB_UINT(Dim)        uint##Dim
 	#define UB_FLOAT(Dim)       float##Dim
-	#if METAL_PROFILE || COMPILER_HLSL
+	#if METAL_ES3_1_PROFILE || COMPILER_HLSL
 		#define UB_HALF_FLOAT(Dim)  float##Dim
 	#else
 		#define UB_HALF_FLOAT(Dim)  half##Dim
@@ -716,7 +726,7 @@ void ShaderYield()
 #define COMPILER_SUPPORTS_MED3 0
 #endif
 
-#if (ES3_1_PROFILE && !METAL_PROFILE) || VULKAN_PROFILE_SM5
+#if (ES3_1_PROFILE && !METAL_ES3_1_PROFILE) || VULKAN_PROFILE_SM5
 #define HALF_TYPE half
 #define HALF2_TYPE half2
 #define HALF3_TYPE half3
@@ -750,7 +760,7 @@ void ShaderYield()
 
 #endif // SM6_PROFILE || SM5_PROFILE || COMPILER_SUPPORTS_ATTRIBUTES
 
-#if SM6_PROFILE || SM5_PROFILE || METAL_MRT_PROFILE || METAL_SM5_PROFILE || METAL_SM6_PROFILE || ES3_1_PROFILE || VULKAN_PROFILE_SM5 || VULKAN_PROFILE_SM6
+#if SM6_PROFILE || SM5_PROFILE || METAL_SM5_IOS_TVOS_PROFILE || METAL_SM5_PROFILE || METAL_SM6_PROFILE || ES3_1_PROFILE || VULKAN_PROFILE_SM5 || VULKAN_PROFILE_SM6
 #define EARLYDEPTHSTENCIL [earlydepthstencil]
 #endif
 
diff --git a/Engine/Source/Developer/Apple/MetalShaderFormat/Private/MetalShaderCompiler.cpp b/Engine/Source/Developer/Apple/MetalShaderFormat/Private/MetalShaderCompiler.cpp
--- a/Engine/Source/Developer/Apple/MetalShaderFormat/Private/MetalShaderCompiler.cpp
+++ b/Engine/Source/Developer/Apple/MetalShaderFormat/Private/MetalShaderCompiler.cpp
@@ -286,10 +286,10 @@ void BuildMetalShaderOutput(
 	const bool bBindlessEnabled = (ShaderInput.Environment.CompilerFlags.Contains(CFLAG_BindlessResources) || ShaderInput.Environment.CompilerFlags.Contains(CFLAG_BindlessSamplers));
 
 	//TODO read from toolchain
-	const bool bIsMobile = (ShaderInput.Target.Platform == SP_METAL || ShaderInput.Target.Platform == SP_METAL_MRT || ShaderInput.Target.Platform == SP_METAL_TVOS || ShaderInput.Target.Platform == SP_METAL_MRT_TVOS || ShaderInput.Target.Platform == SP_METAL_SIM);
+	const bool bIsMobile = (ShaderInput.Target.Platform == SP_METAL_ES3_1_IOS || ShaderInput.Target.Platform == SP_METAL_SM5_IOS || ShaderInput.Target.Platform == SP_METAL_ES3_1_TVOS || ShaderInput.Target.Platform == SP_METAL_SM5_TVOS || ShaderInput.Target.Platform == SP_METAL_SIM);
 	bool bNoFastMath = ShaderInput.Environment.CompilerFlags.Contains(CFLAG_NoFastMath);
 	const bool bUsingWPO = ShaderInput.Environment.GetCompileArgument(TEXT("USES_WORLD_POSITION_OFFSET"), false);
-	if (bUsingWPO && (ShaderInput.Target.Platform == SP_METAL_MRT || ShaderInput.Target.Platform == SP_METAL_MRT_TVOS) && Frequency == SF_Vertex)
+	if (bUsingWPO && (ShaderInput.Target.Platform == SP_METAL_SM5_IOS || ShaderInput.Target.Platform == SP_METAL_SM5_TVOS) && Frequency == SF_Vertex)
 	{
 		// WPO requires that we make all multiply/sincos instructions invariant :(
 		bNoFastMath = true;
@@ -892,13 +892,12 @@ void CompileMetalShader(const FShaderCompilerInput& Input, const FShaderPreproce
 	}
 
 	EMetalGPUSemantics Semantics = EMetalGPUSemanticsMobile;
-	if (Input.ShaderFormat == NAME_SF_METAL_MRT
-		|| Input.ShaderFormat == NAME_SF_METAL_MRT_TVOS
-		|| Input.ShaderFormat == NAME_SF_METAL_MRT_MAC)
+	if (Input.ShaderFormat == NAME_SF_METAL_SM5_IOS
+		|| Input.ShaderFormat == NAME_SF_METAL_SM5_TVOS)
 	{
 		Semantics = EMetalGPUSemanticsTBDRDesktop;
 	}
-	else if (Input.ShaderFormat == NAME_SF_METAL_MACES3_1
+	else if (Input.ShaderFormat == NAME_SF_METAL_ES3_1
 		|| Input.ShaderFormat == NAME_SF_METAL_SM5
 		|| Input.ShaderFormat == NAME_SF_METAL_SM6)
 	{
@@ -914,7 +913,7 @@ void CompileMetalShader(const FShaderCompilerInput& Input, const FShaderPreproce
 
 	// TODO read from toolchain
 	const bool bIsMobile = FMetalCompilerToolchain::Get()->IsMobile((EShaderPlatform)Input.Target.Platform);
-	const bool bAppleTV = (Input.ShaderFormat == NAME_SF_METAL_TVOS || Input.ShaderFormat == NAME_SF_METAL_MRT_TVOS);
+	const bool bAppleTV = (Input.ShaderFormat == NAME_SF_METAL_ES3_1_TVOS || Input.ShaderFormat == NAME_SF_METAL_SM5_TVOS);
 	const bool bIsSimulator = (Input.ShaderFormat == NAME_SF_METAL_SIM);
 
 	FString MinOSVersion;
diff --git a/Engine/Source/Developer/Apple/MetalShaderFormat/Private/MetalShaderFormat.cpp b/Engine/Source/Developer/Apple/MetalShaderFormat/Private/MetalShaderFormat.cpp
--- a/Engine/Source/Developer/Apple/MetalShaderFormat/Private/MetalShaderFormat.cpp
+++ b/Engine/Source/Developer/Apple/MetalShaderFormat/Private/MetalShaderFormat.cpp
@@ -103,28 +103,26 @@ public:
 
 	virtual void GetSupportedFormats(TArray<FName>& OutFormats) const override final
 	{
-		OutFormats.Add(NAME_SF_METAL);
-		OutFormats.Add(NAME_SF_METAL_MRT);
-		OutFormats.Add(NAME_SF_METAL_TVOS);
-		OutFormats.Add(NAME_SF_METAL_MRT_TVOS);
+		OutFormats.Add(NAME_SF_METAL_ES3_1_IOS);
+		OutFormats.Add(NAME_SF_METAL_SM5_IOS);
+		OutFormats.Add(NAME_SF_METAL_ES3_1_TVOS);
+		OutFormats.Add(NAME_SF_METAL_SM5_TVOS);
 		OutFormats.Add(NAME_SF_METAL_SM5);
 		OutFormats.Add(NAME_SF_METAL_SM6);
 		OutFormats.Add(NAME_SF_METAL_SIM);
-		OutFormats.Add(NAME_SF_METAL_MACES3_1);
-		OutFormats.Add(NAME_SF_METAL_MRT_MAC);
+		OutFormats.Add(NAME_SF_METAL_ES3_1);
 	}
 
 	void CheckShaderFormat(FName Format) const
 	{
-		check(Format == NAME_SF_METAL
-			|| Format == NAME_SF_METAL_MRT
-			|| Format == NAME_SF_METAL_TVOS
-			|| Format == NAME_SF_METAL_MRT_TVOS
+		check(Format == NAME_SF_METAL_ES3_1_IOS
+			|| Format == NAME_SF_METAL_SM5_IOS
+			|| Format == NAME_SF_METAL_ES3_1_TVOS
+			|| Format == NAME_SF_METAL_SM5_TVOS
 			|| Format == NAME_SF_METAL_SM5
 			|| Format == NAME_SF_METAL_SM6
 			|| Format == NAME_SF_METAL_SIM
-			|| Format == NAME_SF_METAL_MACES3_1
-			|| Format == NAME_SF_METAL_MRT_MAC);
+			|| Format == NAME_SF_METAL_ES3_1);
 	}
 
 	virtual bool PreprocessShader(
@@ -188,7 +186,7 @@ public:
 		check(Components.Num() == 2);
 		FName ShaderFormatName(Components[0]);
 
-		check(ShaderFormatName == NAME_SF_METAL || ShaderFormatName == NAME_SF_METAL_MRT || ShaderFormatName == NAME_SF_METAL_TVOS || ShaderFormatName == NAME_SF_METAL_MRT_TVOS || ShaderFormatName == NAME_SF_METAL_SM5 || ShaderFormatName == NAME_SF_METAL_SM6 || ShaderFormatName == NAME_SF_METAL_SIM || ShaderFormatName == NAME_SF_METAL_MACES3_1 || ShaderFormatName == NAME_SF_METAL_MRT_MAC);
+		check(ShaderFormatName == NAME_SF_METAL_ES3_1_IOS || ShaderFormatName == NAME_SF_METAL_SM5_IOS || ShaderFormatName == NAME_SF_METAL_ES3_1_TVOS || ShaderFormatName == NAME_SF_METAL_SM5_TVOS || ShaderFormatName == NAME_SF_METAL_SM5 || ShaderFormatName == NAME_SF_METAL_SM6 || ShaderFormatName == NAME_SF_METAL_SIM || ShaderFormatName == NAME_SF_METAL_ES3_1);
 
 		// SM6 needs a lower limit of shaders per library as the packing process takes a significant amount of RAM
 		if(ShaderFormatName == NAME_SF_METAL_SM6)
@@ -364,21 +362,16 @@ public:
 
 		Input.Environment.SetDefine(TEXT("COMPILER_METAL"), 1);
 
-		if (Input.ShaderFormat == NAME_SF_METAL || Input.ShaderFormat == NAME_SF_METAL_TVOS)
+		if (Input.ShaderFormat == NAME_SF_METAL_ES3_1_IOS ||
+			Input.ShaderFormat == NAME_SF_METAL_ES3_1_TVOS ||
+			Input.ShaderFormat == NAME_SF_METAL_SIM ||
+			Input.ShaderFormat == NAME_SF_METAL_ES3_1)
 		{
-			Input.Environment.SetDefine(TEXT("METAL_PROFILE"), 1);
+			Input.Environment.SetDefine(TEXT("METAL_ES3_1_PROFILE"), 1);
 		}
-		else if (Input.ShaderFormat == NAME_SF_METAL_SIM)
+		else if (Input.ShaderFormat == NAME_SF_METAL_SM5_IOS || Input.ShaderFormat == NAME_SF_METAL_SM5_TVOS)
 		{
-			Input.Environment.SetDefine(TEXT("METAL_PROFILE"), 1);
-		}
-		else if (Input.ShaderFormat == NAME_SF_METAL_MRT || Input.ShaderFormat == NAME_SF_METAL_MRT_TVOS)
-		{
-			Input.Environment.SetDefine(TEXT("METAL_MRT_PROFILE"), 1);
-		}
-		else if (Input.ShaderFormat == NAME_SF_METAL_MACES3_1)
-		{
-			Input.Environment.SetDefine(TEXT("METAL_PROFILE"), 1);
+			Input.Environment.SetDefine(TEXT("METAL_SM5_IOS_TVOS_PROFILE"), 1);
 		}
 		else if (Input.ShaderFormat == NAME_SF_METAL_SM5)
 		{
@@ -388,10 +381,6 @@ public:
 		{
 			Input.Environment.SetDefine(TEXT("METAL_SM6_PROFILE"), 1);
 		}
-		else if (Input.ShaderFormat == NAME_SF_METAL_MRT_MAC)
-		{
-			Input.Environment.SetDefine(TEXT("METAL_MRT_PROFILE"), 1);
-		}
 
 		Input.Environment.SetDefine(TEXT("COMPILER_HLSLCC"), 2);
 #if UE_METAL_USE_METAL_SHADER_CONVERTER
@@ -649,15 +638,14 @@ void FMetalCompilerToolchain::Destroy()
 
 EShaderPlatform FMetalCompilerToolchain::MetalShaderFormatToLegacyShaderPlatform(FName ShaderFormat)
 {
-	if (ShaderFormat == NAME_SF_METAL)				return SP_METAL;
-	if (ShaderFormat == NAME_SF_METAL_MRT)			return SP_METAL_MRT;
-	if (ShaderFormat == NAME_SF_METAL_TVOS)			return SP_METAL_TVOS;
-	if (ShaderFormat == NAME_SF_METAL_MRT_TVOS)		return SP_METAL_MRT_TVOS;
-	if (ShaderFormat == NAME_SF_METAL_MRT_MAC)		return SP_METAL_MRT_MAC;
+	if (ShaderFormat == NAME_SF_METAL_ES3_1_IOS)	return SP_METAL_ES3_1_IOS;
+	if (ShaderFormat == NAME_SF_METAL_SM5_IOS)		return SP_METAL_SM5_IOS;
+	if (ShaderFormat == NAME_SF_METAL_ES3_1_TVOS)	return SP_METAL_ES3_1_TVOS;
+	if (ShaderFormat == NAME_SF_METAL_SM5_TVOS)		return SP_METAL_SM5_TVOS;
 	if (ShaderFormat == NAME_SF_METAL_SM5)			return SP_METAL_SM5;
     if (ShaderFormat == NAME_SF_METAL_SM6)          return SP_METAL_SM6;
     if (ShaderFormat == NAME_SF_METAL_SIM)          return SP_METAL_SIM;
-	if (ShaderFormat == NAME_SF_METAL_MACES3_1)		return SP_METAL_MACES3_1;
+	if (ShaderFormat == NAME_SF_METAL_ES3_1)		return SP_METAL_ES3_1;
 
 	return SP_NumPlatforms;
 }
@@ -1134,6 +1122,9 @@ bool FMetalCompilerToolchain::CompileMetalShader(FMetalShaderBytecodeJob& Job, F
 		{
 			MetalParams += TEXT(" -Os");
 		}
+		
+		// We don't use incremental build and this void creating temporary files in C drive in Windows
+		MetalParams += TEXT(" -fno-temp-file");
 
 		bool bSuccess = this->ExecMetalFrontend(SDK, *MetalParams, &Job.ReturnCode, &Job.Results, &Job.Errors);
 
diff --git a/Engine/Source/Developer/Apple/MetalShaderFormat/Private/MetalShaderFormat.h b/Engine/Source/Developer/Apple/MetalShaderFormat/Private/MetalShaderFormat.h
--- a/Engine/Source/Developer/Apple/MetalShaderFormat/Private/MetalShaderFormat.h
+++ b/Engine/Source/Developer/Apple/MetalShaderFormat/Private/MetalShaderFormat.h
@@ -24,15 +24,14 @@ enum EAppleSDKType
 	AppleSDKCount,
 };
 
-static FName NAME_SF_METAL(TEXT("SF_METAL"));
-static FName NAME_SF_METAL_MRT(TEXT("SF_METAL_MRT"));
-static FName NAME_SF_METAL_TVOS(TEXT("SF_METAL_TVOS"));
-static FName NAME_SF_METAL_MRT_TVOS(TEXT("SF_METAL_MRT_TVOS"));
+static FName NAME_SF_METAL_ES3_1_IOS(TEXT("SF_METAL_ES3_1_IOS"));
+static FName NAME_SF_METAL_SM5_IOS(TEXT("SF_METAL_SM5_IOS"));
+static FName NAME_SF_METAL_ES3_1_TVOS(TEXT("SF_METAL_ES3_1_TVOS"));
+static FName NAME_SF_METAL_SM5_TVOS(TEXT("SF_METAL_SM5_TVOS"));
 static FName NAME_SF_METAL_SM5(TEXT("SF_METAL_SM5"));
 static FName NAME_SF_METAL_SM6(TEXT("SF_METAL_SM6"));
 static FName NAME_SF_METAL_SIM(TEXT("SF_METAL_SIM"));
-static FName NAME_SF_METAL_MACES3_1(TEXT("SF_METAL_MACES3_1"));
-static FName NAME_SF_METAL_MRT_MAC(TEXT("SF_METAL_MRT_MAC"));
+static FName NAME_SF_METAL_ES3_1(TEXT("SF_METAL_ES3_1"));
 
 DECLARE_LOG_CATEGORY_EXTERN(LogMetalCompilerSetup, Log, All);
 DECLARE_LOG_CATEGORY_EXTERN(LogMetalShaderCompiler, Log, All);
@@ -137,7 +136,7 @@ public:
 	// True if Platform indicates a mobile (ios, tvos) platform
 	static bool IsMobile(const EShaderPlatform Platform)
 	{
-		return (Platform == SP_METAL || Platform == SP_METAL_MRT || Platform == SP_METAL_SIM || Platform == SP_METAL_TVOS || Platform == SP_METAL_MRT_TVOS);
+		return (Platform == SP_METAL_ES3_1_IOS || Platform == SP_METAL_SM5_IOS || Platform == SP_METAL_SIM || Platform == SP_METAL_ES3_1_TVOS || Platform == SP_METAL_SM5_TVOS);
 	}
 
 	// True if Format is the FName of a mobile shader format (ios, tvos)
@@ -150,7 +149,7 @@ public:
 	// True if Platform is a tvos platform
 	static bool IsTVOS(const EShaderPlatform Platform)
 	{
-		return Platform == SP_METAL_TVOS || Platform == SP_METAL_MRT_TVOS;
+		return Platform == SP_METAL_ES3_1_TVOS || Platform == SP_METAL_SM5_TVOS;
 	}
 
 	// Converts an FName ShaderFormat to the equivalent ShaderPlatform
diff --git a/Engine/Source/Developer/IOS/IOSPlatformEditor/Private/IOSPlatformEditorModule.cpp b/Engine/Source/Developer/IOS/IOSPlatformEditor/Private/IOSPlatformEditorModule.cpp
--- a/Engine/Source/Developer/IOS/IOSPlatformEditor/Private/IOSPlatformEditorModule.cpp
+++ b/Engine/Source/Developer/IOS/IOSPlatformEditor/Private/IOSPlatformEditorModule.cpp
@@ -57,8 +57,8 @@ void FIOSPlatformEditorModule::StartupModule()
 		SelectedSection->OnSelect().BindRaw(this, &FIOSPlatformEditorModule::HandleSelectIOSSection);
 
 		{
-			static FName NAME_SF_METAL(TEXT("SF_METAL"));
-			UShaderPlatformQualitySettings* IOSMaterialQualitySettings = UMaterialShaderQualitySettings::Get()->GetShaderPlatformQualitySettings(NAME_SF_METAL);
+			static FName NAME_SF_METAL_ES3_1_IOS(TEXT("SF_METAL_ES3_1_IOS"));
+			UShaderPlatformQualitySettings* IOSMaterialQualitySettings = UMaterialShaderQualitySettings::Get()->GetShaderPlatformQualitySettings(NAME_SF_METAL_ES3_1_IOS);
 			SettingsModule->RegisterSettings("Project", "Platforms", "iOSMetalQuality",
 				LOCTEXT("IOSMetalQualitySettingsName", "iOS Material Quality"),
 				LOCTEXT("IOSMetalQualitySettingsDescription", "Settings for iOS material quality"),
diff --git a/Engine/Source/Developer/IOS/IOSTargetPlatformSettings/Private/IOSTargetPlatformSettings.cpp b/Engine/Source/Developer/IOS/IOSTargetPlatformSettings/Private/IOSTargetPlatformSettings.cpp
--- a/Engine/Source/Developer/IOS/IOSTargetPlatformSettings/Private/IOSTargetPlatformSettings.cpp
+++ b/Engine/Source/Developer/IOS/IOSTargetPlatformSettings/Private/IOSTargetPlatformSettings.cpp
@@ -53,6 +53,7 @@ static bool SupportsMetal()
 	return bSupportsMetal;
 }
 
+// MetalMRT (Multiple Render Target) represent the SM5 variant of metal for IOS 
 static bool SupportsMetalMRT()
 {
 	// default to NOT supporting metal MRT
@@ -111,31 +112,31 @@ bool FIOSTargetPlatformSettings::SupportsFeature( ETargetPlatformFeatures Featur
 
 void FIOSTargetPlatformSettings::GetAllPossibleShaderFormats( TArray<FName>& OutFormats ) const
 {
-	static FName NAME_SF_METAL(TEXT("SF_METAL"));
+	static FName NAME_SF_METAL_ES3_1_IOS(TEXT("SF_METAL_ES3_1_IOS"));
 	static FName NAME_SF_METAL_SIM(TEXT("SF_METAL_SIM"));
-	static FName NAME_SF_METAL_MRT(TEXT("SF_METAL_MRT"));
-	static FName NAME_SF_METAL_TVOS(TEXT("SF_METAL_TVOS"));
-	static FName NAME_SF_METAL_MRT_TVOS(TEXT("SF_METAL_MRT_TVOS"));
+	static FName NAME_SF_METAL_SM5_IOS(TEXT("SF_METAL_SM5_IOS"));
+	static FName NAME_SF_METAL_ES3_1_TVOS(TEXT("SF_METAL_ES3_1_TVOS"));
+	static FName NAME_SF_METAL_SM5_TVOS(TEXT("SF_METAL_SM5_TVOS"));
 
 	if (bIsTVOS)
 	{
 		if (SupportsMetalMRT())
 		{
-			OutFormats.AddUnique(NAME_SF_METAL_MRT_TVOS);
+			OutFormats.AddUnique(NAME_SF_METAL_SM5_TVOS);
 		}
 
 		// because we are currently using IOS settings, we will always use metal, even if Metal isn't listed as being supported
 		// however, if MetalMRT is specific and Metal is set to false, then we will just use MetalMRT
 		if (SupportsMetal() || !SupportsMetalMRT())
 		{
-			OutFormats.AddUnique(NAME_SF_METAL_TVOS);
+			OutFormats.AddUnique(NAME_SF_METAL_ES3_1_TVOS);
 		}
 	}
 	else
 	{
 		if (SupportsMetal())
 		{
-			OutFormats.AddUnique(NAME_SF_METAL);
+			OutFormats.AddUnique(NAME_SF_METAL_ES3_1_IOS);
 
 			bool bEnableSimulatorSupport = false;
 			GConfig->GetBool(TEXT("/Script/IOSRuntimeSettings.IOSRuntimeSettings"), TEXT("bEnableSimulatorSupport"), bEnableSimulatorSupport, GEngineIni);
@@ -147,7 +148,7 @@ void FIOSTargetPlatformSettings::GetAllPossibleShaderFormats( TArray<FName>& Out
 
 		if (SupportsMetalMRT())
 		{
-			OutFormats.AddUnique(NAME_SF_METAL_MRT);
+			OutFormats.AddUnique(NAME_SF_METAL_SM5_IOS);
 		}
 	}
 }
diff --git a/Engine/Source/Developer/Mac/MacTargetPlatformSettings/Public/GenericMacTargetPlatformSettings.h b/Engine/Source/Developer/Mac/MacTargetPlatformSettings/Public/GenericMacTargetPlatformSettings.h
--- a/Engine/Source/Developer/Mac/MacTargetPlatformSettings/Public/GenericMacTargetPlatformSettings.h
+++ b/Engine/Source/Developer/Mac/MacTargetPlatformSettings/Public/GenericMacTargetPlatformSettings.h
@@ -69,10 +69,8 @@ public:
 			OutFormats.AddUnique(NAME_SF_METAL_SM5);
             static FName NAME_SF_METAL_SM6(TEXT("SF_METAL_SM6"));
             OutFormats.AddUnique(NAME_SF_METAL_SM6);
-			static FName NAME_SF_METAL_MACES3_1(TEXT("SF_METAL_MACES3_1"));
-			OutFormats.AddUnique(NAME_SF_METAL_MACES3_1);
-			static FName NAME_SF_METAL_MRT_MAC(TEXT("SF_METAL_MRT_MAC"));
-			OutFormats.AddUnique(NAME_SF_METAL_MRT_MAC);
+			static FName NAME_SF_METAL_ES3_1(TEXT("SF_METAL_ES3_1"));
+			OutFormats.AddUnique(NAME_SF_METAL_ES3_1);
 		}
 	}
 
diff --git a/Engine/Source/Editor/DetailCustomizations/Private/MacTargetSettingsDetails.cpp b/Engine/Source/Editor/DetailCustomizations/Private/MacTargetSettingsDetails.cpp
--- a/Engine/Source/Editor/DetailCustomizations/Private/MacTargetSettingsDetails.cpp
+++ b/Engine/Source/Editor/DetailCustomizations/Private/MacTargetSettingsDetails.cpp
@@ -121,16 +121,16 @@ static FText GetFriendlyNameFromRHINameMac(FName InRHIName)
 	case SP_OPENGL_ES3_1_ANDROID:
 		FriendlyRHIName = LOCTEXT("OpenGLESMobile", "OpenGLES (Mobile)");
 		break;
-	case SP_METAL:
+	case SP_METAL_ES3_1_IOS:
 		FriendlyRHIName = LOCTEXT("Metal", "iOS Metal Mobile Renderer (Mobile, Metal 2.4+, iOS 15.0 or later)");
 		break;
-	case SP_METAL_MRT:
+	case SP_METAL_SM5_IOS:
 		FriendlyRHIName = LOCTEXT("MetalMRT", "iOS Metal Desktop Renderer (SM5, Metal 2.4+, iOS 15.0 or later)");
 		break;
-	case SP_METAL_TVOS:
+	case SP_METAL_ES3_1_TVOS:
 		FriendlyRHIName = LOCTEXT("MetalTV", "tvOS Metal Mobile Renderer (Mobile, Metal 2.4+, tvOS 15.0 or later)");
 		break;
-	case SP_METAL_MRT_TVOS:
+	case SP_METAL_SM5_TVOS:
 		FriendlyRHIName = LOCTEXT("MetalMRTTV", "tvOS Metal Desktop Renderer (SM5, Metal 2.4+, tvOS 15.0 or later)");
 		break;
 	case SP_METAL_SM5:
@@ -142,12 +142,9 @@ static FText GetFriendlyNameFromRHINameMac(FName InRHIName)
 	case SP_METAL_SIM:
 		FriendlyRHIName = LOCTEXT("MetalSim", "iOS Metal Simulator Mobile Renderer (Simulator, Metal 2.4+, iOS 15.0 or later)");
 		break;
-	case SP_METAL_MACES3_1:
+	case SP_METAL_ES3_1:
 		FriendlyRHIName = LOCTEXT("MetalMobile", "Mac Metal High-End Mobile Preview (Mobile Preview)");
 		break;
-	case SP_METAL_MRT_MAC:
-		FriendlyRHIName = LOCTEXT("MetalMRTMac", "Mac Metal iOS/tvOS Desktop Renderer Preview (SM5)");
-		break;
 	case SP_VULKAN_SM5:
 	case SP_VULKAN_SM5_ANDROID:
 		FriendlyRHIName = LOCTEXT("VulkanSM5", "Vulkan (SM5)");
diff --git a/Engine/Source/Editor/MaterialEditor/Private/MaterialStats.cpp b/Engine/Source/Editor/MaterialEditor/Private/MaterialStats.cpp
--- a/Engine/Source/Editor/MaterialEditor/Private/MaterialStats.cpp
+++ b/Engine/Source/Editor/MaterialEditor/Private/MaterialStats.cpp
@@ -650,12 +650,11 @@ void FMaterialStats::BuildShaderPlatformDB(const bool bAllowIgnoringCompilationE
 	// macOS
 	AddShaderPlatform(EPlatformCategoryType::Desktop, SP_METAL_SM5, TEXT("Metal SM5"), true, TEXT("macOS, Metal, Shader Model 5"));
 	AddShaderPlatform(EPlatformCategoryType::Desktop, SP_METAL_SM6, TEXT("Metal SM6"), true, TEXT("macOS, Metal, Shader Model 6"), !bAllowIgnoringCompilationErrors);
-	AddShaderPlatform(EPlatformCategoryType::Desktop, SP_METAL_MRT_MAC, TEXT("Metal SM5 (MRT)"), true, TEXT("macOS, Metal, Shader Model 5"));
 #endif
 
 	// iOS
-	AddShaderPlatform(EPlatformCategoryType::IOS, SP_METAL, TEXT("Metal"), true, TEXT("iOS, Metal, Mobile"));
-	AddShaderPlatform(EPlatformCategoryType::IOS, SP_METAL_MRT, TEXT("Metal MRT"), true, TEXT("iOS, Metal, Shader Model 5"));
+	AddShaderPlatform(EPlatformCategoryType::IOS, SP_METAL_ES3_1_IOS, TEXT("IOS Metal ES3_1"), true, TEXT("iOS, Metal, Mobile"));
+	AddShaderPlatform(EPlatformCategoryType::IOS, SP_METAL_SM5_IOS, TEXT("IOS Metal SM5"), true, TEXT("iOS, Metal, Shader Model 5"));
 	AddShaderPlatform(EPlatformCategoryType::IOS, SP_METAL_SIM, TEXT("Metal Simulator"), true, TEXT("iOS, Metal, Mobile"));
 
 	ITargetPlatformManagerModule& TPM = GetTargetPlatformManagerRef();
diff --git a/Engine/Source/Editor/PIEPreviewDeviceProfileSelector/Private/PIEPreviewDevice.cpp b/Engine/Source/Editor/PIEPreviewDeviceProfileSelector/Private/PIEPreviewDevice.cpp
--- a/Engine/Source/Editor/PIEPreviewDeviceProfileSelector/Private/PIEPreviewDevice.cpp
+++ b/Engine/Source/Editor/PIEPreviewDeviceProfileSelector/Private/PIEPreviewDevice.cpp
@@ -365,7 +365,7 @@ void FPIEPreviewDevice::ApplyRHIOverrides() const
 
 		case EPIEPreviewDeviceType::IOS:
 		{
-			PreviewPlatform = SP_METAL_MACES3_1;
+			PreviewPlatform = SP_METAL_ES3_1;
 			RHIOverrideState = &DeviceSpecs->IOSProperties.MetalRHIState;
 		}
 		break;
diff --git a/Engine/Source/Editor/UnrealEd/Private/PreferenceStubs.cpp b/Engine/Source/Editor/UnrealEd/Private/PreferenceStubs.cpp
--- a/Engine/Source/Editor/UnrealEd/Private/PreferenceStubs.cpp
+++ b/Engine/Source/Editor/UnrealEd/Private/PreferenceStubs.cpp
@@ -86,7 +86,7 @@ UMaterialStatsOptions::UMaterialStatsOptions(const FObjectInitializer& ObjectIni
 #if PLATFORM_WINDOWS
 	bPlatformUsed[SP_PCD3D_ES3_1] = 1;
 #elif PLATFORM_MAC
-	bPlatformUsed[SP_METAL] = 1;
+	bPlatformUsed[SP_METAL_ES3_1_IOS] = 1;
 #elif PLATFORM_LINUX
 	bPlatformUsed[SP_VULKAN_PCES3_1] = 1;
 #endif
diff --git a/Engine/Source/Runtime/Apple/MetalRHI/Private/MetalRHI.cpp b/Engine/Source/Runtime/Apple/MetalRHI/Private/MetalRHI.cpp
--- a/Engine/Source/Runtime/Apple/MetalRHI/Private/MetalRHI.cpp
+++ b/Engine/Source/Runtime/Apple/MetalRHI/Private/MetalRHI.cpp
@@ -81,11 +81,11 @@ static void ValidateTargetedRHIFeatureLevelExists(EShaderPlatform Platform)
 		}
 	}
 #else
-	if (Platform == SP_METAL || Platform == SP_METAL_TVOS || Platform == SP_METAL_SIM)
+	if (Platform == SP_METAL_ES3_1_IOS || Platform == SP_METAL_ES3_1_TVOS || Platform == SP_METAL_SIM)
 	{
 		GConfig->GetBool(TEXT("/Script/IOSRuntimeSettings.IOSRuntimeSettings"), TEXT("bSupportsMetal"), bSupportsShaderPlatform, GEngineIni);
 	}
-	else if (Platform == SP_METAL_MRT || Platform == SP_METAL_MRT_TVOS)
+	else if (Platform == SP_METAL_SM5_IOS || Platform == SP_METAL_SM5_TVOS)
 	{
 		GConfig->GetBool(TEXT("/Script/IOSRuntimeSettings.IOSRuntimeSettings"), TEXT("bSupportsMetalMRT"), bSupportsShaderPlatform, GEngineIni);
 	}
@@ -292,11 +292,11 @@ FMetalDynamicRHI::FMetalDynamicRHI(ERHIFeatureLevel::Type RequestedFeatureLevel)
     if (bProjectSupportsMRTs && bRequestedMetalMRT && !bIsA8FeatureSet && !bForceES3_1)
     {
 #if PLATFORM_TVOS
-		ValidateTargetedRHIFeatureLevelExists(SP_METAL_MRT);
-		GMaxRHIShaderPlatform = SP_METAL_MRT_TVOS;
+		ValidateTargetedRHIFeatureLevelExists(SP_METAL_SM5_IOS);
+		GMaxRHIShaderPlatform = SP_METAL_SM5_TVOS;
 #else
-		ValidateTargetedRHIFeatureLevelExists(SP_METAL_MRT);
-        GMaxRHIShaderPlatform = SP_METAL_MRT;
+		ValidateTargetedRHIFeatureLevelExists(SP_METAL_SM5_IOS);
+        GMaxRHIShaderPlatform = SP_METAL_SM5_IOS;
 #endif
 		GMaxRHIFeatureLevel = ERHIFeatureLevel::SM5;
     }
@@ -308,15 +308,15 @@ FMetalDynamicRHI::FMetalDynamicRHI(ERHIFeatureLevel::Type RequestedFeatureLevel)
 		}
 		
 #if PLATFORM_TVOS
-		ValidateTargetedRHIFeatureLevelExists(SP_METAL_TVOS);
-		GMaxRHIShaderPlatform = SP_METAL_TVOS;
+		ValidateTargetedRHIFeatureLevelExists(SP_METAL_ES3_1_TVOS);
+		GMaxRHIShaderPlatform = SP_METAL_ES3_1_TVOS;
 #else
 #if WITH_IOS_SIMULATOR
 		ValidateTargetedRHIFeatureLevelExists(SP_METAL_SIM);
 		GMaxRHIShaderPlatform = SP_METAL_SIM;
 #else
-		ValidateTargetedRHIFeatureLevelExists(SP_METAL);
-		GMaxRHIShaderPlatform = SP_METAL;
+		ValidateTargetedRHIFeatureLevelExists(SP_METAL_ES3_1_IOS);
+		GMaxRHIShaderPlatform = SP_METAL_ES3_1_IOS;
 #endif	// WITH_IOS_SIMULATOR
 #endif	// PLATFORM_TVOS
         GMaxRHIFeatureLevel = ERHIFeatureLevel::ES3_1;
@@ -339,13 +339,13 @@ FMetalDynamicRHI::FMetalDynamicRHI(ERHIFeatureLevel::Type RequestedFeatureLevel)
 	
 #if PLATFORM_TVOS
 	GShaderPlatformForFeatureLevel[ERHIFeatureLevel::ES2_REMOVED] = SP_NumPlatforms;
-	GShaderPlatformForFeatureLevel[ERHIFeatureLevel::ES3_1] = SP_METAL_TVOS;
+	GShaderPlatformForFeatureLevel[ERHIFeatureLevel::ES3_1] = SP_METAL_ES3_1_TVOS;
 #else
 	GShaderPlatformForFeatureLevel[ERHIFeatureLevel::ES2_REMOVED] = SP_NumPlatforms;
 #if WITH_IOS_SIMULATOR
 	GShaderPlatformForFeatureLevel[ERHIFeatureLevel::ES3_1] = SP_METAL_SIM;
 #else
-	GShaderPlatformForFeatureLevel[ERHIFeatureLevel::ES3_1] = SP_METAL;
+	GShaderPlatformForFeatureLevel[ERHIFeatureLevel::ES3_1] = SP_METAL_ES3_1_IOS;
 #endif
 #endif
 	GShaderPlatformForFeatureLevel[ERHIFeatureLevel::SM4_REMOVED] = SP_NumPlatforms;
@@ -468,14 +468,7 @@ FMetalDynamicRHI::FMetalDynamicRHI(ERHIFeatureLevel::Type RequestedFeatureLevel)
 	else if(bSupportsSM5 && bRequestedSM5)
 	{
 		GMaxRHIFeatureLevel = ERHIFeatureLevel::SM5;
-		if (!FParse::Param(FCommandLine::Get(),TEXT("metalmrt")))
-		{
-			GMaxRHIShaderPlatform = SP_METAL_SM5;
-		}
-		else
-		{
-			GMaxRHIShaderPlatform = SP_METAL_MRT_MAC;
-		}
+		GMaxRHIShaderPlatform = SP_METAL_SM5;
 	}
 	else
 	{
@@ -512,21 +505,21 @@ FMetalDynamicRHI::FMetalDynamicRHI(ERHIFeatureLevel::Type RequestedFeatureLevel)
 		GMaxRHIFeatureLevel = PreviewFeatureLevel;
 		if (GMaxRHIFeatureLevel == ERHIFeatureLevel::ES3_1)
 		{
-			GMaxRHIShaderPlatform = SP_METAL_MACES3_1;
+			GMaxRHIShaderPlatform = SP_METAL_ES3_1;
 		}
 	}
 
 	ValidateTargetedRHIFeatureLevelExists(GMaxRHIShaderPlatform);
 	
 	GShaderPlatformForFeatureLevel[ERHIFeatureLevel::ES2_REMOVED] = SP_NumPlatforms;
-	GShaderPlatformForFeatureLevel[ERHIFeatureLevel::ES3_1] = (GMaxRHIFeatureLevel >= ERHIFeatureLevel::ES3_1) ? SP_METAL_MACES3_1 : SP_NumPlatforms;
+	GShaderPlatformForFeatureLevel[ERHIFeatureLevel::ES3_1] = (GMaxRHIFeatureLevel >= ERHIFeatureLevel::ES3_1) ? SP_METAL_ES3_1 : SP_NumPlatforms;
 	GShaderPlatformForFeatureLevel[ERHIFeatureLevel::SM4_REMOVED] = SP_NumPlatforms;
 	GShaderPlatformForFeatureLevel[ERHIFeatureLevel::SM5] = (GMaxRHIFeatureLevel >= ERHIFeatureLevel::SM5) ? GMaxRHIShaderPlatform : SP_NumPlatforms;
     GShaderPlatformForFeatureLevel[ERHIFeatureLevel::SM6] = (GMaxRHIFeatureLevel >= ERHIFeatureLevel::SM6) ? GMaxRHIShaderPlatform : SP_NumPlatforms;
 	
 	// Mac GPUs support layer indexing.
-	GSupportsVolumeTextureRendering = (GMaxRHIShaderPlatform != SP_METAL_MRT_MAC);
-	bSupportsPointLights &= (GMaxRHIShaderPlatform != SP_METAL_MRT_MAC);
+	GSupportsVolumeTextureRendering = true;
+	bSupportsPointLights &= true;
 	
 	// Make sure the vendors match - the assumption that order in IORegistry is the order in Metal may not hold up forever.
 	if(GPUDesc.GPUVendorId == GRHIVendorId)
@@ -658,7 +651,7 @@ FMetalDynamicRHI::FMetalDynamicRHI(ERHIFeatureLevel::Type RequestedFeatureLevel)
 	SetEmitDrawEvents(GetEmitDrawEvents() | ENABLE_METAL_GPUEVENTS);
 #endif
 
-	GSupportsShaderFramebufferFetch = !PLATFORM_MAC && GMaxRHIShaderPlatform != SP_METAL_MRT && GMaxRHIShaderPlatform != SP_METAL_MRT_TVOS && MobileAllowFramebufferFetch(GMaxRHIShaderPlatform);
+	GSupportsShaderFramebufferFetch = !PLATFORM_MAC && GMaxRHIShaderPlatform != SP_METAL_SM5_IOS && GMaxRHIShaderPlatform != SP_METAL_SM5_TVOS && MobileAllowFramebufferFetch(GMaxRHIShaderPlatform);
 	GSupportsShaderMRTFramebufferFetch = GSupportsShaderFramebufferFetch;
 	GHardwareHiddenSurfaceRemoval = true;
 	GSupportsRenderTargetFormat_PF_G8 = false;
diff --git a/Engine/Source/Runtime/Engine/Private/ShaderCompiler/ShaderCompiler.cpp b/Engine/Source/Runtime/Engine/Private/ShaderCompiler/ShaderCompiler.cpp
--- a/Engine/Source/Runtime/Engine/Private/ShaderCompiler/ShaderCompiler.cpp
+++ b/Engine/Source/Runtime/Engine/Private/ShaderCompiler/ShaderCompiler.cpp
@@ -3623,7 +3623,7 @@ void GlobalBeginCompileShader(
 		bAllowDevelopmentShaderCompile = false;
 	}
 
-	// Mobile emulation should be defined when a PC platform is using a mobile renderer (limited to feature level ES3_1)...  eg SP_PCD3D_ES3_1,SP_VULKAN_PCES3_1,SP_METAL_MACES3_1
+	// Mobile emulation should be defined when a PC platform is using a mobile renderer (limited to feature level ES3_1)...  eg SP_PCD3D_ES3_1,SP_VULKAN_PCES3_1,SP_METAL_ES3_1
 	if (IsSimulatedPlatform(EShaderPlatform(Target.Platform)) && bAllowDevelopmentShaderCompile)
 	{
 		SET_SHADER_DEFINE(Input.Environment, MOBILE_EMULATION, 1);
diff --git a/Engine/Source/Runtime/NullDrv/Private/NullRHI.cpp b/Engine/Source/Runtime/NullDrv/Private/NullRHI.cpp
--- a/Engine/Source/Runtime/NullDrv/Private/NullRHI.cpp
+++ b/Engine/Source/Runtime/NullDrv/Private/NullRHI.cpp
@@ -29,7 +29,7 @@ void FNullDynamicRHI::Init()
 	GShaderPlatformForFeatureLevel[ERHIFeatureLevel::SM6] = SP_PCD3D_SM6;
 #elif PLATFORM_MAC
     GShaderPlatformForFeatureLevel[ERHIFeatureLevel::ES2_REMOVED] = SP_NumPlatforms;
-    GShaderPlatformForFeatureLevel[ERHIFeatureLevel::ES3_1] = SP_METAL_MACES3_1;
+    GShaderPlatformForFeatureLevel[ERHIFeatureLevel::ES3_1] = SP_METAL_ES3_1;
     GShaderPlatformForFeatureLevel[ERHIFeatureLevel::SM4_REMOVED] = SP_NumPlatforms;
     GShaderPlatformForFeatureLevel[ERHIFeatureLevel::SM5] = SP_METAL_SM5;
     GShaderPlatformForFeatureLevel[ERHIFeatureLevel::SM6] = SP_METAL_SM6;
diff --git a/Engine/Source/Runtime/RHI/Private/Apple/Platform/IOS/IOSPlatformDynamicRHI.h b/Engine/Source/Runtime/RHI/Private/Apple/Platform/IOS/IOSPlatformDynamicRHI.h
--- a/Engine/Source/Runtime/RHI/Private/Apple/Platform/IOS/IOSPlatformDynamicRHI.h
+++ b/Engine/Source/Runtime/RHI/Private/Apple/Platform/IOS/IOSPlatformDynamicRHI.h
@@ -16,14 +16,14 @@
 
 
 #if PLATFORM_TVOS
-#define SHADER_PLATFORM_METAL_MRT	SP_METAL_MRT_TVOS
-#define SHADER_PLATFORM_METAL		SP_METAL_TVOS
+#define SHADER_PLATFORM_METAL_SM5	SP_METAL_SM5_TVOS
+#define SHADER_PLATFORM_METAL_ES3_1	SP_METAL_ES3_1_TVOS
 #elif WITH_IOS_SIMULATOR
-#define SHADER_PLATFORM_METAL		SP_METAL_SIM
-#define SHADER_PLATFORM_METAL_MRT	SP_METAL_MRT
+#define SHADER_PLATFORM_METAL_ES3_1	SP_METAL_SIM
+#define SHADER_PLATFORM_METAL_SM5	SP_METAL_SM5_IOS
 #else
-#define SHADER_PLATFORM_METAL_MRT	SP_METAL_MRT
-#define SHADER_PLATFORM_METAL		SP_METAL
+#define SHADER_PLATFORM_METAL_SM5	SP_METAL_SM5_IOS
+#define SHADER_PLATFORM_METAL_ES3_1	SP_METAL_ES3_1_IOS
 #endif
 
 //------------------------------------------------------------------------------
@@ -51,9 +51,9 @@ void AddTargetedShaderFormats(TArray<FString>& TargetedShaderFormats)
 {
 	if (ShouldSupportMetalMRT())
 	{
-		TargetedShaderFormats.Add(LegacyShaderPlatformToShaderFormat(SHADER_PLATFORM_METAL_MRT).ToString());
+		TargetedShaderFormats.Add(LegacyShaderPlatformToShaderFormat(SHADER_PLATFORM_METAL_SM5).ToString());
 	}
-	TargetedShaderFormats.Add(LegacyShaderPlatformToShaderFormat(SHADER_PLATFORM_METAL).ToString());
+	TargetedShaderFormats.Add(LegacyShaderPlatformToShaderFormat(SHADER_PLATFORM_METAL_ES3_1).ToString());
 }
 
 } // namespace FIOSPlatformDynamicRHI
diff --git a/Engine/Source/Runtime/RHI/Private/RHIStrings.cpp b/Engine/Source/Runtime/RHI/Private/RHIStrings.cpp
--- a/Engine/Source/Runtime/RHI/Private/RHIStrings.cpp
+++ b/Engine/Source/Runtime/RHI/Private/RHIStrings.cpp
@@ -88,15 +88,14 @@ FString LexToString(EShaderPlatform Platform, bool bError)
 	case SP_PCD3D_ES3_1: return TEXT("PCD3D_ES3_1");
 	case SP_OPENGL_PCES3_1: return TEXT("OPENGL_PCES3_1");
 	case SP_OPENGL_ES3_1_ANDROID: return TEXT("OPENGL_ES3_1_ANDROID");
-	case SP_METAL: return TEXT("METAL");
-	case SP_METAL_MRT: return TEXT("METAL_MRT");
-	case SP_METAL_TVOS: return TEXT("METAL_TVOS");
-	case SP_METAL_MRT_TVOS: return TEXT("METAL_MRT_TVOS");
-	case SP_METAL_MRT_MAC: return TEXT("METAL_MRT_MAC");
+	case SP_METAL_ES3_1_IOS: return TEXT("METAL_ES3_1_IOS");
+	case SP_METAL_SM5_IOS: return TEXT("METAL_SM5_IOS");
+	case SP_METAL_ES3_1_TVOS: return TEXT("METAL_ES3_1_TVOS");
+	case SP_METAL_SM5_TVOS: return TEXT("METAL_SM5_TVOS");
 	case SP_METAL_SM5: return TEXT("METAL_SM5");
     case SP_METAL_SM6: return TEXT("METAL_SM6");
     case SP_METAL_SIM: return TEXT("METAL_SIM");
-	case SP_METAL_MACES3_1: return TEXT("METAL_MACES3_1");
+	case SP_METAL_ES3_1: return TEXT("METAL_ES3_1");
 	case SP_VULKAN_ES3_1_ANDROID: return TEXT("VULKAN_ES3_1_ANDROID");
 	case SP_VULKAN_PCES3_1: return TEXT("VULKAN_PCES3_1");
 	case SP_VULKAN_SM5: return TEXT("VULKAN_SM5");
@@ -211,17 +210,16 @@ FName ShaderPlatformToPlatformName(EShaderPlatform Platform)
 	case SP_VULKAN_SM5_ANDROID:
 	case SP_OPENGL_ES3_1_ANDROID:
 		return NAME_PLATFORM_ANDROID;
-	case SP_METAL:
-	case SP_METAL_MRT:
+	case SP_METAL_ES3_1_IOS:
+	case SP_METAL_SM5_IOS:
 	case SP_METAL_SIM:
 		return NAME_PLATFORM_IOS;
 	case SP_METAL_SM5:
     case SP_METAL_SM6:
-	case SP_METAL_MACES3_1:
-	case SP_METAL_MRT_MAC:
+	case SP_METAL_ES3_1:
 		return NAME_PLATFORM_MAC;
-	case SP_METAL_TVOS:
-	case SP_METAL_MRT_TVOS:
+	case SP_METAL_ES3_1_TVOS:
+	case SP_METAL_SM5_TVOS:
 		return NAME_PLATFORM_TVOS;
 
 	default:
diff --git a/Engine/Source/Runtime/RHI/Public/DataDrivenShaderPlatformInfo.h b/Engine/Source/Runtime/RHI/Public/DataDrivenShaderPlatformInfo.h
--- a/Engine/Source/Runtime/RHI/Public/DataDrivenShaderPlatformInfo.h
+++ b/Engine/Source/Runtime/RHI/Public/DataDrivenShaderPlatformInfo.h
@@ -978,7 +978,7 @@ inline bool IsVulkanMobileSM5Platform(const FStaticShaderPlatform Platform)
 // @todo: data drive uses of this function
 inline bool IsMetalMobileSM5Platform(const FStaticShaderPlatform Platform)
 {
-	return Platform == SP_METAL_MRT;
+	return Platform == SP_METAL_SM5_IOS;
 	// 	return FDataDrivenShaderPlatformInfo::GetIsLanguageMetal(Platform)
 	// 		&& FDataDrivenShaderPlatformInfo::GetMaxFeatureLevel(Platform) == ERHIFeatureLevel::SM5
 	// 		&& FDataDrivenShaderPlatformInfo::GetIsMobile(Platform);
diff --git a/Engine/Source/Runtime/RHI/Public/IOS/IOSStaticShaderPlatform.inl b/Engine/Source/Runtime/RHI/Public/IOS/IOSStaticShaderPlatform.inl
--- a/Engine/Source/Runtime/RHI/Public/IOS/IOSStaticShaderPlatform.inl
+++ b/Engine/Source/Runtime/RHI/Public/IOS/IOSStaticShaderPlatform.inl
@@ -4,12 +4,12 @@
 
 #include "Misc/AssertionMacros.h"
 
-#if defined(UE_IOS_SHADER_PLATFORM_METAL)
-	#define UE_IOS_STATIC_SHADER_PLATFORM SP_METAL
+#if defined(UE_IOS_SHADER_PLATFORM_METAL_ES3_1_IOS)
+	#define UE_IOS_STATIC_SHADER_PLATFORM SP_METAL_ES3_1_IOS
 #elif defined(UE_IOS_SHADER_PLATFORM_METAL_SIM)
 	#define UE_IOS_STATIC_SHADER_PLATFORM SP_METAL_SIM
-#elif defined(UE_IOS_SHADER_PLATFORM_METAL_MRT)
-	#define UE_IOS_STATIC_SHADER_PLATFORM SP_METAL_MRT
+#elif defined(UE_IOS_SHADER_PLATFORM_METAL_SM5_IOS)
+	#define UE_IOS_STATIC_SHADER_PLATFORM SP_METAL_SM5_IOS
 #else
 	#error "Unknown IOS static shader platform"
 #endif
diff --git a/Engine/Source/Runtime/RHI/Public/RHIShaderFormatDefinitions.inl b/Engine/Source/Runtime/RHI/Public/RHIShaderFormatDefinitions.inl
--- a/Engine/Source/Runtime/RHI/Public/RHIShaderFormatDefinitions.inl
+++ b/Engine/Source/Runtime/RHI/Public/RHIShaderFormatDefinitions.inl
@@ -16,15 +16,20 @@ static FName NAME_PCD3D_ES3_1(TEXT("PCD3D_ES31"));
 static FName NAME_GLSL_150_ES31(TEXT("GLSL_150_ES31"));
 static FName NAME_GLSL_ES3_1_ANDROID(TEXT("GLSL_ES3_1_ANDROID"));
 
-static FName NAME_SF_METAL(TEXT("SF_METAL"));
-static FName NAME_SF_METAL_MRT(TEXT("SF_METAL_MRT"));
-static FName NAME_SF_METAL_TVOS(TEXT("SF_METAL_TVOS"));
-static FName NAME_SF_METAL_MRT_TVOS(TEXT("SF_METAL_MRT_TVOS"));
-static FName NAME_SF_METAL_MRT_MAC(TEXT("SF_METAL_MRT_MAC"));
+static FName NAME_SF_METAL_ES3_1_IOS_DEPRECATED(TEXT("SF_METAL"));
+static FName NAME_SF_METAL_SM5_IOS_DEPRECATED(TEXT("SF_METAL_MRT"));
+static FName NAME_SF_METAL_ES3_1_TVOS_DEPRECATED(TEXT("SF_METAL_TVOS"));
+static FName NAME_SF_METAL_SM5_TVOS_DEPRECATED(TEXT("SF_METAL_MRT_TVOS"));
+static FName NAME_SF_METAL_MRT_MAC_DEPRECATED(TEXT("SF_METAL_MRT_MAC"));
+
+static FName NAME_SF_METAL_ES3_1_IOS(TEXT("SF_METAL_ES3_1_IOS"));
+static FName NAME_SF_METAL_SM5_IOS(TEXT("SF_METAL_SM5_IOS"));
+static FName NAME_SF_METAL_ES3_1_TVOS(TEXT("SF_METAL_ES3_1_TVOS"));
+static FName NAME_SF_METAL_SM5_TVOS(TEXT("SF_METAL_SM5_TVOS"));
 static FName NAME_SF_METAL_SM5(TEXT("SF_METAL_SM5"));
 static FName NAME_SF_METAL_SM6(TEXT("SF_METAL_SM6"));
 static FName NAME_SF_METAL_SIM(TEXT("SF_METAL_SIM"));
-static FName NAME_SF_METAL_MACES3_1(TEXT("SF_METAL_MACES3_1"));
+static FName NAME_SF_METAL_ES3_1(TEXT("SF_METAL_ES3_1"));
 
 static FName NAME_VULKAN_ES3_1_ANDROID(TEXT("SF_VULKAN_ES31_ANDROID"));
 static FName NAME_VULKAN_ES3_1(TEXT("SF_VULKAN_ES31"));
@@ -34,28 +39,32 @@ static FName NAME_VULKAN_SM5_ANDROID(TEXT("SF_VULKAN_SM5_ANDROID"));
 
 static EShaderPlatform ShaderFormatNameToShaderPlatform(FName ShaderFormat)
 {
-	if (ShaderFormat == NAME_PCD3D_SM6)					return SP_PCD3D_SM6;
-	if (ShaderFormat == NAME_PCD3D_SM5)					return SP_PCD3D_SM5;
-	if (ShaderFormat == NAME_PCD3D_ES3_1)				return SP_PCD3D_ES3_1;
+	if (ShaderFormat == NAME_PCD3D_SM6)							return SP_PCD3D_SM6;
+	if (ShaderFormat == NAME_PCD3D_SM5)							return SP_PCD3D_SM5;
+	if (ShaderFormat == NAME_PCD3D_ES3_1)						return SP_PCD3D_ES3_1;
 
-	if (ShaderFormat == NAME_GLSL_150_ES31)				return SP_OPENGL_PCES3_1;
-	if (ShaderFormat == NAME_GLSL_ES3_1_ANDROID)		return SP_OPENGL_ES3_1_ANDROID;
+	if (ShaderFormat == NAME_GLSL_150_ES31)						return SP_OPENGL_PCES3_1;
+	if (ShaderFormat == NAME_GLSL_ES3_1_ANDROID)				return SP_OPENGL_ES3_1_ANDROID;
 
-	if (ShaderFormat == NAME_SF_METAL)					return SP_METAL;
-	if (ShaderFormat == NAME_SF_METAL_MRT)				return SP_METAL_MRT;
-	if (ShaderFormat == NAME_SF_METAL_TVOS)				return SP_METAL_TVOS;
-	if (ShaderFormat == NAME_SF_METAL_MRT_TVOS)			return SP_METAL_MRT_TVOS;
-	if (ShaderFormat == NAME_SF_METAL_MRT_MAC)			return SP_METAL_MRT_MAC;
-	if (ShaderFormat == NAME_SF_METAL_SM5)				return SP_METAL_SM5;
-    if (ShaderFormat == NAME_SF_METAL_SM6)              return SP_METAL_SM6;
-	if (ShaderFormat == NAME_SF_METAL_SIM)				return SP_METAL_SIM;
-	if (ShaderFormat == NAME_SF_METAL_MACES3_1)			return SP_METAL_MACES3_1;
+	if (ShaderFormat == NAME_SF_METAL_ES3_1_IOS)				return SP_METAL_ES3_1_IOS;
+	if (ShaderFormat == NAME_SF_METAL_ES3_1_IOS_DEPRECATED)		return SP_METAL_ES3_1_IOS;
+	if (ShaderFormat == NAME_SF_METAL_SM5_IOS)					return SP_METAL_SM5_IOS;
+	if (ShaderFormat == NAME_SF_METAL_SM5_IOS_DEPRECATED)		return SP_METAL_SM5_IOS;
+	if (ShaderFormat == NAME_SF_METAL_ES3_1_TVOS)				return SP_METAL_ES3_1_TVOS;
+	if (ShaderFormat == NAME_SF_METAL_ES3_1_TVOS_DEPRECATED)	return SP_METAL_ES3_1_TVOS;
+	if (ShaderFormat == NAME_SF_METAL_SM5_TVOS)					return SP_METAL_SM5_TVOS;
+	if (ShaderFormat == NAME_SF_METAL_SM5_TVOS_DEPRECATED)		return SP_METAL_SM5_TVOS;
+	if (ShaderFormat == NAME_SF_METAL_SM5)						return SP_METAL_SM5;
+	if (ShaderFormat == NAME_SF_METAL_SM6)						return SP_METAL_SM6;
+	if (ShaderFormat == NAME_SF_METAL_SIM)						return SP_METAL_SIM;
+	if (ShaderFormat == NAME_SF_METAL_ES3_1)					return SP_METAL_ES3_1;
+	if (ShaderFormat == NAME_SF_METAL_MRT_MAC_DEPRECATED)				return SP_METAL_ES3_1;
 
-	if (ShaderFormat == NAME_VULKAN_ES3_1_ANDROID)		return SP_VULKAN_ES3_1_ANDROID;
-	if (ShaderFormat == NAME_VULKAN_ES3_1)				return SP_VULKAN_PCES3_1;
-	if (ShaderFormat == NAME_VULKAN_SM5)				return SP_VULKAN_SM5;
-	if (ShaderFormat == NAME_VULKAN_SM6)				return SP_VULKAN_SM6;
-	if (ShaderFormat == NAME_VULKAN_SM5_ANDROID)		return SP_VULKAN_SM5_ANDROID;
+	if (ShaderFormat == NAME_VULKAN_ES3_1_ANDROID)				return SP_VULKAN_ES3_1_ANDROID;
+	if (ShaderFormat == NAME_VULKAN_ES3_1)						return SP_VULKAN_PCES3_1;
+	if (ShaderFormat == NAME_VULKAN_SM5)						return SP_VULKAN_SM5;
+	if (ShaderFormat == NAME_VULKAN_SM6)						return SP_VULKAN_SM6;
+	if (ShaderFormat == NAME_VULKAN_SM5_ANDROID)				return SP_VULKAN_SM5_ANDROID;
 
 	for (int32 StaticPlatform = SP_StaticPlatform_First; StaticPlatform <= SP_StaticPlatform_Last; ++StaticPlatform)
 	{
diff --git a/Engine/Source/Runtime/RHI/Public/RHIShaderPlatform.h b/Engine/Source/Runtime/RHI/Public/RHIShaderPlatform.h
--- a/Engine/Source/Runtime/RHI/Public/RHIShaderPlatform.h
+++ b/Engine/Source/Runtime/RHI/Public/RHIShaderPlatform.h
@@ -10,19 +10,23 @@
 enum EShaderPlatform : uint16
 {
 	SP_PCD3D_SM5 = 0,
-	SP_METAL = 11,
-	SP_METAL_MRT = 12,
+	SP_METAL_ES3_1_IOS = 11,
+	SP_METAL UE_DEPRECATED(5.6, "Use SP_METAL_ES3_1_IOS enumeration value instead.") = SP_METAL_ES3_1_IOS,
+	SP_METAL_SM5_IOS = 12,
+	SP_METAL_MRT UE_DEPRECATED(5.6, "Use SP_METAL_SM5_IOS enumeration value instead.") = SP_METAL_SM5_IOS,
 	SP_PCD3D_ES3_1 = 14,
 	SP_OPENGL_PCES3_1 = 15,
 	SP_METAL_SM5 = 16,
 	SP_VULKAN_PCES3_1 = 17,
 	SP_VULKAN_SM5 = 20,
 	SP_VULKAN_ES3_1_ANDROID = 21,
-	SP_METAL_MACES3_1 = 22,
+	SP_METAL_ES3_1 = 22,
+	SP_METAL_MACES3_1 UE_DEPRECATED(5.6, "Use SP_METAL_ES3_1 enumeration value instead.") = SP_METAL_ES3_1,
 	SP_OPENGL_ES3_1_ANDROID = 24,
-	SP_METAL_MRT_MAC = 27,
-	SP_METAL_TVOS = 30,
-	SP_METAL_MRT_TVOS = 31,
+	SP_METAL_ES3_1_TVOS = 30,
+	SP_METAL_TVOS UE_DEPRECATED(5.6, "Use SP_METAL_ES3_1_TVOS enumeration value instead.") = SP_METAL_ES3_1_TVOS,
+	SP_METAL_SM5_TVOS = 31,
+	SP_METAL_MRT_TVOS UE_DEPRECATED(5.6, "Use SP_METAL_SM5_TVOS enumeration value instead.") = SP_METAL_SM5_TVOS,
 	/**********************************************************************************/
 	/* !! Do not add any new platforms here. Add them below SP_StaticPlatform_Last !! */
 	/**********************************************************************************/
diff --git a/Engine/Source/Runtime/Renderer/Private/VT/RuntimeVirtualTextureRender.cpp b/Engine/Source/Runtime/Renderer/Private/VT/RuntimeVirtualTextureRender.cpp
--- a/Engine/Source/Runtime/Renderer/Private/VT/RuntimeVirtualTextureRender.cpp
+++ b/Engine/Source/Runtime/Renderer/Private/VT/RuntimeVirtualTextureRender.cpp
@@ -488,11 +488,11 @@ namespace RuntimeVirtualTexture
 	{
 		switch (ShaderPlatform)
 		{
-		case SP_METAL:
-		case SP_METAL_MRT:
+		case SP_METAL_ES3_1_IOS:
+		case SP_METAL_SM5_IOS:
 		case SP_METAL_SIM:
-		case SP_METAL_TVOS:
-		case SP_METAL_MRT_TVOS:
+		case SP_METAL_ES3_1_TVOS:
+		case SP_METAL_SM5_TVOS:
 		case SP_VULKAN_ES3_1_ANDROID:
 		case SP_OPENGL_ES3_1_ANDROID:
 		case SP_VULKAN_SM5_ANDROID:
@@ -510,11 +510,11 @@ namespace RuntimeVirtualTexture
 		}
 		switch (ShaderPlatform)
 		{
-		case SP_METAL:
-		case SP_METAL_MRT:
+		case SP_METAL_ES3_1_IOS:
+		case SP_METAL_SM5_IOS:
 		case SP_METAL_SIM:
-		case SP_METAL_TVOS:
-		case SP_METAL_MRT_TVOS:
+		case SP_METAL_ES3_1_TVOS:
+		case SP_METAL_SM5_TVOS:
 		case SP_VULKAN_ES3_1_ANDROID:
 		case SP_OPENGL_ES3_1_ANDROID:
 		case SP_VULKAN_SM5_ANDROID:
