Metal libraries mapped memory lifetime handling fix.

We can only unmap memory, which was used to create the Metal library, when the corresponding dispatch_data_t destructor is called. Previously, we were unmapping it after the last ref-count release on the library instance from our side. That is incorrect, as there might be more existing references to that instance, which we don't control. dispatch_data_t destructor is called exactly at a moment, when it won't be needed anymore, hence we release library's data memory from there.

r.Metal.PersistentMappingMetallib is still left in the code, with it still influencing the unmapping. However, it is now disabled by default. Left as a temporal safety hatch, will be removed soon.

#rnx

diff --git a/Engine/Source/Runtime/Apple/MetalRHI/Private/MetalDynamicRHI_Shaders.cpp b/Engine/Source/Runtime/Apple/MetalRHI/Private/MetalDynamicRHI_Shaders.cpp
--- a/Engine/Source/Runtime/Apple/MetalRHI/Private/MetalDynamicRHI_Shaders.cpp
+++ b/Engine/Source/Runtime/Apple/MetalRHI/Private/MetalDynamicRHI_Shaders.cpp
@@ -13,6 +13,21 @@
 #include "Interfaces/IPluginManager.h"
 
 //------------------------------------------------------------------------------
+static bool GPersistentMappingMetallib = 0;
+static FAutoConsoleVariableRef CVarPersistentMappingMetallib(
+	TEXT("r.Metal.PersistentMappingMetallib"),
+	GPersistentMappingMetallib,
+	TEXT("Makes the metallib file mapping persistent."),
+	ECVF_Default
+);
+
+static void ReleaseLibraryMemory(FMetalShaderLibrary::FShaderLibDataOwner* LibraryMemOwner)
+{
+	if (!GPersistentMappingMetallib)
+	{
+		delete LibraryMemOwner;
+	}
+}
 
 #pragma mark - Metal Dynamic RHI Shader Methods
 
@@ -151,25 +166,25 @@ FRHIShaderLibraryRef FMetalDynamicRHI::RHICreateShaderLibrary(EShaderPlatform Pl
 #endif
 	};
 
-	auto MapShaderLibraryFile = [](const FString& FilePath, FMetalShaderLibrary::FShaderLibDataOwner* MemOwner)
+	auto MapShaderLibraryFile = [](const FString& FilePath, FMetalShaderLibrary::FShaderLibDataOwner& MemOwner)
 	{
 		if (FPlatformProperties::SupportsMemoryMappedFiles())
 		{
 			FOpenMappedResult Res = FPlatformFileManager::Get().GetPlatformFile().OpenMappedEx(*FilePath);
-			MemOwner->MappedCacheFile = Res.HasError() ? nullptr : Res.StealValue();
-			if (MemOwner->MappedCacheFile.IsValid())
+			MemOwner.MappedCacheFile = Res.HasError() ? nullptr : Res.StealValue();
+			if (MemOwner.MappedCacheFile.IsValid())
 			{
-				MemOwner->MappedRegion = TUniquePtr<IMappedFileRegion>(MemOwner->MappedCacheFile->MapRegion(0, MemOwner->MappedCacheFile->GetFileSize()));
+				MemOwner.MappedRegion = TUniquePtr<IMappedFileRegion>(MemOwner.MappedCacheFile->MapRegion(0, MemOwner.MappedCacheFile->GetFileSize(), EMappedFileFlags::ENone));
 			}
 		}
 	};
 
- #if !USE_MMAPPED_SHADERARCHIVE
+#if !USE_MMAPPED_SHADERARCHIVE
     FArchive* BinaryShaderAr = IFileManager::Get().CreateFileReader(*BinaryShaderFile);
 #else
 	FStaticMemoryReader* BinaryShaderAr = nullptr;
     TUniquePtr<FMetalShaderLibrary::FShaderLibDataOwner> MemOwner = MakeUnique<FMetalShaderLibrary::FShaderLibDataOwner>();
-	MapShaderLibraryFile(BinaryShaderFile, MemOwner.Get());
+	MapShaderLibraryFile(BinaryShaderFile, *MemOwner);
 	if (MemOwner->MappedRegion.IsValid())
 	{
 		UE_LOG(LogMetal, Display, TEXT("mmapping %s, %d bytes"), *BinaryShaderFile, MemOwner->MappedCacheFile->GetFileSize());
@@ -203,24 +218,58 @@ FRHIShaderLibraryRef FMetalDynamicRHI::RHICreateShaderLibrary(EShaderPlatform Pl
         {
             check(((SerializedShaders.GetNumShaders() + Header.NumShadersPerLibrary - 1) / Header.NumShadersPerLibrary) == Header.NumLibraries);
 
-			TArray<TUniquePtr<FMetalShaderLibrary::FMetalLibDataOwner>> LibraryMemOwner;
-
-			LibraryMemOwner.Empty(Header.NumLibraries);
+			TArray<MTLLibraryPtr> Libraries;
+			Libraries.Empty(Header.NumLibraries);
 
             for (uint32 i = 0; i < Header.NumLibraries; i++)
             {
 				FString MetalLibraryFilePath = (FilePath / LibName) + FString::Printf(TEXT(".%d.metallib"), i);
-
-				TUniquePtr<FMetalShaderLibrary::FMetalLibDataOwner> MemOwnerLibrary = MakeUnique<FMetalShaderLibrary::FMetalLibDataOwner>();
-				MemOwnerLibrary->MetalLibraryFilePath = MetalLibraryFilePath;
-				MapShaderLibraryFile(MetalLibraryFilePath, MemOwnerLibrary.Get());
-				
-				LibraryMemOwner.Add(MoveTemp(MemOwnerLibrary));
+				NS::Error* Error;
+
+				MTLLibraryPtr Library;
+				TUniquePtr<FMetalShaderLibrary::FShaderLibDataOwner> LibraryMemOwner = MakeUnique<FMetalShaderLibrary::FShaderLibDataOwner>();
+
+				MapShaderLibraryFile(MetalLibraryFilePath, *LibraryMemOwner);
+				if (LibraryMemOwner->MappedRegion.IsValid())
+				{
+					UE_LOG(LogMetal, Display, TEXT("mmapping %s, %d bytes"), *MetalLibraryFilePath, LibraryMemOwner->MappedCacheFile->GetFileSize());
+					FMetalShaderLibrary::FShaderLibDataOwner* MyLibraryMemOwner = LibraryMemOwner.Release();
+					dispatch_data_t data = dispatch_data_create(MyLibraryMemOwner->MappedRegion->GetMappedPtr(), MyLibraryMemOwner->MappedRegion->GetMappedSize(), nil, ^(void) { 
+						ReleaseLibraryMemory(MyLibraryMemOwner);
+					});
+					Library = NS::TransferPtr(Device->GetDevice()->newLibrary(data, &Error));
+					dispatch_release(data);
+				}
+
+				if (Library.get() == nullptr)
+				{
+					TArray<uint8>& FileData = LibraryMemOwner->Mem;
+					if (FFileHelper::LoadFileToArray(FileData, *MetalLibraryFilePath))
+					{
+						UE_LOG(LogMetal, Display, TEXT("emulating mmapping %s, %d bytes!"), *MetalLibraryFilePath, FileData.Num());
+						FMetalShaderLibrary::FShaderLibDataOwner* MyLibraryMemOwner = LibraryMemOwner.Release();
+						dispatch_data_t data = dispatch_data_create(FileData.GetData(), FileData.Num(), nil, ^(void) { 
+							ReleaseLibraryMemory(MyLibraryMemOwner);
+						});
+						Library = NS::TransferPtr(Device->GetDevice()->newLibrary(data, &Error));
+						dispatch_release(data);
+					}
+				}
+
+				if (Library.get() != nullptr)
+				{
+					Libraries.Add(MoveTemp(Library));
+				}
+				else
+				{
+					UE_LOG(LogMetal, Display, TEXT("Failed to create library: %s"), *MetalLibraryFilePath);
+					return nullptr;
+				}
             }
 
-            FMetalShaderLibrary* MtlLib = new FMetalShaderLibrary(*Device, Platform, Name, BinaryShaderFile, Header, MoveTemp(SerializedShaders), MoveTemp(ShaderCode), MoveTemp(LibraryMemOwner)
+			FMetalShaderLibrary* MtlLib = new FMetalShaderLibrary(*Device, Platform, Name, BinaryShaderFile, Header, MoveTemp(SerializedShaders), MoveTemp(ShaderCode), Libraries
 #if USE_MMAPPED_SHADERARCHIVE
-				,MoveTemp(MemOwner)
+				, MoveTemp(MemOwner)
 #endif
 			);
             Result = MtlLib;
diff --git a/Engine/Source/Runtime/Apple/MetalRHI/Private/Shaders/MetalShaderLibrary.cpp b/Engine/Source/Runtime/Apple/MetalRHI/Private/Shaders/MetalShaderLibrary.cpp
--- a/Engine/Source/Runtime/Apple/MetalRHI/Private/Shaders/MetalShaderLibrary.cpp
+++ b/Engine/Source/Runtime/Apple/MetalRHI/Private/Shaders/MetalShaderLibrary.cpp
@@ -12,14 +12,6 @@
 #endif // !UE_BUILD_SHIPPING
 #include "MetalShaderTypes.h"
 
-static bool GPersistentMappingMetallib = 1;
-static FAutoConsoleVariableRef CVarPersistentMappingMetallib(
-	TEXT("r.Metal.PersistentMappingMetallib"),
-	GPersistentMappingMetallib,
-	TEXT("Makes the metallib file mapping persistent."),
-	ECVF_Default
-);
-
 //------------------------------------------------------------------------------
 
 #pragma mark - Metal Shader Library Class Support Routines
@@ -51,15 +43,7 @@ TMap<FString, FRHIShaderLibrary*> FMetalShaderLibrary::LoadedShaderLibraryMap;
 //------------------------------------------------------------------------------
 
 #pragma mark - Metal Shader Library Class
-FMetalShaderLibrary::FMetalLibDataOwner::~FMetalLibDataOwner()
-{
-	// Intentionally leak the metallib file handle to workaround driver bug
-	if (GPersistentMappingMetallib > 0)
-	{
-		MappedCacheFile.Release();
-		MappedRegion.Release();
-	}
-}
+
 
 FMetalShaderLibrary::FMetalShaderLibrary(FMetalDevice& MetalDevice,
 										 EShaderPlatform Platform,
@@ -68,14 +52,14 @@ FMetalShaderLibrary::FMetalShaderLibrary(FMetalDevice& MetalDevice,
 										 const FMetalShaderLibraryHeader& InHeader,
 										 FSerializedShaderArchive&& InSerializedShaders,
 										 FShaderCodeArrayType&& InShaderCode,
-										 TArray<TUniquePtr<FMetalLibDataOwner>>&& InLibraryMemOwner)
+										 const TArray<MTLLibraryPtr>& InLibrary)
 	: FRHIShaderLibrary(Platform, Name)
 	, Device(MetalDevice)
 	, ShaderLibraryFilename(InShaderLibraryFilename)
 	, Header(InHeader)
 	, SerializedShaders(MoveTemp(InSerializedShaders))
 	, ShaderCode(MoveTemp(InShaderCode))
-	, LibraryMemOwner(MoveTemp(InLibraryMemOwner))
+	, Library(InLibrary)
 {
 #if !UE_BUILD_SHIPPING
 	DebugFile = nullptr;
@@ -153,47 +137,20 @@ TRefCountPtr<FRHIShader> FMetalShaderLibrary::CreateShader(int32 Index, bool bRe
 	const TArrayView<const uint8> Code = MakeArrayView(ShaderCode.GetData() + ShaderEntry.Offset, ShaderEntry.Size);
 	const int32 LibraryIndex = Index / Header.NumShadersPerLibrary;
 
-	{
-		TUniquePtr<FMetalLibDataOwner>& LibraryMemOwnerRef = LibraryMemOwner[LibraryIndex];
-		UE::TScopeLock Lock(LibraryMemOwnerRef->LibraryLock);
-		if (LibraryMemOwnerRef->Library.get() == nullptr)
-		{
-			NS::Error* Error;
-			if (LibraryMemOwnerRef->MappedRegion.IsValid())
-			{
-				UE_LOG(LogMetal, Display, TEXT("mmapping %s, %d bytes"), *LibraryMemOwnerRef->MetalLibraryFilePath, LibraryMemOwnerRef->MappedCacheFile->GetFileSize());
-				dispatch_data_t data = dispatch_data_create(LibraryMemOwnerRef->MappedRegion->GetMappedPtr(), LibraryMemOwnerRef->MappedRegion->GetMappedSize(), nil, ^ {});
-				LibraryMemOwnerRef->Library = NS::TransferPtr(Device.GetDevice()->newLibrary(data, &Error));
-				dispatch_release(data);
-			}
-			else
-			{
-				TArray<uint8>& FileData = LibraryMemOwnerRef->Mem;
-				if (FFileHelper::LoadFileToArray(FileData, *LibraryMemOwnerRef->MetalLibraryFilePath))
-				{
-					UE_LOG(LogMetal, Display, TEXT("emulating mmapping %s, %d bytes!"), *LibraryMemOwnerRef->MetalLibraryFilePath, FileData.Num());
-					dispatch_data_t data = dispatch_data_create(FileData.GetData(), FileData.Num(), nil, ^ {});
-					LibraryMemOwnerRef->Library = NS::TransferPtr(Device.GetDevice()->newLibrary(data, &Error));
-					dispatch_release(data);
-				}
-			}
-		}
-	}
-
 	TRefCountPtr<FRHIShader> Shader;
 	switch (ShaderEntry.Frequency)
 	{
 		case SF_Vertex:
-			Shader = CreateMetalShader<FMetalVertexShader>(Device, Code, LibraryMemOwner[LibraryIndex]->Library);
+			Shader = CreateMetalShader<FMetalVertexShader>(Device, Code, Library[LibraryIndex]);
 			break;
 
 		case SF_Pixel:
-			Shader = CreateMetalShader<FMetalPixelShader>(Device, Code, LibraryMemOwner[LibraryIndex]->Library);
+			Shader = CreateMetalShader<FMetalPixelShader>(Device, Code, Library[LibraryIndex]);
 			break;
  
  		case SF_Geometry:
 #if PLATFORM_SUPPORTS_GEOMETRY_SHADERS
-            Shader = CreateMetalShader<FMetalGeometryShader>(Device, Code, LibraryMemOwner[LibraryIndex]->Library);
+            Shader = CreateMetalShader<FMetalGeometryShader>(Device, Code, Library[LibraryIndex]);
 #else
             checkf(false, TEXT("Geometry shaders not supported"));
 #endif
@@ -202,7 +159,7 @@ TRefCountPtr<FRHIShader> FMetalShaderLibrary::CreateShader(int32 Index, bool bRe
 
         case SF_Mesh:
 #if PLATFORM_SUPPORTS_MESH_SHADERS
-            Shader = CreateMetalShader<FMetalMeshShader>(Device, Code, LibraryMemOwner[LibraryIndex]->Library);
+            Shader = CreateMetalShader<FMetalMeshShader>(Device, Code, Library[LibraryIndex]);
 #else
 			checkf(false, TEXT("Mesh shaders not supported"));
 #endif
@@ -210,14 +167,14 @@ TRefCountPtr<FRHIShader> FMetalShaderLibrary::CreateShader(int32 Index, bool bRe
 
         case SF_Amplification:
 #if PLATFORM_SUPPORTS_MESH_SHADERS
-            Shader = CreateMetalShader<FMetalAmplificationShader>(Device, Code, LibraryMemOwner[LibraryIndex]->Library);
+            Shader = CreateMetalShader<FMetalAmplificationShader>(Device, Code, Library[LibraryIndex]);
 #else
 			checkf(false, TEXT("Amplification shaders not supported"));
 #endif
             break;
 
 		case SF_Compute:
-			Shader = CreateMetalShader<FMetalComputeShader>(Device, Code, LibraryMemOwner[LibraryIndex]->Library);
+			Shader = CreateMetalShader<FMetalComputeShader>(Device, Code, Library[LibraryIndex]);
 			break;
 
 		default:
diff --git a/Engine/Source/Runtime/Apple/MetalRHI/Private/Shaders/MetalShaderLibrary.h b/Engine/Source/Runtime/Apple/MetalRHI/Private/Shaders/MetalShaderLibrary.h
--- a/Engine/Source/Runtime/Apple/MetalRHI/Private/Shaders/MetalShaderLibrary.h
+++ b/Engine/Source/Runtime/Apple/MetalRHI/Private/Shaders/MetalShaderLibrary.h
@@ -24,23 +24,13 @@ public:
 		TUniquePtr<IMappedFileHandle> MappedCacheFile;
 		TUniquePtr<IMappedFileRegion> MappedRegion;
 	};
-
-	struct FMetalLibDataOwner : public FShaderLibDataOwner
-	{
-		MTLLibraryPtr Library;
-		FString MetalLibraryFilePath;
-		UE::FMutex LibraryLock;
-
-		~FMetalLibDataOwner();
-	};
 #if !USE_MMAPPED_SHADERARCHIVE
     using FShaderCodeArrayType = TArray<uint8>;
 #else 
 	using FShaderCodeArrayType = TArrayView<const uint8>;
 private:
-    TUniquePtr<FShaderLibDataOwner> MemOwner;
+	TUniquePtr<FShaderLibDataOwner> MemOwner;
 public:
-    
     FMetalShaderLibrary(FMetalDevice& Device,
 		EShaderPlatform Platform,
         FString const& Name,
@@ -48,12 +38,12 @@ public:
         const FMetalShaderLibraryHeader& InHeader,
         FSerializedShaderArchive&& InSerializedShaders,
         FShaderCodeArrayType&& InShaderCode,
-		TArray<TUniquePtr<FMetalLibDataOwner>>&& InLibraryMemOwner,
+		const TArray<MTLLibraryPtr>& InLibrary,
         TUniquePtr<FShaderLibDataOwner>&& InMemOwner)
         :
-        FMetalShaderLibrary(Device, Platform, Name, InShaderLibraryFilename, InHeader, MoveTemp(InSerializedShaders), MoveTemp(InShaderCode), MoveTemp(InLibraryMemOwner))
+        FMetalShaderLibrary(Device, Platform, Name, InShaderLibraryFilename, InHeader, MoveTemp(InSerializedShaders), MoveTemp(InShaderCode), InLibrary)
         {
-            MemOwner = MoveTemp(InMemOwner);
+			MemOwner = MoveTemp(InMemOwner);
         }
 #endif
   
@@ -67,7 +57,7 @@ public:
 						const FMetalShaderLibraryHeader& InHeader,
 						FSerializedShaderArchive&& InSerializedShaders,
 						FShaderCodeArrayType&& InShaderCode,
-						TArray<TUniquePtr<FMetalLibDataOwner>>&& InLibraryMemOwner);
+						const TArray<MTLLibraryPtr>& InLibrary);
 
 	virtual ~FMetalShaderLibrary();
 
@@ -97,7 +87,7 @@ private:
 	FMetalShaderLibraryHeader Header;
 	FSerializedShaderArchive SerializedShaders;
 	FShaderCodeArrayType ShaderCode;
-	TArray<TUniquePtr<FMetalLibDataOwner>> LibraryMemOwner;
+	TArray<MTLLibraryPtr> Library;
 #if !UE_BUILD_SHIPPING
 	class FMetalShaderDebugZipFile* DebugFile;
 #endif
