Disabling new RHI multi-threaded submission on IOS.

Added a CVar to also control the interrupt thread creation for Metal RHI.

#rnx

diff --git a/Engine/Source/Runtime/Apple/MetalRHI/Private/MetalSubmission.cpp b/Engine/Source/Runtime/Apple/MetalRHI/Private/MetalSubmission.cpp
--- a/Engine/Source/Runtime/Apple/MetalRHI/Private/MetalSubmission.cpp
+++ b/Engine/Source/Runtime/Apple/MetalRHI/Private/MetalSubmission.cpp
@@ -9,6 +9,14 @@
 #include "IRenderCaptureProvider.h"
 #include "MetalProfiler.h"
 
+static TAutoConsoleVariable<int32> CVarMetalRHIUseInterruptThread(
+	TEXT("MetalRHI.UseInterruptThread"),
+	1,
+	TEXT("Whether to enable the Metal RHI's interrupt thread.\n")
+	TEXT("  0: No\n")
+	TEXT("  1: Yes\n"),
+	ECVF_ReadOnly);
+
 static TAutoConsoleVariable<int32> CVarMetalRHIUseSubmissionThread(
 	TEXT("MetalRHI.UseSubmissionThread"),
 	1,
@@ -89,7 +97,11 @@ void FMetalDynamicRHI::InitializeSubmissionPipe()
 	if (FPlatformProcess::SupportsMultithreading())
 	{
 #if METAL_USE_INTERRUPT_THREAD
-		InterruptThread = new FMetalThread(TEXT("RHIInterruptThread"), TPri_Highest, this, &FMetalDynamicRHI::ProcessInterruptQueue);
+		bool bUseInterruptThread = CVarMetalRHIUseInterruptThread.GetValueOnAnyThread() == 1;
+		if (bUseInterruptThread)
+		{
+			InterruptThread = new FMetalThread(TEXT("RHIInterruptThread"), TPri_Highest, this, &FMetalDynamicRHI::ProcessInterruptQueue);
+		}
 #endif
 
 #if METAL_USE_SUBMISSION_THREAD
