Metal libraries creation is now lazy again.

This was lost in 41452230 as a part of mmapping lifetime fix and it is now back.

#rnx

diff --git a/Engine/Source/Runtime/Apple/MetalRHI/Private/MetalDynamicRHI_Shaders.cpp b/Engine/Source/Runtime/Apple/MetalRHI/Private/MetalDynamicRHI_Shaders.cpp
--- a/Engine/Source/Runtime/Apple/MetalRHI/Private/MetalDynamicRHI_Shaders.cpp
+++ b/Engine/Source/Runtime/Apple/MetalRHI/Private/MetalDynamicRHI_Shaders.cpp
@@ -13,21 +13,6 @@
 #include "Interfaces/IPluginManager.h"
 
 //------------------------------------------------------------------------------
-static bool GPersistentMappingMetallib = 0;
-static FAutoConsoleVariableRef CVarPersistentMappingMetallib(
-	TEXT("r.Metal.PersistentMappingMetallib"),
-	GPersistentMappingMetallib,
-	TEXT("Makes the metallib file mapping persistent."),
-	ECVF_Default
-);
-
-static void ReleaseLibraryMemory(FMetalShaderLibrary::FShaderLibDataOwner* LibraryMemOwner)
-{
-	if (!GPersistentMappingMetallib)
-	{
-		delete LibraryMemOwner;
-	}
-}
 
 #pragma mark - Metal Dynamic RHI Shader Methods
 
@@ -218,56 +203,22 @@ FRHIShaderLibraryRef FMetalDynamicRHI::RHICreateShaderLibrary(EShaderPlatform Pl
         {
             check(((SerializedShaders.GetNumShaders() + Header.NumShadersPerLibrary - 1) / Header.NumShadersPerLibrary) == Header.NumLibraries);
 
-			TArray<MTLLibraryPtr> Libraries;
-			Libraries.Empty(Header.NumLibraries);
+			TArray<TUniquePtr<FMetalShaderLibrary::FLazyMetalLib>> LazyLibraries;
+			LazyLibraries.Empty(Header.NumLibraries);
 
             for (uint32 i = 0; i < Header.NumLibraries; i++)
-            {
+			{
 				FString MetalLibraryFilePath = (FilePath / LibName) + FString::Printf(TEXT(".%d.metallib"), i);
-				NS::Error* Error;
-
-				MTLLibraryPtr Library;
-				TUniquePtr<FMetalShaderLibrary::FShaderLibDataOwner> LibraryMemOwner = MakeUnique<FMetalShaderLibrary::FShaderLibDataOwner>();
-
-				MapShaderLibraryFile(MetalLibraryFilePath, *LibraryMemOwner);
-				if (LibraryMemOwner->MappedRegion.IsValid())
-				{
-					UE_LOG(LogMetal, Display, TEXT("mmapping %s, %d bytes"), *MetalLibraryFilePath, LibraryMemOwner->MappedCacheFile->GetFileSize());
-					FMetalShaderLibrary::FShaderLibDataOwner* MyLibraryMemOwner = LibraryMemOwner.Release();
-					dispatch_data_t data = dispatch_data_create(MyLibraryMemOwner->MappedRegion->GetMappedPtr(), MyLibraryMemOwner->MappedRegion->GetMappedSize(), nil, ^(void) { 
-						ReleaseLibraryMemory(MyLibraryMemOwner);
-					});
-					Library = NS::TransferPtr(Device->GetDevice()->newLibrary(data, &Error));
-					dispatch_release(data);
-				}
-
-				if (Library.get() == nullptr)
-				{
-					TArray<uint8>& FileData = LibraryMemOwner->Mem;
-					if (FFileHelper::LoadFileToArray(FileData, *MetalLibraryFilePath))
-					{
-						UE_LOG(LogMetal, Display, TEXT("emulating mmapping %s, %d bytes!"), *MetalLibraryFilePath, FileData.Num());
-						FMetalShaderLibrary::FShaderLibDataOwner* MyLibraryMemOwner = LibraryMemOwner.Release();
-						dispatch_data_t data = dispatch_data_create(FileData.GetData(), FileData.Num(), nil, ^(void) { 
-							ReleaseLibraryMemory(MyLibraryMemOwner);
-						});
-						Library = NS::TransferPtr(Device->GetDevice()->newLibrary(data, &Error));
-						dispatch_release(data);
-					}
-				}
-
-				if (Library.get() != nullptr)
-				{
-					Libraries.Add(MoveTemp(Library));
-				}
-				else
-				{
-					UE_LOG(LogMetal, Display, TEXT("Failed to create library: %s"), *MetalLibraryFilePath);
-					return nullptr;
-				}
-            }
-
-			FMetalShaderLibrary* MtlLib = new FMetalShaderLibrary(*Device, Platform, Name, BinaryShaderFile, Header, MoveTemp(SerializedShaders), MoveTemp(ShaderCode), Libraries
+				
+				TUniquePtr<FMetalShaderLibrary::FLazyMetalLib> LazyLibrary = MakeUnique<FMetalShaderLibrary::FLazyMetalLib>();
+				LazyLibrary->MetalLibraryFilePath = MetalLibraryFilePath;
+				LazyLibrary->Data = MakeUnique<FMetalShaderLibrary::FShaderLibDataOwner>();
+				MapShaderLibraryFile(MetalLibraryFilePath, *LazyLibrary->Data);
+				
+				LazyLibraries.Add(MoveTemp(LazyLibrary));
+			}
+
+			FMetalShaderLibrary* MtlLib = new FMetalShaderLibrary(*Device, Platform, Name, BinaryShaderFile, Header, MoveTemp(SerializedShaders), MoveTemp(ShaderCode), MoveTemp(LazyLibraries)
 #if USE_MMAPPED_SHADERARCHIVE
 				, MoveTemp(MemOwner)
 #endif
diff --git a/Engine/Source/Runtime/Apple/MetalRHI/Private/Shaders/MetalShaderLibrary.cpp b/Engine/Source/Runtime/Apple/MetalRHI/Private/Shaders/MetalShaderLibrary.cpp
--- a/Engine/Source/Runtime/Apple/MetalRHI/Private/Shaders/MetalShaderLibrary.cpp
+++ b/Engine/Source/Runtime/Apple/MetalRHI/Private/Shaders/MetalShaderLibrary.cpp
@@ -12,6 +12,22 @@
 #endif // !UE_BUILD_SHIPPING
 #include "MetalShaderTypes.h"
 
+static bool GPersistentMappingMetallib = false;
+static FAutoConsoleVariableRef CVarPersistentMappingMetallib(
+	TEXT("r.Metal.PersistentMappingMetallib"),
+	GPersistentMappingMetallib,
+	TEXT("Makes the metallib file mapping persistent."),
+	ECVF_Default
+);
+
+static void ReleaseLibraryMemory(FMetalShaderLibrary::FShaderLibDataOwner* LibraryMemOwner)
+{
+	if (!GPersistentMappingMetallib)
+	{
+		delete LibraryMemOwner;
+	}
+}
+
 //------------------------------------------------------------------------------
 
 #pragma mark - Metal Shader Library Class Support Routines
@@ -52,14 +68,14 @@ FMetalShaderLibrary::FMetalShaderLibrary(FMetalDevice& MetalDevice,
 										 const FMetalShaderLibraryHeader& InHeader,
 										 FSerializedShaderArchive&& InSerializedShaders,
 										 FShaderCodeArrayType&& InShaderCode,
-										 const TArray<MTLLibraryPtr>& InLibrary)
+										 TArray<TUniquePtr<FLazyMetalLib>>&& InLazyLibraries)
 	: FRHIShaderLibrary(Platform, Name)
 	, Device(MetalDevice)
 	, ShaderLibraryFilename(InShaderLibraryFilename)
 	, Header(InHeader)
 	, SerializedShaders(MoveTemp(InSerializedShaders))
 	, ShaderCode(MoveTemp(InShaderCode))
-	, Library(InLibrary)
+	, LazyLibraries(MoveTemp(InLazyLibraries))
 {
 #if !UE_BUILD_SHIPPING
 	DebugFile = nullptr;
@@ -137,20 +153,63 @@ TRefCountPtr<FRHIShader> FMetalShaderLibrary::CreateShader(int32 Index, bool bRe
 	const TArrayView<const uint8> Code = MakeArrayView(ShaderCode.GetData() + ShaderEntry.Offset, ShaderEntry.Size);
 	const int32 LibraryIndex = Index / Header.NumShadersPerLibrary;
 
+	{
+		TUniquePtr<FLazyMetalLib>& LazyLibrary = LazyLibraries[LibraryIndex];
+		UE::TScopeLock Lock(LazyLibrary->LibraryLock);
+		
+		check(LazyLibrary->Data.Get() != nullptr);
+		
+		if (LazyLibrary->Library.get() == nullptr)
+		{
+			NS::Error* Error;
+			if (LazyLibrary->Data->MappedRegion.IsValid())
+			{
+				UE_LOG(LogMetal, Display, TEXT("mmapping %s, %d bytes"), *LazyLibrary->MetalLibraryFilePath, LazyLibrary->Data->MappedCacheFile->GetFileSize());
+				FMetalShaderLibrary::FShaderLibDataOwner* MyData = LazyLibrary->Data.Release();
+				dispatch_data_t data = dispatch_data_create(MyData->MappedRegion->GetMappedPtr(), MyData->MappedRegion->GetMappedSize(), nil, ^ {
+					ReleaseLibraryMemory(MyData);
+				});
+				LazyLibrary->Library = NS::TransferPtr(Device.GetDevice()->newLibrary(data, &Error));
+				dispatch_release(data);
+			}
+			else
+			{
+				TArray<uint8>& FileData = LazyLibrary->Data->Mem;
+				if (FFileHelper::LoadFileToArray(FileData, *LazyLibrary->MetalLibraryFilePath))
+				{
+					UE_LOG(LogMetal, Display, TEXT("emulating mmapping %s, %d bytes!"), *LazyLibrary->MetalLibraryFilePath, FileData.Num());
+					FMetalShaderLibrary::FShaderLibDataOwner* MyData = LazyLibrary->Data.Release();
+					dispatch_data_t data = dispatch_data_create(FileData.GetData(), FileData.Num(), nil, ^ {
+						ReleaseLibraryMemory(MyData);
+					});
+					LazyLibrary->Library = NS::TransferPtr(Device.GetDevice()->newLibrary(data, &Error));
+					dispatch_release(data);
+				}
+			}
+			
+			if (LazyLibrary->Library.get() == nullptr)
+			{
+				UE_LOG(LogMetal, Fatal, TEXT("Failed to create library: %s"), *LazyLibrary->MetalLibraryFilePath);
+			}
+		}
+	}
+	
+	check(LazyLibraries[LibraryIndex]->Library.get() != nullptr);
+	
 	TRefCountPtr<FRHIShader> Shader;
 	switch (ShaderEntry.Frequency)
 	{
 		case SF_Vertex:
-			Shader = CreateMetalShader<FMetalVertexShader>(Device, Code, Library[LibraryIndex]);
+			Shader = CreateMetalShader<FMetalVertexShader>(Device, Code, LazyLibraries[LibraryIndex]->Library);
 			break;
 
 		case SF_Pixel:
-			Shader = CreateMetalShader<FMetalPixelShader>(Device, Code, Library[LibraryIndex]);
+			Shader = CreateMetalShader<FMetalPixelShader>(Device, Code, LazyLibraries[LibraryIndex]->Library);
 			break;
  
  		case SF_Geometry:
 #if PLATFORM_SUPPORTS_GEOMETRY_SHADERS
-            Shader = CreateMetalShader<FMetalGeometryShader>(Device, Code, Library[LibraryIndex]);
+            Shader = CreateMetalShader<FMetalGeometryShader>(Device, Code, LazyLibraries[LibraryIndex]->Library);
 #else
             checkf(false, TEXT("Geometry shaders not supported"));
 #endif
@@ -159,7 +218,7 @@ TRefCountPtr<FRHIShader> FMetalShaderLibrary::CreateShader(int32 Index, bool bRe
 
         case SF_Mesh:
 #if PLATFORM_SUPPORTS_MESH_SHADERS
-            Shader = CreateMetalShader<FMetalMeshShader>(Device, Code, Library[LibraryIndex]);
+            Shader = CreateMetalShader<FMetalMeshShader>(Device, Code, LazyLibraries[LibraryIndex]->Library);
 #else
 			checkf(false, TEXT("Mesh shaders not supported"));
 #endif
@@ -167,14 +226,14 @@ TRefCountPtr<FRHIShader> FMetalShaderLibrary::CreateShader(int32 Index, bool bRe
 
         case SF_Amplification:
 #if PLATFORM_SUPPORTS_MESH_SHADERS
-            Shader = CreateMetalShader<FMetalAmplificationShader>(Device, Code, Library[LibraryIndex]);
+            Shader = CreateMetalShader<FMetalAmplificationShader>(Device, Code, LazyLibraries[LibraryIndex]->Library);
 #else
 			checkf(false, TEXT("Amplification shaders not supported"));
 #endif
             break;
 
 		case SF_Compute:
-			Shader = CreateMetalShader<FMetalComputeShader>(Device, Code, Library[LibraryIndex]);
+			Shader = CreateMetalShader<FMetalComputeShader>(Device, Code, LazyLibraries[LibraryIndex]->Library);
 			break;
 
 		default:
diff --git a/Engine/Source/Runtime/Apple/MetalRHI/Private/Shaders/MetalShaderLibrary.h b/Engine/Source/Runtime/Apple/MetalRHI/Private/Shaders/MetalShaderLibrary.h
--- a/Engine/Source/Runtime/Apple/MetalRHI/Private/Shaders/MetalShaderLibrary.h
+++ b/Engine/Source/Runtime/Apple/MetalRHI/Private/Shaders/MetalShaderLibrary.h
@@ -24,6 +24,15 @@ public:
 		TUniquePtr<IMappedFileHandle> MappedCacheFile;
 		TUniquePtr<IMappedFileRegion> MappedRegion;
 	};
+	
+	struct FLazyMetalLib
+	{
+		MTLLibraryPtr Library;
+		FString MetalLibraryFilePath;
+		TUniquePtr<FShaderLibDataOwner> Data;
+		UE::FMutex LibraryLock;
+	};
+	
 #if !USE_MMAPPED_SHADERARCHIVE
     using FShaderCodeArrayType = TArray<uint8>;
 #else 
@@ -38,10 +47,10 @@ public:
         const FMetalShaderLibraryHeader& InHeader,
         FSerializedShaderArchive&& InSerializedShaders,
         FShaderCodeArrayType&& InShaderCode,
-		const TArray<MTLLibraryPtr>& InLibrary,
+		TArray<TUniquePtr<FLazyMetalLib>>&& InLazyLibraries,
         TUniquePtr<FShaderLibDataOwner>&& InMemOwner)
         :
-        FMetalShaderLibrary(Device, Platform, Name, InShaderLibraryFilename, InHeader, MoveTemp(InSerializedShaders), MoveTemp(InShaderCode), InLibrary)
+        FMetalShaderLibrary(Device, Platform, Name, InShaderLibraryFilename, InHeader, MoveTemp(InSerializedShaders), MoveTemp(InShaderCode), MoveTemp(InLazyLibraries))
         {
 			MemOwner = MoveTemp(InMemOwner);
         }
@@ -57,7 +66,7 @@ public:
 						const FMetalShaderLibraryHeader& InHeader,
 						FSerializedShaderArchive&& InSerializedShaders,
 						FShaderCodeArrayType&& InShaderCode,
-						const TArray<MTLLibraryPtr>& InLibrary);
+						TArray<TUniquePtr<FLazyMetalLib>>&& InLazyLibraries);
 
 	virtual ~FMetalShaderLibrary();
 
@@ -87,7 +96,7 @@ private:
 	FMetalShaderLibraryHeader Header;
 	FSerializedShaderArchive SerializedShaders;
 	FShaderCodeArrayType ShaderCode;
-	TArray<MTLLibraryPtr> Library;
+	TArray<TUniquePtr<FLazyMetalLib>> LazyLibraries;
 #if !UE_BUILD_SHIPPING
 	class FMetalShaderDebugZipFile* DebugFile;
 #endif
