Metal libraries lazy creation assertion fix.

LazyLibrary data is released after the very first library creation, so it triggers assert on any consecutive shader creation from the same library.

#rnx

diff --git a/Engine/Source/Runtime/Apple/MetalRHI/Private/Shaders/MetalShaderLibrary.cpp b/Engine/Source/Runtime/Apple/MetalRHI/Private/Shaders/MetalShaderLibrary.cpp
--- a/Engine/Source/Runtime/Apple/MetalRHI/Private/Shaders/MetalShaderLibrary.cpp
+++ b/Engine/Source/Runtime/Apple/MetalRHI/Private/Shaders/MetalShaderLibrary.cpp
@@ -157,10 +157,10 @@ TRefCountPtr<FRHIShader> FMetalShaderLibrary::CreateShader(int32 Index, bool bRe
 		TUniquePtr<FLazyMetalLib>& LazyLibrary = LazyLibraries[LibraryIndex];
 		UE::TScopeLock Lock(LazyLibrary->LibraryLock);
 		
-		check(LazyLibrary->Data.Get() != nullptr);
-		
 		if (LazyLibrary->Library.get() == nullptr)
 		{
+			check(LazyLibrary->Data.Get() != nullptr);
+			
 			NS::Error* Error;
 			if (LazyLibrary->Data->MappedRegion.IsValid())
 			{
@@ -189,6 +189,7 @@ TRefCountPtr<FRHIShader> FMetalShaderLibrary::CreateShader(int32 Index, bool bRe
 			
 			if (LazyLibrary->Library.get() == nullptr)
 			{
+				UE_LOG(LogMetal, Error, TEXT("Metal library creationg error: %s"), *NSStringToFString(Error->description()));
 				UE_LOG(LogMetal, Fatal, TEXT("Failed to create library: %s"), *LazyLibrary->MetalLibraryFilePath);
 			}
 		}
