Metal RHI

- Added support for RunOnQueue function to support Apple FX

diff --git a/Engine/Source/Runtime/Apple/MetalRHI/Private/MetalDynamicRHI.cpp b/Engine/Source/Runtime/Apple/MetalRHI/Private/MetalDynamicRHI.cpp
--- a/Engine/Source/Runtime/Apple/MetalRHI/Private/MetalDynamicRHI.cpp
+++ b/Engine/Source/Runtime/Apple/MetalRHI/Private/MetalDynamicRHI.cpp
@@ -268,3 +268,27 @@ bool FMetalDynamicRHI::RHIMatchPrecachePSOInitializers(const FGraphicsPipelineSt
 
 	return true;
 }
+
+void FMetalDynamicRHI::RHIRunOnQueue(TFunction<void(FMetalCommandQueue&)>&& CodeToRun, bool bWaitForSubmission)
+{
+	FGraphEventRef SubmissionEvent;
+
+	TArray<FMetalPayload*> Payloads;
+	FMetalPayload* Payload = new FMetalPayload(Device->GetCommandQueue(EMetalQueueType::Direct));
+	Payloads.Add(Payload);
+
+	Payload->PreExecuteCallback = MoveTemp(CodeToRun);
+
+	if (bWaitForSubmission)
+	{
+		SubmissionEvent = FGraphEvent::CreateGraphEvent();
+		Payload->SubmissionEvent = SubmissionEvent;
+	}
+
+	SubmitPayloads(MoveTemp(Payloads));
+	
+	if (SubmissionEvent && !SubmissionEvent->IsComplete())
+	{
+		SubmissionEvent->Wait();
+	}
+}
diff --git a/Engine/Source/Runtime/Apple/MetalRHI/Private/MetalDynamicRHI.h b/Engine/Source/Runtime/Apple/MetalRHI/Private/MetalDynamicRHI.h
--- a/Engine/Source/Runtime/Apple/MetalRHI/Private/MetalDynamicRHI.h
+++ b/Engine/Source/Runtime/Apple/MetalRHI/Private/MetalDynamicRHI.h
@@ -240,6 +240,8 @@ public:
 	virtual void RHIBeginRenderQueryBatch_TopOfPipe(FRHICommandListBase& RHICmdList, ERenderQueryType QueryType) final override;
 	virtual void RHIEndRenderQueryBatch_TopOfPipe(FRHICommandListBase& RHICmdList, ERenderQueryType QueryType) final override;
 	
+	void RHIRunOnQueue(TFunction<void(FMetalCommandQueue&)>&& CodeToRun, bool bWaitForSubmission);
+	
 	template <typename ...Args>
 	void DeferredDelete(Args&&... InArgs)
 	{
