Ensure valid URL for pipeline cache save location. The current invalid URL generated was invalid as it was constructing a file URL using the absolute string of another URL as a path parameter to the constructor, causing MTLBinaryArchive->serialzeToURL to crash when used.

diff --git a/Engine/Source/Runtime/Apple/MetalRHI/Private/MetalPipeline.cpp b/Engine/Source/Runtime/Apple/MetalRHI/Private/MetalPipeline.cpp
--- a/Engine/Source/Runtime/Apple/MetalRHI/Private/MetalPipeline.cpp
+++ b/Engine/Source/Runtime/Apple/MetalRHI/Private/MetalPipeline.cpp
@@ -857,13 +857,13 @@ static uint32_t GPSOHarvestCount = 0;
 
 static NS::URL* PipelineCacheSaveLocation()
 {
-	NSString* path = NSSearchPathForDirectoriesInDomains(NSDocumentDirectory, NSUserDomainMask, YES)[0];
-	if (!path)
+	NSString* dir = NSSearchPathForDirectoriesInDomains(NSDocumentDirectory, NSUserDomainMask, YES)[0];
+	if (!dir)
 	{
 		UE_LOG(LogMetal, Error, TEXT("Metal Pipeline Cache: Unable to find Documents folder"));
 	}
-	NSURL* url = [NSURL fileURLWithPath : [NSString stringWithFormat : @"%@/mtlarchive.bin", path] ];
-    NS::URL* NativeURL = NS::URL::fileURLWithPath(NS::String::string(url.absoluteString.UTF8String, NS::UTF8StringEncoding));
+	NSString* path = [dir stringByAppendingString : @"/mtlarchive.bin"];
+	NS::URL* NativeURL = NS::URL::fileURLWithPath(NS::String::string(path.UTF8String, NS::UTF8StringEncoding));
 	return NativeURL;
 }
 
