Implement functionality for -MetalPSOCache=use.

diff --git a/Engine/Source/Runtime/Apple/MetalRHI/Private/MetalPipeline.cpp b/Engine/Source/Runtime/Apple/MetalRHI/Private/MetalPipeline.cpp
--- a/Engine/Source/Runtime/Apple/MetalRHI/Private/MetalPipeline.cpp
+++ b/Engine/Source/Runtime/Apple/MetalRHI/Private/MetalPipeline.cpp
@@ -1036,6 +1036,11 @@ static FMetalShaderPipelinePtr CreateMTLRenderPipeline(FMetalDevice& Device, boo
 		MTLRenderPipelineDescriptorPtr RenderPipelineDesc = NS::TransferPtr(MTL::RenderPipelineDescriptor::alloc()->init());
         check(RenderPipelineDesc);
         
+		if ((GPSOCacheMode == CacheMode::Append) || (GPSOCacheMode == CacheMode::Use))
+		{
+			RenderPipelineDesc->setBinaryArchives(NS::Array::array(GPSOBinaryArchive));
+		}
+		
         MTLComputePipelineDescriptorPtr ComputePipelineDesc;
 		
 		if (!ConfigureRenderPipelineDescriptor(Device, RenderPipelineDesc.get(), Key, Init))
@@ -1143,7 +1148,6 @@ static FMetalShaderPipelinePtr CreateMTLRenderPipeline(FMetalDevice& Device, boo
 		Pipeline = new FMetalShaderPipeline(Device);
 		METAL_DEBUG_OPTION(FMemory::Memzero(Pipeline->ResourceMask, sizeof(Pipeline->ResourceMask)));
 		
-		MTLRenderPipelineDescriptorPtr DebugPipelineDesc = NS::TransferPtr(MTL::RenderPipelineDescriptor::alloc()->init());
 		MTLMeshRenderPipelineDescriptorPtr MeshPipelineDesc = NS::TransferPtr(MTL::MeshRenderPipelineDescriptor::alloc()->init());
 		MeshPipelineDesc->setObjectFunction(vertexFunction.get());
 		MeshPipelineDesc->setMeshFunction(geometryFunction.get());
@@ -1237,7 +1241,6 @@ static FMetalShaderPipelinePtr CreateMTLRenderPipeline(FMetalDevice& Device, boo
         Pipeline = new FMetalShaderPipeline(Device);
         METAL_DEBUG_OPTION(FMemory::Memzero(Pipeline->ResourceMask, sizeof(Pipeline->ResourceMask)));
 
-		MTLRenderPipelineDescriptorPtr DebugPipelineDesc 	= NS::TransferPtr(MTL::RenderPipelineDescriptor::alloc()->init());
 		MTLMeshRenderPipelineDescriptorPtr MeshPipelineDesc = NS::TransferPtr(MTL::MeshRenderPipelineDescriptor::alloc()->init());
         MeshPipelineDesc->setObjectFunction(amplificationFunction.get());
         MeshPipelineDesc->setMeshFunction(meshFunction.get());
