Apple Vision Pro audio spatialization fix
-Our audio is already spatialized, but was also being spatialized by apple's audio system.

diff --git a/Engine/Source/Runtime/Apple/AudioMixerAudioUnit/Private/AudioMixerPlatformAudioUnit.cpp b/Engine/Source/Runtime/Apple/AudioMixerAudioUnit/Private/AudioMixerPlatformAudioUnit.cpp
--- a/Engine/Source/Runtime/Apple/AudioMixerAudioUnit/Private/AudioMixerPlatformAudioUnit.cpp
+++ b/Engine/Source/Runtime/Apple/AudioMixerAudioUnit/Private/AudioMixerPlatformAudioUnit.cpp
@@ -114,6 +114,14 @@ namespace Audio
 			UE_LOG(LogAudioMixerAudioUnit, Display, TEXT("Error setting sample rate."));
 		}
 		
+#if PLATFORM_VISIONOS
+		Success = [AudioSession setIntendedSpatialExperience:AVAudioSessionSpatialExperienceBypassed options:nil error:&error];
+		if (!Success)
+		{
+			UE_LOG(LogAudioMixerAudioUnit, Display, TEXT("Error while setting spatial experience to bypass for VisionOS."));
+		}
+#endif
+
 		// By calling setPreferredIOBufferDuration, we indicate that we would prefer that the buffer size not change if possible.
 		float AudioMixerBufferSizeInSec = InternalPlatformSettings.CallbackBufferFrameSize / GraphSampleRate;
 		Success = [AudioSession setPreferredIOBufferDuration:AudioMixerBufferSizeInSec error: &error];
