Fix UpdateTexture2D() case where the source region has an offset and the texture is block compressed.
In this case SourcePitch already accounts for the block size, and previously we were double counting it.

diff --git a/Engine/Source/Runtime/Apple/MetalRHI/Private/MetalTexture.cpp b/Engine/Source/Runtime/Apple/MetalRHI/Private/MetalTexture.cpp
--- a/Engine/Source/Runtime/Apple/MetalRHI/Private/MetalTexture.cpp
+++ b/Engine/Source/Runtime/Apple/MetalRHI/Private/MetalTexture.cpp
@@ -1544,7 +1544,7 @@ static FMetalBufferPtr Internal_CreateBufferAndCopyTexture2DUpdateRegionData(FMe
 	const uint32 WidthInBlocks  = FMath::DivideAndRoundUp<uint32>(UpdateRegion.Width,  FormatInfo.BlockSizeX);
 	const uint32 HeightInBlocks = FMath::DivideAndRoundUp<uint32>(UpdateRegion.Height, FormatInfo.BlockSizeY);
 
-	const uint8* OffsetSourceData = SourceData + FormatInfo.BlockBytes * SrcXInBlocks + InOutSourcePitch * SrcYInBlocks * FormatInfo.BlockSizeY;
+	const uint8* OffsetSourceData = SourceData + FormatInfo.BlockBytes * SrcXInBlocks + InOutSourcePitch * SrcYInBlocks;
 	uint32 UpdatePitch = InOutSourcePitch;
 
 	FMetalBufferPtr OutBuffer;
@@ -1658,8 +1658,8 @@ static FMetalBufferPtr Internal_CreateBufferAndCopyTexture3DUpdateRegionData(FMe
 		const uint32 SrcZInBlocks   = FMath::DivideAndRoundUp<uint32>(UpdateRegion.SrcZ,   	FormatInfo.BlockSizeZ);
 
 		const uint8* OffsetSourceData = SourceData + FormatInfo.BlockBytes * SrcXInBlocks +
-										SourceRowPitch * SrcYInBlocks * FormatInfo.BlockSizeY +
-										SourceDepthPitch * SrcZInBlocks * FormatInfo.BlockSizeZ;
+										SourceRowPitch * SrcYInBlocks +
+										SourceDepthPitch * SrcZInBlocks;
 		
 		OutBuffer = Device.CreatePooledBuffer(FMetalPooledBufferArgs(&Device, ExpandedBufferSize, BUF_Static, MTL::StorageModeShared));
 		InternalExpandR8ToStandardRGBA3D((uint32*)OutBuffer->Contents(), UpdateRegion, SourceRowPitch, SourceDepthPitch, OffsetSourceData);
diff --git a/Engine/Source/Runtime/D3D12RHI/Private/D3D12Texture.cpp b/Engine/Source/Runtime/D3D12RHI/Private/D3D12Texture.cpp
--- a/Engine/Source/Runtime/D3D12RHI/Private/D3D12Texture.cpp
+++ b/Engine/Source/Runtime/D3D12RHI/Private/D3D12Texture.cpp
@@ -2143,7 +2143,7 @@ void FD3D12Texture::UpdateTexture2D(FRHICommandListBase& RHICmdList, uint32 MipI
 	}
 	check(StagingMemory);
 
-	const uint8* CopySrc = SourceData + FormatInfo.BlockBytes * SrcXInBlocks + SourcePitch * SrcYInBlocks * FormatInfo.BlockSizeY;
+	const uint8* CopySrc = SourceData + FormatInfo.BlockBytes * SrcXInBlocks + SourcePitch * SrcYInBlocks;
 	uint8* CopyDst = (uint8*)StagingMemory;
 	for (uint32 BlockRow = 0; BlockRow < HeightInBlocks; BlockRow++)
 	{
diff --git a/Engine/Source/Runtime/OpenGLDrv/Private/OpenGLTexture.cpp b/Engine/Source/Runtime/OpenGLDrv/Private/OpenGLTexture.cpp
--- a/Engine/Source/Runtime/OpenGLDrv/Private/OpenGLTexture.cpp
+++ b/Engine/Source/Runtime/OpenGLDrv/Private/OpenGLTexture.cpp
@@ -1582,7 +1582,7 @@ void FOpenGLDynamicRHI::RHIUpdateTexture2D(FRHICommandListBase& RHICmdList, FRHI
 	const uint32 WidthInBlocks  = FMath::DivideAndRoundUp<uint32>(UpdateRegion.Width,  FormatInfo.BlockSizeX);
 	const uint32 HeightInBlocks = FMath::DivideAndRoundUp<uint32>(UpdateRegion.Height, FormatInfo.BlockSizeY);
 
-	const void* UpdateMemory = SourceData + FormatInfo.BlockBytes * SrcXInBlocks + SourcePitch * SrcYInBlocks * FormatInfo.BlockSizeY;
+	const void* UpdateMemory = SourceData + FormatInfo.BlockBytes * SrcXInBlocks + SourcePitch * SrcYInBlocks;
 	uint32 UpdatePitch = SourcePitch;
 
 	const bool bNeedStagingMemory = RHICmdList.IsTopOfPipe();
diff --git a/Engine/Source/Runtime/VulkanRHI/Private/VulkanTexture.cpp b/Engine/Source/Runtime/VulkanRHI/Private/VulkanTexture.cpp
--- a/Engine/Source/Runtime/VulkanRHI/Private/VulkanTexture.cpp
+++ b/Engine/Source/Runtime/VulkanRHI/Private/VulkanTexture.cpp
@@ -1303,7 +1303,7 @@ void FVulkanDynamicRHI::InternalUpdateTexture2D(FRHICommandListBase& RHICmdList,
 	VulkanRHI::FStagingBuffer* StagingBuffer = Device->GetStagingManager().AcquireBuffer(StagingBufferSize);
 	void* RESTRICT StagingMemory = StagingBuffer->GetMappedPointer();
 
-	const uint8* CopySrc = SourceData + FormatInfo.BlockBytes * SrcXInBlocks + SourcePitch * SrcYInBlocks * FormatInfo.BlockSizeY;
+	const uint8* CopySrc = SourceData + FormatInfo.BlockBytes * SrcXInBlocks + SourcePitch * SrcYInBlocks;
 	uint8* CopyDst = (uint8*)StagingMemory;
 	for (uint32 BlockRow = 0; BlockRow < HeightInBlocks; BlockRow++)
 	{
diff --git a/Engine/Source/Runtime/Windows/D3D11RHI/Private/D3D11Texture.cpp b/Engine/Source/Runtime/Windows/D3D11RHI/Private/D3D11Texture.cpp
--- a/Engine/Source/Runtime/Windows/D3D11RHI/Private/D3D11Texture.cpp
+++ b/Engine/Source/Runtime/Windows/D3D11RHI/Private/D3D11Texture.cpp
@@ -1197,7 +1197,7 @@ void FD3D11DynamicRHI::RHIUpdateTexture2D(FRHICommandListBase& RHICmdList, FRHIT
 	const uint32 WidthInBlocks  = FMath::DivideAndRoundUp<uint32>(UpdateRegion.Width,  FormatInfo.BlockSizeX);
 	const uint32 HeightInBlocks = FMath::DivideAndRoundUp<uint32>(UpdateRegion.Height, FormatInfo.BlockSizeY);
 
-	const void* UpdateMemory = SourceData + FormatInfo.BlockBytes * SrcXInBlocks + SourcePitch * SrcYInBlocks * FormatInfo.BlockSizeY;
+	const void* UpdateMemory = SourceData + FormatInfo.BlockBytes * SrcXInBlocks + SourcePitch * SrcYInBlocks;
 	uint32 UpdatePitch = SourcePitch;
 
 	const bool bNeedStagingMemory = RHICmdList.IsTopOfPipe();
