Metal RHI

Fixed incorrectly named enums for new profiler

diff --git a/Engine/Source/Runtime/Apple/MetalRHI/Private/MetalCommandBuffer.h b/Engine/Source/Runtime/Apple/MetalRHI/Private/MetalCommandBuffer.h
--- a/Engine/Source/Runtime/Apple/MetalRHI/Private/MetalCommandBuffer.h
+++ b/Engine/Source/Runtime/Apple/MetalRHI/Private/MetalCommandBuffer.h
@@ -49,7 +49,7 @@ public:
 	void BeginBreadcrumb(FRHIBreadcrumbNode* Node)
 	{
 		FMetalBreadcrumbTrackerObject TrackerObject;
-		TrackerObject.Type = EMetalBreadcrumbTrackerType_Begin;
+		TrackerObject.Type = EMetalBreadcrumbTrackerType::Begin;
 		TrackerObject.CmdBuffer = this;
 		TrackerObject.Node = Node;
 		BreadcrumbTrackerStream.Add(TrackerObject);
@@ -58,7 +58,7 @@ public:
 	void EndBreadcrumb(FRHIBreadcrumbNode* Node)
 	{
 		FMetalBreadcrumbTrackerObject TrackerObject;
-		TrackerObject.Type = EMetalBreadcrumbTrackerType_End;
+		TrackerObject.Type = EMetalBreadcrumbTrackerType::End;
 		TrackerObject.CmdBuffer = this;
 		TrackerObject.Node = Node;
 		BreadcrumbTrackerStream.Add(TrackerObject);
@@ -82,7 +82,7 @@ public:
 		}
 #elif WITH_RHI_BREADCRUMBS
 		FMetalBreadcrumbTrackerObject TrackerObject;
-		TrackerObject.Type = EMetalBreadcrumbTrackerType_Encode;
+		TrackerObject.Type = EMetalBreadcrumbTrackerType::Encode;
 		TrackerObject.CmdBuffer = this;
 		TrackerObject.CounterSample = CounterSample;
 		BreadcrumbTrackerStream.Add(TrackerObject);
diff --git a/Engine/Source/Runtime/Apple/MetalRHI/Private/MetalProfiler.cpp b/Engine/Source/Runtime/Apple/MetalRHI/Private/MetalProfiler.cpp
--- a/Engine/Source/Runtime/Apple/MetalRHI/Private/MetalProfiler.cpp
+++ b/Engine/Source/Runtime/Apple/MetalRHI/Private/MetalProfiler.cpp
@@ -918,17 +918,17 @@ void FMetalBreadcrumbProfiler::ResolveBreadcrumbTrackerStream(TArray<FMetalBread
 	{
 		switch(Tracker.Type)
 		{
-			case EMetalBreadcrumbTrackerType_Begin:
+			case EMetalBreadcrumbTrackerType::Begin:
 			{
 				OnBreadcrumbBegin(Tracker.Node);
 				break;
 			}
-			case EMetalBreadcrumbTrackerType_End:
+			case EMetalBreadcrumbTrackerType::End:
 			{	
 				OnBreadcrumbEnd(Tracker.Node);
 				break;
 			}
-			case EMetalBreadcrumbTrackerType_Encode:
+			case EMetalBreadcrumbTrackerType::Encode:
 			{
 				AddSample(Tracker.CounterSample);
 				break;
