Metal RHI

Fixed assert when running parallel and the new GPU profiler

diff --git a/Engine/Source/Runtime/Apple/MetalRHI/Private/MetalRHI.cpp b/Engine/Source/Runtime/Apple/MetalRHI/Private/MetalRHI.cpp
--- a/Engine/Source/Runtime/Apple/MetalRHI/Private/MetalRHI.cpp
+++ b/Engine/Source/Runtime/Apple/MetalRHI/Private/MetalRHI.cpp
@@ -1356,20 +1356,23 @@ void FMetalRHICommandContext::RHIBeginBreadcrumbGPU(FRHIBreadcrumbNode* Breadcru
 
 #if ENABLE_METAL_GPUPROFILE
 #if RHI_NEW_GPU_PROFILER
-	if(!CurrentEncoder.GetCommandBuffer())
+	if(!CurrentEncoder.IsParallelEncoding())
 	{
-		StartCommandBuffer();
-	}
-	
-	FMetalCommandBuffer* CmdBuffer = CurrentEncoder.GetCommandBuffer();
-	
-	// Can't process breadcrumbs if we are within a render pass
-	if(Device.SupportsFeature(EMetalFeaturesStageCounterSampling))
-	{
-		FMetalBreadcrumbEvent& Event = FMetalBreadcrumbProfiler::GetInstance()->GetBreadcrumbEvent(Breadcrumb, bWithinRenderPass);
-		Event.TimestampTOP = &CmdBuffer->EmplaceProfilerEvent<UE::RHI::GPUProfiler::FEvent::FBeginBreadcrumb>(Breadcrumb).GPUTimestampTOP;
-		*Event.TimestampTOP = 0;
-		CmdBuffer->BeginBreadcrumb(Breadcrumb);
+		if(!CurrentEncoder.GetCommandBuffer())
+		{
+			StartCommandBuffer();
+		}
+		
+		FMetalCommandBuffer* CmdBuffer = CurrentEncoder.GetCommandBuffer();
+		
+		// Can't process breadcrumbs if we are within a render pass
+		if(Device.SupportsFeature(EMetalFeaturesStageCounterSampling))
+		{
+			FMetalBreadcrumbEvent& Event = FMetalBreadcrumbProfiler::GetInstance()->GetBreadcrumbEvent(Breadcrumb, bWithinRenderPass);
+			Event.TimestampTOP = &CmdBuffer->EmplaceProfilerEvent<UE::RHI::GPUProfiler::FEvent::FBeginBreadcrumb>(Breadcrumb).GPUTimestampTOP;
+			*Event.TimestampTOP = 0;
+			CmdBuffer->BeginBreadcrumb(Breadcrumb);
+		}
 	}
 #else
 	if (Profiler && Profiler->IsProfilingGPU())
@@ -1384,22 +1387,24 @@ void FMetalRHICommandContext::RHIEndBreadcrumbGPU(FRHIBreadcrumbNode* Breadcrumb
 {
 #if ENABLE_METAL_GPUPROFILE
 #if RHI_NEW_GPU_PROFILER
-	if(!CurrentEncoder.GetCommandBuffer())
+	if(!CurrentEncoder.IsParallelEncoding())
 	{
-		StartCommandBuffer();	
-	}
-	
-	FMetalCommandBuffer* CmdBuffer = CurrentEncoder.GetCommandBuffer();
-	
-	if(Device.SupportsFeature(EMetalFeaturesStageCounterSampling))
-	{
-		FMetalBreadcrumbEvent& Event = FMetalBreadcrumbProfiler::GetInstance()->GetBreadcrumbEvent(Breadcrumb, bWithinRenderPass);
-		Event.TimestampBOP = &CmdBuffer->EmplaceProfilerEvent<UE::RHI::GPUProfiler::FEvent::FEndBreadcrumb>(Breadcrumb).GPUTimestampBOP;
-		*Event.TimestampBOP = 0;
+		if(!CurrentEncoder.GetCommandBuffer())
+		{
+			StartCommandBuffer();	
+		}
 		
-		CmdBuffer->EndBreadcrumb(Breadcrumb);
+		FMetalCommandBuffer* CmdBuffer = CurrentEncoder.GetCommandBuffer();
+		
+		if(Device.SupportsFeature(EMetalFeaturesStageCounterSampling))
+		{
+			FMetalBreadcrumbEvent& Event = FMetalBreadcrumbProfiler::GetInstance()->GetBreadcrumbEvent(Breadcrumb, bWithinRenderPass);
+			Event.TimestampBOP = &CmdBuffer->EmplaceProfilerEvent<UE::RHI::GPUProfiler::FEvent::FEndBreadcrumb>(Breadcrumb).GPUTimestampBOP;
+			*Event.TimestampBOP = 0;
+			
+			CmdBuffer->EndBreadcrumb(Breadcrumb);
+		}
 	}
-	
 #else
 	if (Profiler && Profiler->IsProfilingGPU())
 	{
