Fix compile error for missing IsParallelEncoding

diff --git a/Engine/Source/Runtime/Apple/MetalRHI/Private/MetalCommandEncoder.h b/Engine/Source/Runtime/Apple/MetalRHI/Private/MetalCommandEncoder.h
--- a/Engine/Source/Runtime/Apple/MetalRHI/Private/MetalCommandEncoder.h
+++ b/Engine/Source/Runtime/Apple/MetalRHI/Private/MetalCommandEncoder.h
@@ -221,6 +221,11 @@ public:
 				IsParallelRenderCommandEncoderActive();
 	}
 	
+	inline bool IsParallelEncoding()
+	{
+		return bIsParallelEncoding;
+	}
+	
 #if METAL_RHI_RAYTRACING
 	/** @returns True if and only if there is an active acceleration structure command encoder, otherwise false. */
 	bool IsAccelerationStructureCommandEncoderActive(void) const;
