Metal RHI

Fix compiler error for missing include

diff --git a/Engine/Source/Runtime/Apple/MetalRHI/Public/MetalRHIContext.h b/Engine/Source/Runtime/Apple/MetalRHI/Public/MetalRHIContext.h
--- a/Engine/Source/Runtime/Apple/MetalRHI/Public/MetalRHIContext.h
+++ b/Engine/Source/Runtime/Apple/MetalRHI/Public/MetalRHIContext.h
@@ -3,6 +3,7 @@
 #pragma once
 
 // Metal RHI public headers.
+#include "MetalRHIPrivate.h"
 #include "MetalThirdParty.h"
 #include "MetalState.h"
 #include "MetalResources.h"
