MetalRHI

Guard for bindless when compiling on iOS for descriptors

diff --git a/Engine/Source/Runtime/Apple/MetalRHI/Private/MetalRHIContext.cpp b/Engine/Source/Runtime/Apple/MetalRHI/Private/MetalRHIContext.cpp
--- a/Engine/Source/Runtime/Apple/MetalRHI/Private/MetalRHIContext.cpp
+++ b/Engine/Source/Runtime/Apple/MetalRHI/Private/MetalRHIContext.cpp
@@ -57,7 +57,9 @@ FMetalRHICommandContext::~FMetalRHICommandContext()
 
 void FMetalRHICommandContext::ResetContext()
 {
+#if PLATFORM_SUPPORTS_BINDLESS_RENDERING
 	check(!ComputeDescriptorIndices.Num());
+#endif
 	
 	// Reset cached state in the encoder
 	StateCache.Reset();
diff --git a/Engine/Source/Runtime/Apple/MetalRHI/Public/MetalRHIContext.h b/Engine/Source/Runtime/Apple/MetalRHI/Public/MetalRHIContext.h
--- a/Engine/Source/Runtime/Apple/MetalRHI/Public/MetalRHIContext.h
+++ b/Engine/Source/Runtime/Apple/MetalRHI/Public/MetalRHIContext.h
@@ -215,11 +215,13 @@ public:
 	
 	void Finalize(TArray<FMetalPayload*>& OutPayloads);
 	
+#if PLATFORM_SUPPORTS_BINDLESS_RENDERING
 	void EnqueueDescriptorUpdate(IRDescriptorTableEntry Entry, uint32_t Index)
 	{
 		ComputeDescriptorEntries.Add(Entry);
 		ComputeDescriptorIndices.Add(Index);
 	}
+#endif
 	
 	void PushDescriptorUpdates();
 	
@@ -386,8 +388,10 @@ protected:
 	
 	FMetalParallelRenderPassInfo* ParallelRenderPassInfo;
 	
+#if PLATFORM_SUPPORTS_BINDLESS_RENDERING
 	TResourceArray<IRDescriptorTableEntry>    ComputeDescriptorEntries;
 	TResourceArray<uint32_t>                  ComputeDescriptorIndices;
+#endif
 	
 private:
 	void RHIClearMRT(bool bClearColor, int32 NumClearColors, const FLinearColor* ColorArray, bool bClearDepth, float Depth, bool bClearStencil, uint32 Stencil);
