Metal RHI

Temporarily disable compute descriptor updates

diff --git a/Engine/Source/Runtime/Apple/MetalRHI/Private/MetalBindlessDescriptors.cpp b/Engine/Source/Runtime/Apple/MetalRHI/Private/MetalBindlessDescriptors.cpp
--- a/Engine/Source/Runtime/Apple/MetalRHI/Private/MetalBindlessDescriptors.cpp
+++ b/Engine/Source/Runtime/Apple/MetalRHI/Private/MetalBindlessDescriptors.cpp
@@ -28,7 +28,7 @@ static FAutoConsoleVariableRef CVarBindlessSamplerDescriptorHeapSize(
 	ECVF_ReadOnly
 );
 
-static int32_t GMetalBindlessUpdateDescriptorOnCompute = 1;
+static int32_t GMetalBindlessUpdateDescriptorOnCompute = 0;
 static FAutoConsoleVariableRef CVarBindlessUpdateDescriptorOnCompute(
 	TEXT("Metal.Bindless.UpdateDescriptorOnCompute"),
 	GMetalBindlessUpdateDescriptorOnCompute,
