MetalRHI

Fix vision OS compilation error when using new profiler

diff --git a/Engine/Source/Runtime/Apple/MetalRHI/Private/MetalViewport.cpp b/Engine/Source/Runtime/Apple/MetalRHI/Private/MetalViewport.cpp
--- a/Engine/Source/Runtime/Apple/MetalRHI/Private/MetalViewport.cpp
+++ b/Engine/Source/Runtime/Apple/MetalRHI/Private/MetalViewport.cpp
@@ -723,7 +723,7 @@ void FMetalViewport::PresentImmersive(const MetalRHIVisionOS::PresentImmersivePa
 				// command buffer
 				FMetalCommandBuffer* FinalCommandBuffer = Context.GetCurrentCommandBuffer();;
 				
-#if ENABLE_METAL_GPUPROFILE
+#if ENABLE_METAL_GPUPROFILE && RHI_NEW_GPU_PROFILER == 0
 				FMetalProfiler* Profiler = FMetalProfiler::GetProfiler();
 				FMetalCommandBufferStats* Stats = Profiler->AllocateCommandBuffer(FinalCommandBuffer->GetMTLCmdBuffer(), 0);
 #endif
@@ -736,7 +736,9 @@ void FMetalViewport::PresentImmersive(const MetalRHIVisionOS::PresentImmersivePa
 						dispatch_semaphore_signal(FrameSemaphore);
 						dispatch_release(FrameSemaphore);
 						
+#if RHI_NEW_GPU_PROFILER == 0
 						FMetalCommandBufferTimer::RecordPresent(cmd_buf);
+#endif
 					};
 					FinalCommandBuffer->GetMTLCmdBuffer()->addCompletedHandler(CommandBufferHandler);
 				}
@@ -761,7 +763,9 @@ void FMetalViewport::PresentImmersive(const MetalRHIVisionOS::PresentImmersivePa
 		}
 	}
 	
+#if RHI_NEW_GPU_PROFILER == 0
 	FMetalCommandBufferTimer::ResetFrameBufferTimings();
+#endif
 }
 #endif //PLATFORM_VISIONOS
 
