Metal RHI

- Fix for incorrect offset in UpdateDescriptor handle
- Re-enable compute descriptor updates

diff --git a/Engine/Shaders/Private/UpdateDescriptorHandle.usf b/Engine/Shaders/Private/UpdateDescriptorHandle.usf
--- a/Engine/Shaders/Private/UpdateDescriptorHandle.usf
+++ b/Engine/Shaders/Private/UpdateDescriptorHandle.usf
@@ -4,7 +4,7 @@
 	UpdateDescriptorHandle.usf
 =============================================================================*/
 
-#include "Common.ush"
+#include "/Engine/Public/Platform.ush"
 
 #if COMPILER_METAL_SHADER_CONVERTER
 struct MetalDescriptor
diff --git a/Engine/Source/Runtime/Apple/MetalRHI/Private/MetalBindlessDescriptors.cpp b/Engine/Source/Runtime/Apple/MetalRHI/Private/MetalBindlessDescriptors.cpp
--- a/Engine/Source/Runtime/Apple/MetalRHI/Private/MetalBindlessDescriptors.cpp
+++ b/Engine/Source/Runtime/Apple/MetalRHI/Private/MetalBindlessDescriptors.cpp
@@ -28,7 +28,7 @@ static FAutoConsoleVariableRef CVarBindlessSamplerDescriptorHeapSize(
 	ECVF_ReadOnly
 );
 
-static int32_t GMetalBindlessUpdateDescriptorOnCompute = 0;
+static int32_t GMetalBindlessUpdateDescriptorOnCompute = 1;
 static FAutoConsoleVariableRef CVarBindlessUpdateDescriptorOnCompute(
 	TEXT("Metal.Bindless.UpdateDescriptorOnCompute"),
 	GMetalBindlessUpdateDescriptorOnCompute,
@@ -243,7 +243,6 @@ void FMetalDescriptorHeap::BindResource(FRHIDescriptorHandle DescriptorHandle, F
 	};
 
 	{
-		FScopeLock ScopeLock(&UpdateCS);
 		if(Context && UpdateType != EDescriptorUpdateType_Immediate)
 		{
 			check(GIsRHIInitialized);
@@ -273,8 +272,6 @@ void FMetalBindlessDescriptorManager::UpdateDescriptorsWithCopy(FMetalRHICommand
 																   TResourceArray<IRDescriptorTableEntry>&& ComputeDescriptorEntriesCopy,
 																   TResourceArray<uint32_t>&& ComputeDescriptorIndicesCopy)
 {
-	FScopeLock ScopeLock(&StandardResources.UpdateCS);
-	
 	const uint32_t NumDescriptors = ComputeDescriptorIndicesCopy.Num();
 	
 	if(!NumDescriptors)
@@ -302,8 +299,6 @@ void FMetalBindlessDescriptorManager::UpdateDescriptorsWithCompute(FMetalRHIComm
 																   TResourceArray<IRDescriptorTableEntry>&& ComputeDescriptorEntriesCopy,
 																   TResourceArray<uint32_t>&& ComputeDescriptorIndicesCopy)
 {
-	FScopeLock ScopeLock(&StandardResources.UpdateCS);
-	
 	const uint32_t NumDescriptors = ComputeDescriptorIndicesCopy.Num();
 	
 	if(!NumDescriptors)
@@ -382,7 +377,6 @@ void FMetalDescriptorHeap::BindTexture(FRHICommandListBase& RHICmdList, FRHIDesc
 	
 	RHICmdList.EnqueueLambda([this, InUpdateType=UpdateType, Data=DescriptorData, Handle=DescriptorHandle](FRHICommandListBase& RHICmdList)
 	{
-		FScopeLock ScopeLock(&UpdateCS);
 		if(InUpdateType == EDescriptorUpdateType_Immediate)
 		{
 			UpdateDescriptor(Handle, Data);
diff --git a/Engine/Source/Runtime/Apple/MetalRHI/Private/MetalBindlessDescriptors.h b/Engine/Source/Runtime/Apple/MetalRHI/Private/MetalBindlessDescriptors.h
--- a/Engine/Source/Runtime/Apple/MetalRHI/Private/MetalBindlessDescriptors.h
+++ b/Engine/Source/Runtime/Apple/MetalRHI/Private/MetalBindlessDescriptors.h
@@ -43,7 +43,6 @@ struct FMetalDescriptorHeap
 	FMetalDevice& 					Device;
 	
     FCriticalSection                FreeListCS;
-	FCriticalSection                UpdateCS;
     TQueue<uint32>                  FreeList;
 
     std::atomic<uint32>             PeakDescriptorCount;
