Metal RHI

Fix for missing deferred delete in MetalRHIRenderQuery

diff --git a/Engine/Source/Runtime/Apple/MetalRHI/Private/MetalRHIRenderQuery.cpp b/Engine/Source/Runtime/Apple/MetalRHI/Private/MetalRHIRenderQuery.cpp
--- a/Engine/Source/Runtime/Apple/MetalRHI/Private/MetalRHIRenderQuery.cpp
+++ b/Engine/Source/Runtime/Apple/MetalRHI/Private/MetalRHIRenderQuery.cpp
@@ -172,7 +172,10 @@ void FMetalQueryBufferPool::ReleaseCurrentQueryBuffer()
 {
 	if (IsValidRef(CurrentBuffer) && (CurrentBuffer->WriteOffset > 0))
 	{
-		CurrentBuffer.SafeRelease();
+		FMetalDynamicRHI::Get().DeferredDelete([InBuffer = MoveTemp(CurrentBuffer)]() mutable
+		{
+			InBuffer.SafeRelease();
+		});
 	}
 }
 
