Disallow resolve of R32F MSAA targets as this format isn't MSAA resolvable in Metal. Additionally, make the texture memoryless as we'll never be resolving to it.

diff --git a/Engine/Source/Runtime/Apple/MetalRHI/Private/MetalStateCache.cpp b/Engine/Source/Runtime/Apple/MetalRHI/Private/MetalStateCache.cpp
--- a/Engine/Source/Runtime/Apple/MetalRHI/Private/MetalStateCache.cpp
+++ b/Engine/Source/Runtime/Apple/MetalRHI/Private/MetalStateCache.cpp
@@ -617,9 +617,18 @@ bool FMetalStateCache::SetRenderPassInfo(FRHIRenderPassInfo const& InRenderTarge
 					NewColorStore[RenderTargetIndex] = GetMetalRTStoreAction(Device, ERenderTargetStoreAction::EMultisampleResolve);
 					MTL::StoreAction FinalStoreAction = !bMemoryless && GRHIDeviceId > 2 ? MTL::StoreActionUnknown : NewColorStore[RenderTargetIndex];
 					FinalStoreAction = ConditionalOverrideStoreAction(FinalStoreAction, false, !!Surface.MSAATexture);
+					
+					MTL::Texture* resolveTexture = Surface.MSAAResolveTexture ? Surface.MSAAResolveTexture.get() : Surface.Texture.get();
+#if PLATFORM_IOS
+					if (Surface.MSAATexture->pixelFormat() == MTL::PixelFormatR32Float)
+					{
+						FinalStoreAction = MTL::StoreActionDontCare;
+						resolveTexture = nullptr;
+					}
+#endif
 					ColorAttachment->setStoreAction(FinalStoreAction);
 					
-					ColorAttachment->setResolveTexture(Surface.MSAAResolveTexture ? Surface.MSAAResolveTexture.get() : Surface.Texture.get());
+					ColorAttachment->setResolveTexture(resolveTexture);
 					SampleCount = Surface.MSAATexture->sampleCount();
 				}
 				else
diff --git a/Engine/Source/Runtime/Apple/MetalRHI/Private/MetalTexture.cpp b/Engine/Source/Runtime/Apple/MetalRHI/Private/MetalTexture.cpp
--- a/Engine/Source/Runtime/Apple/MetalRHI/Private/MetalTexture.cpp
+++ b/Engine/Source/Runtime/Apple/MetalRHI/Private/MetalTexture.cpp
@@ -619,7 +619,7 @@ FMetalSurface::FMetalSurface(FMetalDevice& MetalDevice, FMetalTextureCreateDesc
         FMetalTextureCreateDesc NewCreateDesc = CreateDesc;
         
         // If we are attempting to create an MSAA texture the texture cannot be memoryless unless we are creating a depth texture
-        if(ShouldUseMemoryless(CreateDesc.Flags) && bIsMSAARequired && CreateDesc.Format != PF_DepthStencil)
+        if(ShouldUseMemoryless(CreateDesc.Flags) && bIsMSAARequired && CreateDesc.Format != PF_DepthStencil && CreateDesc.Format != PF_R32_FLOAT)
         {
             NewCreateDesc.Flags &= ~TexCreate_Memoryless;
             
