Metal RHI

- Missing viewport update on Parallel render pass setup causing rare flickering
- Bindless was missing handle when binding backbuffer textures

diff --git a/Engine/Source/Runtime/Apple/MetalRHI/Private/MetalBindlessDescriptors.cpp b/Engine/Source/Runtime/Apple/MetalRHI/Private/MetalBindlessDescriptors.cpp
--- a/Engine/Source/Runtime/Apple/MetalRHI/Private/MetalBindlessDescriptors.cpp
+++ b/Engine/Source/Runtime/Apple/MetalRHI/Private/MetalBindlessDescriptors.cpp
@@ -368,27 +368,35 @@ void FMetalBindlessDescriptorManager::UpdateDescriptorsWithCompute(FMetalRHIComm
 	}
 }
 
-void FMetalDescriptorHeap::BindTexture(FRHICommandListBase& RHICmdList, FRHIDescriptorHandle DescriptorHandle, MTL::Texture* Texture, EDescriptorUpdateType UpdateType)
+void FMetalDescriptorHeap::BindTexture(FRHICommandListBase* RHICmdList, FRHIDescriptorHandle DescriptorHandle, MTL::Texture* Texture, EDescriptorUpdateType UpdateType)
 {
 	IRDescriptorTableEntry DescriptorData = {0};
 	IRDescriptorTableSetTexture(&DescriptorData, Texture, 0.0f, 0u);
 	
 	UpdateType = !GIsRHIInitialized ? EDescriptorUpdateType_Immediate : UpdateType;
 	
-	RHICmdList.EnqueueLambda([this, InUpdateType=UpdateType, Data=DescriptorData, Handle=DescriptorHandle](FRHICommandListBase& RHICmdList)
+	if(RHICmdList)
 	{
-		if(InUpdateType == EDescriptorUpdateType_Immediate)
+		RHICmdList->EnqueueLambda([this, InUpdateType=UpdateType, Data=DescriptorData, Handle=DescriptorHandle](FRHICommandListBase& RHICmdList)
 		{
-			UpdateDescriptor(Handle, Data);
-		}
-		else
-		{	
-			FMetalRHICommandContext* Context = &FMetalRHICommandContext::Get(RHICmdList);
-			Context->EnqueueDescriptorUpdate(Data, Handle.GetIndex());
-		}
-	});
-	
-	RHICmdList.RHIThreadFence(true);
+			if(InUpdateType == EDescriptorUpdateType_Immediate)
+			{
+				UpdateDescriptor(Handle, Data);
+			}
+			else
+			{	
+				FMetalRHICommandContext* Context = &FMetalRHICommandContext::Get(RHICmdList);
+				Context->EnqueueDescriptorUpdate(Data, Handle.GetIndex());
+			}
+		});
+		
+		RHICmdList->RHIThreadFence(true);
+	}
+	else
+	{
+		check(!IsRunningRHIInSeparateThread() || IsInRHIThread() || IsInParallelRHIThread());
+		UpdateDescriptor(DescriptorHandle, DescriptorData);
+	}
 }
 
 void FMetalBindlessDescriptorManager::BindDescriptorHeapsToEncoder(FMetalCommandEncoder* Encoder, MTL::FunctionType FunctionType, EMetalShaderStages Frequency)
@@ -408,7 +416,7 @@ void FMetalBindlessDescriptorManager::BindResource(FRHIDescriptorHandle Descript
 	StandardResources.BindResource(DescriptorHandle, Resource, Context, UpdateType);
 }
 
-void FMetalBindlessDescriptorManager::BindTexture(FRHICommandListBase& RHICmdList, FRHIDescriptorHandle DescriptorHandle, MTL::Texture* Texture, EDescriptorUpdateType UpdateType)
+void FMetalBindlessDescriptorManager::BindTexture(FRHICommandListBase* RHICmdList, FRHIDescriptorHandle DescriptorHandle, MTL::Texture* Texture, EDescriptorUpdateType UpdateType)
 {
 	StandardResources.BindTexture(RHICmdList, DescriptorHandle, Texture, UpdateType);
 }
diff --git a/Engine/Source/Runtime/Apple/MetalRHI/Private/MetalBindlessDescriptors.h b/Engine/Source/Runtime/Apple/MetalRHI/Private/MetalBindlessDescriptors.h
--- a/Engine/Source/Runtime/Apple/MetalRHI/Private/MetalBindlessDescriptors.h
+++ b/Engine/Source/Runtime/Apple/MetalRHI/Private/MetalBindlessDescriptors.h
@@ -37,7 +37,7 @@ struct FMetalDescriptorHeap
 	void                    BindBufferImmediate(FRHIDescriptorHandle DescriptorHandle, FMetalBufferPtr Resource);
 	void                    BindResource(FRHIDescriptorHandle DescriptorHandle, FMetalResourceViewBase* Resource,
 												 FMetalRHICommandContext* Context, const EDescriptorUpdateType UpdateType = EDescriptorUpdateType_GPU);
-	void                    BindTexture(FRHICommandListBase& RHICmdList, FRHIDescriptorHandle DescriptorHandle, MTL::Texture* Texture, const EDescriptorUpdateType UpdateType);
+	void                    BindTexture(FRHICommandListBase* RHICmdList, FRHIDescriptorHandle DescriptorHandle, MTL::Texture* Texture, const EDescriptorUpdateType UpdateType);
 	
     //
 	FMetalDevice& 					Device;
@@ -67,7 +67,7 @@ public:
     void                    BindSampler(FRHIDescriptorHandle DescriptorHandle, MTL::SamplerState* Sampler);
     void                    BindResource(FRHIDescriptorHandle DescriptorHandle, FMetalResourceViewBase* Resource,
 										 FMetalRHICommandContext* Context, const EDescriptorUpdateType UpdateType = EDescriptorUpdateType_GPU);
-    void                    BindTexture(FRHICommandListBase& RHICmdList, FRHIDescriptorHandle DescriptorHandle, MTL::Texture* Texture, EDescriptorUpdateType UpdateType);
+    void                    BindTexture(FRHICommandListBase* RHICmdList, FRHIDescriptorHandle DescriptorHandle, MTL::Texture* Texture, EDescriptorUpdateType UpdateType);
 
     void                    BindDescriptorHeapsToEncoder(FMetalCommandEncoder* Encoder, MTL::FunctionType FunctionType, EMetalShaderStages Frequency);
 
diff --git a/Engine/Source/Runtime/Apple/MetalRHI/Private/MetalRHIContext.cpp b/Engine/Source/Runtime/Apple/MetalRHI/Private/MetalRHIContext.cpp
--- a/Engine/Source/Runtime/Apple/MetalRHI/Private/MetalRHIContext.cpp
+++ b/Engine/Source/Runtime/Apple/MetalRHI/Private/MetalRHIContext.cpp
@@ -94,6 +94,18 @@ void FMetalRHICommandContext::SetupParallelContext(const FRHIParallelRenderPassI
 	
 	bWithinRenderPass = true;
 	bIsParallelContext = true;
+	
+	// Set the viewport to the full size of render target 0.
+	if (RenderPassInfo.ColorRenderTargets[0].RenderTarget)
+	{
+		const FRHIRenderPassInfo::FColorEntry& RenderTargetView = RenderPassInfo.ColorRenderTargets[0];
+		FMetalSurface* RenderTarget = GetMetalSurfaceFromRHITexture(RenderTargetView.RenderTarget);
+
+		uint32 Width = FMath::Max((uint32)(RenderTarget->Texture->width() >> RenderTargetView.MipIndex), (uint32)1);
+		uint32 Height = FMath::Max((uint32)(RenderTarget->Texture->height() >> RenderTargetView.MipIndex), (uint32)1);
+
+		RHISetViewport(0.0f, 0.0f, 0.0f, (float)Width, (float)Height, 1.0f);
+	}
 }
 
 static uint32_t MAX_COLOR_RENDER_TARGETS_PER_DESC = 8;
diff --git a/Engine/Source/Runtime/Apple/MetalRHI/Private/MetalStateCache.cpp b/Engine/Source/Runtime/Apple/MetalRHI/Private/MetalStateCache.cpp
--- a/Engine/Source/Runtime/Apple/MetalRHI/Private/MetalStateCache.cpp
+++ b/Engine/Source/Runtime/Apple/MetalRHI/Private/MetalStateCache.cpp
@@ -274,6 +274,12 @@ void FMetalStateCache::Reset()
 			ShaderTextures[Frequency].Usage[i] = MTL::ResourceUsage(0);
 		}
 		ShaderTextures[Frequency].Bound = 0;
+		
+#if PLATFORM_SUPPORTS_BINDLESS_RENDERING
+		HeapsUsedByStage[Frequency].Reset();
+		ROResourcesByStage[Frequency].Reset();
+		RWResourcesByStage[Frequency].Reset();
+#endif
 	}
 	
 	VisibilityResults = nullptr;
diff --git a/Engine/Source/Runtime/Apple/MetalRHI/Private/MetalTexture.cpp b/Engine/Source/Runtime/Apple/MetalRHI/Private/MetalTexture.cpp
--- a/Engine/Source/Runtime/Apple/MetalRHI/Private/MetalTexture.cpp
+++ b/Engine/Source/Runtime/Apple/MetalRHI/Private/MetalTexture.cpp
@@ -894,12 +894,9 @@ FMetalSurface::~FMetalSurface()
 	ImageSurfaceRef = nullptr;
 
 #if PLATFORM_SUPPORTS_BINDLESS_RENDERING
-    FMetalBindlessDescriptorManager* BindlessDescriptorManager = Device.GetBindlessDescriptorManager();
-    check(BindlessDescriptorManager);
-
 	if(IsMetalBindlessEnabled())
 	{
-		if (!(GetDesc().Flags & TexCreate_Presentable))
+		if(BindlessHandle.IsValid())
 		{
 			FMetalDynamicRHI::Get().DeferredDelete(BindlessHandle);
 		}
@@ -1303,14 +1300,42 @@ uint32 FMetalSurface::GetNumFaces()
 
 MTLTexturePtr FMetalSurface::GetDrawableTexture()
 {
+	FScopeLock Lock(&DrawableMutex);
 	if (!Texture && EnumHasAnyFlags(GetDesc().Flags, TexCreate_Presentable))
 	{
 		check(Viewport);
 		Texture = NS::RetainPtr(Viewport->GetDrawableTexture(EMetalViewportAccessRHI));
+#if PLATFORM_SUPPORTS_BINDLESS_RENDERING
+		if(IsMetalBindlessEnabled())
+		{
+			FMetalBindlessDescriptorManager* BindlessDescriptorManager = Device.GetBindlessDescriptorManager();
+			check(BindlessDescriptorManager);
+			
+			BindlessHandle = BindlessDescriptorManager->ReserveDescriptor(ERHIDescriptorHeapType::Standard);
+
+			BindlessDescriptorManager->BindTexture(nullptr, BindlessHandle, Texture.get(), EDescriptorUpdateType_Immediate);
+		}
+#endif
 	}
 	return Texture;
 }
 
+void FMetalSurface::ReleaseDrawableTexture()
+{
+#if PLATFORM_SUPPORTS_BINDLESS_RENDERING
+	if(IsMetalBindlessEnabled())
+	{
+		if(BindlessHandle.IsValid())
+		{
+			FMetalDynamicRHI::Get().DeferredDelete(BindlessHandle);
+			BindlessHandle = {};
+		}
+	}
+#endif
+
+	Texture.reset();
+}
+
 MTLTexturePtr FMetalSurface::GetCurrentTexture()
 {
     MTLTexturePtr Tex;
diff --git a/Engine/Source/Runtime/Apple/MetalRHI/Private/MetalViewport.cpp b/Engine/Source/Runtime/Apple/MetalRHI/Private/MetalViewport.cpp
--- a/Engine/Source/Runtime/Apple/MetalRHI/Private/MetalViewport.cpp
+++ b/Engine/Source/Runtime/Apple/MetalRHI/Private/MetalViewport.cpp
@@ -435,7 +435,7 @@ void FMetalViewport::ReleaseDrawable()
 
 		if (!GMetalSupportsIntermediateBackBuffer && IsValidRef(BackBuffer[GetViewportIndex(EMetalViewportAccessRHI)]))
 		{
-			BackBuffer[GetViewportIndex(EMetalViewportAccessRHI)]->Texture.reset();
+			BackBuffer[GetViewportIndex(EMetalViewportAccessRHI)]->ReleaseDrawableTexture();
 		}
 	}
 }
diff --git a/Engine/Source/Runtime/Apple/MetalRHI/Public/MetalResources.h b/Engine/Source/Runtime/Apple/MetalRHI/Public/MetalResources.h
--- a/Engine/Source/Runtime/Apple/MetalRHI/Public/MetalResources.h
+++ b/Engine/Source/Runtime/Apple/MetalRHI/Public/MetalResources.h
@@ -217,6 +217,8 @@ public:
 	
 	/** Gets the drawable texture if this is a back-buffer surface. */
 	MTLTexturePtr GetDrawableTexture();
+	void 		  ReleaseDrawableTexture();
+	
     MTLTexturePtr GetCurrentTexture();
 
     MTLTexturePtr Reallocate(MTLTexturePtr Texture, MTL::TextureUsage UsageModifier);
@@ -259,10 +261,11 @@ public:
 	}
 	
 #if PLATFORM_SUPPORTS_BINDLESS_RENDERING
-    FRHIDescriptorHandle BindlessHandle;
-
     virtual FRHIDescriptorHandle GetDefaultBindlessHandle() const override final
     {
+#if PLATFORM_SUPPORTS_BINDLESS_RENDERING
+		check(!IsMetalBindlessEnabled() || BindlessHandle.IsValid());
+#endif
         return BindlessHandle;
     }
 #endif // PLATFORM_SUPPORTS_BINDLESS_RENDERING
@@ -275,6 +278,11 @@ private:
 	
 	// The movie playback IOSurface/CVTexture wrapper to avoid page-off
 	CFTypeRef ImageSurfaceRef;
+	
+#if PLATFORM_SUPPORTS_BINDLESS_RENDERING
+	FRHIDescriptorHandle 	BindlessHandle;
+#endif
+	FCriticalSection 		DrawableMutex;
 };
 
 class FMetalBufferData
