Fixed compile error in MetalTexture.cpp

diff --git a/Engine/Source/Runtime/Apple/MetalRHI/Private/MetalTexture.cpp b/Engine/Source/Runtime/Apple/MetalRHI/Private/MetalTexture.cpp
--- a/Engine/Source/Runtime/Apple/MetalRHI/Private/MetalTexture.cpp
+++ b/Engine/Source/Runtime/Apple/MetalRHI/Private/MetalTexture.cpp
@@ -820,7 +820,7 @@ void FMetalSurface::Initialize(FRHICommandListBase& RHICmdList)
 		FMetalBindlessDescriptorManager* BindlessDescriptorManager = Device.GetBindlessDescriptorManager();
 		check(BindlessDescriptorManager);
 
-		BindlessDescriptorManager->BindTexture(RHICmdList, BindlessHandle, Texture.get(), EDescriptorUpdateType_Immediate);
+		BindlessDescriptorManager->BindTexture(&RHICmdList, BindlessHandle, Texture.get(), EDescriptorUpdateType_Immediate);
 	}
 #endif
 }
