Metal RHI

- Fixing Metal GPU timer not showing results when incorrect samples returned
- Moved FrameBoundary event

diff --git a/Engine/Source/Runtime/Apple/MetalRHI/Private/MetalCounterSampler.cpp b/Engine/Source/Runtime/Apple/MetalRHI/Private/MetalCounterSampler.cpp
--- a/Engine/Source/Runtime/Apple/MetalRHI/Private/MetalCounterSampler.cpp
+++ b/Engine/Source/Runtime/Apple/MetalRHI/Private/MetalCounterSampler.cpp
@@ -17,7 +17,7 @@ void FMetalCounterSample::ResolveStageCounters(uint64_t& InStartTime, uint64_t&
 		
 		if(SampleType == EMetalCounterSampleType::RenderStage)
 		{
-			NS::Range Range = NS::Range::Make(Offset, 4);
+			NS::Range Range = NS::Range::Make(Offset, 2);
 			Data = SampleBuffer->resolveCounterRange(Range);
 		}
 		else
@@ -38,9 +38,7 @@ void FMetalCounterSample::ResolveStageCounters(uint64_t& InStartTime, uint64_t&
 		if(SampleType == EMetalCounterSampleType::RenderStage)
 		{
 			uint64_t StartTimeVertex = Timestamps[0].timestamp;
-			uint64_t EndTimeVertex = Timestamps[1].timestamp;
-			uint64_t StartTimeFragment = Timestamps[2].timestamp;
-			uint64_t EndTimeFragment = Timestamps[3].timestamp;
+			uint64_t EndTimeFragment = Timestamps[1].timestamp;
 			
 			StartTime = StartTimeVertex;
 			EndTime = EndTimeFragment;
@@ -51,6 +49,11 @@ void FMetalCounterSample::ResolveStageCounters(uint64_t& InStartTime, uint64_t&
 			EndTime = Timestamps[1].timestamp;
 		}
 		
+		if(StartTime == 0 || EndTime == 0)
+		{
+			StartTime = EndTime = FMath::Max(StartTime, EndTime);
+		}
+		
 		double OneOverNanosToCycles = 1.0 / 1000000000.0 / FPlatformTime::GetSecondsPerCycle64();
 		StartTime = FMath::TruncToInt64(StartTime * OneOverNanosToCycles);
 		EndTime = FMath::TruncToInt64(EndTime * OneOverNanosToCycles);
@@ -180,7 +183,7 @@ FMetalCounterSamplePtr FMetalCounterSampler::SetupStageCounters(MTL::BlitPassDes
 FMetalCounterSamplePtr FMetalCounterSampler::SetupStageCounters(MTL::RenderPassDescriptor* InRenderPassDesc)
 {
 	uint32_t OutOffset;
-	MTL::CounterSampleBuffer* Buffer = SwapOrAllocateBuffer(4, OutOffset);
+	MTL::CounterSampleBuffer* Buffer = SwapOrAllocateBuffer(2, OutOffset);
 	check(Buffer);
 	
 	FMetalCounterSamplePtr Sample = FMetalCounterSamplePtr(new FMetalCounterSample(EMetalCounterSampleType::RenderStage, Buffer, OutOffset));
@@ -188,9 +191,7 @@ FMetalCounterSamplePtr FMetalCounterSampler::SetupStageCounters(MTL::RenderPassD
 	
 	SampleDesc->setSampleBuffer(Buffer);
 	SampleDesc->setStartOfVertexSampleIndex(OutOffset);
-	SampleDesc->setEndOfVertexSampleIndex(OutOffset+1);
-	SampleDesc->setStartOfFragmentSampleIndex(OutOffset+2);
-	SampleDesc->setEndOfFragmentSampleIndex(OutOffset+3);
+	SampleDesc->setEndOfFragmentSampleIndex(OutOffset+1);
 	
 	return Sample;
 }
diff --git a/Engine/Source/Runtime/Apple/MetalRHI/Private/MetalRHI.cpp b/Engine/Source/Runtime/Apple/MetalRHI/Private/MetalRHI.cpp
--- a/Engine/Source/Runtime/Apple/MetalRHI/Private/MetalRHI.cpp
+++ b/Engine/Source/Runtime/Apple/MetalRHI/Private/MetalRHI.cpp
@@ -1315,7 +1315,7 @@ void FMetalDynamicRHI::RHIEndFrame(const FRHIEndFrameArgs& Args)
 		
 		ERHIPipeline Pipeline = ERHIPipeline::Graphics;
 		
-		Payload.EventStream.Emplace<UE::RHI::GPUProfiler::FEvent::FFrameBoundary>(
+		Payload.EndFrameEvent = UE::RHI::GPUProfiler::FEvent::FFrameBoundary(
 			0, Args.FrameNumber
 		#if WITH_RHI_BREADCRUMBS
 			, (Pipeline != ERHIPipeline::None) ? Args.GPUBreadcrumbs[Pipeline] : nullptr
diff --git a/Engine/Source/Runtime/Apple/MetalRHI/Private/MetalRHIContext.cpp b/Engine/Source/Runtime/Apple/MetalRHI/Private/MetalRHIContext.cpp
--- a/Engine/Source/Runtime/Apple/MetalRHI/Private/MetalRHIContext.cpp
+++ b/Engine/Source/Runtime/Apple/MetalRHI/Private/MetalRHIContext.cpp
@@ -645,7 +645,7 @@ void FMetalRHICommandContext::StartCommandBuffer()
 	
 #if RHI_NEW_GPU_PROFILER
 	auto& Event = CurrentEncoder.GetCommandBuffer()->EmplaceProfilerEvent<UE::RHI::GPUProfiler::FEvent::FBeginWork>(0);
-	Event.GPUTimestampTOP = UINT64_MAX;
+	Event.GPUTimestampTOP = 0;
 	CurrentEncoder.GetCommandBuffer()->SetBeginWorkTimestamp(&Event.GPUTimestampTOP);
 #endif
 	
diff --git a/Engine/Source/Runtime/Apple/MetalRHI/Private/MetalSubmission.cpp b/Engine/Source/Runtime/Apple/MetalRHI/Private/MetalSubmission.cpp
--- a/Engine/Source/Runtime/Apple/MetalRHI/Private/MetalSubmission.cpp
+++ b/Engine/Source/Runtime/Apple/MetalRHI/Private/MetalSubmission.cpp
@@ -567,6 +567,14 @@ void FMetalDynamicRHI::FlushBatchedPayloads(FMetalCommandQueue::FPayloadArray& P
 			Queue.CommitCommandBuffer(CommandBuffer);
 			Queue.GetSignalEvent().LastSignaledValue.store(Payload->CompletionFenceValue, std::memory_order_release);
 		}
+		
+#if RHI_NEW_GPU_PROFILER
+		if (Payload->EndFrameEvent.IsSet())
+		{
+			Payload->EndFrameEvent->CPUTimestamp = FPlatformTime::Cycles64();
+			Payload->EventStream.Emplace<UE::RHI::GPUProfiler::FEvent::FFrameBoundary>(*Payload->EndFrameEvent);
+		}
+#endif
 
 		// Submission of this payload is completed. Signal the submission event if one was provided.
 		if (Payload->SubmissionEvent)
diff --git a/Engine/Source/Runtime/Apple/MetalRHI/Private/MetalSubmission.h b/Engine/Source/Runtime/Apple/MetalRHI/Private/MetalSubmission.h
--- a/Engine/Source/Runtime/Apple/MetalRHI/Private/MetalSubmission.h
+++ b/Engine/Source/Runtime/Apple/MetalRHI/Private/MetalSubmission.h
@@ -208,6 +208,7 @@ struct FMetalPayload : public FMetalPayloadBaseFixLayout
 #if RHI_NEW_GPU_PROFILER
 	TOptional<FMetalTiming*> Timing;
 	UE::RHI::GPUProfiler::FEventStream EventStream;
+	TOptional<UE::RHI::GPUProfiler::FEvent::FFrameBoundary> EndFrameEvent;
 #endif
 
 	virtual ~FMetalPayload();
