Fixing Metal Texture Initializer handling of non-Default Bulk Data.

diff --git a/Engine/Source/Runtime/Apple/MetalRHI/Private/MetalTexture.cpp b/Engine/Source/Runtime/Apple/MetalRHI/Private/MetalTexture.cpp
--- a/Engine/Source/Runtime/Apple/MetalRHI/Private/MetalTexture.cpp
+++ b/Engine/Source/Runtime/Apple/MetalRHI/Private/MetalTexture.cpp
@@ -1534,9 +1534,13 @@ FRHITextureInitializer FMetalDynamicRHI::RHICreateTextureInitializer(FRHICommand
 	LLM_SCOPE_DYNAMIC_STAT_OBJECTPATH_FNAME(CreateDesc.GetTraceClassName(), ELLMTagSet::AssetClasses);
 	UE_TRACE_METADATA_SCOPE_ASSET_FNAME(CreateDesc.DebugName, CreateDesc.GetTraceClassName(), CreateDesc.OwnerName);
 
+	// If the FMetalSurface constructor handled the bulk data, just return the default initializer.
+	const bool bForceDefaultInitializer =
+		(CreateDesc.InitAction == ERHITextureInitAction::BulkData && CreateDesc.BulkData && CreateDesc.BulkData->GetResourceType() != FResourceBulkDataInterface::EBulkDataType::Default);
+
 	FMetalSurface* Texture = new FMetalSurface(*Device, FMetalTextureCreateDesc(*Device, CreateDesc));
 
-	if (CreateDesc.InitAction == ERHITextureInitAction::Default)
+	if (CreateDesc.InitAction == ERHITextureInitAction::Default || bForceDefaultInitializer)
 	{
 		return UE::RHICore::FBaseTextureInitializerImplementation(RHICmdList, Texture,
 			[Texture = TRefCountPtr<FMetalSurface>(Texture)](FRHICommandListBase& RHICmdList) mutable
