Metal RHI

Added auto release pool to ensure NSData is cleaned up when calling resolveCounterSamples
Fixed a bug where the GPU returns an incorrect time stamp

diff --git a/Engine/Source/Runtime/Apple/MetalRHI/Private/MetalCounterSampler.cpp b/Engine/Source/Runtime/Apple/MetalRHI/Private/MetalCounterSampler.cpp
--- a/Engine/Source/Runtime/Apple/MetalRHI/Private/MetalCounterSampler.cpp
+++ b/Engine/Source/Runtime/Apple/MetalRHI/Private/MetalCounterSampler.cpp
@@ -11,6 +11,8 @@ static TAutoConsoleVariable<bool> CVarMetalRHIInsertCounterSampleBarrier(
 
 void FMetalCounterSample::ResolveStageCounters(uint64_t& InStartTime, uint64_t& InEndTime)
 {
+	MTL_SCOPED_AUTORELEASE_POOL;
+	
 	if(!bResolved)
 	{
 		NS::Data* Data;
@@ -53,6 +55,11 @@ void FMetalCounterSample::ResolveStageCounters(uint64_t& InStartTime, uint64_t&
 		{
 			StartTime = EndTime = FMath::Max(StartTime, EndTime);
 		}
+		else if (StartTime > EndTime)
+		{
+			// Looks like a driver bug that randomly causes StartVertex to be a huge value even though EndFragment is valid
+			StartTime = EndTime;
+		}
 		
 		double OneOverNanosToCycles = 1.0 / 1000000000.0 / FPlatformTime::GetSecondsPerCycle64();
 		StartTime = FMath::TruncToInt64(StartTime * OneOverNanosToCycles);
@@ -67,6 +74,8 @@ void FMetalCounterSample::ResolveStageCounters(uint64_t& InStartTime, uint64_t&
 
 void FMetalCounterSample::ResolveBoundaryCounter(uint64_t& Time)
 {
+	MTL_SCOPED_AUTORELEASE_POOL;
+	
 	if(!bResolved)
 	{
 		NS::Range Range = NS::Range::Make(Offset, 1);
