Fix Cocoa memory leaks

A set of NS* allocations were going into the global autoreleasepool, which lives for the whole application lifetime, which is basically a leak.
Fixed by introducing local pools to deallocate unused things asap.

#rnx

diff --git a/Engine/Source/Runtime/Apple/MetalRHI/Private/MetalBuffer.cpp b/Engine/Source/Runtime/Apple/MetalRHI/Private/MetalBuffer.cpp
--- a/Engine/Source/Runtime/Apple/MetalRHI/Private/MetalBuffer.cpp
+++ b/Engine/Source/Runtime/Apple/MetalRHI/Private/MetalBuffer.cpp
@@ -937,6 +937,7 @@ FMetalRingBufferRef::FMetalRingBufferRef(FMetalDevice& InDevice, FMetalBufferPtr
 , Buffer(Buf)
 , LastRead(Buffer->GetLength())
 {
+	MTL_SCOPED_AUTORELEASE_POOL;
     Buffer->GetMTLBuffer()->setLabel(NS::String::string("Ring Buffer", NS::UTF8StringEncoding));
 }
 
diff --git a/Engine/Source/Runtime/Apple/MetalRHI/Private/MetalCommandEncoder.cpp b/Engine/Source/Runtime/Apple/MetalRHI/Private/MetalCommandEncoder.cpp
--- a/Engine/Source/Runtime/Apple/MetalRHI/Private/MetalCommandEncoder.cpp
+++ b/Engine/Source/Runtime/Apple/MetalRHI/Private/MetalCommandEncoder.cpp
@@ -262,6 +262,8 @@ void FMetalCommandEncoder::StartCommandBuffer(void)
 	
 void FMetalCommandEncoder::EndCommandBuffer(FMetalRHICommandContext* Context)
 {
+	MTL_SCOPED_AUTORELEASE_POOL;
+	
 	check(CommandBuffer);
 	check(IsRenderCommandEncoderActive() == false
 		  && IsComputeCommandEncoderActive() == false
@@ -513,6 +515,8 @@ MTLParallelRenderCommandEncoderPtr FMetalCommandEncoder::BeginParallelRenderComm
 
 void FMetalCommandEncoder::BeginComputeCommandEncoding(MTL::DispatchType DispatchType, FMetalCounterSampler* Sampler)
 {
+	MTL_SCOPED_AUTORELEASE_POOL;
+	
 	check(CommandBuffer);
 
 	static bool bSupportsFences = Device.SupportsFeature(EMetalFeaturesFences);
@@ -557,6 +561,8 @@ void FMetalCommandEncoder::BeginComputeCommandEncoding(MTL::DispatchType Dispatc
 
 void FMetalCommandEncoder::BeginBlitCommandEncoding(FMetalCounterSampler* Sampler)
 {
+	MTL_SCOPED_AUTORELEASE_POOL;
+	
 	check(CommandBuffer);
 	check(IsRenderCommandEncoderActive() == false && IsComputeCommandEncoderActive() == false && IsBlitCommandEncoderActive() == false
 #if METAL_RHI_RAYTRACING
diff --git a/Engine/Source/Runtime/Apple/MetalRHI/Private/MetalRHIContext.cpp b/Engine/Source/Runtime/Apple/MetalRHI/Private/MetalRHIContext.cpp
--- a/Engine/Source/Runtime/Apple/MetalRHI/Private/MetalRHIContext.cpp
+++ b/Engine/Source/Runtime/Apple/MetalRHI/Private/MetalRHIContext.cpp
@@ -112,6 +112,8 @@ static uint32_t MAX_COLOR_RENDER_TARGETS_PER_DESC = 8;
 
 void FMetalRHICommandContext::BeginComputeEncoder()
 {
+	MTL_SCOPED_AUTORELEASE_POOL;
+
 	SCOPE_CYCLE_COUNTER(STAT_MetalSwitchToComputeTime);
 	
 	check(!bWithinRenderPass);
@@ -159,6 +161,8 @@ void FMetalRHICommandContext::EndComputeEncoder()
 
 void FMetalRHICommandContext::BeginBlitEncoder()
 {
+	MTL_SCOPED_AUTORELEASE_POOL;
+	
 	SCOPE_CYCLE_COUNTER(STAT_MetalSwitchToBlitTime);
 	check(!bWithinRenderPass);
 	
diff --git a/Engine/Source/Runtime/Apple/MetalRHI/Private/Shaders/Types/Templates/MetalBaseShader.h b/Engine/Source/Runtime/Apple/MetalRHI/Private/Shaders/Types/Templates/MetalBaseShader.h
--- a/Engine/Source/Runtime/Apple/MetalRHI/Private/Shaders/Types/Templates/MetalBaseShader.h
+++ b/Engine/Source/Runtime/Apple/MetalRHI/Private/Shaders/Types/Templates/MetalBaseShader.h
@@ -276,6 +276,8 @@ void TMetalBaseShader<BaseResourceType, ShaderType>::Init(TArrayView<const uint8
 		}
 		else
 		{
+			MTL_SCOPED_AUTORELEASE_POOL;
+			
 			METAL_GPUPROFILE(FScopedMetalCPUStats CPUStat(FString::Printf(TEXT("NewLibrarySource: %d_%d"), SourceLen, SourceCRC)));
 			NS::String* ShaderString = ((OfflineCompiledFlag == 0) ? NS::String::string(SourceCode, NS::UTF8StringEncoding) : GlslCodeNSString);
 
@@ -406,6 +408,8 @@ void TMetalBaseShader<BaseResourceType, ShaderType>::Destroy()
 template<typename BaseResourceType, int32 ShaderType>
 inline NS::String* TMetalBaseShader<BaseResourceType, ShaderType>::GetSourceCode()
 {
+	MTL_SCOPED_AUTORELEASE_POOL;
+	
 	if (!GlslCodeNSString && CodeSize && CompressedSource.Num())
 	{
 		GlslCodeNSString = DecodeMetalSourceCode(CodeSize, CompressedSource);
@@ -422,6 +426,8 @@ inline NS::String* TMetalBaseShader<BaseResourceType, ShaderType>::GetSourceCode
 template<typename BaseResourceType, int32 ShaderType>
 MTLFunctionPtr TMetalBaseShader<BaseResourceType, ShaderType>::GetCompiledFunction(bool const bAsync, const int32 FunctionIndex)
 {
+	MTL_SCOPED_AUTORELEASE_POOL;
+	
 	MTLFunctionPtr Func = Function;
 
 	bool bNeedToRecreateFunction = (LibraryFunctionIndex != FunctionIndex);
diff --git a/Engine/Source/Runtime/Core/Private/IOS/IOSPlatformFile.cpp b/Engine/Source/Runtime/Core/Private/IOS/IOSPlatformFile.cpp
--- a/Engine/Source/Runtime/Core/Private/IOS/IOSPlatformFile.cpp
+++ b/Engine/Source/Runtime/Core/Private/IOS/IOSPlatformFile.cpp
@@ -1,6 +1,7 @@
 // Copyright Epic Games, Inc. All Rights Reserved.
 
 #include "IOS/IOSPlatformFile.h"
+#include "Apple/ApplePlatformMisc.h"
 #include "HAL/PlatformTLS.h"
 #include "Containers/StringConv.h"
 #include "HAL/PlatformTime.h"
@@ -832,6 +833,8 @@ bool FIOSPlatformFile::DeleteDirectory(const TCHAR* Directory)
 
 bool FIOSPlatformFile::IterateDirectory(const TCHAR* Directory, FDirectoryVisitor& Visitor)
 {
+	SCOPED_AUTORELEASE_POOL;
+
 	const FString DirectoryStr = Directory;
 	return IterateDirectoryCommon(Directory, [&](struct dirent* InEntry) -> bool
 	{
@@ -845,6 +848,8 @@ bool FIOSPlatformFile::IterateDirectory(const TCHAR* Directory, FDirectoryVisito
 
 bool FIOSPlatformFile::IterateDirectoryStat(const TCHAR* Directory, FDirectoryStatVisitor& Visitor)
 {
+	SCOPED_AUTORELEASE_POOL;
+
 	const FString DirectoryStr = Directory;
 	const FString NormalizedDirectoryStr = NormalizeFilename(Directory);
 
diff --git a/Engine/Source/Runtime/Core/Private/IOS/IOSPlatformMisc.cpp b/Engine/Source/Runtime/Core/Private/IOS/IOSPlatformMisc.cpp
--- a/Engine/Source/Runtime/Core/Private/IOS/IOSPlatformMisc.cpp
+++ b/Engine/Source/Runtime/Core/Private/IOS/IOSPlatformMisc.cpp
@@ -5,6 +5,7 @@
 =============================================================================*/
 
 #include "IOS/IOSPlatformMisc.h"
+#include "Apple/ApplePlatformMisc.h"
 #include "Async/TaskGraphInterfaces.h"
 #include "Apple/ApplePlatformCrashContext.h"
 #include "HAL/ExceptionHandling.h"
@@ -1015,6 +1016,8 @@ FString FIOSPlatformMisc::GetOSVersion()
 
 bool FIOSPlatformMisc::GetDiskTotalAndFreeSpace(const FString& InPath, uint64& TotalNumberOfBytes, uint64& NumberOfFreeBytes, bool IncludeFreeable)
 {
+    SCOPED_AUTORELEASE_POOL;
+
     //On iOS 11 use new method to return disk space available for important usages
 #if !PLATFORM_TVOS
     bool GetValueSuccess = false;
