Metal RHI

Disable Metal Counter Sampling in Test/Shipping

diff --git a/Engine/Source/Runtime/Apple/MetalRHI/Private/MetalDevice.cpp b/Engine/Source/Runtime/Apple/MetalRHI/Private/MetalDevice.cpp
--- a/Engine/Source/Runtime/Apple/MetalRHI/Private/MetalDevice.cpp
+++ b/Engine/Source/Runtime/Apple/MetalRHI/Private/MetalDevice.cpp
@@ -623,6 +623,7 @@ void FMetalDevice::EnumerateFeatureSupport()
 	}
 #endif
 	
+#if WITH_PROFILEGPU
 	// Counter Sampling Features
 	if(Device->supportsCounterSampling(MTL::CounterSamplingPointAtStageBoundary))
 	{
@@ -635,6 +636,7 @@ void FMetalDevice::EnumerateFeatureSupport()
 	{
 		Features |= EMetalFeaturesBoundaryCounterSampling;
 	}
+#endif
 }
 
 void FMetalDevice::EndDrawingViewport(bool bPresent)
