Adding an empty IMetalDynamicRHI.

diff --git a/Engine/Source/Runtime/Apple/MetalRHI/Private/MetalDynamicRHI.h b/Engine/Source/Runtime/Apple/MetalRHI/Private/MetalDynamicRHI.h
--- a/Engine/Source/Runtime/Apple/MetalRHI/Private/MetalDynamicRHI.h
+++ b/Engine/Source/Runtime/Apple/MetalRHI/Private/MetalDynamicRHI.h
@@ -7,7 +7,7 @@
 #pragma once 
 
 #include "CoreMinimal.h"
-#include "RHI.h"
+#include "IMetalDynamicRHI.h"
 
 // Metal RHI public headers.
 #include "MetalThirdParty.h"
@@ -62,7 +62,7 @@ struct FMetalFinalizedCommands final : public IRHIPlatformCommandList, public TA
 {};
 
 /** The interface which is implemented by the dynamically bound RHI. */
-class FMetalDynamicRHI : public FDynamicRHI
+class FMetalDynamicRHI : public IMetalDynamicRHI
 {
 	static inline FMetalDynamicRHI* Singleton = nullptr;
 public:
@@ -95,7 +95,6 @@ public:
 	virtual void Init();
 	virtual void Shutdown() {}
 	virtual const TCHAR* GetName() override { return TEXT("Metal"); }
-	virtual ERHIInterfaceType GetInterfaceType() const override { return ERHIInterfaceType::Metal; }
 
 	virtual void RHIEndFrame_RenderThread(FRHICommandListImmediate& RHICmdList) final override;
 	virtual void RHIEndFrame(const FRHIEndFrameArgs& Args) final override;
diff --git a/Engine/Source/Runtime/Apple/MetalRHI/Public/IMetalDynamicRHI.h b/Engine/Source/Runtime/Apple/MetalRHI/Public/IMetalDynamicRHI.h
new file mode 100644
--- /dev/null
+++ b/Engine/Source/Runtime/Apple/MetalRHI/Public/IMetalDynamicRHI.h
@@ -0,0 +1,25 @@
+// Copyright Epic Games, Inc. All Rights Reserved.
+
+#pragma once 
+
+#include "RHI.h"
+
+struct IMetalDynamicRHI : public FDynamicRHI
+{
+	virtual ERHIInterfaceType GetInterfaceType() const override final
+	{
+		return ERHIInterfaceType::Metal;
+	}
+};
+
+inline bool IsRHIMetal()
+{
+	return GDynamicRHI != nullptr && GDynamicRHI->GetInterfaceType() == ERHIInterfaceType::Metal;
+}
+
+inline IMetalDynamicRHI* GetIMetalDynamicRHI()
+{
+	check(GDynamicRHI->GetInterfaceType() == ERHIInterfaceType::Metal);
+	return GetDynamicRHI<IMetalDynamicRHI>();
+}
+
