Replace EBulkDataType in MetalRHI with Metal-specific RHI functions

diff --git a/Engine/Plugins/Experimental/AVCodecs/AVCodecsCore/Source/AVCodecsCore/Public/Video/Resources/Metal/VideoResourceMetal.h b/Engine/Plugins/Experimental/AVCodecs/AVCodecsCore/Source/AVCodecsCore/Public/Video/Resources/Metal/VideoResourceMetal.h
--- a/Engine/Plugins/Experimental/AVCodecs/AVCodecsCore/Source/AVCodecsCore/Public/Video/Resources/Metal/VideoResourceMetal.h
+++ b/Engine/Plugins/Experimental/AVCodecs/AVCodecsCore/Source/AVCodecsCore/Public/Video/Resources/Metal/VideoResourceMetal.h
@@ -46,43 +46,6 @@ public:
 	virtual FAVResult Validate() const override;
 };
 
-/**
- * Passes a CVPixelBufferRef through to the RHI to wrap in an RHI texture without traversing system memory.
- */
-class AVCODECSCORE_API FBulkDataMetal : public FResourceBulkDataInterface
-{
-public:
-    FBulkDataMetal(CFTypeRef InImageBuffer)
-        : ImageBuffer(InImageBuffer)
-    {
-        check(ImageBuffer);
-        CFRetain(ImageBuffer);
-    }
-    virtual ~FBulkDataMetal()
-    {
-        CFRelease(ImageBuffer);
-        ImageBuffer = nullptr;
-    }
-public:
-    virtual void Discard() override
-    {
-        delete this;
-    }
-    virtual const void* GetResourceBulkData() const override
-    {
-        return ImageBuffer;
-    }
-    virtual uint32 GetResourceBulkDataSize() const override
-    {
-        return ImageBuffer ? ~0u : 0;
-    }
-    virtual EBulkDataType GetResourceType() const override
-    {
-        return EBulkDataType::MediaTexture;
-    }
-    CFTypeRef ImageBuffer;
-};
-
 DECLARE_TYPEID(FVideoContextMetal, AVCODECSCORE_API);
 DECLARE_TYPEID(FVideoResourceMetal, AVCODECSCORE_API);
 
diff --git a/Engine/Plugins/Media/AvfMedia/Source/AvfMedia/AvfMedia.Build.cs b/Engine/Plugins/Media/AvfMedia/Source/AvfMedia/AvfMedia.Build.cs
--- a/Engine/Plugins/Media/AvfMedia/Source/AvfMedia/AvfMedia.Build.cs
+++ b/Engine/Plugins/Media/AvfMedia/Source/AvfMedia/AvfMedia.Build.cs
@@ -19,9 +19,12 @@ namespace UnrealBuildTool.Rules
 					"ApplicationCore",
 					"MediaUtils",
 					"RenderCore",
-					"RHI"
+					"RHI",
+					"MetalRHI"
 				});
 			
+			AddEngineThirdPartyPrivateStaticDependencies(Target, "MetalCPP");
+			        
 			PrivateIncludePathModuleNames.AddRange(
 				new string[] {
 					"Media",
diff --git a/Engine/Plugins/Media/AvfMedia/Source/AvfMedia/Private/Player/AvfMediaTracks.cpp b/Engine/Plugins/Media/AvfMedia/Source/AvfMedia/Private/Player/AvfMediaTracks.cpp
--- a/Engine/Plugins/Media/AvfMedia/Source/AvfMedia/Private/Player/AvfMediaTracks.cpp
+++ b/Engine/Plugins/Media/AvfMedia/Source/AvfMedia/Private/Player/AvfMediaTracks.cpp
@@ -444,8 +444,6 @@ void FAvfMediaTracks::Initialize(AVPlayerItem* InPlayerItem, FString& OutInfo)
 		}
 		else if ([MediaType isEqualToString:AVMediaTypeVideo])
 		{
-			check(COREVIDEO_SUPPORTS_METAL);
-			
 			NSMutableDictionary* OutputSettings = [NSMutableDictionary dictionary];
 			CMFormatDescriptionRef DescRef = (CMFormatDescriptionRef)[AssetTrack.formatDescriptions objectAtIndex:0];
 			CMVideoCodecType CodecType = CMFormatDescriptionGetMediaSubType(DescRef);
diff --git a/Engine/Plugins/Media/AvfMedia/Source/AvfMedia/Private/Player/AvfMediaVideoSampler.cpp b/Engine/Plugins/Media/AvfMedia/Source/AvfMedia/Private/Player/AvfMediaVideoSampler.cpp
--- a/Engine/Plugins/Media/AvfMedia/Source/AvfMedia/Private/Player/AvfMediaVideoSampler.cpp
+++ b/Engine/Plugins/Media/AvfMedia/Source/AvfMedia/Private/Player/AvfMediaVideoSampler.cpp
@@ -17,6 +17,7 @@
 #include "PipelineStateCache.h"
 #include "RHIStaticStates.h"
 #include "UObject/Class.h"
+#include "IMetalDynamicRHI.h"
 #endif
 
 #include "AvfMediaTextureSample.h"
@@ -188,7 +189,6 @@ void FAvfMediaVideoSampler::ProcessFrame(CVPixelBufferRef Frame, FTimespan Sampl
 	TRefCountPtr<FRHITexture> ShaderResource;
 
 	// We have to support Metal for this object now
-	check(COREVIDEO_SUPPORTS_METAL);
 	check(IsMetalPlatform(GMaxRHIShaderPlatform));
 	{
 		if (!MetalTextureCache)
@@ -222,20 +222,11 @@ void FAvfMediaVideoSampler::ProcessFrame(CVPixelBufferRef Frame, FTimespan Sampl
 			check(Result == kCVReturnSuccess);
 			check(UVTextureRef);
 			
-			// Metal can upload directly from an IOSurface to a 2D texture, so we can just wrap it.
-
-			const FRHITextureCreateDesc YDesc =
-				FRHITextureCreateDesc::Create2D(TEXT("YTex"), YWidth, YHeight, PF_G8)
-				.SetFlags(ETextureCreateFlags::NoTiling | ETextureCreateFlags::ShaderResource)
-				.SetInitActionBulkData(new FAvfTexture2DResourceWrapper(YTextureRef));
-
-			const FRHITextureCreateDesc UVDesc =
-				FRHITextureCreateDesc::Create2D(TEXT("UVTex"), UVWidth, UVHeight, PF_R8G8)
-				.SetFlags(ETextureCreateFlags::NoTiling | ETextureCreateFlags::ShaderResource)
-				.SetInitActionBulkData(new FAvfTexture2DResourceWrapper(UVTextureRef));
-			
-			TRefCountPtr<FRHITexture> YTex = RHICmdList.CreateTexture(YDesc);
-			TRefCountPtr<FRHITexture> UVTex = RHICmdList.CreateTexture(UVDesc);
+			// Metal can upload directly from an IOSurface to a 2D texture
+			IMetalDynamicRHI* MetalDynamicRHI = GetIMetalDynamicRHI();
+						
+			TRefCountPtr<FRHITexture> YTex = MetalDynamicRHI->RHICreateTexture2DFromCVMetalTexture(PF_G8, ETextureCreateFlags::NoTiling | ETextureCreateFlags::ShaderResource, FClearValueBinding::White, YTextureRef);
+			TRefCountPtr<FRHITexture> UVTex = MetalDynamicRHI->RHICreateTexture2DFromCVMetalTexture(PF_R8G8, ETextureCreateFlags::NoTiling | ETextureCreateFlags::ShaderResource, FClearValueBinding::White, UVTextureRef);
 
 			const FRHITextureCreateDesc Desc =
 				FRHITextureCreateDesc::Create2D(TEXT("Info"), YWidth, YHeight, PF_B8G8R8A8)
@@ -292,12 +283,9 @@ void FAvfMediaVideoSampler::ProcessFrame(CVPixelBufferRef Frame, FTimespan Sampl
 			check(Result == kCVReturnSuccess);
 			check(TextureRef);
 
-			const FRHITextureCreateDesc Desc =
-				FRHITextureCreateDesc::Create2D(TEXT("FAvfMediaVideoSampler"), Width, Height, PF_B8G8R8A8)
-				.SetFlags(ETextureCreateFlags::SRGB | ETextureCreateFlags::NoTiling | ETextureCreateFlags::ShaderResource)
-				.SetInitActionBulkData(new FAvfTexture2DResourceWrapper(TextureRef));
-
-			ShaderResource = RHICmdList.CreateTexture(Desc);
+			IMetalDynamicRHI* MetalDynamicRHI = GetIMetalDynamicRHI();
+						
+			ShaderResource = MetalDynamicRHI->RHICreateTexture2DFromCVMetalTexture(PF_B8G8R8A8, ETextureCreateFlags::NoTiling | ETextureCreateFlags::ShaderResource, FClearValueBinding::White, TextureRef);
 			
 			CFRelease(TextureRef);
 		}
diff --git a/Engine/Plugins/Media/AvfMedia/Source/AvfMedia/Public/AvfMediaTextureSample.h b/Engine/Plugins/Media/AvfMedia/Source/AvfMedia/Public/AvfMediaTextureSample.h
--- a/Engine/Plugins/Media/AvfMedia/Source/AvfMedia/Public/AvfMediaTextureSample.h
+++ b/Engine/Plugins/Media/AvfMedia/Source/AvfMedia/Public/AvfMediaTextureSample.h
@@ -200,51 +200,3 @@ private:
 
 /** Implements a pool for AVF texture sample objects. */
 class FAvfMediaTextureSamplePool : public TMediaObjectPool<FAvfMediaTextureSample> { };
-
-/**
- * Passes a CV*TextureRef or CVPixelBufferRef through to the RHI to wrap in an RHI texture without traversing system memory.
- */
-class FAvfTexture2DResourceWrapper
-	: public FResourceBulkDataInterface
-{
-public:
-
-	FAvfTexture2DResourceWrapper(CFTypeRef InImageBuffer)
-		: ImageBuffer(InImageBuffer)
-	{
-		check(ImageBuffer);
-		CFRetain(ImageBuffer);
-	}
-
-	virtual ~FAvfTexture2DResourceWrapper()
-	{
-		CFRelease(ImageBuffer);
-		ImageBuffer = nullptr;
-	}
-
-public:
-
-	//~ FResourceBulkDataInterface interface
-
-	virtual void Discard() override
-	{
-		delete this;
-	}
-
-	virtual const void* GetResourceBulkData() const override
-	{
-		return ImageBuffer;
-	}
-	
-	virtual uint32 GetResourceBulkDataSize() const override
-	{
-		return ImageBuffer ? ~0u : 0;
-	}
-
-	virtual EBulkDataType GetResourceType() const override
-	{
-		return EBulkDataType::MediaTexture;
-	}
-			
-	CFTypeRef ImageBuffer;
-};
diff --git a/Engine/Plugins/Media/AvfMedia/Source/AvfMediaCapture/AvfMediaCapture.Build.cs b/Engine/Plugins/Media/AvfMedia/Source/AvfMediaCapture/AvfMediaCapture.Build.cs
--- a/Engine/Plugins/Media/AvfMedia/Source/AvfMediaCapture/AvfMediaCapture.Build.cs
+++ b/Engine/Plugins/Media/AvfMedia/Source/AvfMediaCapture/AvfMediaCapture.Build.cs
@@ -18,9 +18,12 @@ namespace UnrealBuildTool.Rules
 					"ApplicationCore",
 					"MediaUtils",
 					"RenderCore",
-					"RHI"
+					"RHI",
+					"MetalRHI"
 				});
 
+			AddEngineThirdPartyPrivateStaticDependencies(Target, "MetalCPP");
+			
 			PrivateIncludePathModuleNames.AddRange(
 				new string[] {
 					"Media",
diff --git a/Engine/Plugins/Media/AvfMedia/Source/AvfMediaCapture/Private/Player/AvfMediaCapturePlayer.cpp b/Engine/Plugins/Media/AvfMedia/Source/AvfMediaCapture/Private/Player/AvfMediaCapturePlayer.cpp
--- a/Engine/Plugins/Media/AvfMedia/Source/AvfMediaCapture/Private/Player/AvfMediaCapturePlayer.cpp
+++ b/Engine/Plugins/Media/AvfMedia/Source/AvfMediaCapture/Private/Player/AvfMediaCapturePlayer.cpp
@@ -4,6 +4,8 @@
 #include "MediaSamples.h"
 #include "Misc/MessageDialog.h"
 
+#include "IMetalDynamicRHI.h"
+
 inline FString FStringFromAppleFourCharWordCode(uint32 Word)
 {
 	char TypeCode[5];
@@ -737,21 +739,19 @@ void FAvfMediaCapturePlayer::ProcessSampleBufferVideo(CMSampleBufferRef SampleBu
 		check(Result == kCVReturnSuccess);
 		check(TextureRef);
 
-		FRHITextureCreateDesc Desc =
-			FRHITextureCreateDesc::Create2D(TEXT("FAvfMediaCapturePlayer"), Width, Height, PF_B8G8R8A8)
-			.SetFlags(ETextureCreateFlags::SRGB | ETextureCreateFlags::NoTiling | ETextureCreateFlags::ShaderResource)
-			.SetInitActionBulkData(new FAvfTexture2DResourceWrapper(TextureRef));
-
-		CFRelease(TextureRef);
-		
 		FIntPoint Dim(Width, Height);
-		
+
 		ENQUEUE_RENDER_COMMAND(FAvfMediaCapturePlayer_ProcessSampleBufferVideo_CreateTexture)(
-            [WeakPlayer = this->AsWeak(), InDesc = MoveTemp(Desc), Dim, SampleDuration](FRHICommandListImmediate& RHICmdList)
+            [WeakPlayer = this->AsWeak(), Dim, SampleDuration, TextureRef](FRHICommandListImmediate& RHICmdList)
             {
             	if (const TSharedPtr<FAvfMediaCapturePlayer> Player = WeakPlayer.Pin())
             	{
-            		TRefCountPtr<FRHITexture> ShaderResource = RHICreateTexture(InDesc);
+					IMetalDynamicRHI* MetalDynamicRHI = GetIMetalDynamicRHI();
+										
+					TRefCountPtr<FRHITexture> ShaderResource = MetalDynamicRHI->RHICreateTexture2DFromCVMetalTexture(PF_B8G8R8A8, ETextureCreateFlags::NoTiling | ETextureCreateFlags::ShaderResource, FClearValueBinding::White, TextureRef);
+					CFRelease(TextureRef);
+					
+					
 					TSharedRef<FAvfMediaTextureSample, ESPMode::ThreadSafe> VideoSample = Player->VideoSamplePool.AcquireShared();
 					VideoSample->Initialize(ShaderResource, Dim, Dim, Player->CurrentTime, SampleDuration);
             
diff --git a/Engine/Plugins/Media/ElectraUtil/Source/ElectraSamples/Private/ElectraTextureSampleAVF.cpp b/Engine/Plugins/Media/ElectraUtil/Source/ElectraSamples/Private/ElectraTextureSampleAVF.cpp
--- a/Engine/Plugins/Media/ElectraUtil/Source/ElectraSamples/Private/ElectraTextureSampleAVF.cpp
+++ b/Engine/Plugins/Media/ElectraUtil/Source/ElectraSamples/Private/ElectraTextureSampleAVF.cpp
@@ -22,6 +22,8 @@ THIRD_PARTY_INCLUDES_START
 #include "MetalInclude.h"
 THIRD_PARTY_INCLUDES_END
 
+#include "IMetalDynamicRHI.h"
+
 #if WITH_ENGINE
 extern void SafeReleaseMetalObject(NS::Object* Object);
 #endif
@@ -49,112 +51,6 @@ FElectraMediaTexConvApple::~FElectraMediaTexConvApple()
 }
 
 #if WITH_ENGINE
-namespace
-{
-
-/**
- * Passes a CV*TextureRef or CVPixelBufferRef through to the RHI to wrap in an RHI texture without traversing system memory.
- */
-class FTexConvTexResourceWrapper
-	: public FResourceBulkDataInterface
-{
-public:
-
-	FTexConvTexResourceWrapper(CFTypeRef InImageBuffer)
-		: ImageBuffer(InImageBuffer)
-	{
-		check(ImageBuffer);
-		CFRetain(ImageBuffer);
-	}
-
-	virtual ~FTexConvTexResourceWrapper()
-	{
-		CFRelease(ImageBuffer);
-		ImageBuffer = nullptr;
-	}
-
-public:
-
-	//~ FResourceBulkDataInterface interface
-
-	virtual void Discard() override
-	{
-		delete this;
-	}
-
-	virtual const void* GetResourceBulkData() const override
-	{
-		return ImageBuffer;
-	}
-
-	virtual uint32 GetResourceBulkDataSize() const override
-	{
-		return ImageBuffer ? ~0u : 0;
-	}
-
-	virtual EBulkDataType GetResourceType() const override
-	{
-		return EBulkDataType::MediaTexture;
-	}
-
-	CFTypeRef ImageBuffer;
-};
-
-
-/**
- * Allows for direct GPU mem allocation for texture resource from a CVImageBufferRef's system memory backing store.
- */
-class FTexConvTexResourceMemory
-	: public FResourceBulkDataInterface
-{
-public:
-	FTexConvTexResourceMemory(CVImageBufferRef InImageBuffer)
-	: ImageBuffer(InImageBuffer)
-	{
-		check(ImageBuffer);
-		CFRetain(ImageBuffer);
-	}
-
-	/**
-	 * @return ptr to the resource memory which has been preallocated
-	 */
-	virtual const void* GetResourceBulkData() const override
-	{
-		CVPixelBufferLockBaseAddress(ImageBuffer, kCVPixelBufferLock_ReadOnly);
-		return CVPixelBufferGetBaseAddress(ImageBuffer);
-	}
-
-	/**
-	 * @return size of resource memory
-	 */
-	virtual uint32 GetResourceBulkDataSize() const override
-	{
-		int32 Pitch = CVPixelBufferGetBytesPerRow(ImageBuffer);
-		int32 Height = CVPixelBufferGetHeight(ImageBuffer);
-		uint32 Size = (Pitch * Height);
-
-		return Size;
-	}
-
-	/**
-	 * Free memory after it has been used to initialize RHI resource
-	 */
-	virtual void Discard() override
-	{
-		CVPixelBufferUnlockBaseAddress(ImageBuffer, kCVPixelBufferLock_ReadOnly);
-		delete this;
-	}
-
-	virtual ~FTexConvTexResourceMemory()
-	{
-		CFRelease(ImageBuffer);
-		ImageBuffer = nullptr;
-	}
-
-	CVImageBufferRef ImageBuffer;
-};
-
-} // namespace anonymous
 
 
 void FElectraMediaTexConvApple::ConvertTexture(FTextureRHIRef & InDstTexture, CVImageBufferRef InImageBufferRef, bool bFullRange, EMediaTextureSampleFormat Format, const FMatrix44f& YUVMtx, const UE::Color::FColorSpace& SourceColorSpace, UE::Color::EEncoding EncodingType, float NormalizationFactor)
@@ -166,7 +62,6 @@ void FElectraMediaTexConvApple::ConvertTexture(FTextureRHIRef & InDstTexture, CV
 	const int32 FrameStride = CVPixelBufferGetBytesPerRow(InImageBufferRef);
 
 	// We have to support Metal for this object now
-	check(COREVIDEO_SUPPORTS_METAL);
 	check(IsMetalPlatform(GMaxRHIShaderPlatform));
 	{
 		if (!MetalTextureCache)
@@ -204,22 +99,11 @@ void FElectraMediaTexConvApple::ConvertTexture(FTextureRHIRef & InDstTexture, CV
 			check(Result == kCVReturnSuccess);
 			check(UVTextureRef);
 
-			// Metal can upload directly from an IOSurface to a 2D texture, so we can just wrap it.
-
-			const FRHITextureCreateDesc YDesc =
-				FRHITextureCreateDesc::Create2D(TEXT("YTex"), YWidth, YHeight, bIs8Bit ? PF_G8 : PF_G16)
-				.SetFlags(ETextureCreateFlags::NoTiling | ETextureCreateFlags::ShaderResource)
-				.SetInitActionBulkData(new FTexConvTexResourceWrapper(YTextureRef))
-				.SetInitialState(ERHIAccess::SRVMask);
-
-			const FRHITextureCreateDesc UVDesc =
-				FRHITextureCreateDesc::Create2D(TEXT("UVTex"), UVWidth, UVHeight, bIs8Bit ? PF_R8G8 : PF_G16R16)
-				.SetFlags(ETextureCreateFlags::NoTiling | ETextureCreateFlags::ShaderResource)
-				.SetInitActionBulkData(new FTexConvTexResourceWrapper(UVTextureRef))
-				.SetInitialState(ERHIAccess::SRVMask);
-
-			TRefCountPtr<FRHITexture> YTex = RHICreateTexture(YDesc);
-			TRefCountPtr<FRHITexture> UVTex = RHICreateTexture(UVDesc);
+			// Metal can upload directly from an IOSurface to a 2D texture
+			IMetalDynamicRHI* MetalDynamicRHI = GetIMetalDynamicRHI();
+				
+			TRefCountPtr<FRHITexture> YTex = MetalDynamicRHI->RHICreateTexture2DFromCVMetalTexture(bIs8Bit ? PF_G8 : PF_G16, ETextureCreateFlags::NoTiling | ETextureCreateFlags::ShaderResource, FClearValueBinding::White, YTextureRef);
+			TRefCountPtr<FRHITexture> UVTex = MetalDynamicRHI->RHICreateTexture2DFromCVMetalTexture(bIs8Bit ? PF_R8G8 : PF_G16R16, ETextureCreateFlags::NoTiling | ETextureCreateFlags::ShaderResource, FClearValueBinding::White, UVTextureRef);
 
 			{
 				// configure media shaders
@@ -332,15 +216,10 @@ void FElectraMediaTexConvApple::ConvertTexture(FTextureRHIRef & InDstTexture, CV
 				check(Result == kCVReturnSuccess);
 				check(TextureRef);
 
-				const FRHITextureCreateDesc Desc =
-					FRHITextureCreateDesc::Create2D(TEXT("DstTexture"), Width, Height, PF_R16G16B16A16_UNORM)
-					.SetFlags(ETextureCreateFlags::NoTiling | ETextureCreateFlags::ShaderResource)
-					.SetInitActionBulkData(new FTexConvTexResourceWrapper(TextureRef));
-
-				InDstTexture = RHICreateTexture(Desc);
-
+				IMetalDynamicRHI* MetalDynamicRHI = GetIMetalDynamicRHI();
+							
+				InDstTexture = MetalDynamicRHI->RHICreateTexture2DFromCVMetalTexture(PF_R16G16B16A16_UNORM, ETextureCreateFlags::NoTiling | ETextureCreateFlags::ShaderResource, FClearValueBinding::White, TextureRef);
 				CFRelease(TextureRef);
-
 			}
 			else
 			{
@@ -360,13 +239,9 @@ void FElectraMediaTexConvApple::ConvertTexture(FTextureRHIRef & InDstTexture, CV
 				check(Result == kCVReturnSuccess);
 				check(TextureRef);
 
-				const FRHITextureCreateDesc Desc =
-					FRHITextureCreateDesc::Create2D(TEXT("DstTexture"), Width, Height, PF_B8G8R8A8)
-					.SetFlags(ETextureCreateFlags::SRGB | ETextureCreateFlags::NoTiling | ETextureCreateFlags::ShaderResource)
-					.SetInitActionBulkData(new FTexConvTexResourceWrapper(TextureRef));
-
-				InDstTexture = RHICreateTexture(Desc);
+				IMetalDynamicRHI* MetalDynamicRHI = GetIMetalDynamicRHI();
 
+				InDstTexture = MetalDynamicRHI->RHICreateTexture2DFromCVMetalTexture(PF_B8G8R8A8, ETextureCreateFlags::NoTiling | ETextureCreateFlags::ShaderResource, FClearValueBinding::White, TextureRef);
 				CFRelease(TextureRef);
 			}
 		}
diff --git a/Engine/Plugins/Runtime/AR/AppleAR/AppleARKit/Source/AppleARKit/Private/AppleARKitTextures.cpp b/Engine/Plugins/Runtime/AR/AppleAR/AppleARKit/Source/AppleARKit/Private/AppleARKitTextures.cpp
--- a/Engine/Plugins/Runtime/AR/AppleAR/AppleARKit/Source/AppleARKit/Private/AppleARKitTextures.cpp
+++ b/Engine/Plugins/Runtime/AR/AppleAR/AppleARKit/Source/AppleARKit/Private/AppleARKitTextures.cpp
@@ -499,61 +499,6 @@ void UAppleARKitEnvironmentCaptureProbeTexture::Init(float InTimestamp, id<MTLTe
 	}
 }
 
-/**
- * Passes a metaltexture through to the RHI to wrap in an RHI texture without traversing system memory.
- */
-class FAppleARKitMetalTextureResourceWrapper :
-	public FResourceBulkDataInterface
-{
-public:
-	FAppleARKitMetalTextureResourceWrapper(id<MTLTexture> InImageBuffer)
-		: ImageBuffer(InImageBuffer)
-	{
-		check(ImageBuffer);
-		CFRetain(ImageBuffer);
-	}
-	
-	virtual ~FAppleARKitMetalTextureResourceWrapper()
-	{
-		CFRelease(ImageBuffer);
-		ImageBuffer = nullptr;
-	}
-
-	/**
-	 * @return ptr to the resource memory which has been preallocated
-	 */
-	virtual const void* GetResourceBulkData() const override
-	{
-		return ImageBuffer;
-	}
-	
-	/**
-	 * @return size of resource memory
-	 */
-	virtual uint32 GetResourceBulkDataSize() const override
-	{
-		return 0;
-	}
-	
-	/**
-	 * @return the type of bulk data for special handling
-	 */
-	virtual EBulkDataType GetResourceType() const override
-	{
-		return EBulkDataType::MediaTexture;
-	}
-	
-	/**
-	 * Free memory after it has been used to initialize RHI resource
-	 */
-	virtual void Discard() override
-	{
-		delete this;
-	}
-	
-	id<MTLTexture> ImageBuffer;
-};
-
 class FARMetalResource :
 	public FTextureResource
 {
@@ -787,63 +732,6 @@ FTextureResource* UAppleARKitOcclusionTexture::CreateResource()
 #endif // PLATFORM_APPLE
 
 
-#if SUPPORTS_ARKIT_1_0
-/**
-* Passes a CVMetalTextureRef through to the RHI to wrap in an RHI texture without traversing system memory.
-* @see FAvfTexture2DResourceWrapper & FMetalSurface::FMetalSurface
-*/
-class FAppleARKitCameraTextureResourceWrapper : public FResourceBulkDataInterface
-{
-public:
-	FAppleARKitCameraTextureResourceWrapper(CFTypeRef InImageBuffer)
-		: ImageBuffer(InImageBuffer)
-	{
-		check(ImageBuffer);
-		CFRetain(ImageBuffer);
-	}
-
-	/**
-	* @return ptr to the resource memory which has been preallocated
-	*/
-	virtual const void* GetResourceBulkData() const override
-	{
-		return ImageBuffer;
-	}
-
-	/**
-	* @return size of resource memory
-	*/
-	virtual uint32 GetResourceBulkDataSize() const override
-	{
-		return 0;
-	}
-
-	/**
-	* @return the type of bulk data for special handling
-	*/
-	virtual EBulkDataType GetResourceType() const override
-	{
-		return EBulkDataType::MediaTexture;
-	}
-
-	/**
-	* Free memory after it has been used to initialize RHI resource
-	*/
-	virtual void Discard() override
-	{
-		delete this;
-	}
-
-	virtual ~FAppleARKitCameraTextureResourceWrapper()
-	{
-		CFRelease(ImageBuffer);
-		ImageBuffer = nullptr;
-	}
-
-	CFTypeRef ImageBuffer;
-};
-#endif
-
 class FComputeShaderYCbCrToRGB : public FGlobalShader
 {
 public:
@@ -949,20 +837,13 @@ public:
 		}
 		
 		{
-			// Update the RHI texture wrapper for the Y and CbCr images
-			const FRHITextureCreateDesc YDesc =
-				FRHITextureCreateDesc::Create2D(TEXT("VideoTextureY"), CapturedYImageSize, PF_B8G8R8A8)
-				.SetFlags(ETextureCreateFlags::NoTiling | ETextureCreateFlags::ShaderResource)
-				.SetInitActionBulkData(new FAppleARKitCameraTextureResourceWrapper(CapturedYImage));
-
-			const FRHITextureCreateDesc cBcRDesc =
-				FRHITextureCreateDesc::Create2D(TEXT("VideoTextureCbCr"), CapturedCbCrImageSize, PF_B8G8R8A8)
-				.SetFlags(ETextureCreateFlags::NoTiling | ETextureCreateFlags::ShaderResource)
-				.SetInitActionBulkData(new FAppleARKitCameraTextureResourceWrapper(CapturedCbCrImage));
-
-			// pull the Y and CbCr textures out of the captured image planes (format is fake here, it will get the format from the FAppleARKitCameraTextureResourceWrapper)
-			VideoTextureY = RHICmdList.CreateTexture(YDesc);
-			VideoTextureCbCr = RHICmdList.CreateTexture(cBcRDesc);
+			IMetalDynamicRHI* MetalDynamicRHI = GetIMetalDynamicRHI();
+			
+			VideoTextureY = MetalDynamicRHI->RHICreateTexture2DFromCVMetalTexture(PF_B8G8R8A8, ETextureCreateFlags::NoTiling | ETextureCreateFlags::ShaderResource, FClearValueBinding::White, CapturedYTexture);
+			CFRelease(CapturedYTexture);
+	
+			VideoTextureCbCr = MetalDynamicRHI->RHICreateTexture2DFromCVMetalTexture(PF_B8G8R8A8, ETextureCreateFlags::NoTiling | ETextureCreateFlags::ShaderResource, FClearValueBinding::White, CapturedCbCrTexture);
+			CFRelease(CapturedCbCrTexture);
 
 			// todo: Add an update call to the registry instead of this unregister/re-register
 			FExternalTextureRegistry::Get().UnregisterExternalTexture(ARKitPassthroughCameraExternalTextureYGuid);
@@ -977,9 +858,7 @@ public:
 			//Make sure AR camera pass through materials are updated properly
 			FMaterialRenderProxy::UpdateDeferredCachedUniformExpressions();
 			
-			CFRelease(CapturedYImage);
 			CapturedYImage = nullptr;
-			CFRelease(CapturedCbCrImage);
 			CapturedCbCrImage = nullptr;
 		}
 		
diff --git a/Engine/Plugins/Runtime/AppleMoviePlayer/Source/AppleMoviePlayer/AppleMoviePlayer.Build.cs b/Engine/Plugins/Runtime/AppleMoviePlayer/Source/AppleMoviePlayer/AppleMoviePlayer.Build.cs
--- a/Engine/Plugins/Runtime/AppleMoviePlayer/Source/AppleMoviePlayer/AppleMoviePlayer.Build.cs
+++ b/Engine/Plugins/Runtime/AppleMoviePlayer/Source/AppleMoviePlayer/AppleMoviePlayer.Build.cs
@@ -15,10 +15,13 @@ namespace UnrealBuildTool.Rules
                     "RenderCore",
                     "RHI",
                     "Slate",
-					"ApplicationCore"
+					"ApplicationCore",
+					"MetalRHI"
 				}
 			);
 
+			AddEngineThirdPartyPrivateStaticDependencies(Target, "MetalCPP");
+        
 			PrivateDependencyModuleNames.AddRange(
 				new string[] {
 					"SlateCore",
diff --git a/Engine/Plugins/Runtime/AppleMoviePlayer/Source/AppleMoviePlayer/Private/AppleMovieStreamer.cpp b/Engine/Plugins/Runtime/AppleMoviePlayer/Source/AppleMoviePlayer/Private/AppleMovieStreamer.cpp
--- a/Engine/Plugins/Runtime/AppleMoviePlayer/Source/AppleMoviePlayer/Private/AppleMovieStreamer.cpp
+++ b/Engine/Plugins/Runtime/AppleMoviePlayer/Source/AppleMoviePlayer/Private/AppleMovieStreamer.cpp
@@ -12,6 +12,8 @@
 #include "Misc/CoreDelegates.h"
 #include "Containers/ResourceArray.h"
 
+#include "IMetalDynamicRHI.h"
+
 DEFINE_LOG_CATEGORY_STATIC(LogMoviePlayer, Log, All);
 
 #define MOVIE_FILE_EXTENSION @"mp4"
@@ -71,54 +73,6 @@ static FString ConvertToNativePath(const FString& Filename, bool bForWrite)
 	return Result;
 }
 
-//
-// Clone from FAvfMediaVideoSampler
-//
-class FAppleMovieStreamerTexture2DResourceWrapper
-	: public FResourceBulkDataInterface
-{
-public:
-
-	FAppleMovieStreamerTexture2DResourceWrapper(CFTypeRef InImageBuffer)
-		: ImageBuffer(InImageBuffer)
-	{
-		check(ImageBuffer);
-		CFRetain(ImageBuffer);
-	}
-
-	virtual ~FAppleMovieStreamerTexture2DResourceWrapper()
-	{
-		CFRelease(ImageBuffer);
-		ImageBuffer = nullptr;
-	}
-
-public:
-
-	//~ FResourceBulkDataInterface interface
-
-	virtual void Discard() override
-	{
-		delete this;
-	}
-
-	virtual const void* GetResourceBulkData() const override
-	{
-		return ImageBuffer;
-	}
-	
-	virtual uint32 GetResourceBulkDataSize() const override
-	{
-		return ImageBuffer ? ~0u : 0;
-	}
-
-	virtual EBulkDataType GetResourceType() const override
-	{
-		return EBulkDataType::MediaTexture;
-	}
-	
-	CFTypeRef ImageBuffer;
-};
-
 //
 // Methods
 //
@@ -593,13 +547,11 @@ bool FAVPlayerMovieStreamer::CheckForNextFrameAndCopy()
 		check(Result == kCVReturnSuccess);
 		check(TextureCacheRef);
 
-		const FRHITextureCreateDesc Desc =
-			FRHITextureCreateDesc::Create2D(TEXT("FAVPlayerMovieStreamer"), Width, Height, PF_B8G8R8A8)
-			.SetFlags(ETextureCreateFlags::NoTiling)
-			.SetInitActionBulkData(new FAppleMovieStreamerTexture2DResourceWrapper(TextureCacheRef));
-
-		FTextureRHIRef MediaWrappedTexture = RHICreateTexture(Desc);
-
+		IMetalDynamicRHI* MetalDynamicRHI = GetIMetalDynamicRHI();
+					
+		FTextureRHIRef MediaWrappedTexture = MetalDynamicRHI->RHICreateTexture2DFromCVMetalTexture(PF_B8G8R8A8, ETextureCreateFlags::NoTiling | ETextureCreateFlags::ShaderResource, FClearValueBinding::White, TextureCacheRef);
+		CFRelease(TextureCacheRef);
+		
 		if(SlateTexture.IsValid())
 		{
 			if(SlateTexture->IsValid())
@@ -614,7 +566,6 @@ bool FAVPlayerMovieStreamer::CheckForNextFrameAndCopy()
 			MovieViewport->SetTexture(SlateTexture);
 		}
 
-		CFRelease(TextureCacheRef);
         CFRelease(LatestSamples);
         LatestSamples = NULL;
 
diff --git a/Engine/Source/Runtime/Apple/MetalRHI/Private/MetalDynamicRHI.cpp b/Engine/Source/Runtime/Apple/MetalRHI/Private/MetalDynamicRHI.cpp
--- a/Engine/Source/Runtime/Apple/MetalRHI/Private/MetalDynamicRHI.cpp
+++ b/Engine/Source/Runtime/Apple/MetalRHI/Private/MetalDynamicRHI.cpp
@@ -269,7 +269,7 @@ bool FMetalDynamicRHI::RHIMatchPrecachePSOInitializers(const FGraphicsPipelineSt
 	return true;
 }
 
-void FMetalDynamicRHI::RHIRunOnQueue(TFunction<void(FMetalCommandQueue&)>&& CodeToRun, bool bWaitForSubmission)
+void FMetalDynamicRHI::RHIRunOnQueue(TFunction<void(MTL::CommandQueue*)>&& CodeToRun, bool bWaitForSubmission)
 {
 	FGraphEventRef SubmissionEvent;
 
diff --git a/Engine/Source/Runtime/Apple/MetalRHI/Private/MetalDynamicRHI.h b/Engine/Source/Runtime/Apple/MetalRHI/Private/MetalDynamicRHI.h
--- a/Engine/Source/Runtime/Apple/MetalRHI/Private/MetalDynamicRHI.h
+++ b/Engine/Source/Runtime/Apple/MetalRHI/Private/MetalDynamicRHI.h
@@ -238,7 +238,9 @@ public:
 	virtual void RHIBeginRenderQueryBatch_TopOfPipe(FRHICommandListBase& RHICmdList, ERenderQueryType QueryType) final override;
 	virtual void RHIEndRenderQueryBatch_TopOfPipe(FRHICommandListBase& RHICmdList, ERenderQueryType QueryType) final override;
 	
-	void RHIRunOnQueue(TFunction<void(FMetalCommandQueue&)>&& CodeToRun, bool bWaitForSubmission);
+	virtual void 				RHIRunOnQueue(TFunction<void(MTL::CommandQueue*)>&& CodeToRun, bool bWaitForSubmission) final override;
+	
+	virtual FTextureRHIRef		RHICreateTexture2DFromCVMetalTexture(EPixelFormat Format, ETextureCreateFlags TexCreateFlags, const FClearValueBinding& ClearValueBinding, CVMetalTextureRef Resource) final override;
 	
 	template <typename ...Args>
 	void DeferredDelete(Args&&... InArgs)
diff --git a/Engine/Source/Runtime/Apple/MetalRHI/Private/MetalSubmission.cpp b/Engine/Source/Runtime/Apple/MetalRHI/Private/MetalSubmission.cpp
--- a/Engine/Source/Runtime/Apple/MetalRHI/Private/MetalSubmission.cpp
+++ b/Engine/Source/Runtime/Apple/MetalRHI/Private/MetalSubmission.cpp
@@ -934,6 +934,6 @@ void FMetalPayload::PreExecute()
 {
 	if (PreExecuteCallback)
 	{
-		PreExecuteCallback(Queue);
+		PreExecuteCallback(Queue.GetQueue());
 	}
 }
diff --git a/Engine/Source/Runtime/Apple/MetalRHI/Private/MetalSubmission.h b/Engine/Source/Runtime/Apple/MetalRHI/Private/MetalSubmission.h
--- a/Engine/Source/Runtime/Apple/MetalRHI/Private/MetalSubmission.h
+++ b/Engine/Source/Runtime/Apple/MetalRHI/Private/MetalSubmission.h
@@ -183,7 +183,7 @@ struct FMetalPayload : public FMetalPayloadBaseFixLayout
 	std::atomic<bool> bSubmitted { false };
 
 	// Used by RHIRunOnQueue
-	TFunction<void(FMetalCommandQueue&)> PreExecuteCallback;
+	TFunction<void(MTL::CommandQueue*)> PreExecuteCallback;
 
 	// Execute
 	TArray<FMetalCommandBuffer*> CommandBuffersToExecute;
diff --git a/Engine/Source/Runtime/Apple/MetalRHI/Private/MetalTexture.cpp b/Engine/Source/Runtime/Apple/MetalRHI/Private/MetalTexture.cpp
--- a/Engine/Source/Runtime/Apple/MetalRHI/Private/MetalTexture.cpp
+++ b/Engine/Source/Runtime/Apple/MetalRHI/Private/MetalTexture.cpp
@@ -259,8 +259,6 @@ MTLTexturePtr FMetalSurface::Reallocate(MTLTexturePtr InTexture, MTL::TextureUsa
 
 void FMetalSurface::MakeAliasable(void)
 {
-	check(ImageSurfaceRef == nullptr);
-	
 	static bool bSupportsHeaps = Device.SupportsFeature(EMetalFeaturesHeaps);
 	if (bSupportsHeaps && Texture->storageMode() == MTL::StorageModePrivate && Texture->heap())
 	{
@@ -542,9 +540,8 @@ FMetalSurface::FMetalSurface(FMetalDevice& MetalDevice, FMetalTextureCreateDesc
 	: FRHITexture       (CreateDesc)
 	, Device			(MetalDevice)
 	, FormatKey         (CreateDesc.FormatKey)
-	, TotalTextureSize  (0)
 	, Viewport          (nullptr)
-	, ImageSurfaceRef   (nullptr)
+	, TotalTextureSize  (0)
 {
 	FPlatformAtomics::InterlockedExchange(&Written, 0);
 	check(CreateDesc.Extent.X > 0 && CreateDesc.Extent.Y > 0 && CreateDesc.NumMips > 0);
@@ -557,54 +554,6 @@ FMetalSurface::FMetalSurface(FMetalDevice& MetalDevice, FMetalTextureCreateDesc
     
     bool bIsMSAARequired = CreateDesc.NumSamples > 1 && AllowMSAA();
 
-	FResourceBulkDataInterface* BulkData = CreateDesc.BulkData;
-
-	// The bulk data interface can be used to create external textures for VR and media player.
-	// Handle these first.
-	if (BulkData != nullptr)
-	{
-		switch (BulkData->GetResourceType())
-		{
-		case FResourceBulkDataInterface::EBulkDataType::MediaTexture:
-			{
-				checkf(CreateDesc.NumMips == 1 && CreateDesc.ArraySize == 1, TEXT("Only handling bulk data with 1 mip and 1 array length"));
-				ImageSurfaceRef = (CFTypeRef)BulkData->GetResourceBulkData();
-				CFRetain(ImageSurfaceRef);
-				
-#if !COREVIDEO_SUPPORTS_METAL
-                Texture = Device.GetDevice()->newTexture(CreateDesc.Desc, CVPixelBufferGetIOSurface((CVPixelBufferRef)ImageSurfaceRef));
-#else
-				Texture = NS::RetainPtr((__bridge MTL::Texture*)CVMetalTextureGetTexture((CVPixelBufferRef)ImageSurfaceRef));
-#endif
-				METAL_FATAL_ASSERT(Texture, TEXT("Failed to create texture, desc %s"), *NSStringToFString(CreateDesc.Desc->description()));
-
-				BulkData->Discard();
-				BulkData = nullptr;
-			}
-			break;
-
-#if PLATFORM_MAC
-		case FResourceBulkDataInterface::EBulkDataType::VREyeBuffer:
-			{
-				ImageSurfaceRef = (CFTypeRef)BulkData->GetResourceBulkData();
-				CFRetain(ImageSurfaceRef);
-
-				MTLTextureDescriptorPtr DescCopy = NS::TransferPtr(CreateDesc.Desc->copy());
-				DescCopy->setStorageMode(MTL::StorageModeManaged);
-				DescCopy->setResourceOptions((MTL::ResourceOptions)((DescCopy->resourceOptions() & ~(MTL::ResourceStorageModeMask)) | MTL::ResourceStorageModeManaged));
-
-				Texture = NS::TransferPtr(Device.GetDevice()->newTexture(DescCopy.get(), (IOSurfaceRef)ImageSurfaceRef, 0));
-
-				METAL_FATAL_ASSERT(Texture, TEXT("Failed to create texture, desc %s"), *NSStringToFString(DescCopy->description()));
-
-				BulkData->Discard();
-				BulkData = nullptr;
-			}
-			break;
-#endif
-		}
-	}
-
 #if RHI_USE_RESOURCE_DEBUG_NAME
     NS::String* DebugName = FStringToNSString(CreateDesc.DebugName);
 #endif
@@ -812,6 +761,60 @@ FMetalSurface::FMetalSurface(FMetalDevice& MetalDevice, FMetalTextureCreateDesc
 #endif
 }
 
+FMetalExternalSurface::FMetalExternalSurface(FMetalDevice& MetalDevice, MTLTexturePtr InTexture, FRHITextureCreateDesc const& CreateDesc, TUniqueFunction<void()>&& InOnDeleteFunction)
+	: 	FMetalSurface	(MetalDevice, InTexture, CreateDesc),
+		OnDeleteFunction(MoveTemp(InOnDeleteFunction))
+{
+	check(CreateDesc.Extent.X > 0 && CreateDesc.Extent.Y > 0 && CreateDesc.NumMips > 0);
+	
+	FormatKey = GetMetalPixelFormatKey(Texture->pixelFormat());
+	
+	// calculate size of the texture
+	TotalTextureSize = GetMemorySize();
+	
+	const bool bOnlyStreamableTexturesInTexturePool = false;
+	UE::RHICore::UpdateGlobalTextureStats(GetDesc(), TotalTextureSize, bOnlyStreamableTexturesInTexturePool, true);
+	
+#if PLATFORM_SUPPORTS_BINDLESS_RENDERING
+	if(IsMetalBindlessEnabled())
+	{
+		FMetalBindlessDescriptorManager* BindlessDescriptorManager = Device.GetBindlessDescriptorManager();
+		check(BindlessDescriptorManager);
+
+		BindlessHandle = BindlessDescriptorManager->ReserveDescriptor(ERHIDescriptorHeapType::Standard);
+	}
+#endif
+}
+
+FMetalExternalSurface::~FMetalExternalSurface()
+{
+	FMetalDynamicRHI::Get().DeferredDelete([DeleteFunction=MoveTemp(OnDeleteFunction), InTexture=Texture]{
+		DeleteFunction();
+	});
+	
+	// Reset the texture so the parent doesn't free it as well
+	Texture.reset();
+}
+
+FTextureRHIRef FMetalDynamicRHI::RHICreateTexture2DFromCVMetalTexture(EPixelFormat Format, ETextureCreateFlags TexCreateFlags, const FClearValueBinding& ClearValueBinding, CVMetalTextureRef Resource)
+{
+	CFRetain(Resource);
+	
+	MTLTexturePtr MetalTexture = NS::RetainPtr((__bridge MTL::Texture*)CVMetalTextureGetTexture((CVPixelBufferRef)Resource));
+	
+	const FRHITextureCreateDesc Desc =
+		FRHITextureCreateDesc::Create2D(TEXT("MetalTexture2DFromResource"), MetalTexture->width(), MetalTexture->height(), Format)
+		.SetClearValue(ClearValueBinding)
+		.SetFlags(TexCreateFlags)
+		.SetNumMips(MetalTexture->mipmapLevelCount())
+		.SetNumSamples(MetalTexture->sampleCount())
+		.DetermineInititialState();
+
+	return new FMetalExternalSurface(*Device, MetalTexture, Desc, [Resource]() {
+		CFRelease(Resource);
+	});
+}
+
 void FMetalSurface::Initialize(FRHICommandListBase& RHICmdList)
 {
 #if PLATFORM_SUPPORTS_BINDLESS_RENDERING
@@ -882,16 +885,6 @@ FMetalSurface::~FMetalSurface()
 	FMetalDynamicRHI::Get().DeferredDelete([DeferredStats]() {
         delete DeferredStats;
     });
-	
-	if(ImageSurfaceRef)
-	{
-		// CFArray can contain CFType objects and is toll-free bridged with NSArray
-		CFArrayRef Temp = CFArrayCreate(kCFAllocatorSystemDefault, &ImageSurfaceRef, 1, &kCFTypeArrayCallBacks);
-		FMetalDynamicRHI::Get().DeferredDelete((__bridge NS::Array*)Temp);
-		CFRelease(ImageSurfaceRef);
-	}
-	
-	ImageSurfaceRef = nullptr;
 
 #if PLATFORM_SUPPORTS_BINDLESS_RENDERING
 	if(IsMetalBindlessEnabled())
@@ -1534,13 +1527,9 @@ FRHITextureInitializer FMetalDynamicRHI::RHICreateTextureInitializer(FRHICommand
 	LLM_SCOPE_DYNAMIC_STAT_OBJECTPATH_FNAME(CreateDesc.GetTraceClassName(), ELLMTagSet::AssetClasses);
 	UE_TRACE_METADATA_SCOPE_ASSET_FNAME(CreateDesc.DebugName, CreateDesc.GetTraceClassName(), CreateDesc.OwnerName);
 
-	// If the FMetalSurface constructor handled the bulk data, just return the default initializer.
-	const bool bForceDefaultInitializer =
-		(CreateDesc.InitAction == ERHITextureInitAction::BulkData && CreateDesc.BulkData && CreateDesc.BulkData->GetResourceType() != FResourceBulkDataInterface::EBulkDataType::Default);
-
 	FMetalSurface* Texture = new FMetalSurface(*Device, FMetalTextureCreateDesc(*Device, CreateDesc));
 
-	if (CreateDesc.InitAction == ERHITextureInitAction::Default || bForceDefaultInitializer)
+	if (CreateDesc.InitAction == ERHITextureInitAction::Default)
 	{
 		return UE::RHICore::FBaseTextureInitializerImplementation(RHICmdList, Texture,
 			[Texture = TRefCountPtr<FMetalSurface>(Texture)](FRHICommandListBase& RHICmdList) mutable
diff --git a/Engine/Source/Runtime/Apple/MetalRHI/Public/IMetalDynamicRHI.h b/Engine/Source/Runtime/Apple/MetalRHI/Public/IMetalDynamicRHI.h
--- a/Engine/Source/Runtime/Apple/MetalRHI/Public/IMetalDynamicRHI.h
+++ b/Engine/Source/Runtime/Apple/MetalRHI/Public/IMetalDynamicRHI.h
@@ -3,6 +3,7 @@
 #pragma once 
 
 #include "RHI.h"
+#include "MetalThirdParty.h"
 
 struct IMetalDynamicRHI : public FDynamicRHI
 {
@@ -10,6 +11,9 @@ struct IMetalDynamicRHI : public FDynamicRHI
 	{
 		return ERHIInterfaceType::Metal;
 	}
+	
+	virtual void 			  RHIRunOnQueue(TFunction<void(MTL::CommandQueue*)>&& CodeToRun, bool bWaitForSubmission) = 0;
+	virtual FTextureRHIRef	  RHICreateTexture2DFromCVMetalTexture(EPixelFormat Format, ETextureCreateFlags TexCreateFlags, const FClearValueBinding& ClearValueBinding, CVMetalTextureRef Resource) = 0;
 };
 
 inline bool IsRHIMetal()
diff --git a/Engine/Source/Runtime/Apple/MetalRHI/Public/MetalResources.h b/Engine/Source/Runtime/Apple/MetalRHI/Public/MetalResources.h
--- a/Engine/Source/Runtime/Apple/MetalRHI/Public/MetalResources.h
+++ b/Engine/Source/Runtime/Apple/MetalRHI/Public/MetalResources.h
@@ -162,7 +162,6 @@ inline uint32 GetTypeHash(const MTLTexturePtr& TexturePtr)
 class METALRHI_API FMetalSurface : public FRHITexture, public FMetalViewableResource
 {
 public:
-
 	/** 
 	 * Constructor that will create Texture and Color/DepthBuffers as needed
 	 */
@@ -172,6 +171,21 @@ public:
 	 * Destructor
 	 */
 	virtual ~FMetalSurface();
+	
+protected:
+	/** 
+	 * Constructor for derived classes
+	 */
+	FMetalSurface(FMetalDevice& MetalDevice, MTLTexturePtr InTexture, FRHITextureCreateDesc const& CreateDesc)
+		: FRHITexture		(CreateDesc)
+		, Device			(MetalDevice)
+		, FormatKey			(0)
+		, Texture			(InTexture)
+		, Viewport          (nullptr)
+		, TotalTextureSize	(0)
+	{}
+	
+public:
 
 	void Initialize(FRHICommandListBase& RHICmdList);
 
@@ -223,32 +237,6 @@ public:
 
     MTLTexturePtr Reallocate(MTLTexturePtr Texture, MTL::TextureUsage UsageModifier);
 	void MakeAliasable(void);
-	
-	FMetalDevice& Device;
-	int16 volatile Written;
-	uint8 const FormatKey;
-
-	//texture used for store actions and binding to shader params
-    MTLTexturePtr Texture;
-	//if surface is MSAA, texture used to bind for RT
-    MTLTexturePtr MSAATexture;
-
-	//texture used for a resolve target.  Same as texture on iOS.  
-	//Dummy target on Mac where RHISupportsSeparateMSAAAndResolveTextures is true.	In this case we don't always want a resolve texture but we
-	//have to have one until renderpasses are implemented at a high level.
-	// Mac / RHISupportsSeparateMSAAAndResolveTextures == true
-	// iOS A9+ where depth resolve is available
-	// iOS < A9 where depth resolve is unavailable.
-    MTLTexturePtr MSAAResolveTexture;
-
-	// how much memory is allocated for this texture
-	uint64 TotalTextureSize;
-	
-	// Used for atomics
-	FMetalBufferPtr BackingBuffer;
-	
-	// For back-buffers, the owning viewport.
-	class FMetalViewport* Viewport;
 
 	virtual void* GetTextureBaseRHI() override final
 	{
@@ -270,19 +258,58 @@ public:
     }
 #endif // PLATFORM_SUPPORTS_BINDLESS_RENDERING
 
+	FMetalDevice& Device;
+	int16 volatile Written;
+	uint8 FormatKey;
+
+	//texture used for store actions and binding to shader params
+	MTLTexturePtr Texture;
+	//if surface is MSAA, texture used to bind for RT
+	MTLTexturePtr MSAATexture;
+
+	//texture used for a resolve target.  Same as texture on iOS.  
+	//Dummy target on Mac where RHISupportsSeparateMSAAAndResolveTextures is true.	In this case we don't always want a resolve texture but we
+	//have to have one until renderpasses are implemented at a high level.
+	// Mac / RHISupportsSeparateMSAAAndResolveTextures == true
+	// iOS A9+ where depth resolve is available
+	// iOS < A9 where depth resolve is unavailable.
+	MTLTexturePtr MSAAResolveTexture;
+	
+	// Used for atomics
+	FMetalBufferPtr BackingBuffer;
+	
+	// For back-buffers, the owning viewport.
+	class FMetalViewport* Viewport;
+	
 private:
 	/** Safely releases the texture when not in use
 	 * @param MTLTexturePtr Texture from surface to be released
 	 */
 	void SafeRelease(MTLTexturePtr Texture);
 	
-	// The movie playback IOSurface/CVTexture wrapper to avoid page-off
-	CFTypeRef ImageSurfaceRef;
+	FCriticalSection DrawableMutex;
+	
+protected:
+	
+	// how much memory is allocated for this texture
+	uint64 TotalTextureSize;
 	
 #if PLATFORM_SUPPORTS_BINDLESS_RENDERING
 	FRHIDescriptorHandle 	BindlessHandle;
 #endif
-	FCriticalSection 		DrawableMutex;
+};
+
+/** 
+ * Derived class that can be used for creating surfaces with external textures 
+ */
+class METALRHI_API FMetalExternalSurface : public FMetalSurface
+{
+public:
+	FMetalExternalSurface(FMetalDevice& MetalDevice, MTLTexturePtr InTexture, FRHITextureCreateDesc const& CreateDesc, TUniqueFunction<void()>&& InOnDeleteFunction);
+	~FMetalExternalSurface();
+
+private:	
+	TUniqueFunction<void()> OnDeleteFunction;
 };
 
 class FMetalBufferData
diff --git a/Engine/Source/Runtime/Core/Public/Containers/ResourceArray.h b/Engine/Source/Runtime/Core/Public/Containers/ResourceArray.h
--- a/Engine/Source/Runtime/Core/Public/Containers/ResourceArray.h
+++ b/Engine/Source/Runtime/Core/Public/Containers/ResourceArray.h
@@ -124,21 +124,6 @@ public:
 	 * Free memory after it has been used to initialize RHI resource 
 	 */
 	virtual void Discard() = 0;
-	
-	enum class EBulkDataType
-	{
-		Default,
-		MediaTexture,
-		VREyeBuffer,
-	};
-	
-	/**
-	 * @return the type of bulk data for special handling
-	 */
-	[[nodiscard]] virtual EBulkDataType GetResourceType() const
-	{
-		return EBulkDataType::Default;
-	}
 };
 
 /** Utility to do a simple upload of data from an array managed by the caller. */
