Remove RHIDiscardRenderTargets from IRHICommandContext interface
- The only implementation that was still called was in OpenGL RHI, and that was an internal call. The function does not need to be on the RHI interface.
- Removed dead code in Metal RHI.

diff --git a/Engine/Source/Runtime/Apple/MetalRHI/Private/MetalCommands.cpp b/Engine/Source/Runtime/Apple/MetalRHI/Private/MetalCommands.cpp
--- a/Engine/Source/Runtime/Apple/MetalRHI/Private/MetalCommands.cpp
+++ b/Engine/Source/Runtime/Apple/MetalRHI/Private/MetalCommands.cpp
@@ -1317,11 +1317,6 @@ void FMetalRHICommandContext::RHISetDepthBounds(float MinDepth, float MaxDepth)
 	METAL_IGNORED(FMetalRHICommandContextSetDepthBounds);
 }
 
-void FMetalRHICommandContext::RHIDiscardRenderTargets(bool Depth, bool Stencil, uint32 ColorBitMask)
-{
-	StateCache.DiscardRenderTargets(Depth, Stencil, ColorBitMask);
-}
-
 #if PLATFORM_USES_FIXED_RHI_CLASS
 #define INTERNAL_DECORATOR(Method) ((FMetalRHICommandContext&)CmdList.GetContext()).FMetalRHICommandContext::Method
 #include "RHICommandListCommandExecutes.inl"
diff --git a/Engine/Source/Runtime/Apple/MetalRHI/Private/MetalStateCache.cpp b/Engine/Source/Runtime/Apple/MetalRHI/Private/MetalStateCache.cpp
--- a/Engine/Source/Runtime/Apple/MetalRHI/Private/MetalStateCache.cpp
+++ b/Engine/Source/Runtime/Apple/MetalRHI/Private/MetalStateCache.cpp
@@ -2862,49 +2862,6 @@ void FMetalStateCache::CommitResourceTable(EMetalShaderStages const Frequency, M
 	}
 }
 
-void FMetalStateCache::DiscardRenderTargets(bool Depth, bool Stencil, uint32 ColorBitMask)
-{
-	if (Depth)
-	{
-		switch (DepthStore)
-		{
-			case MTL::StoreActionUnknown:
-			case MTL::StoreActionStore:
-				DepthStore = MTL::StoreActionDontCare;
-				break;
-			case MTL::StoreActionStoreAndMultisampleResolve:
-				DepthStore = MTL::StoreActionMultisampleResolve;
-				break;
-			default:
-				break;
-		}
-	}
-
-	if (Stencil)
-	{
-		StencilStore = MTL::StoreActionDontCare;
-	}
-
-	for (uint32 Index = 0; Index < MaxSimultaneousRenderTargets; ++Index)
-	{
-		if ((ColorBitMask & (1u << Index)) != 0)
-		{
-			switch (ColorStore[Index])
-			{
-				case MTL::StoreActionUnknown:
-				case MTL::StoreActionStore:
-					ColorStore[Index] = MTL::StoreActionDontCare;
-					break;
-				case MTL::StoreActionStoreAndMultisampleResolve:
-					ColorStore[Index] = MTL::StoreActionMultisampleResolve;
-					break;
-				default:
-					break;
-			}
-		}
-	}
-}
-
 void FMetalStateCache::ReleaseDescriptor(MTL::RenderPassDescriptor* Desc)
 {
 	FMetalRenderPassDescriptorPool::Get().ReleaseDescriptor(Desc);
diff --git a/Engine/Source/Runtime/Apple/MetalRHI/Private/MetalStateCache.h b/Engine/Source/Runtime/Apple/MetalRHI/Private/MetalStateCache.h
--- a/Engine/Source/Runtime/Apple/MetalRHI/Private/MetalStateCache.h
+++ b/Engine/Source/Runtime/Apple/MetalRHI/Private/MetalStateCache.h
@@ -191,8 +191,6 @@ public:
 	void SetRenderPipelineState(FMetalCommandEncoder& CommandEncoder);
     void SetComputePipelineState(FMetalCommandEncoder& CommandEncoder);
 	void FlushVisibilityResults(FMetalCommandEncoder& CommandEncoder);
-
-	void DiscardRenderTargets(bool Depth, bool Stencil, uint32 ColorBitMask);
 	
 	void ReleaseDescriptor(MTL::RenderPassDescriptor* Desc);
 	void ClearPreviousComputeState();
diff --git a/Engine/Source/Runtime/Apple/MetalRHI/Public/MetalRHIContext.h b/Engine/Source/Runtime/Apple/MetalRHI/Public/MetalRHIContext.h
--- a/Engine/Source/Runtime/Apple/MetalRHI/Public/MetalRHIContext.h
+++ b/Engine/Source/Runtime/Apple/MetalRHI/Public/MetalRHIContext.h
@@ -89,8 +89,6 @@ public:
     
 	virtual void RHIBeginRenderQuery(FRHIRenderQuery* RenderQuery) final override;
 	virtual void RHIEndRenderQuery(FRHIRenderQuery* RenderQuery) final override;
-
-	virtual void RHIDiscardRenderTargets(bool Depth, bool Stencil, uint32 ColorBitMask) final override;
 	
 	// This method is queued with an RHIThread, otherwise it will flush after it is queued; without an RHI thread there is no benefit to queuing this frame advance commands
 	virtual void RHIEndDrawingViewport(FRHIViewport* Viewport, bool bPresent, bool bLockToVsync) override;
diff --git a/Engine/Source/Runtime/OpenGLDrv/Private/OpenGLCommands.cpp b/Engine/Source/Runtime/OpenGLDrv/Private/OpenGLCommands.cpp
--- a/Engine/Source/Runtime/OpenGLDrv/Private/OpenGLCommands.cpp
+++ b/Engine/Source/Runtime/OpenGLDrv/Private/OpenGLCommands.cpp
@@ -1681,7 +1681,7 @@ void FOpenGLDynamicRHI::SetRenderTargets(
 	}
 }
 
-void FOpenGLDynamicRHI::RHIDiscardRenderTargets(bool Depth, bool Stencil, uint32 ColorBitMaskIn)
+void FOpenGLDynamicRHI::DiscardRenderTargets(bool Depth, bool Stencil, uint32 ColorBitMaskIn)
 {
 	if (FOpenGL::SupportsDiscardFrameBuffer())
 	{
diff --git a/Engine/Source/Runtime/OpenGLDrv/Private/OpenGLDrv.h b/Engine/Source/Runtime/OpenGLDrv/Private/OpenGLDrv.h
--- a/Engine/Source/Runtime/OpenGLDrv/Private/OpenGLDrv.h
+++ b/Engine/Source/Runtime/OpenGLDrv/Private/OpenGLDrv.h
@@ -446,7 +446,7 @@ public:
 	UE_API virtual void RHIBeginBreadcrumbGPU(FRHIBreadcrumbNode* Breadcrumb) final override;
 	UE_API virtual void RHIEndBreadcrumbGPU  (FRHIBreadcrumbNode* Breadcrumb) final override;
 #endif
-	UE_API virtual void RHIDiscardRenderTargets(bool Depth,bool Stencil,uint32 ColorBitMask) final override;
+	void DiscardRenderTargets(bool Depth, bool Stencil, uint32 ColorBitMask);
 
 	// FIXME: Broken on Android for cubemaps
 	UE_API virtual void RHICopyTexture(FRHITexture* SourceTexture, FRHITexture* DestTexture, const FRHICopyTextureInfo& CopyInfo) final override;
diff --git a/Engine/Source/Runtime/OpenGLDrv/Private/OpenGLRenderTarget.cpp b/Engine/Source/Runtime/OpenGLDrv/Private/OpenGLRenderTarget.cpp
--- a/Engine/Source/Runtime/OpenGLDrv/Private/OpenGLRenderTarget.cpp
+++ b/Engine/Source/Runtime/OpenGLDrv/Private/OpenGLRenderTarget.cpp
@@ -1109,7 +1109,7 @@ void FOpenGLDynamicRHI::RHIEndRenderPass()
 
 	if (bDiscardDepthStencil || ColorMask != 0)
 	{
-		RHIDiscardRenderTargets(bDiscardDepthStencil, bDiscardDepthStencil, ColorMask);
+		DiscardRenderTargets(bDiscardDepthStencil, bDiscardDepthStencil, ColorMask);
 	}
 
 	FRHIRenderTargetView RTV(nullptr, ERenderTargetLoadAction::ENoAction);
diff --git a/Engine/Source/Runtime/RHI/Public/RHICommandList.h b/Engine/Source/Runtime/RHI/Public/RHICommandList.h
--- a/Engine/Source/Runtime/RHI/Public/RHICommandList.h
+++ b/Engine/Source/Runtime/RHI/Public/RHICommandList.h
@@ -2521,22 +2521,6 @@ FRHICOMMAND_MACRO(FRHICommandEndDrawingViewport)
 	RHI_EXECUTE_API void Execute(FRHICommandListBase& CmdList);
 };
 
-FRHICOMMAND_MACRO(FRHICommandDiscardRenderTargets)
-{
-	uint32 ColorBitMask;
-	bool Depth;
-	bool Stencil;
-
-	FORCEINLINE_DEBUGGABLE FRHICommandDiscardRenderTargets(bool InDepth, bool InStencil, uint32 InColorBitMask)
-		: ColorBitMask(InColorBitMask)
-		, Depth(InDepth)
-		, Stencil(InStencil)
-	{
-	}
-	
-	RHI_EXECUTE_API void Execute(FRHICommandListBase& CmdList);
-};
-
 FRHICOMMAND_MACRO(FRHICommandCopyBufferRegion)
 {
 	FRHIBuffer* DestBuffer;
diff --git a/Engine/Source/Runtime/RHI/Public/RHICommandListCommandExecutes.inl b/Engine/Source/Runtime/RHI/Public/RHICommandListCommandExecutes.inl
--- a/Engine/Source/Runtime/RHI/Public/RHICommandListCommandExecutes.inl
+++ b/Engine/Source/Runtime/RHI/Public/RHICommandListCommandExecutes.inl
@@ -558,9 +558,3 @@ void FRHICommandEndDrawingViewport::Execute(FRHICommandListBase& CmdList)
 	RHISTAT(EndDrawingViewport);
 	INTERNAL_DECORATOR(RHIEndDrawingViewport)(Viewport, bPresent, bLockToVsync);
 }
-
-void FRHICommandDiscardRenderTargets::Execute(FRHICommandListBase& CmdList)
-{
-	RHISTAT(RHIDiscardRenderTargets);
-	INTERNAL_DECORATOR(RHIDiscardRenderTargets)(Depth, Stencil, ColorBitMask);
-}
diff --git a/Engine/Source/Runtime/RHI/Public/RHIContext.h b/Engine/Source/Runtime/RHI/Public/RHIContext.h
--- a/Engine/Source/Runtime/RHI/Public/RHIContext.h
+++ b/Engine/Source/Runtime/RHI/Public/RHIContext.h
@@ -678,9 +678,6 @@ public:
 	}
 #endif
 
-	// Not all RHIs need this (Mobile specific)
-	virtual void RHIDiscardRenderTargets(bool Depth, bool Stencil, uint32 ColorBitMask) {}
-
 	// This method is queued with an RHIThread, otherwise it will flush after it is queued; without an RHI thread there is no benefit to queuing this frame advance commands
 	virtual void RHIEndDrawingViewport(FRHIViewport* Viewport, bool bPresent, bool bLockToVsync) = 0;
 
diff --git a/Engine/Source/Runtime/RHI/Public/RHIValidationContext.h b/Engine/Source/Runtime/RHI/Public/RHIValidationContext.h
--- a/Engine/Source/Runtime/RHI/Public/RHIValidationContext.h
+++ b/Engine/Source/Runtime/RHI/Public/RHIValidationContext.h
@@ -752,12 +752,6 @@ public:
 	}
 #endif
 
-	// Not all RHIs need this (Mobile specific)
-	virtual void RHIDiscardRenderTargets(bool bDepth, bool bStencil, uint32 ColorBitMask) override final
-	{
-		RHIContext->RHIDiscardRenderTargets(bDepth, bStencil, ColorBitMask);
-	}
-
 	// This method is queued with an RHIThread, otherwise it will flush after it is queued; without an RHI thread there is no benefit to queuing this frame advance commands
 	virtual void RHIEndDrawingViewport(FRHIViewport* Viewport, bool bPresent, bool bLockToVsync) override final
 	{
