Metal RHI

Re-enable SetBytes and SetBufferOffset on iOS and enable on tvOS

diff --git a/Engine/Source/Runtime/Apple/MetalRHI/Private/MetalDevice.cpp b/Engine/Source/Runtime/Apple/MetalRHI/Private/MetalDevice.cpp
--- a/Engine/Source/Runtime/Apple/MetalRHI/Private/MetalDevice.cpp
+++ b/Engine/Source/Runtime/Apple/MetalRHI/Private/MetalDevice.cpp
@@ -458,9 +458,9 @@ void FMetalDevice::EnumerateFeatureSupport()
 	}
 	
 #if PLATFORM_IOS
-#if PLATFORM_TVOS
-	Features &= ~(EMetalFeaturesSetBytes);
-		
+	Features |= EMetalFeaturesSetBufferOffset | EMetalFeaturesSetBytes;
+	
+#if PLATFORM_TVOS	
 	if(Device->supportsFeatureSet(MTL::FeatureSet_tvOS_GPUFamily2_v1))
 	{
 		Features |= EMetalFeaturesCountingQueries | EMetalFeaturesBaseVertexInstance | EMetalFeaturesIndirectBuffer | EMetalFeaturesMSAADepthResolve | EMetalFeaturesMSAAStoreAndResolve;
