Metal RHI

- Removed a number of the Metal feature flags as many of the features are now available across all devices.
- Stripped dead code

diff --git a/Engine/Source/Developer/Mac/MacTargetPlatformSettings/Classes/MacTargetSettings.h b/Engine/Source/Developer/Mac/MacTargetPlatformSettings/Classes/MacTargetSettings.h
--- a/Engine/Source/Developer/Mac/MacTargetPlatformSettings/Classes/MacTargetSettings.h
+++ b/Engine/Source/Developer/Mac/MacTargetPlatformSettings/Classes/MacTargetSettings.h
@@ -122,10 +122,6 @@ public:
 	 */
 	UPROPERTY(EditAnywhere, config, Category=Rendering, meta = (DisplayName = "Enable Fast-Math optimisations", ConfigRestartRequired = true))
 	bool EnableMathOptimisations;
-	
-	/** Whether to compile shaders using a tier Indirect Argument Buffers. */
-	UPROPERTY(config, EditAnywhere, Category = Rendering, Meta = (DisplayName = "Tier of Indirect Argument Buffers to use when compiling shaders", ConfigRestartRequired = true))
-	int32 IndirectArgumentTier;
 
 	/** Sample rate to run the audio mixer with. */
 	UPROPERTY(config, EditAnywhere, Category = "Audio", Meta = (DisplayName = "Audio Mixer Sample Rate"))
diff --git a/Engine/Source/Developer/Mac/MacTargetPlatformSettings/Private/MacTargetPlatformSettingsModule.cpp b/Engine/Source/Developer/Mac/MacTargetPlatformSettings/Private/MacTargetPlatformSettingsModule.cpp
--- a/Engine/Source/Developer/Mac/MacTargetPlatformSettings/Private/MacTargetPlatformSettingsModule.cpp
+++ b/Engine/Source/Developer/Mac/MacTargetPlatformSettings/Private/MacTargetPlatformSettingsModule.cpp
@@ -101,11 +101,6 @@ public:
 			TargetSettings->EnableMathOptimisations = true;
 		}
 		
-		if (!GConfig->GetInt(TEXT("/Script/MacTargetPlatform.MacTargetSettings"), TEXT("IndirectArgumentTier"), TargetSettings->IndirectArgumentTier, GEngineIni))
-		{
-			TargetSettings->IndirectArgumentTier = 0;
-		}
-		
 		TargetSettings->AddToRoot();
         
         ProjectSettings = NewObject<UXcodeProjectSettings>(GetTransientPackage(), "XcodeProjectSettings", RF_Standalone);
diff --git a/Engine/Source/Runtime/Apple/MetalRHI/Private/MetalCommandEncoder.cpp b/Engine/Source/Runtime/Apple/MetalRHI/Private/MetalCommandEncoder.cpp
--- a/Engine/Source/Runtime/Apple/MetalRHI/Private/MetalCommandEncoder.cpp
+++ b/Engine/Source/Runtime/Apple/MetalRHI/Private/MetalCommandEncoder.cpp
@@ -69,7 +69,6 @@ FString FMetalCommandData::ToString() const
 FMetalCommandEncoder::FMetalCommandEncoder(FMetalDevice& MetalDevice, FMetalCommandList& CmdList)
 : Device(MetalDevice)
 , CommandList(CmdList)
-, bSupportsMetalFeaturesSetBytes(Device.SupportsFeature(EMetalFeaturesSetBytes))
 , RingBuffer(Device, EncoderRingBufferSize, BufferOffsetAlignment, FMetalCommandQueue::GetCompatibleResourceOptions((MTL::ResourceOptions)(MTL::ResourceHazardTrackingModeUntracked | BUFFER_RESOURCE_STORAGE_MANAGED)))
 , RenderPassDesc(nullptr)
 , EncoderFence(nullptr)
@@ -1210,8 +1209,7 @@ void FMetalCommandEncoder::SetShaderBuffer(MTL::FunctionType const FunctionType,
 	FenceResource(Buffer->GetMTLBuffer(), FunctionType);
 	check(index < ML_MaxBuffers);
     
-    if(Device.SupportsFeature(EMetalFeaturesSetBufferOffset) &&
-	   Buffer && (ShaderBuffers[uint32(FunctionType)].Bound & (1 << index)) &&
+    if(Buffer && (ShaderBuffers[uint32(FunctionType)].Bound & (1 << index)) &&
 	   ShaderBuffers[uint32(FunctionType)].Buffers[index] == Buffer)
     {
 		if(ShaderBuffers[uint32(FunctionType)].Offsets[index] != Offset ||
@@ -1280,45 +1278,37 @@ void FMetalCommandEncoder::SetShaderBytes(MTL::FunctionType const FunctionType,
 	{
 		ShaderBuffers[uint32(FunctionType)].Bound |= (1 << Index);
 
-		if (bSupportsMetalFeaturesSetBytes)
+		switch (FunctionType)
 		{
-			switch (FunctionType)
-			{
-				case MTL::FunctionTypeVertex:
-					check(RenderCommandEncoder);
-					RenderCommandEncoder->setVertexBytes(Bytes, Length, Index);
-					break;
-				case MTL::FunctionTypeFragment:
-					check(RenderCommandEncoder);
-					RenderCommandEncoder->setFragmentBytes(Bytes, Length, Index);
-					break;
-				case MTL::FunctionTypeKernel:
-					check(ComputeCommandEncoder);
-					ComputeCommandEncoder->setBytes(Bytes, Length, Index);
-					break;
+			case MTL::FunctionTypeVertex:
+				check(RenderCommandEncoder);
+				RenderCommandEncoder->setVertexBytes(Bytes, Length, Index);
+				break;
+			case MTL::FunctionTypeFragment:
+				check(RenderCommandEncoder);
+				RenderCommandEncoder->setFragmentBytes(Bytes, Length, Index);
+				break;
+			case MTL::FunctionTypeKernel:
+				check(ComputeCommandEncoder);
+				ComputeCommandEncoder->setBytes(Bytes, Length, Index);
+				break;
 #if PLATFORM_SUPPORTS_MESH_SHADERS
-				case MTL::FunctionTypeMesh:
-					check(RenderCommandEncoder);
-					RenderCommandEncoder->setMeshBytes(Bytes, Length, Index);
-					break;
-				case MTL::FunctionTypeObject:
-					check(RenderCommandEncoder);
-					RenderCommandEncoder->setObjectBytes(Bytes, Length, Index);
-                    break;
+			case MTL::FunctionTypeMesh:
+				check(RenderCommandEncoder);
+				RenderCommandEncoder->setMeshBytes(Bytes, Length, Index);
+				break;
+			case MTL::FunctionTypeObject:
+				check(RenderCommandEncoder);
+				RenderCommandEncoder->setObjectBytes(Bytes, Length, Index);
+				break;
 #endif
-				default:
-					check(false);
-					break;
-			}
-			
-			ShaderBuffers[uint32(FunctionType)].Buffers[Index] = nullptr;
-		}
-		else
-		{
-			FMetalBufferPtr Buffer = RingBuffer.NewBuffer(Length, BufferOffsetAlignment);
-			FMemory::Memcpy(((uint8*)Buffer->Contents()), Bytes, Length);
-			ShaderBuffers[uint32(FunctionType)].Buffers[Index] = Buffer;
+			default:
+				check(false);
+				break;
 		}
+		
+		ShaderBuffers[uint32(FunctionType)].Buffers[Index] = nullptr;
+
 #if METAL_RHI_RAYTRACING
 		ShaderBuffers[uint32(FunctionType)].AccelerationStructure[Index] = nullptr;
 #endif // METAL_RHI_RAYTRACING
@@ -1350,7 +1340,6 @@ void FMetalCommandEncoder::SetShaderBufferOffset(MTL::FunctionType FunctionType,
 {
 	check(index < ML_MaxBuffers);
     checkf(ShaderBuffers[uint32(FunctionType)].Buffers[index] && (ShaderBuffers[uint32(FunctionType)].Bound & (1 << index)), TEXT("Buffer must already be bound"));
-	check(Device.SupportsFeature(EMetalFeaturesSetBufferOffset));
 	ShaderBuffers[uint32(FunctionType)].Offsets[index] = Offset;
 	ShaderBuffers[uint32(FunctionType)].SetBufferMetaData(index, Length, GMetalBufferFormats[PF_Unknown].DataFormat, 0);
 	
@@ -1644,16 +1633,6 @@ void FMetalCommandEncoder::SetShaderBufferInternal(MTL::FunctionType Function, u
     NS::UInteger Offset = Binding.Offsets[Index];
 	
 	bool bBufferHasBytes = Binding.Bytes[Index] != nullptr;
-	if (!Binding.Buffers[Index] && bBufferHasBytes && !bSupportsMetalFeaturesSetBytes)
-	{
-		uint8 const* Bytes = (((uint8 const*)Binding.Bytes[Index]->Data) + Binding.Offsets[Index]);
-		uint32 Len = Binding.Bytes[Index]->Len - Binding.Offsets[Index];
-		
-		Offset = 0;
-		Binding.Buffers[Index] = RingBuffer.NewBuffer(Len, BufferOffsetAlignment);
-		
-		FMemory::Memcpy(((uint8*)Binding.Buffers[Index]->Contents()) + Offset, Bytes, Len);
-	}
 	
 	TArray<TTuple<MTL::Resource*, MTL::ResourceUsage>> ReferencedResources = ShaderBuffers[uint32(Function)].ReferencedResources[Index];
 	switch (Function)
@@ -1783,7 +1762,7 @@ void FMetalCommandEncoder::SetShaderBufferInternal(MTL::FunctionType Function, u
 				break;
 		}
 	}
-	else if (bBufferHasBytes && bSupportsMetalFeaturesSetBytes)
+	else if (bBufferHasBytes)
 	{
 		uint8 const* Bytes = (((uint8 const*)Binding.Bytes[Index]->Data) + Binding.Offsets[Index]);
 		uint32 Len = Binding.Bytes[Index]->Len - Binding.Offsets[Index];
diff --git a/Engine/Source/Runtime/Apple/MetalRHI/Private/MetalCommandEncoder.h b/Engine/Source/Runtime/Apple/MetalRHI/Private/MetalCommandEncoder.h
--- a/Engine/Source/Runtime/Apple/MetalRHI/Private/MetalCommandEncoder.h
+++ b/Engine/Source/Runtime/Apple/MetalRHI/Private/MetalCommandEncoder.h
@@ -598,9 +598,6 @@ public:
 	FMetalDevice& Device;
 	FMetalCommandList& CommandList;
 
-    // Cache Queue feature
-    bool bSupportsMetalFeaturesSetBytes;
-    
 	FMetalBufferBindings ShaderBuffers[int(MTL::FunctionTypeObject)+1];
 	
     MTL::StoreAction ColorStoreActions[MaxSimultaneousRenderTargets];
diff --git a/Engine/Source/Runtime/Apple/MetalRHI/Private/MetalCommandQueue.cpp b/Engine/Source/Runtime/Apple/MetalRHI/Private/MetalCommandQueue.cpp
--- a/Engine/Source/Runtime/Apple/MetalRHI/Private/MetalCommandQueue.cpp
+++ b/Engine/Source/Runtime/Apple/MetalRHI/Private/MetalCommandQueue.cpp
@@ -37,7 +37,6 @@ FMetalCommandQueue::FMetalCommandQueue(FMetalDevice& MetalDevice, uint32 const M
 	: Device(MetalDevice)
 	, RuntimeDebuggingLevel(EMetalDebugLevelOff)
 {
-	int32 IndirectArgumentTier = 0;
     int32 MetalShaderVersion = 0;
 #if PLATFORM_MAC
 	const TCHAR* const Settings = TEXT("/Script/MacTargetPlatform.MacTargetSettings");
@@ -46,10 +45,6 @@ FMetalCommandQueue::FMetalCommandQueue(FMetalDevice& MetalDevice, uint32 const M
 #endif
     GConfig->GetInt(Settings, TEXT("MetalLanguageVersion"), MetalShaderVersion, GEngineIni);
 	
-    if(!GConfig->GetInt(Settings, TEXT("IndirectArgumentTier"), IndirectArgumentTier, GEngineIni))
-	{
-		IndirectArgumentTier = 0;
-	}
 	ValidateVersion(MetalShaderVersion);
 
 	if(MaxNumCommandBuffers == 0)
diff --git a/Engine/Source/Runtime/Apple/MetalRHI/Private/MetalDevice.cpp b/Engine/Source/Runtime/Apple/MetalRHI/Private/MetalDevice.cpp
--- a/Engine/Source/Runtime/Apple/MetalRHI/Private/MetalDevice.cpp
+++ b/Engine/Source/Runtime/Apple/MetalRHI/Private/MetalDevice.cpp
@@ -444,173 +444,64 @@ FMetalDevice::~FMetalDevice()
 
 void FMetalDevice::EnumerateFeatureSupport()
 {
-#if PLATFORM_MAC
-	const TCHAR* const Settings = TEXT("/Script/MacTargetPlatform.MacTargetSettings");
-#else
-	const TCHAR* const Settings = TEXT("/Script/IOSRuntimeSettings.IOSRuntimeSettings");
-#endif
-	
-	int32 IndirectArgumentTier = 0;
-	
-	if(!GConfig->GetInt(Settings, TEXT("IndirectArgumentTier"), IndirectArgumentTier, GEngineIni))
-	{
-		IndirectArgumentTier = 0;
-	}
-	
-#if PLATFORM_IOS
-	Features |= EMetalFeaturesSetBufferOffset | EMetalFeaturesSetBytes;
-	
-#if PLATFORM_TVOS	
-	if(Device->supportsFeatureSet(MTL::FeatureSet_tvOS_GPUFamily2_v1))
-	{
-		Features |= EMetalFeaturesCountingQueries | EMetalFeaturesBaseVertexInstance | EMetalFeaturesIndirectBuffer | EMetalFeaturesMSAADepthResolve | EMetalFeaturesMSAAStoreAndResolve;
-	}
-
-	Features |= EMetalFeaturesPrivateBufferSubAllocation;
-
-	Features |= EMetalFeaturesGPUCaptureManager | EMetalFeaturesBufferSubAllocation | EMetalFeaturesParallelRenderEncoders | EMetalFeaturesPipelineBufferMutability;
-
-	Features |= EMetalFeaturesMaxThreadsPerThreadgroup;
+	FString DeviceName(Device->name()->cString(NS::UTF8StringEncoding));
 
-	if (FParse::Param(FCommandLine::Get(), TEXT("metalfence")))
+	bool bSupportsMetal3 = Device->supportsFamily(MTL::GPUFamilyMetal3);
+	
+	if (!GForceNoMetalFence && FParse::Param(FCommandLine::Get(),TEXT("metalfence")))
 	{
 		Features |= EMetalFeaturesFences;
 	}
-					
-	if (FParse::Param(FCommandLine::Get(),TEXT("metalheap")))
+	
+	if (!GForceNoMetalHeap && 
+		  !DeviceName.Contains(TEXT("Intel")) &&
+		  Device->hasUnifiedMemory() &&
+		  FParse::Param(FCommandLine::Get(),TEXT("metalheap")))
 	{
 		Features |= EMetalFeaturesHeaps;
 	}
-
-	Features |= EMetalFeaturesTextureBuffers;
-#else
-	if (Device->supportsFeatureSet(MTL::FeatureSet_iOS_GPUFamily3_v1))
+	
+	if(bSupportsMetal3 || Device->supportsFamily(MTL::GPUFamilyApple3))
 	{
-		Features |= EMetalFeaturesCountingQueries | EMetalFeaturesBaseVertexInstance | EMetalFeaturesIndirectBuffer | EMetalFeaturesMSAADepthResolve;
+		Features |= EMetalFeaturesCountingQueries | EMetalFeaturesBaseVertexInstance | EMetalFeaturesIndirectBuffer | EMetalFeaturesMSAADepthResolve | EMetalFeaturesMSAAStoreAndResolve;
 	}
-		
-	if(Device->supportsFeatureSet(MTL::FeatureSet_iOS_GPUFamily3_v2) || Device->supportsFeatureSet(MTL::FeatureSet_iOS_GPUFamily2_v3) || Device->supportsFeatureSet(MTL::FeatureSet_iOS_GPUFamily1_v3))
+	
+	if(bSupportsMetal3 || Device->supportsFamily(MTL::GPUFamilyApple4))
 	{
-		if (FParse::Param(FCommandLine::Get(),TEXT("metalfence")))
-		{
-			Features |= EMetalFeaturesFences;
-		}
-			
-		if (FParse::Param(FCommandLine::Get(),TEXT("metalheap")))
-		{
-			Features |= EMetalFeaturesHeaps;
-		}
+		Features |= EMetalFeaturesCubemapArrays;
 	}
-		
-	if(Device->supportsFeatureSet(MTL::FeatureSet_iOS_GPUFamily3_v2))
+	
+	if(bSupportsMetal3 || Device->supportsFamily(MTL::GPUFamilyApple5))
 	{
-		Features |= EMetalFeaturesMSAAStoreAndResolve;
+		Features |= EMetalFeaturesLayeredRendering;
 	}
-
-	// Turning the below option on will allocate more buffer memory which isn't generally desirable on iOS
-	// Features |= EMetalFeaturesEfficientBufferBlits;
-			
-	// These options are fine however as thye just change how we allocate small buffers
-	Features |= EMetalFeaturesBufferSubAllocation;
+	
+#if PLATFORM_IOS
 	Features |= EMetalFeaturesPrivateBufferSubAllocation;
-
-	Features |= EMetalFeaturesPresentMinDuration | EMetalFeaturesGPUCaptureManager | EMetalFeaturesBufferSubAllocation | EMetalFeaturesParallelRenderEncoders | EMetalFeaturesPipelineBufferMutability;
-
-	Features |= EMetalFeaturesMaxThreadsPerThreadgroup;
-	if (!GForceNoMetalFence && !FParse::Param(FCommandLine::Get(), TEXT("nometalfence")))
-	{
-		Features |= EMetalFeaturesFences;
-	}
-					
-	if (!GForceNoMetalHeap && !FParse::Param(FCommandLine::Get(),TEXT("nometalheap")))
-	{
-		Features |= EMetalFeaturesHeaps;
-	}
-
-	Features |= EMetalFeaturesTextureBuffers;
-
-	if (Device->supportsFeatureSet(MTL::FeatureSet_iOS_GPUFamily4_v1))
+	Features |= EMetalFeaturesBufferSubAllocation;
+	
+	if(Device->supportsFamily(MTL::GPUFamilyApple4))
 	{
 		Features |= EMetalFeaturesTileShaders;
 	}
 					
-	if (Device->supportsFeatureSet(MTL::FeatureSet_iOS_GPUFamily5_v1))
-	{
-		Features |= EMetalFeaturesLayeredRendering;
-	}
+#if !PLATFORM_TVOS
+	Features |= EMetalFeaturesPresentMinDuration;	
 #endif
+
+	// Turning the below option on will allocate more buffer memory which isn't generally desirable on iOS
+	// Features |= EMetalFeaturesEfficientBufferBlits;
+
 #else // Assume that Mac & other platforms all support these from the start. They can diverge later.
-	Features = EMetalFeaturesCountingQueries | EMetalFeaturesBaseVertexInstance | EMetalFeaturesIndirectBuffer |
-				EMetalFeaturesLayeredRendering | EMetalFeaturesCubemapArrays | EMetalFeaturesSetBufferOffset;
 
-	FString DeviceName(Device->name()->cString(NS::UTF8StringEncoding));
-	
-	Features |= EMetalFeaturesMSAADepthResolve | EMetalFeaturesMSAAStoreAndResolve;
-	
-	// Assume that set*Bytes only works on macOS Sierra and above as no-one has tested it anywhere else.
-	Features |= EMetalFeaturesSetBytes;
-	
-	// On earlier OS versions Intel Broadwell couldn't suballocate properly
-	if (!(DeviceName.Contains(TEXT("Intel")) && (DeviceName.Contains(TEXT("5300")) || DeviceName.Contains(TEXT("6000")) || DeviceName.Contains(TEXT("6100")))))
-	{
-		// Using Private Memory & BlitEncoders for Vertex & Index data should be *much* faster.
-		Features |= EMetalFeaturesEfficientBufferBlits;
-		
-		Features |= EMetalFeaturesBufferSubAllocation;
-				
-		// On earlier OS versions Vega didn't like non-zero blit offsets
-		if (!DeviceName.Contains(TEXT("Vega")))
-		{
-			Features |= EMetalFeaturesPrivateBufferSubAllocation;
-		}
-	}
-	
-	if (!FParse::Param(FCommandLine::Get(), TEXT("nometalparallelencoder")))
-	{
-		Features |= EMetalFeaturesParallelRenderEncoders;
-	}
-	Features |= EMetalFeaturesTextureBuffers;
-	if (IndirectArgumentTier >= 1)
-	{
-		Features |= EMetalFeaturesIABs;
+	// Using Private Memory & BlitEncoders for Vertex & Index data should be *much* faster.
+	Features |= EMetalFeaturesEfficientBufferBlits;
+	Features |= EMetalFeaturesBufferSubAllocation;
 			
-		if (IndirectArgumentTier >= 2)
-		{
-			Features |= EMetalFeaturesTier2IABs;
-		}
-	}
-
-	// The editor spawns so many viewports and preview icons that we can run out of hardware fences!
-	// Need to figure out a way to safely flush the rendering and reuse the fences when that happens.
-#if WITH_EDITORONLY_DATA
-	if (!GIsEditor)
-#endif
-	{
-		if (FParse::Param(FCommandLine::Get(),TEXT("metalfence")))
-		{
-			Features |= EMetalFeaturesFences;
-		}
-	}
-	
-	// Temporarily only support heaps for devices with unified memory
-	// Disable this by default code while we work on metal heaps
-	if (!DeviceName.Contains(TEXT("Intel")) &&
-		  Device->hasUnifiedMemory() &&
-		  FParse::Param(FCommandLine::Get(),TEXT("metalheap")))
-	{
-		Features |= EMetalFeaturesHeaps;
-	}
-	
-	Features |= EMetalFeaturesMultipleViewports | EMetalFeaturesPipelineBufferMutability | EMetalFeaturesGPUCaptureManager;
-	
-	if (FParse::Param(FCommandLine::Get(),TEXT("metalfence")))
-	{
-		Features |= EMetalFeaturesFences;
-	}
-	
-	if (FParse::Param(FCommandLine::Get(),TEXT("metaliabs")))
+	// On earlier OS versions Vega didn't like non-zero blit offsets
+	if (!DeviceName.Contains(TEXT("Vega")))
 	{
-		Features |= EMetalFeaturesIABs;
+		Features |= EMetalFeaturesPrivateBufferSubAllocation;
 	}
 #endif
 	
diff --git a/Engine/Source/Runtime/Apple/MetalRHI/Private/MetalDevice.h b/Engine/Source/Runtime/Apple/MetalRHI/Private/MetalDevice.h
--- a/Engine/Source/Runtime/Apple/MetalRHI/Private/MetalDevice.h
+++ b/Engine/Source/Runtime/Apple/MetalRHI/Private/MetalDevice.h
@@ -28,10 +28,8 @@ enum class EMetalQueueType
  * Enumeration of features which are present only on some OS/device combinations.
  * These have to be checked at runtime as well as compile time to ensure backward compatibility.
  */
-typedef NS_OPTIONS(uint64, EMetalFeatures)
+typedef NS_OPTIONS(uint32, EMetalFeatures)
 {
-	/** Support for specifying an update to the buffer offset only */
-	EMetalFeaturesSetBufferOffset = 1 << 0,
 	/** Supports NSUInteger counting visibility queries */
 	EMetalFeaturesCountingQueries = 1 << 1,
 	/** Supports base vertex/instance for draw calls */
@@ -40,56 +38,34 @@ typedef NS_OPTIONS(uint64, EMetalFeatures)
 	EMetalFeaturesIndirectBuffer = 1 << 3,
 	/** Supports layered rendering */
 	EMetalFeaturesLayeredRendering = 1 << 4,
-	/** Support for specifying small buffers as byte arrays */
-	EMetalFeaturesSetBytes = 1 << 5,
-	/** Unused Reserved Bit */
-	EMetalFeaturesUnusedReservedBit6 = 1 << 6, // was EMetalFeaturesTessellation
 	/** Supports framework-level validation */
-	EMetalFeaturesValidation = 1 << 7,
-	/** Supports detailed statistics */
-	EMetalFeaturesStatistics = 1 << 8,
+	EMetalFeaturesValidation = 1 << 5,
 	/** Supports the explicit MTLHeap APIs */
-	EMetalFeaturesHeaps = 1 << 9,
+	EMetalFeaturesHeaps = 1 << 6,
 	/** Supports the explicit MTLFence APIs */
-	EMetalFeaturesFences = 1 << 10,
+	EMetalFeaturesFences = 1 << 7,
 	/** Supports MSAA Depth Resolves */
-	EMetalFeaturesMSAADepthResolve = 1 << 11,
+	EMetalFeaturesMSAADepthResolve = 1 << 8,
 	/** Supports Store & Resolve in a single store action */
-	EMetalFeaturesMSAAStoreAndResolve = 1 << 12,
+	EMetalFeaturesMSAAStoreAndResolve = 1 << 9,
 	/** Supports the use of cubemap arrays */
-	EMetalFeaturesCubemapArrays = 1 << 13,
+	EMetalFeaturesCubemapArrays = 1 << 10,
 	/** Supports the specification of multiple viewports and scissor rects */
-	EMetalFeaturesMultipleViewports = 1 << 14,
+	EMetalFeaturesMultipleViewports = 1 << 11,
 	/** Supports minimum on-glass duration for drawables */
-	EMetalFeaturesPresentMinDuration = 1llu << 15llu,
-	/** Supports programmatic frame capture API */
-	EMetalFeaturesGPUCaptureManager = 1llu << 16llu,
+	EMetalFeaturesPresentMinDuration = 1 << 12,
 	/** Supports efficient buffer-blits */
-	EMetalFeaturesEfficientBufferBlits = 1llu << 17llu,
+	EMetalFeaturesEfficientBufferBlits = 1 << 13,
 	/** Supports any kind of buffer sub-allocation */
-	EMetalFeaturesBufferSubAllocation = 1llu << 18llu,
+	EMetalFeaturesBufferSubAllocation = 1 << 14,
 	/** Supports private buffer sub-allocation */
-	EMetalFeaturesPrivateBufferSubAllocation = 1llu << 19llu,
-	/** Supports texture buffers */
-	EMetalFeaturesTextureBuffers = 1llu << 20llu,
-	/** Supports max. compute threads per threadgroup */
-	EMetalFeaturesMaxThreadsPerThreadgroup = 1llu << 21llu,
-	/** Supports parallel render encoders */
-	EMetalFeaturesParallelRenderEncoders = 1llu << 22llu,
-	/** Supports indirect argument buffers */
-	EMetalFeaturesIABs = 1llu << 23llu,
-	/** Supports specifying the mutability of buffers bound to PSOs */
-	EMetalFeaturesPipelineBufferMutability = 1llu << 24llu,
+	EMetalFeaturesPrivateBufferSubAllocation = 1 << 15,
 	/** Supports tile shaders */
-	EMetalFeaturesTileShaders = 1llu << 25llu,
-	/** Unused Reserved Bit */
-	EMetalFeaturesUnusedReservedBit27 = 1llu << 26llu, // was EMetalFeaturesSeparateTessellation
-	/** Supports indirect argument buffers Tier 2 */
-	EMetalFeaturesTier2IABs = 1llu << 27llu,
+	EMetalFeaturesTileShaders = 1 << 16,
 	/** Supports counter sampling on encoder stages */
-	EMetalFeaturesStageCounterSampling = 1llu << 28llu,
+	EMetalFeaturesStageCounterSampling = 1 << 17,
 	/** Supports counter sampling on the stage boundary */
-	EMetalFeaturesBoundaryCounterSampling = 1llu << 29llu,
+	EMetalFeaturesBoundaryCounterSampling = 1 << 18,
 };
 
 /**
diff --git a/Engine/Source/Runtime/Apple/MetalRHI/Private/MetalPipeline.cpp b/Engine/Source/Runtime/Apple/MetalRHI/Private/MetalPipeline.cpp
--- a/Engine/Source/Runtime/Apple/MetalRHI/Private/MetalPipeline.cpp
+++ b/Engine/Source/Runtime/Apple/MetalRHI/Private/MetalPipeline.cpp
@@ -769,40 +769,15 @@ static bool ConfigureRenderPipelineDescriptor(FMetalDevice& Device, TDescriptorT
 	}
 #endif
 	
-	if (Device.SupportsFeature(EMetalFeaturesPipelineBufferMutability))
+	if constexpr(bIsRenderPipelineDesc)
 	{
-		if constexpr(bIsRenderPipelineDesc)
-		{
-			FMetalVertexShader* VertexShader = (FMetalVertexShader*)Init.BoundShaderState.VertexShaderRHI;
-			
-			MTL::PipelineBufferDescriptorArray* VertexPipelineBuffers = RenderPipelineDesc->vertexBuffers();
-			FMetalShaderBindings& VertexBindings = VertexShader->Bindings;
-			int8 VertexSideTable = VertexShader->SideTableBinding;
-			{
-				uint32 ImmutableBuffers = VertexBindings.ConstantBuffers | VertexBindings.ArgumentBuffers;
-				while(ImmutableBuffers)
-				{
-					uint32 Index = __builtin_ctz(ImmutableBuffers);
-					ImmutableBuffers &= ~(1 << Index);
-					
-					if (Index < ML_MaxBuffers)
-					{
-						MTL::PipelineBufferDescriptor* PipelineBuffer = VertexPipelineBuffers->object(Index);
-						PipelineBuffer->setMutability(MTL::MutabilityImmutable);
-					}
-				}
-				if (VertexSideTable > 0)
-				{
-					MTL::PipelineBufferDescriptor* PipelineBuffer = VertexPipelineBuffers->object(VertexSideTable);
-					PipelineBuffer->setMutability(MTL::MutabilityImmutable);
-				}
-			}
-		}
+		FMetalVertexShader* VertexShader = (FMetalVertexShader*)Init.BoundShaderState.VertexShaderRHI;
 		
-		if (PixelShader)
+		MTL::PipelineBufferDescriptorArray* VertexPipelineBuffers = RenderPipelineDesc->vertexBuffers();
+		FMetalShaderBindings& VertexBindings = VertexShader->Bindings;
+		int8 VertexSideTable = VertexShader->SideTableBinding;
 		{
-            MTL::PipelineBufferDescriptorArray* FragmentPipelineBuffers = RenderPipelineDesc->fragmentBuffers();
-			uint32 ImmutableBuffers = PixelShader->Bindings.ConstantBuffers | PixelShader->Bindings.ArgumentBuffers;
+			uint32 ImmutableBuffers = VertexBindings.ConstantBuffers | VertexBindings.ArgumentBuffers;
 			while(ImmutableBuffers)
 			{
 				uint32 Index = __builtin_ctz(ImmutableBuffers);
@@ -810,18 +785,40 @@ static bool ConfigureRenderPipelineDescriptor(FMetalDevice& Device, TDescriptorT
 				
 				if (Index < ML_MaxBuffers)
 				{
-					MTL::PipelineBufferDescriptor* PipelineBuffer = FragmentPipelineBuffers->object(Index);
+					MTL::PipelineBufferDescriptor* PipelineBuffer = VertexPipelineBuffers->object(Index);
 					PipelineBuffer->setMutability(MTL::MutabilityImmutable);
 				}
 			}
-			if (PixelShader->SideTableBinding > 0)
+			if (VertexSideTable > 0)
 			{
-				MTL::PipelineBufferDescriptor* PipelineBuffer = FragmentPipelineBuffers->object(PixelShader->SideTableBinding);
+				MTL::PipelineBufferDescriptor* PipelineBuffer = VertexPipelineBuffers->object(VertexSideTable);
 				PipelineBuffer->setMutability(MTL::MutabilityImmutable);
 			}
 		}
 	}
 	
+	if (PixelShader)
+	{
+		MTL::PipelineBufferDescriptorArray* FragmentPipelineBuffers = RenderPipelineDesc->fragmentBuffers();
+		uint32 ImmutableBuffers = PixelShader->Bindings.ConstantBuffers | PixelShader->Bindings.ArgumentBuffers;
+		while(ImmutableBuffers)
+		{
+			uint32 Index = __builtin_ctz(ImmutableBuffers);
+			ImmutableBuffers &= ~(1 << Index);
+			
+			if (Index < ML_MaxBuffers)
+			{
+				MTL::PipelineBufferDescriptor* PipelineBuffer = FragmentPipelineBuffers->object(Index);
+				PipelineBuffer->setMutability(MTL::MutabilityImmutable);
+			}
+		}
+		if (PixelShader->SideTableBinding > 0)
+		{
+			MTL::PipelineBufferDescriptor* PipelineBuffer = FragmentPipelineBuffers->object(PixelShader->SideTableBinding);
+			PipelineBuffer->setMutability(MTL::MutabilityImmutable);
+		}
+	}
+	
 	return true;
 }
 
diff --git a/Engine/Source/Runtime/Apple/MetalRHI/Private/MetalRHIBuffer.cpp b/Engine/Source/Runtime/Apple/MetalRHI/Private/MetalRHIBuffer.cpp
--- a/Engine/Source/Runtime/Apple/MetalRHI/Private/MetalRHIBuffer.cpp
+++ b/Engine/Source/Runtime/Apple/MetalRHI/Private/MetalRHIBuffer.cpp
@@ -75,11 +75,11 @@ static EMetalBufferUsage GetMetalBufferUsage(EBufferUsageFlags InUsage)
 bool FMetalRHIBuffer::UsePrivateMemory() const
 {
 	if(EnumHasAnyFlags(GetUsage(), BUF_KeepCPUAccessible) && FMetalCommandQueue::IsUMASystem())
+	{
 		return false;
+	}
 	
-	return Device.SupportsFeature(EMetalFeaturesEfficientBufferBlits)
-			|| (Device.SupportsFeature(EMetalFeaturesIABs) && EnumHasAnyFlags(GetUsage(), BUF_ShaderResource|BUF_UnorderedAccess))
-			&& !FMetalCommandQueue::IsUMASystem();
+	return Device.SupportsFeature(EMetalFeaturesEfficientBufferBlits);
 }
 
 FMetalRHIBuffer::FMetalRHIBuffer(FRHICommandListBase& RHICmdList, FMetalDevice& MetalDevice, const FRHIBufferCreateDesc& CreateDesc, FResourceArrayUploadInterface* InResourceArray)
@@ -210,7 +210,7 @@ void FMetalRHIBuffer::AllocateBuffer()
 	
 	const EMetalBufferUsage MetalUsage = GetMetalBufferUsage(GetDesc().Usage);
 	
-	if (EnumHasAnyFlags(MetalUsage, EMetalBufferUsage::LinearTex) && !Device.SupportsFeature(EMetalFeaturesTextureBuffers))
+	if (EnumHasAnyFlags(MetalUsage, EMetalBufferUsage::LinearTex))
 	{
 		if (EnumHasAnyFlags(GetDesc().Usage, BUF_UnorderedAccess))
 		{
diff --git a/Engine/Source/Runtime/Apple/MetalRHI/Private/MetalState.cpp b/Engine/Source/Runtime/Apple/MetalRHI/Private/MetalState.cpp
--- a/Engine/Source/Runtime/Apple/MetalRHI/Private/MetalState.cpp
+++ b/Engine/Source/Runtime/Apple/MetalRHI/Private/MetalState.cpp
@@ -278,12 +278,10 @@ static MTL::SamplerState* FindOrCreateSamplerState(FMetalDevice& Device, const F
 #if PLATFORM_MAC
 		Desc->setBorderColor(Initializer.BorderColor == 0 ? MTL::SamplerBorderColorTransparentBlack : MTL::SamplerBorderColorOpaqueWhite);
 #endif
-#if !METAL_USE_METAL_SHADER_CONVERTER
-		if (Device.SupportsFeature(EMetalFeaturesIABs))
+		
+#if METAL_USE_METAL_SHADER_CONVERTER
+		Desc->setSupportArgumentBuffers(true);
 #endif
-		{
-			Desc->setSupportArgumentBuffers(true);
-		}
 		
 		State = Device.GetDevice()->newSamplerState(Desc);
         Desc->release();
diff --git a/Engine/Source/Runtime/Apple/MetalRHI/Private/MetalStateCache.cpp b/Engine/Source/Runtime/Apple/MetalRHI/Private/MetalStateCache.cpp
--- a/Engine/Source/Runtime/Apple/MetalRHI/Private/MetalStateCache.cpp
+++ b/Engine/Source/Runtime/Apple/MetalRHI/Private/MetalStateCache.cpp
@@ -1884,90 +1884,83 @@ void FMetalStateCache::SetResourcesFromTables(ShaderType Shader, CrossCompiler::
 	
 	EMetalShaderStages Frequency = TranslateShaderStage(ShaderStage);
 
-	if (!Device.SupportsFeature(EMetalFeaturesIABs))
+	struct FUniformResourceBinder
 	{
-		struct FUniformResourceBinder
-		{
-			FMetalStateCache& StateCache;
-			EMetalShaderStages Frequency;
+		FMetalStateCache& StateCache;
+		EMetalShaderStages Frequency;
 
-			void SetUAV(FRHIUnorderedAccessView* UAV, uint8 Index)
-			{
+		void SetUAV(FRHIUnorderedAccessView* UAV, uint8 Index)
+		{
 #if PLATFORM_SUPPORTS_BINDLESS_RENDERING
-				if(IsMetalBindlessEnabled())
-				{
-					StateCache.IRMakeUAVResident(Frequency, static_cast<FMetalUnorderedAccessView*>(UAV));
-				}
-				else
+			if(IsMetalBindlessEnabled())
+			{
+				StateCache.IRMakeUAVResident(Frequency, static_cast<FMetalUnorderedAccessView*>(UAV));
+			}
+			else
 #endif
-				{
-					StateCache.SetShaderUnorderedAccessView(Frequency, Index, static_cast<FMetalUnorderedAccessView*>(UAV));
-				}
+			{
+				StateCache.SetShaderUnorderedAccessView(Frequency, Index, static_cast<FMetalUnorderedAccessView*>(UAV));
 			}
+		}
 
-			void SetSRV(FRHIShaderResourceView* SRV, uint8 Index)
-			{
+		void SetSRV(FRHIShaderResourceView* SRV, uint8 Index)
+		{
 #if PLATFORM_SUPPORTS_BINDLESS_RENDERING
-				if(IsMetalBindlessEnabled())
-				{
-					StateCache.IRMakeSRVResident(Frequency, static_cast<FMetalShaderResourceView*>(SRV));
-				}
-				else
+			if(IsMetalBindlessEnabled())
+			{
+				StateCache.IRMakeSRVResident(Frequency, static_cast<FMetalShaderResourceView*>(SRV));
+			}
+			else
 #endif
-				{
-					StateCache.SetShaderResourceView(Frequency, Index, static_cast<FMetalShaderResourceView*>(SRV));
-				}
+			{
+				StateCache.SetShaderResourceView(Frequency, Index, static_cast<FMetalShaderResourceView*>(SRV));
 			}
+		}
 
-			void SetTexture(FRHITexture* Texture, uint8 Index)
-			{
+		void SetTexture(FRHITexture* Texture, uint8 Index)
+		{
 #if PLATFORM_SUPPORTS_BINDLESS_RENDERING
-				if(IsMetalBindlessEnabled())
-				{
-					StateCache.IRMakeTextureResident(Frequency, GetMetalSurfaceFromRHITexture(Texture)->Texture.get());
-				}
-				else
+			if(IsMetalBindlessEnabled())
+			{
+				StateCache.IRMakeTextureResident(Frequency, GetMetalSurfaceFromRHITexture(Texture)->Texture.get());
+			}
+			else
 #endif
-				{
-					
-					StateCache.SetShaderTexture(Frequency,
-												GetMetalSurfaceFromRHITexture(Texture)->Texture.get(),
-												Index,
-												MTL::ResourceUsage(MTL::ResourceUsageRead | MTL::ResourceUsageSample));
-				}
+			{
+				
+				StateCache.SetShaderTexture(Frequency,
+											GetMetalSurfaceFromRHITexture(Texture)->Texture.get(),
+											Index,
+											MTL::ResourceUsage(MTL::ResourceUsageRead | MTL::ResourceUsageSample));
 			}
+		}
 
-			void SetSampler(FRHISamplerState* Sampler, uint8 Index)
+		void SetSampler(FRHISamplerState* Sampler, uint8 Index)
+		{
+			if(!IsMetalBindlessEnabled())
 			{
-				if(!IsMetalBindlessEnabled())
-				{
-					StateCache.SetShaderSamplerState(Frequency, static_cast<FMetalSamplerState*>(Sampler), Index);
-				}
+				StateCache.SetShaderSamplerState(Frequency, static_cast<FMetalSamplerState*>(Sampler), Index);
 			}
+		}
 
 #if PLATFORM_SUPPORTS_BINDLESS_RENDERING
-			void SetResourceCollection(FRHIResourceCollection* ResourceCollection, uint32 Index)
-			{
-				FMetalResourceCollection* MetalResourceCollection = ResourceCast(ResourceCollection);
-				SetSRV(MetalResourceCollection->GetShaderResourceView(), Index);
-			}
+		void SetResourceCollection(FRHIResourceCollection* ResourceCollection, uint32 Index)
+		{
+			FMetalResourceCollection* MetalResourceCollection = ResourceCast(ResourceCollection);
+			SetSRV(MetalResourceCollection->GetShaderResourceView(), Index);
+		}
 #endif
-		};
+	};
 
-		UE::RHI::Private::SetUniformBufferResourcesFromTables(
-			  FUniformResourceBinder { *this, Frequency }
-			, *Shader
-			, DirtyUniformBuffers[Frequency]
-			, BoundUniformBuffers[Frequency]
+	UE::RHI::Private::SetUniformBufferResourcesFromTables(
+		  FUniformResourceBinder { *this, Frequency }
+		, *Shader
+		, DirtyUniformBuffers[Frequency]
+		, BoundUniformBuffers[Frequency]
 #if ENABLE_RHI_VALIDATION
-			, nullptr /*Tracker*/ // @todo: the current structure of the Metal RHI prevents easily passing the RHI validation layer tracker here
+		, nullptr /*Tracker*/ // @todo: the current structure of the Metal RHI prevents easily passing the RHI validation layer tracker here
 #endif
-		);
-	}
-	else
-	{
-		DirtyUniformBuffers[Frequency] = 0;
-	}
+	);
 }
 
 #if METAL_USE_METAL_SHADER_CONVERTER
diff --git a/Engine/Source/Runtime/Apple/MetalRHI/Private/MetalUAV.cpp b/Engine/Source/Runtime/Apple/MetalRHI/Private/MetalUAV.cpp
--- a/Engine/Source/Runtime/Apple/MetalRHI/Private/MetalUAV.cpp
+++ b/Engine/Source/Runtime/Apple/MetalRHI/Private/MetalUAV.cpp
@@ -217,8 +217,6 @@ void FMetalShaderResourceView::UpdateView(FMetalRHICommandContext* Context, cons
 		{
 		case FRHIViewDesc::EBufferType::Typed:
 			{
-				check(Device.SupportsFeature(EMetalFeaturesTextureBuffers));
-
 				MTL::PixelFormat Format = (MTL::PixelFormat)GMetalBufferFormats[Info.Format].LinearTextureFormat;
 				NS::UInteger Options = ((NS::UInteger)Buffer->Mode) << MTL::ResourceStorageModeShift;
 
@@ -444,8 +442,6 @@ void FMetalUnorderedAccessView::UpdateView(FMetalRHICommandContext* Context, con
 			{
 			case FRHIViewDesc::EBufferType::Typed:
 			{
-				check(Device.SupportsFeature(EMetalFeaturesTextureBuffers));
-
 				MTL::PixelFormat Format = (MTL::PixelFormat)GMetalBufferFormats[Info.Format].LinearTextureFormat;
                 NS::UInteger Options = ((NS::UInteger)Buffer->Mode) << MTL::ResourceStorageModeShift;
 
diff --git a/Engine/Source/Runtime/Apple/MetalRHI/Private/Shaders/Types/MetalComputeShader.cpp b/Engine/Source/Runtime/Apple/MetalRHI/Private/Shaders/Types/MetalComputeShader.cpp
--- a/Engine/Source/Runtime/Apple/MetalRHI/Private/Shaders/Types/MetalComputeShader.cpp
+++ b/Engine/Source/Runtime/Apple/MetalRHI/Private/Shaders/Types/MetalComputeShader.cpp
@@ -58,33 +58,27 @@ FMetalShaderPipelinePtr FMetalComputeShader::GetPipeline()
 		Descriptor->setLabel(Func->name());
 		Descriptor->setComputeFunction(Func.get());
         
-		if (Device.SupportsFeature(EMetalFeaturesTextureBuffers))
-		{
-			Descriptor->setMaxTotalThreadsPerThreadgroup(NumThreadsX*NumThreadsY*NumThreadsZ);
-		}
+		Descriptor->setMaxTotalThreadsPerThreadgroup(NumThreadsX*NumThreadsY*NumThreadsZ);
 
-		if (Device.SupportsFeature(EMetalFeaturesPipelineBufferMutability))
+		MTL::PipelineBufferDescriptorArray* PipelineBuffers = Descriptor->buffers();
+
+		uint32 ImmutableBuffers = Bindings.ConstantBuffers | Bindings.ArgumentBuffers;
+		while(ImmutableBuffers)
 		{
-			MTL::PipelineBufferDescriptorArray* PipelineBuffers = Descriptor->buffers();
+			uint32 Index = __builtin_ctz(ImmutableBuffers);
+			ImmutableBuffers &= ~(1 << Index);
 
-			uint32 ImmutableBuffers = Bindings.ConstantBuffers | Bindings.ArgumentBuffers;
-			while(ImmutableBuffers)
+			if (Index < ML_MaxBuffers)
 			{
-				uint32 Index = __builtin_ctz(ImmutableBuffers);
-				ImmutableBuffers &= ~(1 << Index);
-
-				if (Index < ML_MaxBuffers)
-				{
-					MTL::PipelineBufferDescriptor* PipelineBuffer = PipelineBuffers->object(Index);
-					PipelineBuffer->setMutability(MTL::MutabilityImmutable);
-				}
-			}
-			if (SideTableBinding > 0)
-			{
-				MTL::PipelineBufferDescriptor* PipelineBuffer = PipelineBuffers->object(SideTableBinding);
+				MTL::PipelineBufferDescriptor* PipelineBuffer = PipelineBuffers->object(Index);
 				PipelineBuffer->setMutability(MTL::MutabilityImmutable);
 			}
 		}
+		if (SideTableBinding > 0)
+		{
+			MTL::PipelineBufferDescriptor* PipelineBuffer = PipelineBuffers->object(SideTableBinding);
+			PipelineBuffer->setMutability(MTL::MutabilityImmutable);
+		}
 
         MTLComputePipelineStatePtr Kernel;
 		MTL::ComputePipelineReflection* Reflection = nullptr;
diff --git a/Engine/Source/Runtime/Apple/MetalRHI/Private/Shaders/Types/Templates/MetalBaseShader.h b/Engine/Source/Runtime/Apple/MetalRHI/Private/Shaders/Types/Templates/MetalBaseShader.h
--- a/Engine/Source/Runtime/Apple/MetalRHI/Private/Shaders/Types/Templates/MetalBaseShader.h
+++ b/Engine/Source/Runtime/Apple/MetalRHI/Private/Shaders/Types/Templates/MetalBaseShader.h
@@ -528,30 +528,6 @@ MTLFunctionPtr TMetalBaseShader<BaseResourceType, ShaderType>::GetCompiledFuncti
 		}
 	}
 
-	if (Device.SupportsFeature(EMetalFeaturesIABs) && Bindings.ArgumentBuffers && ArgumentEncoders.Num() == 0)
-	{
-		uint32 ArgumentBuffers = Bindings.ArgumentBuffers;
-		while(ArgumentBuffers)
-		{
-			uint32 Index = __builtin_ctz(ArgumentBuffers);
-			ArgumentBuffers &= ~(1 << Index);
-
-			MTL::ArgumentEncoder* ArgumentEncoder = Function->newArgumentEncoder(Index);
-			ArgumentEncoders.Add(Index, ArgumentEncoder);
-
-			TBitArray<> Resources;
-			for (uint8 Id : Bindings.ArgumentBufferMasks[Index])
-			{
-				if (Id >= Resources.Num())
-				{
-					Resources.Add(false, (Id + 1) - Resources.Num());
-				}
-				Resources[Id] = true;
-			}
-			ArgumentBitmasks.Add(Index, Resources);
-		}
-	}
-
 	check(Func);
 	return Func;
 }
diff --git a/Engine/Source/Runtime/Engine/Private/ShaderCompiler/ShaderCompiler.cpp b/Engine/Source/Runtime/Engine/Private/ShaderCompiler/ShaderCompiler.cpp
--- a/Engine/Source/Runtime/Engine/Private/ShaderCompiler/ShaderCompiler.cpp
+++ b/Engine/Source/Runtime/Engine/Private/ShaderCompiler/ShaderCompiler.cpp
@@ -3578,7 +3578,6 @@ void GlobalBeginCompileShader(
 			
 			bool bAllowFastIntrinsics = false;
 			bool bForceFloats = false;
-			int32 IndirectArgumentTier = 0;
 			bool bEnableMathOptimisations = true;
             bool bSupportAppleA8 = false;
 			bool bMetalOptimizeForSize = false;
@@ -3587,7 +3586,6 @@ void GlobalBeginCompileShader(
 			{
 				GConfig->GetBool(TEXT("/Script/MacTargetPlatform.MacTargetSettings"), TEXT("UseFastIntrinsics"), bAllowFastIntrinsics, GEngineIni);
 				GConfig->GetBool(TEXT("/Script/MacTargetPlatform.MacTargetSettings"), TEXT("EnableMathOptimisations"), bEnableMathOptimisations, GEngineIni);
-				GConfig->GetInt(TEXT("/Script/MacTargetPlatform.MacTargetSettings"), TEXT("IndirectArgumentTier"), IndirectArgumentTier, GEngineIni);
 				GConfig->GetBool(TEXT("/Script/MacTargetPlatform.MacTargetSettings"), TEXT("MetalOptimizeForSize"), bMetalOptimizeForSize, GEngineIni);
                 
                 // No half precision support on MacOS at the moment
@@ -3598,7 +3596,6 @@ void GlobalBeginCompileShader(
 				GConfig->GetBool(TEXT("/Script/IOSRuntimeSettings.IOSRuntimeSettings"), TEXT("UseFastIntrinsics"), bAllowFastIntrinsics, GEngineIni);
 				GConfig->GetBool(TEXT("/Script/IOSRuntimeSettings.IOSRuntimeSettings"), TEXT("EnableMathOptimisations"), bEnableMathOptimisations, GEngineIni);
 				GConfig->GetBool(TEXT("/Script/IOSRuntimeSettings.IOSRuntimeSettings"), TEXT("ForceFloats"), bForceFloats, GEngineIni);
-				GConfig->GetInt(TEXT("/Script/IOSRuntimeSettings.IOSRuntimeSettings"), TEXT("IndirectArgumentTier"), IndirectArgumentTier, GEngineIni);
                 GConfig->GetBool(TEXT("/Script/IOSRuntimeSettings.IOSRuntimeSettings"), TEXT("bSupportAppleA8"), bSupportAppleA8, GEngineIni);
 				GConfig->GetBool(TEXT("/Script/IOSRuntimeSettings.IOSRuntimeSettings"), TEXT("MetalOptimizeForSize"), bMetalOptimizeForSize, GEngineIni);
                 
@@ -3609,7 +3606,6 @@ void GlobalBeginCompileShader(
             Input.Environment.FullPrecisionInPS |= bForceFloats;
             
 			Input.Environment.SetCompileArgument(TEXT("METAL_USE_FAST_INTRINSICS"), bAllowFastIntrinsics);
-			Input.Environment.SetCompileArgument(TEXT("METAL_INDIRECT_ARGUMENT_BUFFERS"), IndirectArgumentTier);
             Input.Environment.SetCompileArgument(TEXT("SUPPORT_APPLE_A8"), bSupportAppleA8);
             Input.Environment.SetCompileArgument(TEXT("METAL_OPTIMIZE_FOR_SIZE"), bMetalOptimizeForSize);
 			
diff --git a/Engine/Source/Runtime/IOS/IOSRuntimeSettings/Classes/IOSRuntimeSettings.h b/Engine/Source/Runtime/IOS/IOSRuntimeSettings/Classes/IOSRuntimeSettings.h
--- a/Engine/Source/Runtime/IOS/IOSRuntimeSettings/Classes/IOSRuntimeSettings.h
+++ b/Engine/Source/Runtime/IOS/IOSRuntimeSettings/Classes/IOSRuntimeSettings.h
@@ -505,10 +505,6 @@ public:
 	UPROPERTY(EditAnywhere, config, Category=Rendering, meta = (DisplayName = "Enable Fast-Math optimisations", ConfigRestartRequired = true))
 	bool EnableMathOptimisations;
 	
-	/** Whether to compile shaders using a tier Indirect Argument Buffers. */
-	UPROPERTY(config, EditAnywhere, Category = Rendering, Meta = (DisplayName = "Tier of Indirect Argument Buffers to use when compiling shaders", ConfigRestartRequired = true))
-	int32 IndirectArgumentTier;
-	
     /** Supports Apple A8 devices.
      * Disables 3d texture compression
      * Virtual Textures are not supported on A8 devices due to missing hardware features
diff --git a/Engine/Source/Runtime/RenderCore/Private/Shader.cpp b/Engine/Source/Runtime/RenderCore/Private/Shader.cpp
--- a/Engine/Source/Runtime/RenderCore/Private/Shader.cpp
+++ b/Engine/Source/Runtime/RenderCore/Private/Shader.cpp
@@ -2173,14 +2173,12 @@ void ShaderMapAppendKey(EShaderPlatform Platform, FShaderKeyGenerator& KeyGen)
 		bool bEnableMathOptimisations = true;
 		bool bForceFloats = false;
         bool bSupportAppleA8 = false;
-		int32 IndirectArgumentTier = 0;
 		bool bMetalOptimizeBySize = false;
 
 		if (IsPCPlatform(Platform))
 		{
 			GConfig->GetBool(TEXT("/Script/MacTargetPlatform.MacTargetSettings"), TEXT("UseFastIntrinsics"), bAllowFastIntrinsics, GEngineIni);
 			GConfig->GetBool(TEXT("/Script/MacTargetPlatform.MacTargetSettings"), TEXT("EnableMathOptimisations"), bEnableMathOptimisations, GEngineIni);
-			GConfig->GetInt(TEXT("/Script/MacTargetPlatform.MacTargetSettings"), TEXT("IndirectArgumentTier"), IndirectArgumentTier, GEngineIni);
 			GConfig->GetBool(TEXT("/Script/MacTargetPlatform.MacTargetSettings"), TEXT("MetalOptimizeBySize"), bMetalOptimizeBySize, GEngineIni);
 		}
 		else
@@ -2188,7 +2186,6 @@ void ShaderMapAppendKey(EShaderPlatform Platform, FShaderKeyGenerator& KeyGen)
 			GConfig->GetBool(TEXT("/Script/IOSRuntimeSettings.IOSRuntimeSettings"), TEXT("UseFastIntrinsics"), bAllowFastIntrinsics, GEngineIni);
 			GConfig->GetBool(TEXT("/Script/IOSRuntimeSettings.IOSRuntimeSettings"), TEXT("EnableMathOptimisations"), bEnableMathOptimisations, GEngineIni);
 			GConfig->GetBool(TEXT("/Script/IOSRuntimeSettings.IOSRuntimeSettings"), TEXT("ForceFloats"), bForceFloats, GEngineIni);
-			GConfig->GetInt(TEXT("/Script/IOSRuntimeSettings.IOSRuntimeSettings"), TEXT("IndirectArgumentTier"), IndirectArgumentTier, GEngineIni);
             GConfig->GetBool(TEXT("/Script/IOSRuntimeSettings.IOSRuntimeSettings"), TEXT("bSupportAppleA8"), bSupportAppleA8, GEngineIni);
 			GConfig->GetBool(TEXT("/Script/IOSRuntimeSettings.IOSRuntimeSettings"), TEXT("MetalOptimizeBySize"), bMetalOptimizeBySize, GEngineIni);
 		}
@@ -2226,7 +2223,6 @@ void ShaderMapAppendKey(EShaderPlatform Platform, FShaderKeyGenerator& KeyGen)
         
 		KeyGen.AppendSeparator();
 		KeyGen.Append(TEXT("IAB"));
-		KeyGen.Append(IndirectArgumentTier);
 		
 		// Shaders built for archiving - for Metal that requires compiling the code in a different way so that we can strip it later
 		bool bArchive = false;
