Re-enable support for devices on Mac2 Platform

diff --git a/Engine/Source/Runtime/Apple/MetalRHI/Private/MetalDevice.cpp b/Engine/Source/Runtime/Apple/MetalRHI/Private/MetalDevice.cpp
--- a/Engine/Source/Runtime/Apple/MetalRHI/Private/MetalDevice.cpp
+++ b/Engine/Source/Runtime/Apple/MetalRHI/Private/MetalDevice.cpp
@@ -447,6 +447,8 @@ void FMetalDevice::EnumerateFeatureSupport()
 	FString DeviceName(Device->name()->cString(NS::UTF8StringEncoding));
 
 	bool bSupportsMetal3 = Device->supportsFamily(MTL::GPUFamilyMetal3);
+	bool bSupportsMac2 = Device->supportsFamily(MTL::GPUFamilyMac2);
+	bool bSupportsMac2Minimum = bSupportsMac2 || bSupportsMetal3;
 	
 	if (!GForceNoMetalFence && FParse::Param(FCommandLine::Get(),TEXT("metalfence")))
 	{
@@ -461,17 +463,17 @@ void FMetalDevice::EnumerateFeatureSupport()
 		Features |= EMetalFeaturesHeaps;
 	}
 	
-	if(bSupportsMetal3 || Device->supportsFamily(MTL::GPUFamilyApple3))
+	if(bSupportsMac2Minimum || Device->supportsFamily(MTL::GPUFamilyApple3))
 	{
 		Features |= EMetalFeaturesCountingQueries | EMetalFeaturesBaseVertexInstance | EMetalFeaturesIndirectBuffer | EMetalFeaturesMSAADepthResolve | EMetalFeaturesMSAAStoreAndResolve;
 	}
 	
-	if(bSupportsMetal3 || Device->supportsFamily(MTL::GPUFamilyApple4))
+	if(bSupportsMac2Minimum || Device->supportsFamily(MTL::GPUFamilyApple4))
 	{
 		Features |= EMetalFeaturesCubemapArrays;
 	}
 	
-	if(bSupportsMetal3 || Device->supportsFamily(MTL::GPUFamilyApple5))
+	if(bSupportsMac2Minimum || Device->supportsFamily(MTL::GPUFamilyApple5))
 	{
 		Features |= EMetalFeaturesLayeredRendering;
 	}
