Adding ERHIDescriptorType for use when allocating descriptors instead of ERHIDescriptorHeapType.

diff --git a/Engine/Source/Runtime/Apple/MetalRHI/Private/MetalBindlessDescriptors.cpp b/Engine/Source/Runtime/Apple/MetalRHI/Private/MetalBindlessDescriptors.cpp
--- a/Engine/Source/Runtime/Apple/MetalRHI/Private/MetalBindlessDescriptors.cpp
+++ b/Engine/Source/Runtime/Apple/MetalRHI/Private/MetalBindlessDescriptors.cpp
@@ -37,10 +37,9 @@ static FAutoConsoleVariableRef CVarBindlessUpdateDescriptorOnCompute(
 	ECVF_ReadOnly
 );
 
-FMetalDescriptorHeap::FMetalDescriptorHeap(FMetalDevice& MetalDevice, const ERHIDescriptorHeapType DescriptorType)
+FMetalDescriptorHeap::FMetalDescriptorHeap(FMetalDevice& MetalDevice, ERHIDescriptorTypeMask InTypeMask)
 	: Device(MetalDevice)
-	, ResourceHeap(nullptr)
-	, Type(DescriptorType)
+	, TypeMask(InTypeMask)
 {
 }
 
@@ -48,7 +47,7 @@ FMetalDescriptorHeap::~FMetalDescriptorHeap()
 {
 }
 
-void FMetalDescriptorHeap::Init(FRHICommandListBase& RHICmdList, const int32 HeapSize)
+void FMetalDescriptorHeap::Init(FRHICommandListBase& RHICmdList, int32 HeapSize)
 {
 	const FRHIBufferCreateDesc CreateDesc =
 		FRHIBufferCreateDesc::CreateStructured(TEXT("ResourceHeap"), HeapSize, 1)
@@ -90,10 +89,10 @@ uint32 FMetalDescriptorHeap::GetFreeResourceIndex()
 	return ResourceIndex;
 }
 
-FRHIDescriptorHandle FMetalDescriptorHeap::ReserveDescriptor()
+FRHIDescriptorHandle FMetalDescriptorHeap::ReserveDescriptor(ERHIDescriptorType DescriptorType)
 {
 	const uint32 ResourceIndex = GetFreeResourceIndex();
-	return FRHIDescriptorHandle(Type, ResourceIndex);
+	return FRHIDescriptorHandle(DescriptorType, ResourceIndex);
 }
 
 void FMetalDescriptorHeap::UpdateDescriptor(FRHIDescriptorHandle DescriptorHandle, IRDescriptorTableEntry DescriptorData)
@@ -104,7 +103,7 @@ void FMetalDescriptorHeap::UpdateDescriptor(FRHIDescriptorHandle DescriptorHandl
 	Descriptors[DescriptorIndex] = DescriptorData;
 }
 
-void FMetalDescriptorHeap::BindHeap(FMetalCommandEncoder* Encoder, MTL::FunctionType FunctionType, const uint32 BindIndex)
+void FMetalDescriptorHeap::BindHeap(FMetalCommandEncoder* Encoder, MTL::FunctionType FunctionType, uint32 BindIndex)
 {
 	uint32 DescriptorCount = PeakDescriptorCount.load();
 	const uint64 HeapSize = DescriptorCount * sizeof(IRDescriptorTableEntry);
@@ -115,8 +114,8 @@ void FMetalDescriptorHeap::BindHeap(FMetalCommandEncoder* Encoder, MTL::Function
 
 FMetalBindlessDescriptorManager::FMetalBindlessDescriptorManager(FMetalDevice& MetalDevice)
 	: Device(MetalDevice)
-	, StandardResources(Device, ERHIDescriptorHeapType::Standard)
-	, SamplerResources(Device, ERHIDescriptorHeapType::Sampler)
+	, StandardResources(Device, ERHIDescriptorTypeMask::Standard)
+	, SamplerResources(Device, ERHIDescriptorTypeMask::Sampler)
 {
 }
 
@@ -135,17 +134,17 @@ void FMetalBindlessDescriptorManager::Init()
 	bIsSupported = true;
 }
 
-FRHIDescriptorHandle FMetalBindlessDescriptorManager::ReserveDescriptor(ERHIDescriptorHeapType InType)
+FRHIDescriptorHandle FMetalBindlessDescriptorManager::ReserveDescriptor(ERHIDescriptorType DescriptorType)
 {
-	switch (InType)
+	if (StandardResources.HandlesAllocation(DescriptorType))
 	{
-	case ERHIDescriptorHeapType::Standard:
-		return StandardResources.ReserveDescriptor();
-	case ERHIDescriptorHeapType::Sampler:
-		return SamplerResources.ReserveDescriptor();
-	default:
-		checkNoEntry();
-	};
+		return StandardResources.ReserveDescriptor(DescriptorType);
+	}
+
+	if (SamplerResources.HandlesAllocation(DescriptorType))
+	{
+		return SamplerResources.ReserveDescriptor(DescriptorType);
+	}
 
 	return FRHIDescriptorHandle();
 }
@@ -153,17 +152,20 @@ FRHIDescriptorHandle FMetalBindlessDescriptorManager::ReserveDescriptor(ERHIDesc
 void FMetalBindlessDescriptorManager::FreeDescriptor(FRHIDescriptorHandle DescriptorHandle)
 {
 	check(DescriptorHandle.IsValid());
-	switch (DescriptorHandle.GetType())
+
+	const ERHIDescriptorType DescriptorType = DescriptorHandle.GetType();
+
+	if (StandardResources.HandlesAllocation(DescriptorType))
 	{
-	case ERHIDescriptorHeapType::Standard:
-		StandardResources.FreeDescriptor(DescriptorHandle);
-		break;
-	case ERHIDescriptorHeapType::Sampler:
-		SamplerResources.FreeDescriptor(DescriptorHandle);
-		break;
-	default:
-		checkNoEntry();
-	};
+		return StandardResources.FreeDescriptor(DescriptorHandle);
+	}
+
+	if (SamplerResources.HandlesAllocation(DescriptorType))
+	{
+		return SamplerResources.FreeDescriptor(DescriptorHandle);
+	}
+
+	checkNoEntry();
 }
 
 void FMetalDescriptorHeap::BindSampler(FRHIDescriptorHandle DescriptorHandle, MTL::SamplerState* Sampler)
@@ -306,7 +308,7 @@ void FMetalBindlessDescriptorManager::UpdateDescriptorsWithCompute(FMetalRHIComm
 		return;
 	}
 
-	FMetalDescriptorHeap* Heap = new FMetalDescriptorHeap(Device, ERHIDescriptorHeapType::Standard);
+	FMetalDescriptorHeap* Heap = new FMetalDescriptorHeap(Device, ERHIDescriptorTypeMask::Standard);
 	
 	TRHICommandList_RecursiveHazardous<FMetalRHICommandContext> RHICmdList(Context);
 	Heap->Init(RHICmdList, 256);
@@ -324,13 +326,13 @@ void FMetalBindlessDescriptorManager::UpdateDescriptorsWithCompute(FMetalRHIComm
 	FMemory::Memcpy(DescriptorIndiciesBuffer->Contents(), ComputeDescriptorIndicesCopy.GetResourceData(), ComputeDescriptorIndicesCopy.GetResourceDataSize());
 	
 	// Create descriptor handles and bind buffers to descriptor heap
-	FRHIDescriptorHandle DescriptorEntriesHandle = Heap->ReserveDescriptor();
+	FRHIDescriptorHandle DescriptorEntriesHandle = Heap->ReserveDescriptor(ERHIDescriptorType::BufferSRV);
 	Heap->BindBufferImmediate(DescriptorEntriesHandle, DescriptorEntriesBuffer);
 	
-	FRHIDescriptorHandle DescriptorIndicesHandle = Heap->ReserveDescriptor();
+	FRHIDescriptorHandle DescriptorIndicesHandle = Heap->ReserveDescriptor(ERHIDescriptorType::BufferSRV);
 	Heap->BindBufferImmediate(DescriptorIndicesHandle, DescriptorIndiciesBuffer);
 	
-	FRHIDescriptorHandle DestBufferHandle = Heap->ReserveDescriptor();
+	FRHIDescriptorHandle DestBufferHandle = Heap->ReserveDescriptor(ERHIDescriptorType::BufferUAV);
 	
 	FMetalRHIBuffer* DestBuffer = (FMetalRHIBuffer*)StandardResources.ResourceHeap.GetReference();
 	Heap->BindBufferImmediate(DestBufferHandle, DestBuffer->GetCurrentBuffer());
diff --git a/Engine/Source/Runtime/Apple/MetalRHI/Private/MetalBindlessDescriptors.h b/Engine/Source/Runtime/Apple/MetalRHI/Private/MetalBindlessDescriptors.h
--- a/Engine/Source/Runtime/Apple/MetalRHI/Private/MetalBindlessDescriptors.h
+++ b/Engine/Source/Runtime/Apple/MetalRHI/Private/MetalBindlessDescriptors.h
@@ -13,6 +13,8 @@
 class FMetalCommandEncoder;
 class FMetalDevice;
 
+struct IRDescriptorTableEntry;
+
 enum EDescriptorUpdateType
 {
 	EDescriptorUpdateType_Immediate,
@@ -21,23 +23,33 @@ enum EDescriptorUpdateType
 
 struct FMetalDescriptorHeap
 {
-	FMetalDescriptorHeap(FMetalDevice& MetalDevice, const ERHIDescriptorHeapType DescriptorType);
+	FMetalDescriptorHeap(FMetalDevice& MetalDevice, ERHIDescriptorTypeMask InTypeMask);
 	~FMetalDescriptorHeap();
 	
-    void                    Init(FRHICommandListBase& RHICmdList, const int32 HeapSize);
+    void                    Init(FRHICommandListBase& RHICmdList, int32 HeapSize);
 
-    FRHIDescriptorHandle    ReserveDescriptor();
+    FRHIDescriptorHandle    ReserveDescriptor(ERHIDescriptorType DescriptorType);
     void                    FreeDescriptor(FRHIDescriptorHandle DescriptorHandle);
     uint32                  GetFreeResourceIndex();
 
     void                    UpdateDescriptor(FRHIDescriptorHandle DescriptorHandle, struct IRDescriptorTableEntry DescriptorData);
-    void                    BindHeap(FMetalCommandEncoder* Encoder, MTL::FunctionType FunctionType, const uint32 BindIndex);
+    void                    BindHeap(FMetalCommandEncoder* Encoder, MTL::FunctionType FunctionType, uint32 BindIndex);
     
 	void                   	BindSampler(FRHIDescriptorHandle DescriptorHandle, MTL::SamplerState* Sampler);
 	void                    BindBufferImmediate(FRHIDescriptorHandle DescriptorHandle, FMetalBufferPtr Resource);
 	void                    BindResource(FRHIDescriptorHandle DescriptorHandle, FMetalResourceViewBase* Resource,
-												 FMetalRHICommandContext* Context, const EDescriptorUpdateType UpdateType = EDescriptorUpdateType_GPU);
+	                                     FMetalRHICommandContext* Context, const EDescriptorUpdateType UpdateType = EDescriptorUpdateType_GPU);
 	void                    BindTexture(FRHICommandListBase* RHICmdList, FRHIDescriptorHandle DescriptorHandle, MTL::Texture* Texture, const EDescriptorUpdateType UpdateType);
+
+	constexpr ERHIDescriptorTypeMask GetTypeMask() const
+	{
+		return TypeMask;
+	}
+
+	constexpr bool HandlesAllocation(ERHIDescriptorType InType) const
+	{
+		return EnumHasAnyFlags(GetTypeMask(), RHIDescriptorTypeMaskFromType(InType));
+	}
 	
     //
 	FMetalDevice& 					Device;
@@ -46,11 +58,11 @@ struct FMetalDescriptorHeap
     TQueue<uint32>                  FreeList;
 
     std::atomic<uint32>             PeakDescriptorCount;
-    struct IRDescriptorTableEntry*  Descriptors;
+    IRDescriptorTableEntry*         Descriptors = nullptr;
     uint32_t                        ResourceHeapLength;
     FBufferRHIRef			        ResourceHeap;
     
-    const ERHIDescriptorHeapType    Type;
+    const ERHIDescriptorTypeMask    TypeMask;
 };
 
 class FMetalBindlessDescriptorManager
@@ -61,7 +73,7 @@ public:
 
     void                    Init();
 
-    FRHIDescriptorHandle    ReserveDescriptor(ERHIDescriptorHeapType InType);
+    FRHIDescriptorHandle    ReserveDescriptor(ERHIDescriptorType DescriptorType);
     void                    FreeDescriptor(FRHIDescriptorHandle DescriptorHandle);
 
     void                    BindSampler(FRHIDescriptorHandle DescriptorHandle, MTL::SamplerState* Sampler);
diff --git a/Engine/Source/Runtime/Apple/MetalRHI/Private/MetalState.cpp b/Engine/Source/Runtime/Apple/MetalRHI/Private/MetalState.cpp
--- a/Engine/Source/Runtime/Apple/MetalRHI/Private/MetalState.cpp
+++ b/Engine/Source/Runtime/Apple/MetalRHI/Private/MetalState.cpp
@@ -309,7 +309,7 @@ FMetalSamplerState::FMetalSamplerState(FMetalDevice& MetalDevice, const FSampler
 
 	if(IsMetalBindlessEnabled())
 	{
-		BindlessHandle = BindlessDescriptorManager->ReserveDescriptor(ERHIDescriptorHeapType::Sampler);
+		BindlessHandle = BindlessDescriptorManager->ReserveDescriptor(ERHIDescriptorType::Sampler);
 		BindlessDescriptorManager->BindSampler(BindlessHandle, State);
 	}
 #endif
diff --git a/Engine/Source/Runtime/Apple/MetalRHI/Private/MetalTexture.cpp b/Engine/Source/Runtime/Apple/MetalRHI/Private/MetalTexture.cpp
--- a/Engine/Source/Runtime/Apple/MetalRHI/Private/MetalTexture.cpp
+++ b/Engine/Source/Runtime/Apple/MetalRHI/Private/MetalTexture.cpp
@@ -756,7 +756,7 @@ FMetalSurface::FMetalSurface(FMetalDevice& MetalDevice, FMetalTextureCreateDesc
 		FMetalBindlessDescriptorManager* BindlessDescriptorManager = Device.GetBindlessDescriptorManager();
 		check(BindlessDescriptorManager);
 
-		BindlessHandle = BindlessDescriptorManager->ReserveDescriptor(ERHIDescriptorHeapType::Standard);
+		BindlessHandle = BindlessDescriptorManager->ReserveDescriptor(ERHIDescriptorType::TextureSRV);
 	}
 #endif
 }
@@ -781,7 +781,7 @@ FMetalExternalSurface::FMetalExternalSurface(FMetalDevice& MetalDevice, MTLTextu
 		FMetalBindlessDescriptorManager* BindlessDescriptorManager = Device.GetBindlessDescriptorManager();
 		check(BindlessDescriptorManager);
 
-		BindlessHandle = BindlessDescriptorManager->ReserveDescriptor(ERHIDescriptorHeapType::Standard);
+		BindlessHandle = BindlessDescriptorManager->ReserveDescriptor(ERHIDescriptorType::TextureSRV);
 	}
 #endif
 }
@@ -1304,7 +1304,7 @@ MTLTexturePtr FMetalSurface::GetDrawableTexture()
 			FMetalBindlessDescriptorManager* BindlessDescriptorManager = Device.GetBindlessDescriptorManager();
 			check(BindlessDescriptorManager);
 			
-			BindlessHandle = BindlessDescriptorManager->ReserveDescriptor(ERHIDescriptorHeapType::Standard);
+			BindlessHandle = BindlessDescriptorManager->ReserveDescriptor(ERHIDescriptorType::TextureSRV);
 
 			BindlessDescriptorManager->BindTexture(nullptr, BindlessHandle, Texture.get(), EDescriptorUpdateType_Immediate);
 		}
diff --git a/Engine/Source/Runtime/Apple/MetalRHI/Private/MetalUAV.cpp b/Engine/Source/Runtime/Apple/MetalRHI/Private/MetalUAV.cpp
--- a/Engine/Source/Runtime/Apple/MetalRHI/Private/MetalUAV.cpp
+++ b/Engine/Source/Runtime/Apple/MetalRHI/Private/MetalUAV.cpp
@@ -90,7 +90,7 @@ FMetalShaderResourceView::FMetalShaderResourceView(FMetalDevice& InDevice, FRHIC
 
 	if(IsMetalBindlessEnabled())
 	{
-		BindlessHandle = BindlessDescriptorManager->ReserveDescriptor(ERHIDescriptorHeapType::Standard);
+		BindlessHandle = BindlessDescriptorManager->ReserveDescriptor(ERHIDescriptorType::TextureSRV);
 	}
 	
 	SurfaceOverride = nullptr;
@@ -388,7 +388,7 @@ FMetalUnorderedAccessView::FMetalUnorderedAccessView(FMetalDevice& InDevice, FRH
 
 	if(IsMetalBindlessEnabled())
 	{
-		BindlessHandle = BindlessDescriptorManager->ReserveDescriptor(ERHIDescriptorHeapType::Standard);
+		BindlessHandle = BindlessDescriptorManager->ReserveDescriptor(ERHIDescriptorType::TextureUAV);
 	}
 #endif
 
diff --git a/Engine/Source/Runtime/D3D12RHI/Private/D3D12BindlessDescriptors.cpp b/Engine/Source/Runtime/D3D12RHI/Private/D3D12BindlessDescriptors.cpp
--- a/Engine/Source/Runtime/D3D12RHI/Private/D3D12BindlessDescriptors.cpp
+++ b/Engine/Source/Runtime/D3D12RHI/Private/D3D12BindlessDescriptors.cpp
@@ -80,7 +80,7 @@ void FD3D12BindlessSamplerManager::CleanupResources()
 
 void FD3D12BindlessSamplerManager::InitializeDescriptor(FRHIDescriptorHandle DstHandle, FD3D12SamplerState* SamplerState)
 {
-	check(DstHandle.GetType() == ERHIDescriptorHeapType::Sampler);
+	check(DstHandle.GetType() == ERHIDescriptorType::Sampler);
 
 	UE::D3D12Descriptors::CopyDescriptor(GetParentDevice(), GpuHeap, DstHandle, SamplerState->OfflineDescriptor);
 }
@@ -138,7 +138,7 @@ void FD3D12BindlessDescriptorAllocator::Init()
 		NumResourceDescriptors += GBindlessOnlineDescriptorHeapBlockSize;
 #endif
 
-		ResourceAllocator = new FRHIHeapDescriptorAllocator(ERHIDescriptorHeapType::Standard, NumResourceDescriptors, ResourceStats);
+		ResourceAllocator = new FRHIHeapDescriptorAllocator(D3D12DescriptorTypeMaskFromHeapType(ERHIDescriptorHeapType::Standard), NumResourceDescriptors, ResourceStats);
 
 		const TStatId SamplerStats[] =
 		{
@@ -156,33 +156,25 @@ void FD3D12BindlessDescriptorAllocator::Init()
 			NumSamplerDescriptors = D3D12_MAX_SHADER_VISIBLE_SAMPLER_HEAP_SIZE;
 		}
 
-		SamplerAllocator = new FRHIHeapDescriptorAllocator(ERHIDescriptorHeapType::Sampler, NumSamplerDescriptors, SamplerStats);
+		SamplerAllocator = new FRHIHeapDescriptorAllocator(D3D12DescriptorTypeMaskFromHeapType(ERHIDescriptorHeapType::Sampler), NumSamplerDescriptors, SamplerStats);
 	}
 }
 
-FRHIDescriptorHandle FD3D12BindlessDescriptorAllocator::AllocateSamplerHandle()
+FRHIDescriptorHandle FD3D12BindlessDescriptorAllocator::AllocateDescriptor(ERHIDescriptorType DescriptorType)
 {
-	FRHIDescriptorHandle Result = SamplerAllocator->Allocate();
-	check(Result.IsValid());
-	return Result;
-}
-
-void FD3D12BindlessDescriptorAllocator::FreeSamplerHandle(FRHIDescriptorHandle InHandle)
-{
-	if (InHandle.IsValid())
+	if (DescriptorType == ERHIDescriptorType::Sampler)
 	{
-		SamplerAllocator->Free(InHandle);
+		FRHIDescriptorHandle Result = SamplerAllocator->Allocate(DescriptorType);
+		check(Result.IsValid());
+		return Result;
 	}
-}
 
-FRHIDescriptorHandle FD3D12BindlessDescriptorAllocator::AllocateResourceHandle()
-{
 	if (!AreResourcesBindless())
 	{
 		return FRHIDescriptorHandle();
 	}
 
-	FRHIDescriptorHandle Result = ResourceAllocator->Allocate();
+	FRHIDescriptorHandle Result = ResourceAllocator->Allocate(DescriptorType);
 
 #if D3D12RHI_BINDLESS_RESOURCE_MANAGER_SUPPORTS_RESIZING
 	if (!Result.IsValid())
@@ -200,7 +192,7 @@ FRHIDescriptorHandle FD3D12BindlessDescriptorAllocator::AllocateResourceHandle()
 			UE_LOG(LogD3D12RHI, Fatal, TEXT("Hit D3D12 device limits on descriptors when attempting to allocate a larger descriptor heap."));
 		}
 
-		Result = ResourceAllocator->ResizeGrowAndAllocate(NewNumDescriptors, ResourceAllocator->GetType());
+		Result = ResourceAllocator->ResizeGrowAndAllocate(NewNumDescriptors, DescriptorType);
 
 		// Grow the CPU heaps for all devices
 		for (FD3D12Device* ParentDevice : GetParentAdapter()->GetDevices())
@@ -217,11 +209,23 @@ FRHIDescriptorHandle FD3D12BindlessDescriptorAllocator::AllocateResourceHandle()
 	return Result;
 }
 
-void FD3D12BindlessDescriptorAllocator::FreeResourceHandle(FRHIDescriptorHandle InHandle)
+void FD3D12BindlessDescriptorAllocator::FreeDescriptor(FRHIDescriptorHandle Handle)
 {
-	if (InHandle.IsValid())
+	if (Handle.IsValid())
 	{
-		ResourceAllocator->Free(InHandle);
+		if (SamplerAllocator && SamplerAllocator->HandlesAllocation(Handle.GetType()))
+		{
+			SamplerAllocator->Free(Handle);
+		}
+		else if (ResourceAllocator && ResourceAllocator->HandlesAllocation(Handle.GetType()))
+		{
+			ResourceAllocator->Free(Handle);
+		}
+		else
+		{
+			// bad configuration?
+			checkNoEntry();
+		}
 	}
 }
 
@@ -760,20 +764,7 @@ void FD3D12BindlessDescriptorManager::Recycle(FD3D12DescriptorHeap* DescriptorHe
 
 void FD3D12BindlessDescriptorManager::ImmediateFree(FRHIDescriptorHandle InHandle)
 {
-	if (InHandle.GetType() == ERHIDescriptorHeapType::Standard && ResourceManager)
-	{
-		Allocator.FreeResourceHandle(InHandle);
-		return;
-	}
-
-	if (InHandle.GetType() == ERHIDescriptorHeapType::Sampler && SamplerManager)
-	{
-		Allocator.FreeSamplerHandle(InHandle);
-		return;
-	}
-
-	// Bad configuration?
-	checkNoEntry();
+	Allocator.FreeDescriptor(InHandle);
 }
 
 void FD3D12BindlessDescriptorManager::DeferredFreeFromDestructor(FRHIDescriptorHandle InHandle)
@@ -798,7 +789,7 @@ void FD3D12BindlessDescriptorManager::InitializeDescriptor(FRHIDescriptorHandle
 
 void FD3D12BindlessDescriptorManager::InitializeDescriptor(FRHIDescriptorHandle DstHandle, FD3D12View* View)
 {
-	if (DstHandle.GetType() == ERHIDescriptorHeapType::Standard && ResourceManager)
+	if (ResourceManager && ResourceManager->HandlesAllocation(DstHandle.GetType()))
 	{
 		ResourceManager->InitializeDescriptor(DstHandle, View);
 		return;
diff --git a/Engine/Source/Runtime/D3D12RHI/Private/D3D12BindlessDescriptors.h b/Engine/Source/Runtime/D3D12RHI/Private/D3D12BindlessDescriptors.h
--- a/Engine/Source/Runtime/D3D12RHI/Private/D3D12BindlessDescriptors.h
+++ b/Engine/Source/Runtime/D3D12RHI/Private/D3D12BindlessDescriptors.h
@@ -41,11 +41,9 @@ public:
 	bool AreSamplersBindless() const { return SamplerAllocator != nullptr; }
 
 	// Bindless descriptor allocators are stored in the adapter, so descriptor handles can be allocated once and shared for multi-GPU objects
-	FRHIDescriptorHandle AllocateSamplerHandle();
-	FRHIDescriptorHandle AllocateResourceHandle();
+	FRHIDescriptorHandle AllocateDescriptor(ERHIDescriptorType DescriptorType);
 
-	void FreeSamplerHandle(FRHIDescriptorHandle Handle);
-	void FreeResourceHandle(FRHIDescriptorHandle Handle);
+	void FreeDescriptor(FRHIDescriptorHandle Handle);
 
 	FCriticalSection& GetResourceHeapsCS() { return ResourceHeapsCS; }
 
@@ -88,6 +86,16 @@ public:
 	FD3D12DescriptorHeap* GetHeap() const { return GpuHeap.GetReference(); }
 	ERHIBindlessConfiguration GetConfiguration() const { return Configuration; }
 
+	constexpr ERHIDescriptorTypeMask GetTypeMask() const
+	{
+		return ERHIDescriptorTypeMask::Sampler;
+	}
+
+	constexpr bool HandlesAllocation(ERHIDescriptorType InType) const
+	{
+		return EnumHasAnyFlags(GetTypeMask(), RHIDescriptorTypeMaskFromType(InType));
+	}
+
 private:
 	FD3D12DescriptorHeapPtr      GpuHeap;
 	ERHIBindlessConfiguration    Configuration;
@@ -188,6 +196,16 @@ public:
 	// Called from FD3D12Adapter::AllocateBindlessResourceHandle
 	void GrowCPUHeap(uint32 OriginalNumDescriptors, uint32 NewNumDescriptors);
 
+	constexpr ERHIDescriptorTypeMask GetTypeMask() const
+	{
+		return ERHIDescriptorTypeMask::CBV | ERHIDescriptorTypeMask::SRV | ERHIDescriptorTypeMask::UAV;
+	}
+
+	constexpr bool HandlesAllocation(ERHIDescriptorType InType) const
+	{
+		return EnumHasAnyFlags(GetTypeMask(), RHIDescriptorTypeMaskFromType(InType));
+	}
+
 private:
 	void CopyCpuHeap(FD3D12DescriptorHeap* DestinationHeap);
 	void AssignHeapToState(FD3D12ContextBindlessState& State);
diff --git a/Engine/Source/Runtime/D3D12RHI/Private/D3D12Descriptors.cpp b/Engine/Source/Runtime/D3D12RHI/Private/D3D12Descriptors.cpp
--- a/Engine/Source/Runtime/D3D12RHI/Private/D3D12Descriptors.cpp
+++ b/Engine/Source/Runtime/D3D12RHI/Private/D3D12Descriptors.cpp
@@ -192,7 +192,7 @@ FD3D12DescriptorHeap::~FD3D12DescriptorHeap()
 
 FD3D12DescriptorManager::FD3D12DescriptorManager(FD3D12Device* Device, FD3D12DescriptorHeap* InHeap, TConstArrayView<TStatId> InStats)
 	: FD3D12DeviceChild(Device)
-	, FRHIHeapDescriptorAllocator(InHeap->GetType(), InHeap->GetNumDescriptors(), InStats)
+	, FRHIHeapDescriptorAllocator(InHeap->GetDescriptorTypeMask(), InHeap->GetNumDescriptors(), InStats)
 	, Heap(InHeap)
 {
 }
@@ -565,7 +565,7 @@ FD3D12DescriptorHeap* FD3D12DescriptorHeapManager::AllocateHeap(const TCHAR* InD
 {
 	for (FD3D12DescriptorManager& GlobalHeap : GlobalHeaps)
 	{
-		if (GlobalHeap.HandlesAllocationWithFlags(InHeapType, InHeapFlags))
+		if (GlobalHeap.HandlesHeapAllocationsWithFlags(InHeapType, InHeapFlags))
 		{
 			uint32 Offset = 0;
 			if (GlobalHeap.Allocate(InNumDescriptors, Offset))
diff --git a/Engine/Source/Runtime/D3D12RHI/Private/D3D12Descriptors.h b/Engine/Source/Runtime/D3D12RHI/Private/D3D12Descriptors.h
--- a/Engine/Source/Runtime/D3D12RHI/Private/D3D12Descriptors.h
+++ b/Engine/Source/Runtime/D3D12RHI/Private/D3D12Descriptors.h
@@ -28,6 +28,38 @@ inline D3D12_DESCRIPTOR_HEAP_TYPE Translate(ERHIDescriptorHeapType InHeapType)
 	}
 }
 
+inline D3D12_DESCRIPTOR_HEAP_TYPE Translate(ERHIDescriptorType InDescriptorType)
+{
+	const ERHIDescriptorTypeMask TypeMask = RHIDescriptorTypeMaskFromType(InDescriptorType);
+	if (EnumHasAnyFlags(TypeMask, ERHIDescriptorTypeMask::CBV | ERHIDescriptorTypeMask::SRV | ERHIDescriptorTypeMask::UAV))
+	{
+		return D3D12_DESCRIPTOR_HEAP_TYPE_CBV_SRV_UAV;
+	}
+
+	if (EnumHasAnyFlags(TypeMask, ERHIDescriptorTypeMask::Sampler))
+	{
+		return FD3D12_DESCRIPTOR_HEAP_TYPE_SAMPLER;
+	}
+
+	checkNoEntry();
+	return D3D12_DESCRIPTOR_HEAP_TYPE_CBV_SRV_UAV;
+}
+
+inline ERHIDescriptorTypeMask D3D12DescriptorTypeMaskFromHeapType(ERHIDescriptorHeapType InHeapType)
+{
+	if (InHeapType == ERHIDescriptorHeapType::Standard)
+	{
+		return ERHIDescriptorTypeMask::CBV | ERHIDescriptorTypeMask::SRV | ERHIDescriptorTypeMask::UAV;
+	}
+
+	if (InHeapType == ERHIDescriptorHeapType::Sampler)
+	{
+		return ERHIDescriptorTypeMask::Sampler;
+	}
+
+	return ERHIDescriptorTypeMask::None;
+}
+
 enum class ED3D12DescriptorHeapFlags : uint8
 {
 	None       = 0,
@@ -72,22 +104,27 @@ public:
 
 	~FD3D12DescriptorHeap();
 
-	inline ID3D12DescriptorHeap*       GetHeap()  const { return Heap; }
-	inline ERHIDescriptorHeapType      GetType()  const { return Type; }
-	inline ED3D12DescriptorHeapFlags   GetFlags() const { return Flags; }
+	ID3D12DescriptorHeap*       GetHeap()  const { return Heap; }
+	ERHIDescriptorHeapType      GetType()  const { return Type; }
+	ED3D12DescriptorHeapFlags   GetFlags() const { return Flags; }
+
+	ERHIDescriptorTypeMask GetDescriptorTypeMask() const
+	{
+		return D3D12DescriptorTypeMaskFromHeapType(GetType());
+	}
 
-	inline uint32 GetOffset()         const { return Offset; }
-	inline uint32 GetNumDescriptors() const { return NumDescriptors; }
-	inline uint32 GetDescriptorSize() const { return DescriptorSize; }
-	inline bool   IsGlobal()          const { return bIsGlobal; }
-	inline bool   IsSuballocation()   const { return bIsSuballocation; }
+	uint32 GetOffset()         const { return Offset; }
+	uint32 GetNumDescriptors() const { return NumDescriptors; }
+	uint32 GetDescriptorSize() const { return DescriptorSize; }
+	bool   IsGlobal()          const { return bIsGlobal; }
+	bool   IsSuballocation()   const { return bIsSuballocation; }
 
-	inline uint32 GetMemorySize() const { return DescriptorSize * NumDescriptors; }
+	uint32 GetMemorySize() const { return DescriptorSize * NumDescriptors; }
 
-	inline D3D12_CPU_DESCRIPTOR_HANDLE GetCPUSlotHandle(uint32 Slot) const { return CD3DX12_CPU_DESCRIPTOR_HANDLE(CpuBase, Slot, DescriptorSize); }
-	inline D3D12_GPU_DESCRIPTOR_HANDLE GetGPUSlotHandle(uint32 Slot) const { return CD3DX12_GPU_DESCRIPTOR_HANDLE(GpuBase, Slot, DescriptorSize); }
+	D3D12_CPU_DESCRIPTOR_HANDLE GetCPUSlotHandle(uint32 Slot) const { return CD3DX12_CPU_DESCRIPTOR_HANDLE(CpuBase, Slot, DescriptorSize); }
+	D3D12_GPU_DESCRIPTOR_HANDLE GetGPUSlotHandle(uint32 Slot) const { return CD3DX12_GPU_DESCRIPTOR_HANDLE(GpuBase, Slot, DescriptorSize); }
 
-	inline bool CanBePooled() const { return EnumHasAnyFlags(GetFlags(), ED3D12DescriptorHeapFlags::Poolable); }
+	bool CanBePooled() const { return EnumHasAnyFlags(GetFlags(), ED3D12DescriptorHeapFlags::Poolable); }
 
 private:
 	TRefCountPtr<ID3D12DescriptorHeap> Heap;
@@ -126,11 +163,17 @@ public:
 	inline       FD3D12DescriptorHeap* GetHeap()       { return Heap.GetReference(); }
 	inline const FD3D12DescriptorHeap* GetHeap() const { return Heap.GetReference(); }
 
-	inline bool HandlesAllocationWithFlags(ERHIDescriptorHeapType InHeapType, ED3D12DescriptorHeapFlags InHeapFlags) const
+	bool HandlesHeapAllocationsWithFlags(ERHIDescriptorHeapType InHeapType, ED3D12DescriptorHeapFlags InHeapFlags) const
+	{
+		return HandlesAllocations(D3D12DescriptorTypeMaskFromHeapType(InHeapType)) && Heap->GetFlags() == InHeapFlags;
+	}
+
+	bool HandlesAllocationWithFlags(ERHIDescriptorType InDescriptorType, ED3D12DescriptorHeapFlags InHeapFlags) const
 	{
-		return HandlesAllocation(InHeapType) && Heap->GetFlags() == InHeapFlags;
+		return HandlesAllocation(InDescriptorType) && Heap->GetFlags() == InHeapFlags;
 	}
 
+
 	inline bool IsHeapAChild(const FD3D12DescriptorHeap* InHeap)
 	{
 		return InHeap->GetHeap() == Heap->GetHeap();
diff --git a/Engine/Source/Runtime/D3D12RHI/Private/D3D12RayTracing.cpp b/Engine/Source/Runtime/D3D12RHI/Private/D3D12RayTracing.cpp
--- a/Engine/Source/Runtime/D3D12RHI/Private/D3D12RayTracing.cpp
+++ b/Engine/Source/Runtime/D3D12RHI/Private/D3D12RayTracing.cpp
@@ -3879,7 +3879,7 @@ void FD3D12RayTracingGeometry::AllocateBufferSRVs(uint32 InGPUIndex)
 		SRVDesc.Buffer.Flags = D3D12_BUFFER_SRV_FLAG_RAW;
 		SRVDesc.Buffer.StructureByteStride = 0;
 
-		HitGroupSystemIndexBufferSRV[InGPUIndex] = MakeShared<FD3D12ShaderResourceView>(GetParentAdapter()->GetDevice(InGPUIndex), InGPUIndex > 0 ? HitGroupSystemIndexBufferSRV[0].Get() : nullptr);
+		HitGroupSystemIndexBufferSRV[InGPUIndex] = MakeShared<FD3D12ShaderResourceView>(GetParentAdapter()->GetDevice(InGPUIndex), InGPUIndex > 0 ? HitGroupSystemIndexBufferSRV[0].Get() : nullptr, ERHIDescriptorType::BufferSRV);
 		HitGroupSystemIndexBufferSRV[InGPUIndex]->CreateView(IndexBuffer, SRVDesc, FD3D12ShaderResourceView::EFlags::None);
 	}
 
@@ -3914,7 +3914,7 @@ void FD3D12RayTracingGeometry::AllocateBufferSRVs(uint32 InGPUIndex)
 				FirstLinkedObject = HitGroupSystemSegmentVertexBufferSRVs[0][SegmentIndex].Get();
 			}
 		}
-		TSharedPtr<FD3D12ShaderResourceView> VertexBufferSRV = MakeShared<FD3D12ShaderResourceView>(GetParentAdapter()->GetDevice(InGPUIndex), FirstLinkedObject);
+		TSharedPtr<FD3D12ShaderResourceView> VertexBufferSRV = MakeShared<FD3D12ShaderResourceView>(GetParentAdapter()->GetDevice(InGPUIndex), FirstLinkedObject, ERHIDescriptorType::BufferSRV);
 		VertexBufferSRV->CreateView(VertexBuffer, SRVDesc, FD3D12ShaderResourceView::EFlags::None);
 		HitGroupSystemSegmentVertexBufferSRVs[InGPUIndex].Add(VertexBufferSRV);
 	}
diff --git a/Engine/Source/Runtime/D3D12RHI/Private/D3D12ResourceCollection.cpp b/Engine/Source/Runtime/D3D12RHI/Private/D3D12ResourceCollection.cpp
--- a/Engine/Source/Runtime/D3D12RHI/Private/D3D12ResourceCollection.cpp
+++ b/Engine/Source/Runtime/D3D12RHI/Private/D3D12ResourceCollection.cpp
@@ -56,7 +56,7 @@ FD3D12ResourceCollection::FD3D12ResourceCollection(FD3D12Device* InParent, FRHIC
 	SRVDesc.Buffer.NumElements = UE::RHICore::CalculateResourceCollectionMemorySize(InMembers) / 4;
 	SRVDesc.Buffer.Flags = D3D12_BUFFER_SRV_FLAG_RAW;
 
-	BufferSRV = MakeShared<FD3D12ShaderResourceView>(InParent, FirstLinkedObject ? FirstLinkedObject->BufferSRV.Get() : nullptr);
+	BufferSRV = MakeShared<FD3D12ShaderResourceView>(InParent, FirstLinkedObject ? FirstLinkedObject->BufferSRV.Get() : nullptr, ERHIDescriptorType::BufferSRV);
 	BufferSRV->CreateView(InBuffer, SRVDesc, FD3D12ShaderResourceView::EFlags::None);
 }
 
diff --git a/Engine/Source/Runtime/D3D12RHI/Private/D3D12SRV.cpp b/Engine/Source/Runtime/D3D12RHI/Private/D3D12SRV.cpp
--- a/Engine/Source/Runtime/D3D12RHI/Private/D3D12SRV.cpp
+++ b/Engine/Source/Runtime/D3D12RHI/Private/D3D12SRV.cpp
@@ -3,19 +3,8 @@
 #include "D3D12RHIPrivate.h"
 #include "D3D12RayTracing.h"
 
-// -----------------------------------------------------------------------------------------------------
-//
-//                                       FD3D12ShaderResourceView                                       
-//
-// -----------------------------------------------------------------------------------------------------
-
-FD3D12ShaderResourceView::FD3D12ShaderResourceView(FD3D12Device* InDevice, FD3D12ShaderResourceView* FirstLinkedObject)
-	: TD3D12View(InDevice, ED3D12ViewType::ShaderResource, ERHIDescriptorHeapType::Standard, FirstLinkedObject)
-{
-}
-
-FD3D12ShaderResourceView::FD3D12ShaderResourceView(FD3D12Device* InDevice, FD3D12ShaderResourceView* FirstLinkedObject, FD3D12RayTracingScene* InRayTracingScene)
-	: TD3D12View(InDevice, ED3D12ViewType::ShaderResource, ERHIDescriptorHeapType::Standard, FirstLinkedObject)
+FD3D12ShaderResourceView::FD3D12ShaderResourceView(FD3D12Device* InDevice, FD3D12ShaderResourceView* FirstLinkedObject, ERHIDescriptorType InDescriptorType, FD3D12RayTracingScene* InRayTracingScene)
+	: TD3D12View(InDevice, InDescriptorType, ERHIDescriptorHeapType::Standard, FirstLinkedObject)
 	, RayTracingScene(InRayTracingScene)
 {
 }
@@ -302,16 +291,23 @@ void FD3D12ShaderResourceView_RHI::UpdateView(FD3D12ContextArray const& Contexts
 	}
 }
 
-FD3D12ShaderResourceView_RHI::FD3D12ShaderResourceView_RHI(FD3D12Device* InDevice, FRHIViewableResource* InResource, FRHIViewDesc const& InViewDesc, FD3D12ShaderResourceView_RHI* FirstLinkedObject)
-	: FRHIShaderResourceView(InResource, InViewDesc)
+inline FD3D12RayTracingScene* GetRayTracingSceneFromViewDesc(const FRHIViewDesc& InViewDesc)
+{
 #if RHI_RAYTRACING
-	, FD3D12ShaderResourceView(InDevice, FirstLinkedObject, InViewDesc.Buffer.SRV.BufferType == FRHIViewDesc::EBufferType::AccelerationStructure ? FD3D12DynamicRHI::ResourceCast(InViewDesc.Buffer.SRV.RayTracingScene) : nullptr)
-#else
-	, FD3D12ShaderResourceView(InDevice, FirstLinkedObject)
+	if (InViewDesc.Buffer.SRV.BufferType == FRHIViewDesc::EBufferType::AccelerationStructure)
+	{
+		return FD3D12DynamicRHI::ResourceCast(InViewDesc.Buffer.SRV.RayTracingScene);
+	}
 #endif
-{}
 
+	return nullptr;
+}
 
+FD3D12ShaderResourceView_RHI::FD3D12ShaderResourceView_RHI(FD3D12Device* InDevice, FRHIViewableResource* InResource, FRHIViewDesc const& InViewDesc, FD3D12ShaderResourceView_RHI* FirstLinkedObject)
+	: FRHIShaderResourceView(InResource, InViewDesc)
+	, FD3D12ShaderResourceView(InDevice, FirstLinkedObject, RHIDescriptorTypeFromViewDesc(InViewDesc), GetRayTracingSceneFromViewDesc(InViewDesc))
+{
+}
 
 // -----------------------------------------------------------------------------------------------------
 //
diff --git a/Engine/Source/Runtime/D3D12RHI/Private/D3D12State.cpp b/Engine/Source/Runtime/D3D12RHI/Private/D3D12State.cpp
--- a/Engine/Source/Runtime/D3D12RHI/Private/D3D12State.cpp
+++ b/Engine/Source/Runtime/D3D12RHI/Private/D3D12State.cpp
@@ -810,7 +810,7 @@ FD3D12SamplerState::FD3D12SamplerState(FD3D12Device* InParent, const D3D12_SAMPL
 #if PLATFORM_SUPPORTS_BINDLESS_RENDERING
 	if (InParent->GetBindlessDescriptorAllocator().GetConfiguration() != ERHIBindlessConfiguration::Disabled)
 	{
-		BindlessHandle = FirstLinkedObject ? FirstLinkedObject->BindlessHandle : InParent->GetBindlessDescriptorAllocator().AllocateSamplerHandle();
+		BindlessHandle = FirstLinkedObject ? FirstLinkedObject->BindlessHandle : InParent->GetBindlessDescriptorAllocator().AllocateDescriptor(ERHIDescriptorType::Sampler);
 
 		InParent->GetBindlessDescriptorManager().InitializeDescriptor(BindlessHandle, this);
 	}
diff --git a/Engine/Source/Runtime/D3D12RHI/Private/D3D12Texture.h b/Engine/Source/Runtime/D3D12RHI/Private/D3D12Texture.h
--- a/Engine/Source/Runtime/D3D12RHI/Private/D3D12Texture.h
+++ b/Engine/Source/Runtime/D3D12RHI/Private/D3D12Texture.h
@@ -121,11 +121,13 @@ public:
 	{
 		check(!ShaderResourceView);
 
+		const ERHIDescriptorType DescriptorType = RHIDescriptorTypeFromD3D12ViewDesc(SRVDesc);
+
 		FD3D12ShaderResourceView::EFlags Flags = SkipsFastClearFinalize()
 			? FD3D12ShaderResourceView::EFlags::SkipFastClearFinalize
 			: FD3D12ShaderResourceView::EFlags::None;
 
-		ShaderResourceView = MakeShared<FD3D12ShaderResourceView>(GetParentDevice(), FirstLinkedObject ? FirstLinkedObject->ShaderResourceView.Get() : nullptr);
+		ShaderResourceView = MakeShared<FD3D12ShaderResourceView>(GetParentDevice(), FirstLinkedObject ? FirstLinkedObject->ShaderResourceView.Get() : nullptr, DescriptorType);
 		ShaderResourceView->CreateView(this, SRVDesc, Flags);
 	}
 
diff --git a/Engine/Source/Runtime/D3D12RHI/Private/D3D12TextureReference.cpp b/Engine/Source/Runtime/D3D12RHI/Private/D3D12TextureReference.cpp
--- a/Engine/Source/Runtime/D3D12RHI/Private/D3D12TextureReference.cpp
+++ b/Engine/Source/Runtime/D3D12RHI/Private/D3D12TextureReference.cpp
@@ -5,7 +5,7 @@
 FD3D12RHITextureReference::FD3D12RHITextureReference(FD3D12Device* InDevice, FD3D12Texture* InReferencedTexture, FD3D12RHITextureReference* FirstLinkedObject)
 	: FD3D12DeviceChild(InDevice)
 #if PLATFORM_SUPPORTS_BINDLESS_RENDERING
-	, FRHITextureReference(InReferencedTexture, FirstLinkedObject ? FirstLinkedObject->BindlessHandle : InDevice->GetBindlessDescriptorAllocator().AllocateResourceHandle())
+	, FRHITextureReference(InReferencedTexture, FirstLinkedObject ? FirstLinkedObject->BindlessHandle : InDevice->GetBindlessDescriptorAllocator().AllocateDescriptor(ERHIDescriptorType::TextureSRV))
 #else
 	, FRHITextureReference(InReferencedTexture)
 #endif
diff --git a/Engine/Source/Runtime/D3D12RHI/Private/D3D12UAV.cpp b/Engine/Source/Runtime/D3D12RHI/Private/D3D12UAV.cpp
--- a/Engine/Source/Runtime/D3D12RHI/Private/D3D12UAV.cpp
+++ b/Engine/Source/Runtime/D3D12RHI/Private/D3D12UAV.cpp
@@ -3,14 +3,8 @@
 #include "D3D12RHIPrivate.h"
 #include "ClearReplacementShaders.h"
 
-// -----------------------------------------------------------------------------------------------------
-//
-//                                      FD3D12UnorderedAccessView                                       
-//
-// -----------------------------------------------------------------------------------------------------
-
-FD3D12UnorderedAccessView::FD3D12UnorderedAccessView(FD3D12Device* InDevice, FD3D12UnorderedAccessView* FirstLinkedObject)
-	: TD3D12View(InDevice, ED3D12ViewType::UnorderedAccess, ERHIDescriptorHeapType::Standard, FirstLinkedObject)
+FD3D12UnorderedAccessView::FD3D12UnorderedAccessView(FD3D12Device* InDevice, FD3D12UnorderedAccessView* FirstLinkedObject, ERHIDescriptorType InDescriptorType)
+	: TD3D12View(InDevice, InDescriptorType, ERHIDescriptorHeapType::Standard, FirstLinkedObject)
 {
 }
 
@@ -256,10 +250,9 @@ void FD3D12UnorderedAccessView_RHI::UpdateView(FD3D12ContextArray const& Context
 
 FD3D12UnorderedAccessView_RHI::FD3D12UnorderedAccessView_RHI(FD3D12Device* InDevice, FRHIViewableResource* InResource, FRHIViewDesc const& InViewDesc, FD3D12UnorderedAccessView_RHI* FirstLinkedObject)
 	: FRHIUnorderedAccessView(InResource, InViewDesc)
-	, FD3D12UnorderedAccessView(InDevice, FirstLinkedObject)
-{}
-
-
+	, FD3D12UnorderedAccessView(InDevice, FirstLinkedObject, RHIDescriptorTypeFromViewDesc(InViewDesc))
+{
+}
 
 // -----------------------------------------------------------------------------------------------------
 //
@@ -374,7 +367,8 @@ void FD3D12CommandContext::ClearUAV(TRHICommandList_RecursiveHazardous<FD3D12Com
 				}
 
 				// Scoped view will free the offline CPU handle once we return
-				FD3D12UnorderedAccessView UAV(ParentDevice, nullptr);			// Always single GPU object, so FirstLinkedObject is nullptr
+				// Always single GPU object, so FirstLinkedObject is nullptr
+				FD3D12UnorderedAccessView UAV(ParentDevice, nullptr, ERHIDescriptorType::BufferUAV);
 				UAV.CreateView(UnorderedAccessView->GetResourceLocation(), R32UAVDesc, FD3D12UnorderedAccessView::EFlags::None);
 
 				FD3D12OfflineDescriptor OfflineHandle = UAV.GetOfflineCpuHandle();
diff --git a/Engine/Source/Runtime/D3D12RHI/Private/D3D12View.cpp b/Engine/Source/Runtime/D3D12RHI/Private/D3D12View.cpp
--- a/Engine/Source/Runtime/D3D12RHI/Private/D3D12View.cpp
+++ b/Engine/Source/Runtime/D3D12RHI/Private/D3D12View.cpp
@@ -220,7 +220,7 @@ inline bool ViewTypeSupportsBindless(ED3D12ViewType InViewType)
 }
 #endif
 
-FD3D12View::FD3D12View(FD3D12Device* InDevice, ED3D12ViewType InViewType, ERHIDescriptorHeapType InHeapType, FD3D12View* FirstLinkedObject)
+FD3D12View::FD3D12View(FD3D12Device* InDevice, ERHIDescriptorType InDescriptorType, ERHIDescriptorHeapType InHeapType, FD3D12View* FirstLinkedObject)
 	: FD3D12DeviceChild(InDevice)
 	, OfflineCpuHandle(InDevice->GetOfflineDescriptorManager(InHeapType).AllocateHeapSlot())
 #if PLATFORM_SUPPORTS_BINDLESS_RENDERING
@@ -229,10 +229,10 @@ FD3D12View::FD3D12View(FD3D12Device* InDevice, ED3D12ViewType InViewType, ERHIDe
 	, HeapType(InHeapType)
 {
 #if PLATFORM_SUPPORTS_BINDLESS_RENDERING
-	if (ViewTypeSupportsBindless(InViewType))
+	if (InDescriptorType != ERHIDescriptorType::Invalid)
 	{
 		// Bindless handle is shared across linked objects, so pull the handle from the first object if available, or allocate if we are the first
-		BindlessHandle = FirstLinkedObject ? FirstLinkedObject->BindlessHandle : InDevice->GetBindlessDescriptorAllocator().AllocateResourceHandle();
+		BindlessHandle = FirstLinkedObject ? FirstLinkedObject->BindlessHandle : InDevice->GetBindlessDescriptorAllocator().AllocateDescriptor(InDescriptorType);
 	}
 #endif
 }
@@ -276,7 +276,6 @@ void FD3D12View::UpdateBindlessSlot(FD3D12ContextArray const& Contexts)
 #if PLATFORM_SUPPORTS_BINDLESS_RENDERING
 	if (BindlessHandle.IsValid())
 	{
-		check(BindlessHandle.GetType() == ERHIDescriptorHeapType::Standard);
 		GetParentDevice()->GetBindlessDescriptorManager().UpdateDescriptor(Contexts, BindlessHandle, this);
 	}
 #endif
@@ -339,7 +338,7 @@ void FD3D12View::UpdateView(FD3D12ContextArray const& Contexts, const FResourceI
 // -----------------------------------------------------------------------------------------------------
 
 FD3D12ConstantBufferView::FD3D12ConstantBufferView(FD3D12Device* InDevice, FD3D12ConstantBufferView* FirstLinkedObject)
-	: TD3D12View(InDevice, ED3D12ViewType::ConstantBuffer, ERHIDescriptorHeapType::Standard, FirstLinkedObject)
+	: TD3D12View(InDevice, ERHIDescriptorType::CBV, ERHIDescriptorHeapType::Standard, FirstLinkedObject)
 {
 }
 
@@ -376,7 +375,7 @@ void FD3D12ConstantBufferView::UpdateDescriptor()
 // -----------------------------------------------------------------------------------------------------
 
 FD3D12RenderTargetView::FD3D12RenderTargetView(FD3D12Device* InDevice, FD3D12RenderTargetView* FirstLinkedObject)
-	: TD3D12View(InDevice, ED3D12ViewType::RenderTarget, ERHIDescriptorHeapType::RenderTarget, FirstLinkedObject)
+	: TD3D12View(InDevice, ERHIDescriptorType::Invalid, ERHIDescriptorHeapType::RenderTarget, FirstLinkedObject)
 {
 }
 
@@ -400,7 +399,7 @@ void FD3D12RenderTargetView::UpdateDescriptor()
 // -----------------------------------------------------------------------------------------------------
 
 FD3D12DepthStencilView::FD3D12DepthStencilView(FD3D12Device* InDevice, FD3D12DepthStencilView* FirstLinkedObject)
-	: TD3D12View(InDevice, ED3D12ViewType::DepthStencil, ERHIDescriptorHeapType::DepthStencil, FirstLinkedObject)
+	: TD3D12View(InDevice, ERHIDescriptorType::Invalid, ERHIDescriptorHeapType::DepthStencil, FirstLinkedObject)
 {
 }
 
diff --git a/Engine/Source/Runtime/D3D12RHI/Private/D3D12View.h b/Engine/Source/Runtime/D3D12RHI/Private/D3D12View.h
--- a/Engine/Source/Runtime/D3D12RHI/Private/D3D12View.h
+++ b/Engine/Source/Runtime/D3D12RHI/Private/D3D12View.h
@@ -13,6 +13,43 @@ class FD3D12Texture;
 class FD3D12RayTracingScene;
 struct FD3D12ResidencyHandle;
 
+inline ERHIDescriptorType RHIDescriptorTypeFromD3D12ViewDesc(const D3D12_SHADER_RESOURCE_VIEW_DESC& ViewDesc)
+{
+	if (ViewDesc.ViewDimension == D3D12_SRV_DIMENSION_UNKNOWN)
+	{
+		return ERHIDescriptorType::Invalid;
+	}
+
+	if (ViewDesc.ViewDimension == D3D12_SRV_DIMENSION_BUFFER)
+	{
+		return ERHIDescriptorType::BufferSRV;
+	}
+
+#if D3D12_RHI_RAYTRACING
+	if (ViewDesc.ViewDimension == D3D12_SRV_DIMENSION_RAYTRACING_ACCELERATION_STRUCTURE)
+	{
+		return ERHIDescriptorType::AccelerationStructure;
+	}
+#endif
+
+	return ERHIDescriptorType::TextureSRV;
+}
+
+inline ERHIDescriptorType RHIDescriptorTypeFromD3D12ViewDesc(const D3D12_UNORDERED_ACCESS_VIEW_DESC& ViewDesc)
+{
+	if (ViewDesc.ViewDimension == D3D12_UAV_DIMENSION_UNKNOWN)
+	{
+		return ERHIDescriptorType::Invalid;
+	}
+
+	if (ViewDesc.ViewDimension == D3D12_UAV_DIMENSION_BUFFER)
+	{
+		return ERHIDescriptorType::BufferUAV;
+	}
+
+	return ERHIDescriptorType::TextureUAV;
+}
+
 struct FD3D12DefaultViews
 {
 	FD3D12OfflineDescriptor NullSRV;
@@ -231,7 +268,7 @@ public:
 
 protected:
 	FD3D12View() = delete;
-	FD3D12View(FD3D12Device* InDevice, ED3D12ViewType InViewType, ERHIDescriptorHeapType InHeapType, FD3D12View* FirstLinkedObject);
+	FD3D12View(FD3D12Device* InDevice, ERHIDescriptorType InDescriptorType, ERHIDescriptorHeapType InHeapType, FD3D12View* FirstLinkedObject);
 	virtual ~FD3D12View();
 		
 	virtual void UpdateDescriptor() = 0;
@@ -265,8 +302,8 @@ class TD3D12View : public FD3D12View
 protected:
 	TDesc D3DViewDesc;
 
-	TD3D12View(FD3D12Device* InDevice, ED3D12ViewType InViewType, ERHIDescriptorHeapType InHeapType, TD3D12View* FirstLinkedObject)
-		: FD3D12View(InDevice, InViewType, InHeapType, FirstLinkedObject)
+	TD3D12View(FD3D12Device* InDevice, ERHIDescriptorType InDescriptorType, ERHIDescriptorHeapType InHeapType, TD3D12View* FirstLinkedObject)
+		: FD3D12View(InDevice, InDescriptorType, InHeapType, FirstLinkedObject)
 	{}
 
 	void CreateView(FResourceInfo const& InResource, TDesc const& InD3DViewDesc)
@@ -322,8 +359,7 @@ public:
 	};
 	FRIEND_ENUM_CLASS_FLAGS(EFlags)
 
-	FD3D12ShaderResourceView(FD3D12Device* InDevice, FD3D12ShaderResourceView* FirstLinkedObject);
-	FD3D12ShaderResourceView(FD3D12Device* InDevice, FD3D12ShaderResourceView* FirstLinkedObject, FD3D12RayTracingScene* InRayTracingScene);
+	FD3D12ShaderResourceView(FD3D12Device* InDevice, FD3D12ShaderResourceView* FirstLinkedObject, ERHIDescriptorType InDescriptorType, FD3D12RayTracingScene* InRayTracingScene = nullptr);
 	void CreateView(FResourceInfo const& InResource, D3D12_SHADER_RESOURCE_VIEW_DESC const& InD3DViewDesc, EFlags InFlags);
 	void UpdateView(FD3D12ContextArray const& Contexts, const FResourceInfo& InResource, const D3D12_SHADER_RESOURCE_VIEW_DESC& InD3DViewDesc, EFlags InFlags);
 
@@ -360,7 +396,7 @@ public:
 	};
 	FRIEND_ENUM_CLASS_FLAGS(EFlags)
 
-	FD3D12UnorderedAccessView(FD3D12Device* InDevice, FD3D12UnorderedAccessView* FirstLinkedObject);
+	FD3D12UnorderedAccessView(FD3D12Device* InDevice, FD3D12UnorderedAccessView* FirstLinkedObject, ERHIDescriptorType InDescriptorType);
 	void CreateView(FResourceInfo const& InResource, D3D12_UNORDERED_ACCESS_VIEW_DESC const& InD3DViewDesc, EFlags InFlags);
 	void UpdateView(FD3D12ContextArray const& Contexts, const FResourceInfo& InResource, const D3D12_UNORDERED_ACCESS_VIEW_DESC& InD3DViewDesc, EFlags InFlags);
 
diff --git a/Engine/Source/Runtime/D3D12RHI/Private/D3D12WorkGraph.cpp b/Engine/Source/Runtime/D3D12RHI/Private/D3D12WorkGraph.cpp
--- a/Engine/Source/Runtime/D3D12RHI/Private/D3D12WorkGraph.cpp
+++ b/Engine/Source/Runtime/D3D12RHI/Private/D3D12WorkGraph.cpp
@@ -89,7 +89,8 @@ FD3D12WorkGraphPipelineState::FD3D12WorkGraphPipelineState(FD3D12Device* Device,
 				const FString NodePathName = FString::Printf(TEXT("%s_%d"), *NodeName, NodeArrayIndex);
 
 				CD3DX12_DXIL_LIBRARY_SUBOBJECT* Lib = StateObjectDesc.CreateSubobject<CD3DX12_DXIL_LIBRARY_SUBOBJECT>();
-				CD3DX12_SHADER_BYTECODE LibCode(NodeShader->Code.GetData(), NodeShader->Code.Num());
+				D3D12_SHADER_BYTECODE LibCode = NodeShader->GetShaderBytecode();
+
 				Lib->SetDXILLibrary(&LibCode);
 				Lib->DefineExport(*NodePathName, *NodeShader->EntryPoint);
 
@@ -933,7 +934,9 @@ void FD3D12CommandContext::DispatchWorkGraphShaderBundle(
 	SRVDesc.Buffer.Flags = D3D12_BUFFER_SRV_FLAG_RAW;
 
 	TSharedPtr<FD3D12ShaderResourceView> RecordArgBufferSRV;
-	RecordArgBufferSRV = MakeShared<FD3D12ShaderResourceView>(GetParentDevice(), nullptr); // Always single GPU object, so FirstLinkedObject is nullptr
+
+	// Always single GPU object, so FirstLinkedObject is nullptr
+	RecordArgBufferSRV = MakeShared<FD3D12ShaderResourceView>(GetParentDevice(), nullptr, ERHIDescriptorType::BufferSRV);
 	RecordArgBufferSRV->CreateView(RecordArgBufferPtr, SRVDesc, FD3D12ShaderResourceView::EFlags::None);
 
 	uint32 RecordArgBufferBindlessHandle = 0;
@@ -1286,7 +1289,8 @@ void FD3D12CommandContext::DispatchWorkGraphShaderBundle(
 	SRVDesc.Buffer.Flags = D3D12_BUFFER_SRV_FLAG_RAW;
 
 	TSharedPtr<FD3D12ShaderResourceView> RecordArgBufferSRV;
-	RecordArgBufferSRV = MakeShared<FD3D12ShaderResourceView>(GetParentDevice(), nullptr); // Always single GPU object, so FirstLinkedObject is nullptr
+	// Always single GPU object, so FirstLinkedObject is nullptr
+	RecordArgBufferSRV = MakeShared<FD3D12ShaderResourceView>(GetParentDevice(), nullptr, ERHIDescriptorType::BufferSRV);
 	RecordArgBufferSRV->CreateView(RecordArgBufferPtr, SRVDesc, FD3D12ShaderResourceView::EFlags::None);
 
 	// Gather root arguments for shader bundle entry node.
diff --git a/Engine/Source/Runtime/RHI/Public/RHIDefinitions.h b/Engine/Source/Runtime/RHI/Public/RHIDefinitions.h
--- a/Engine/Source/Runtime/RHI/Public/RHIDefinitions.h
+++ b/Engine/Source/Runtime/RHI/Public/RHIDefinitions.h
@@ -1331,29 +1331,87 @@ enum class ERHIDescriptorHeapType : uint8
 	Invalid = MAX_uint8
 };
 
+enum class ERHIDescriptorType : uint8
+{
+	BufferSRV,
+	BufferUAV,
+
+	TypedBufferSRV,
+	TypedBufferUAV,
+
+	TextureSRV,
+	TextureUAV,
+
+	CBV,
+	Sampler,
+
+	AccelerationStructure,
+
+	Invalid
+};
+
+enum class ERHIDescriptorTypeMask : uint16
+{
+	None = 0,
+
+	BufferSRV = (1 << static_cast<int32>(ERHIDescriptorType::BufferSRV)),
+	BufferUAV = (1 << static_cast<int32>(ERHIDescriptorType::BufferUAV)),
+
+	TypedBufferSRV = (1 << static_cast<int32>(ERHIDescriptorType::TypedBufferSRV)),
+	TypedBufferUAV = (1 << static_cast<int32>(ERHIDescriptorType::TypedBufferUAV)),
+
+	TextureSRV = (1 << static_cast<int32>(ERHIDescriptorType::TextureSRV)),
+	TextureUAV = (1 << static_cast<int32>(ERHIDescriptorType::TextureUAV)),
+
+	CBV = (1 << static_cast<int32>(ERHIDescriptorType::CBV)),
+	Sampler = (1 << static_cast<int32>(ERHIDescriptorType::Sampler)),
+
+	AccelerationStructure = (1 << static_cast<int32>(ERHIDescriptorType::AccelerationStructure)),
+
+	SRV = (BufferSRV | TypedBufferSRV | TextureSRV),
+	UAV = (BufferUAV | TypedBufferUAV | TextureUAV),
+
+	Standard = (SRV | UAV | CBV | AccelerationStructure),
+};
+ENUM_CLASS_FLAGS(ERHIDescriptorTypeMask);
+
+inline constexpr ERHIDescriptorTypeMask RHIDescriptorTypeMaskFromType(ERHIDescriptorType InType)
+{
+	if (InType == ERHIDescriptorType::Invalid)
+	{
+		return ERHIDescriptorTypeMask::None;
+	}
+	return static_cast<ERHIDescriptorTypeMask>(1 << static_cast<int32>(InType));
+}
+
 struct FRHIDescriptorHandle
 {
 	FRHIDescriptorHandle() = default;
-	FRHIDescriptorHandle(ERHIDescriptorHeapType InType, uint32 InIndex)
+
+	FRHIDescriptorHandle(ERHIDescriptorType InType, uint32 InIndex)
 		: Index(InIndex)
-		, Type((uint8)InType)
+		, Type(InType)
 	{
 	}
-	FRHIDescriptorHandle(uint8 InType, uint32 InIndex)
-		: Index(InIndex)
-		, Type(InType)
+
+	uint32 GetIndex() const
 	{
+		return Index;
 	}
 
-	inline uint32                 GetIndex() const { return Index; }
-	inline ERHIDescriptorHeapType GetType() const { return (ERHIDescriptorHeapType)Type; }
-	inline uint8                  GetRawType() const { return Type; }
+	ERHIDescriptorType GetType() const
+	{
+		return Type;
+	}
 
-	inline bool IsValid() const { return Index != MAX_uint32 && Type != (uint8)ERHIDescriptorHeapType::Invalid; }
+	bool IsValid() const
+	{
+		return Index != MAX_uint32 && Type != ERHIDescriptorType::Invalid;
+	}
 
 private:
-	uint32    Index{ MAX_uint32 };
-	uint8     Type{ (uint8)ERHIDescriptorHeapType::Invalid };
+	uint32             Index{ MAX_uint32 };
+	ERHIDescriptorType Type{ ERHIDescriptorType::Invalid };
 };
 
 // Configuration levels for Bindless. Each value should include more shaders than the previous.
diff --git a/Engine/Source/Runtime/RHI/Public/RHIResources.h b/Engine/Source/Runtime/RHI/Public/RHIResources.h
--- a/Engine/Source/Runtime/RHI/Public/RHIResources.h
+++ b/Engine/Source/Runtime/RHI/Public/RHIResources.h
@@ -3326,6 +3326,30 @@ public:
 	}
 };
 
+inline ERHIDescriptorType RHIDescriptorTypeFromViewType(FRHIViewDesc::EViewType ViewType)
+{
+	switch (ViewType)
+	{
+	default:
+	case FRHIViewDesc::EViewType::BufferSRV:
+		return ERHIDescriptorType::BufferSRV;
+
+	case FRHIViewDesc::EViewType::BufferUAV:
+		return ERHIDescriptorType::BufferUAV;
+
+	case FRHIViewDesc::EViewType::TextureSRV:
+		return ERHIDescriptorType::TextureSRV;
+
+	case FRHIViewDesc::EViewType::TextureUAV:
+		return ERHIDescriptorType::TextureUAV;
+	}
+}
+
+inline ERHIDescriptorType RHIDescriptorTypeFromViewDesc(const FRHIViewDesc& InViewDesc)
+{
+	return RHIDescriptorTypeFromViewType(InViewDesc.Common.ViewType);
+}
+
 //
 // Ray tracing resources
 //
diff --git a/Engine/Source/Runtime/RHICore/Private/RHIDescriptorAllocator.cpp b/Engine/Source/Runtime/RHICore/Private/RHIDescriptorAllocator.cpp
--- a/Engine/Source/Runtime/RHICore/Private/RHIDescriptorAllocator.cpp
+++ b/Engine/Source/Runtime/RHICore/Private/RHIDescriptorAllocator.cpp
@@ -33,7 +33,7 @@ void FRHIDescriptorAllocator::Shutdown()
 	Capacity = 0;
 }
 
-FRHIDescriptorHandle FRHIDescriptorAllocator::ResizeGrowAndAllocate(uint32 InNewNumDescriptors, ERHIDescriptorHeapType InType)
+FRHIDescriptorHandle FRHIDescriptorAllocator::ResizeGrowAndAllocate(uint32 InNewNumDescriptors, ERHIDescriptorType InType)
 {
 	check(Capacity < InNewNumDescriptors);
 	
@@ -65,7 +65,7 @@ FRHIDescriptorHandle FRHIDescriptorAllocator::ResizeGrowAndAllocate(uint32 InNew
 	return FRHIDescriptorHandle(InType, Index);
 }
 
-FRHIDescriptorHandle FRHIDescriptorAllocator::Allocate(ERHIDescriptorHeapType InType)
+FRHIDescriptorHandle FRHIDescriptorAllocator::Allocate(ERHIDescriptorType InType)
 {
 	uint32 Index{};
 	if (Allocate(1, Index))
@@ -267,22 +267,22 @@ bool FRHIDescriptorAllocator::GetAllocatedRange(FRHIDescriptorAllocatorRange& Ou
 ////////////////////////////////////////////////////////////////////////////////////////////////////////////////////////
 // FRHIHeapDescriptorAllocator
 
-FRHIHeapDescriptorAllocator::FRHIHeapDescriptorAllocator(ERHIDescriptorHeapType InType, uint32 InDescriptorCount, TConstArrayView<TStatId> InStats)
+FRHIHeapDescriptorAllocator::FRHIHeapDescriptorAllocator(ERHIDescriptorTypeMask InTypeMask, uint32 InDescriptorCount, TConstArrayView<TStatId> InStats)
 	: FRHIDescriptorAllocator(InDescriptorCount, InStats)
-	, Type(InType)
+	, TypeMask(InTypeMask)
 {
 }
 
-FRHIDescriptorHandle FRHIHeapDescriptorAllocator::Allocate()
+FRHIDescriptorHandle FRHIHeapDescriptorAllocator::Allocate(ERHIDescriptorType InType)
 {
-	return FRHIDescriptorAllocator::Allocate(Type);
+	return FRHIDescriptorAllocator::Allocate(InType);
 }
 
 void FRHIHeapDescriptorAllocator::Free(FRHIDescriptorHandle InHandle)
 {
 	if (InHandle.IsValid())
 	{
-		check(InHandle.GetType() == Type);
+		check(HandlesAllocation(InHandle.GetType()));
 		FRHIDescriptorAllocator::Free(InHandle);
 	}
 }
@@ -300,15 +300,15 @@ void FRHIHeapDescriptorAllocator::Free(uint32 Slot, uint32 NumDescriptors)
 ////////////////////////////////////////////////////////////////////////////////////////////////////////////////////////
 // FRHIOffsetHeapDescriptorAllocator
 
-FRHIOffsetHeapDescriptorAllocator::FRHIOffsetHeapDescriptorAllocator(ERHIDescriptorHeapType InType, uint32 InDescriptorCount, uint32 InHeapOffset, TConstArrayView<TStatId> InStats)
+FRHIOffsetHeapDescriptorAllocator::FRHIOffsetHeapDescriptorAllocator(ERHIDescriptorTypeMask InType, uint32 InDescriptorCount, uint32 InHeapOffset, TConstArrayView<TStatId> InStats)
 	: FRHIHeapDescriptorAllocator(InType, InDescriptorCount, InStats)
 	, HeapOffset(InHeapOffset)
 {
 }
 
-FRHIDescriptorHandle FRHIOffsetHeapDescriptorAllocator::Allocate()
+FRHIDescriptorHandle FRHIOffsetHeapDescriptorAllocator::Allocate(ERHIDescriptorType InType)
 {
-	const FRHIDescriptorHandle AlocatorHandle = FRHIHeapDescriptorAllocator::Allocate();
+	const FRHIDescriptorHandle AlocatorHandle = FRHIHeapDescriptorAllocator::Allocate(InType);
 	if (AlocatorHandle.IsValid())
 	{
 		return FRHIDescriptorHandle(AlocatorHandle.GetType(), AlocatorHandle.GetIndex() + HeapOffset);
diff --git a/Engine/Source/Runtime/RHICore/Public/RHIDescriptorAllocator.h b/Engine/Source/Runtime/RHICore/Public/RHIDescriptorAllocator.h
--- a/Engine/Source/Runtime/RHICore/Public/RHIDescriptorAllocator.h
+++ b/Engine/Source/Runtime/RHICore/Public/RHIDescriptorAllocator.h
@@ -3,11 +3,9 @@
 
 #include "Containers/Array.h"
 #include "HAL/CriticalSection.h"
+#include "RHIDefinitions.h"
 #include "Stats/Stats.h"
 
-enum class ERHIDescriptorHeapType : uint8;
-struct FRHIDescriptorHandle;
-
 struct FRHIDescriptorAllocatorRange
 {
 	FRHIDescriptorAllocatorRange(uint32 InFirst, uint32 InLast) : First(InFirst), Last(InLast) {}
@@ -25,9 +23,9 @@ public:
 	RHICORE_API void Init(uint32 InNumDescriptors, TConstArrayView<TStatId> InStats);
 	RHICORE_API void Shutdown();
 
-	RHICORE_API FRHIDescriptorHandle ResizeGrowAndAllocate(uint32 InNewNumDescriptors, ERHIDescriptorHeapType InType);
+	RHICORE_API FRHIDescriptorHandle ResizeGrowAndAllocate(uint32 InNewNumDescriptors, ERHIDescriptorType InType);
 
-	RHICORE_API FRHIDescriptorHandle Allocate(ERHIDescriptorHeapType InType);
+	RHICORE_API FRHIDescriptorHandle Allocate(ERHIDescriptorType InType);
 	RHICORE_API void Free(FRHIDescriptorHandle InHandle);
 
 	RHICORE_API bool Allocate(uint32 NumDescriptors, uint32& OutSlot);
@@ -77,9 +75,9 @@ class FRHIHeapDescriptorAllocator : protected FRHIDescriptorAllocator
 {
 public:
 	FRHIHeapDescriptorAllocator() = delete;
-	RHICORE_API FRHIHeapDescriptorAllocator(ERHIDescriptorHeapType InType, uint32 InDescriptorCount, TConstArrayView<TStatId> InStats);
+	RHICORE_API FRHIHeapDescriptorAllocator(ERHIDescriptorTypeMask InTypeMask, uint32 InDescriptorCount, TConstArrayView<TStatId> InStats);
 
-	RHICORE_API FRHIDescriptorHandle Allocate();
+	RHICORE_API FRHIDescriptorHandle Allocate(ERHIDescriptorType InType);
 	RHICORE_API void Free(FRHIDescriptorHandle InHandle);
 
 	RHICORE_API bool Allocate(uint32 NumDescriptors, uint32& OutSlot);
@@ -89,30 +87,41 @@ public:
 	using FRHIDescriptorAllocator::GetCapacity;
 	using FRHIDescriptorAllocator::ResizeGrowAndAllocate;
 
-	inline ERHIDescriptorHeapType GetType() const { return Type; }
+	ERHIDescriptorTypeMask GetTypeMask() const
+	{
+		return TypeMask;
+	}
 
-	inline bool HandlesAllocation(ERHIDescriptorHeapType InType) const { return GetType() == InType; }
+	bool HandlesAllocation(ERHIDescriptorType InType) const
+	{
+		return EnumHasAnyFlags(GetTypeMask(), RHIDescriptorTypeMaskFromType(InType));
+	}
+
+	bool HandlesAllocations(ERHIDescriptorTypeMask InTypeMask) const
+	{
+		return EnumHasAllFlags(GetTypeMask(), InTypeMask);
+	}
 
 private:
-	ERHIDescriptorHeapType Type;
+	const ERHIDescriptorTypeMask TypeMask;
 };
 
 class FRHIOffsetHeapDescriptorAllocator : protected FRHIHeapDescriptorAllocator
 {
 public:
 	FRHIOffsetHeapDescriptorAllocator() = delete;
-	RHICORE_API FRHIOffsetHeapDescriptorAllocator(ERHIDescriptorHeapType InType, uint32 InDescriptorCount, uint32 InHeapOffset, TConstArrayView<TStatId> InStats);
+	RHICORE_API FRHIOffsetHeapDescriptorAllocator(ERHIDescriptorTypeMask InTypeMask, uint32 InDescriptorCount, uint32 InHeapOffset, TConstArrayView<TStatId> InStats);
 
-	RHICORE_API FRHIDescriptorHandle Allocate();
+	RHICORE_API FRHIDescriptorHandle Allocate(ERHIDescriptorType InType);
 	RHICORE_API void Free(FRHIDescriptorHandle InHandle);
 
 	using FRHIHeapDescriptorAllocator::GetCapacity;
-	using FRHIHeapDescriptorAllocator::GetType;
+	using FRHIHeapDescriptorAllocator::GetTypeMask;
 	using FRHIHeapDescriptorAllocator::HandlesAllocation;
 
 private:
 	// Offset from start of heap we belong to
-	uint32 HeapOffset;
+	const uint32 HeapOffset;
 };
 
 #if UE_ENABLE_INCLUDE_ORDER_DEPRECATED_IN_5_4
diff --git a/Engine/Source/Runtime/VulkanRHI/Private/VulkanDescriptorSets.cpp b/Engine/Source/Runtime/VulkanRHI/Private/VulkanDescriptorSets.cpp
--- a/Engine/Source/Runtime/VulkanRHI/Private/VulkanDescriptorSets.cpp
+++ b/Engine/Source/Runtime/VulkanRHI/Private/VulkanDescriptorSets.cpp
@@ -125,6 +125,43 @@ static constexpr uint8 GetIndexForDescriptorType(VkDescriptorType DescriptorType
 	return VulkanBindless::MaxNumSets;
 }
 
+static constexpr uint8 GetIndexForDescriptorType(ERHIDescriptorType DescriptorType)
+{
+	switch (DescriptorType)
+	{
+	case ERHIDescriptorType::Sampler:               return VulkanBindless::BindlessSamplerSet;
+	case ERHIDescriptorType::TextureSRV:            return VulkanBindless::BindlessSampledImageSet;
+	case ERHIDescriptorType::TextureUAV:            return VulkanBindless::BindlessStorageImageSet;
+	case ERHIDescriptorType::TypedBufferSRV:        return VulkanBindless::BindlessUniformTexelBufferSet;
+	case ERHIDescriptorType::TypedBufferUAV:        return VulkanBindless::BindlessStorageTexelBufferSet;
+	case ERHIDescriptorType::BufferUAV:             return VulkanBindless::BindlessStorageBufferSet;
+	case ERHIDescriptorType::BufferSRV:             return VulkanBindless::BindlessUniformBufferSet;
+	case ERHIDescriptorType::AccelerationStructure: return VulkanBindless::BindlessAccelerationStructureSet;
+	default: checkNoEntry();
+	}
+
+	return VulkanBindless::MaxNumSets;
+}
+
+
+static constexpr ERHIDescriptorType GetUEDescriptorTypeFromVkDescriptorType(VkDescriptorType DescriptorType)
+{
+	switch (DescriptorType)
+	{
+	case VK_DESCRIPTOR_TYPE_SAMPLER:                    return ERHIDescriptorType::Sampler;
+	case VK_DESCRIPTOR_TYPE_SAMPLED_IMAGE:              return ERHIDescriptorType::TextureSRV;
+	case VK_DESCRIPTOR_TYPE_STORAGE_IMAGE:              return ERHIDescriptorType::TextureUAV;
+	case VK_DESCRIPTOR_TYPE_UNIFORM_TEXEL_BUFFER:       return ERHIDescriptorType::TypedBufferSRV;
+	case VK_DESCRIPTOR_TYPE_STORAGE_TEXEL_BUFFER:       return ERHIDescriptorType::TypedBufferUAV;
+	case VK_DESCRIPTOR_TYPE_STORAGE_BUFFER:             return ERHIDescriptorType::BufferUAV;
+	case VK_DESCRIPTOR_TYPE_UNIFORM_BUFFER:             return ERHIDescriptorType::BufferSRV;
+	case VK_DESCRIPTOR_TYPE_ACCELERATION_STRUCTURE_KHR: return ERHIDescriptorType::AccelerationStructure;
+	default: checkNoEntry();
+	}
+
+	return ERHIDescriptorType::Invalid;
+}
+
 
 extern TAutoConsoleVariable<int32> GCVarRobustBufferAccess;
 static inline uint32 GetDescriptorTypeSize(FVulkanDevice& Device, VkDescriptorType DescriptorType)
@@ -675,9 +712,11 @@ FRHIDescriptorHandle FVulkanBindlessDescriptorManager::ReserveDescriptor(VkDescr
 	if (bIsSupported)
 	{
 		const uint8 SetIndex = GetIndexForDescriptorType(DescriptorType);
+		const ERHIDescriptorType RHIType = GetUEDescriptorTypeFromVkDescriptorType(DescriptorType);
+
 		BindlessSetState& State = BindlessSetStates[SetIndex];
 		const uint32 ResourceIndex = GetFreeResourceIndex(State);
-		return FRHIDescriptorHandle(SetIndex, ResourceIndex);
+		return FRHIDescriptorHandle(RHIType, ResourceIndex);
 	}
 	return FRHIDescriptorHandle();
 }
@@ -686,8 +725,8 @@ void FVulkanBindlessDescriptorManager::UpdateDescriptor(FRHIDescriptorHandle Des
 {
 	checkf(DescriptorHandle.IsValid(), TEXT("Attemping to update invalid descriptor handle!"));
 
-	const uint8 SetIndex = DescriptorHandle.GetRawType();
-	check(SetIndex == GetIndexForDescriptorType(DescriptorType));
+	const uint8 SetIndex = GetIndexForDescriptorType(DescriptorHandle.GetType());
+
 	BindlessSetState& State = BindlessSetStates[SetIndex];
 	const uint32 ByteOffset = DescriptorHandle.GetIndex() * State.DescriptorSize;
 	checkSlow(State.DescriptorTypes.Contains(DescriptorType));
@@ -879,7 +918,7 @@ void FVulkanBindlessDescriptorManager::Unregister(FRHIDescriptorHandle Descripto
 	{
 		checkf(bIsSupported, TEXT("Unregistering a valid handle but bindless is not supported!"));
 
-		const uint8 SetIndex = DescriptorHandle.GetRawType();
+		const uint8 SetIndex = GetIndexForDescriptorType(DescriptorHandle.GetType());
 		BindlessSetState& State = BindlessSetStates[SetIndex];
 
 		FScopeLock ScopeLock(&State.FreeListCS);
diff --git a/Engine/Source/Runtime/VulkanRHI/Private/VulkanDevice.h b/Engine/Source/Runtime/VulkanRHI/Private/VulkanDevice.h
--- a/Engine/Source/Runtime/VulkanRHI/Private/VulkanDevice.h
+++ b/Engine/Source/Runtime/VulkanRHI/Private/VulkanDevice.h
@@ -219,10 +219,8 @@ namespace VulkanRHI
 		{
 			if (DescriptorHandle.IsValid())
 			{
-				const uint64 Type = (uint64)DescriptorHandle.GetRawType();
-				const uint64 Index = (uint64)DescriptorHandle.GetIndex();
-				const uint64 AsUInt64 = (Type << 32) | Index;
-				EnqueueResource(EType::BindlessHandle, AsUInt64);
+				static_assert(sizeof(FRHIDescriptorHandle) <= sizeof(uint64), "FRHIDescriptorHandle is too large");
+				EnqueueResource(EType::BindlessHandle, *reinterpret_cast<uint64*>(&DescriptorHandle));
 			}
 		}
 
diff --git a/Engine/Source/Runtime/VulkanRHI/Private/VulkanMemory.cpp b/Engine/Source/Runtime/VulkanRHI/Private/VulkanMemory.cpp
--- a/Engine/Source/Runtime/VulkanRHI/Private/VulkanMemory.cpp
+++ b/Engine/Source/Runtime/VulkanRHI/Private/VulkanMemory.cpp
@@ -4884,10 +4884,7 @@ namespace VulkanRHI
 			case EType::BindlessHandle:
 			{
 				check(Device->SupportsBindless());
-				uint32 BindlessType = (uint32)(Entry.Handle >> 32);
-				check(BindlessType < MAX_uint8);
-				uint32 BindlessIndex = (uint32)(Entry.Handle & 0xFFFFFFFF);
-				FRHIDescriptorHandle DescriptorHandle((uint8)BindlessType, BindlessIndex);
+				const FRHIDescriptorHandle& DescriptorHandle = reinterpret_cast<const FRHIDescriptorHandle&>(Entry.Handle);
 				Device->GetBindlessDescriptorManager()->Unregister(DescriptorHandle);
 				break;
 			}
