Metal RHI

- Fixed bug where Metal runs out of counter sample buffers, now we return nullptr and ignore the stats.
- Also added MetalRHI option to disable Blit and/or Compute Encoder sampling

diff --git a/Engine/Source/Runtime/Apple/MetalRHI/Private/MetalCommandBuffer.h b/Engine/Source/Runtime/Apple/MetalRHI/Private/MetalCommandBuffer.h
--- a/Engine/Source/Runtime/Apple/MetalRHI/Private/MetalCommandBuffer.h
+++ b/Engine/Source/Runtime/Apple/MetalRHI/Private/MetalCommandBuffer.h
@@ -68,6 +68,11 @@ public:
     
 	void AddCounterSample(FMetalCounterSamplePtr CounterSample)
 	{
+		if(!CounterSample)
+		{
+			return;
+		}
+		
 #if RHI_NEW_GPU_PROFILER == 0
 		for(FMetalEventNode* Node : ActiveEventNodes)
 		{
diff --git a/Engine/Source/Runtime/Apple/MetalRHI/Private/MetalCounterSampler.cpp b/Engine/Source/Runtime/Apple/MetalRHI/Private/MetalCounterSampler.cpp
--- a/Engine/Source/Runtime/Apple/MetalRHI/Private/MetalCounterSampler.cpp
+++ b/Engine/Source/Runtime/Apple/MetalRHI/Private/MetalCounterSampler.cpp
@@ -140,16 +140,29 @@ MTL::CounterSampleBuffer* FMetalCounterSampler::SwapOrAllocateBuffer(uint32_t Sa
 		BufferDesc->setCounterSet(CounterSet);
 		BufferDesc->setStorageMode(MTL::StorageModeShared);
 		
-		Offset = 0;
-		OutOffset = 0;
-		
 		BufferDesc->setSampleCount(Size);
 		
 		NS::Error* DeviceError;
 		SampleBuffer = Device->GetDevice()->newCounterSampleBuffer(BufferDesc, &DeviceError);
-		check(SampleBuffer);
 		
 		BufferDesc->release();
+		
+		if(SampleBuffer)
+		{
+			Offset = 0;
+			OutOffset = 0;
+		}
+		else
+		{
+			static bool bLoggedWarning = false;
+			if (!bLoggedWarning)
+			{
+				UE_LOG(LogMetal, Warning, TEXT("We have run out of Metal counter samplers which will negatively affect the GPU stats, you can try disabling Blit counter samplers with rhi.Metal.SampleBlitEncoderTimings=0"));
+				bLoggedWarning = true;
+			}
+		}
+		
+		return SampleBuffer;
 	}
 	
 	check(Offset <= Size);
@@ -161,7 +174,10 @@ FMetalCounterSamplePtr FMetalCounterSampler::SetupStageCounters(MTL::ComputePass
 {
 	uint32_t OutOffset;
 	MTL::CounterSampleBuffer* Buffer = SwapOrAllocateBuffer(2, OutOffset);
-	check(Buffer);
+	if(!Buffer)
+	{
+		return nullptr;
+	}
 	
 	FMetalCounterSamplePtr Sample = FMetalCounterSamplePtr(new FMetalCounterSample(EMetalCounterSampleType::ComputeStage, Buffer, OutOffset));
 	MTL::ComputePassSampleBufferAttachmentDescriptor* SampleDesc = ComputePassDesc->sampleBufferAttachments()->object(0);
@@ -177,7 +193,10 @@ FMetalCounterSamplePtr FMetalCounterSampler::SetupStageCounters(MTL::BlitPassDes
 {
 	uint32_t OutOffset;
 	MTL::CounterSampleBuffer* Buffer = SwapOrAllocateBuffer(2, OutOffset);
-	check(Buffer);
+	if(!Buffer)
+	{
+		return nullptr;
+	}
 	
 	FMetalCounterSamplePtr Sample = FMetalCounterSamplePtr(new FMetalCounterSample(EMetalCounterSampleType::BlitStage, Buffer, OutOffset));
 	MTL::BlitPassSampleBufferAttachmentDescriptor* SampleDesc = BlitPassDesc->sampleBufferAttachments()->object(0);
@@ -193,7 +212,10 @@ FMetalCounterSamplePtr FMetalCounterSampler::SetupStageCounters(MTL::RenderPassD
 {
 	uint32_t OutOffset;
 	MTL::CounterSampleBuffer* Buffer = SwapOrAllocateBuffer(2, OutOffset);
-	check(Buffer);
+	if(!Buffer)
+	{
+		return nullptr;
+	}
 	
 	FMetalCounterSamplePtr Sample = FMetalCounterSamplePtr(new FMetalCounterSample(EMetalCounterSampleType::RenderStage, Buffer, OutOffset));
 	MTL::RenderPassSampleBufferAttachmentDescriptor* SampleDesc = InRenderPassDesc->sampleBufferAttachments()->object(0);
@@ -209,7 +231,10 @@ FMetalCounterSamplePtr FMetalCounterSampler::SetupBoundaryCounters(MTL::RenderCo
 {
 	uint32_t OutOffset;
 	MTL::CounterSampleBuffer* Buffer = SwapOrAllocateBuffer(1, OutOffset);
-	check(Buffer);
+	if(!Buffer)
+	{
+		return nullptr;
+	}
 	
 	FMetalCounterSamplePtr Sample = FMetalCounterSamplePtr(new FMetalCounterSample(EMetalCounterSampleType::DrawBoundary, Buffer, OutOffset));
 	
@@ -222,7 +247,10 @@ FMetalCounterSamplePtr FMetalCounterSampler::SetupBoundaryCounters(MTL::ComputeC
 {
 	uint32_t OutOffset;
 	MTL::CounterSampleBuffer* Buffer = SwapOrAllocateBuffer(1, OutOffset);
-	check(Buffer);
+	if(!Buffer)
+	{
+		return nullptr;
+	}
 	
 	FMetalCounterSamplePtr Sample = FMetalCounterSamplePtr(new FMetalCounterSample(EMetalCounterSampleType::DispatchBoundary, Buffer, OutOffset));
 	
@@ -235,7 +263,10 @@ FMetalCounterSamplePtr FMetalCounterSampler::SetupBoundaryCounters(MTL::BlitComm
 {
 	uint32_t OutOffset;
 	MTL::CounterSampleBuffer* Buffer = SwapOrAllocateBuffer(1, OutOffset);
-	check(Buffer);
+	if(!Buffer)
+	{
+		return nullptr;
+	}
 	
 	FMetalCounterSamplePtr Sample = FMetalCounterSamplePtr(new FMetalCounterSample(EMetalCounterSampleType::BlitBoundary, Buffer, OutOffset));
 	
diff --git a/Engine/Source/Runtime/Apple/MetalRHI/Private/MetalRHIContext.cpp b/Engine/Source/Runtime/Apple/MetalRHI/Private/MetalRHIContext.cpp
--- a/Engine/Source/Runtime/Apple/MetalRHI/Private/MetalRHIContext.cpp
+++ b/Engine/Source/Runtime/Apple/MetalRHI/Private/MetalRHIContext.cpp
@@ -13,6 +13,18 @@
 #import <CompositorServices/CompositorServices.h>
 #endif
 
+static int32 GMetalSampleBlitEncoderTimings = 1;
+static FAutoConsoleVariableRef CVarMetalSampleBlitEncoderTimings(
+	TEXT("rhi.Metal.SampleBlitEncoderTimings"),
+	GMetalSampleBlitEncoderTimings,
+	TEXT("Whether to include Blit Encoder Timings in GPU Stats, can cause poor perf due to splitting encoders: on"));
+
+static int32 GMetalSampleComputeEncoderTimings = 1;
+static FAutoConsoleVariableRef CVarMetalSampleComputeEncoderTimings(
+	TEXT("rhi.Metal.SampleComputeEncoderTimings"),
+	GMetalSampleComputeEncoderTimings,
+	TEXT("Whether to include Compute Encoder Timings in GPU Stats, can cause poor perf due to splitting encoders: on"));
+
 void METALRHI_API SafeReleaseMetalObject(NS::Object* Object)
 {
 	if(GIsMetalInitialized && GDynamicRHI && Object)
@@ -134,7 +146,7 @@ void FMetalRHICommandContext::BeginComputeEncoder()
 		{
 			CurrentEncoderFence = CurrentEncoder.EndEncoding();
 		}
-		bool bUseStageCounterSamples = Device.SupportsFeature(EMetalFeaturesStageCounterSampling);
+		bool bUseStageCounterSamples = Device.SupportsFeature(EMetalFeaturesStageCounterSampling) && GMetalSampleComputeEncoderTimings;
 		CurrentEncoder.BeginComputeCommandEncoding(MTL::DispatchTypeSerial,
 												   bUseStageCounterSamples ? Device.GetCounterSampler() : nullptr);
 	}
@@ -153,12 +165,16 @@ void FMetalRHICommandContext::EndComputeEncoder()
 	check(CurrentEncoder.IsComputeCommandEncoderActive());
 	
 	// If we are using breadcrumbs then end the encoding here so that our stat tracking is correct
-#if WITH_RHI_BREADCRUMBS
-	CurrentEncoderFence = CurrentEncoder.EndEncoding();
+#if WITH_PROFILEGPU
+	if(Device.SupportsFeature(EMetalFeaturesStageCounterSampling) && GMetalSampleComputeEncoderTimings)
+	{
+		CurrentEncoderFence = CurrentEncoder.EndEncoding();
+	}
 #endif
 	StateCache.SetRenderTargetsActive(false);
 }
 
+
 void FMetalRHICommandContext::BeginBlitEncoder()
 {
 	MTL_SCOPED_AUTORELEASE_POOL;
@@ -178,7 +194,8 @@ void FMetalRHICommandContext::BeginBlitEncoder()
 		{
 			CurrentEncoderFence = CurrentEncoder.EndEncoding();
 		}
-		bool bUseStageCounterSamples = Device.SupportsFeature(EMetalFeaturesStageCounterSampling);
+
+		bool bUseStageCounterSamples = Device.SupportsFeature(EMetalFeaturesStageCounterSampling) && GMetalSampleBlitEncoderTimings;
 		CurrentEncoder.BeginBlitCommandEncoding(bUseStageCounterSamples ? Device.GetCounterSampler() : nullptr);
 	}
 	
@@ -195,8 +212,11 @@ void FMetalRHICommandContext::EndBlitEncoder()
 {
 	check(CurrentEncoder.IsBlitCommandEncoderActive());
 	
-#if WITH_RHI_BREADCRUMBS
-	CurrentEncoderFence = CurrentEncoder.EndEncoding();
+#if WITH_PROFILEGPU
+	if(Device.SupportsFeature(EMetalFeaturesStageCounterSampling) && GMetalSampleBlitEncoderTimings)
+	{
+		CurrentEncoderFence = CurrentEncoder.EndEncoding();
+	}
 #endif
 	
 	StateCache.SetRenderTargetsActive(false);
