MetalRHI

MetalRHIRenderQuery could be cleaned up whilst waiting for completion and readback on the submission thread.

diff --git a/Engine/Source/Runtime/Apple/MetalRHI/Private/MetalRHIContext.cpp b/Engine/Source/Runtime/Apple/MetalRHI/Private/MetalRHIContext.cpp
--- a/Engine/Source/Runtime/Apple/MetalRHI/Private/MetalRHIContext.cpp
+++ b/Engine/Source/Runtime/Apple/MetalRHI/Private/MetalRHIContext.cpp
@@ -259,7 +259,7 @@ void FMetalRHICommandContext::RHIBeginParallelRenderPass(TSharedPtr<FRHIParallel
 	StateCache.SetStateDirty();
 	StateCache.SetRenderTargetsActive(true);
 	
-	FMetalQueryBuffer* VisBuffer = nullptr;
+	FMetalQueryBufferRef VisBuffer = nullptr;
 	if(RenderPassInfo.NumOcclusionQueries > 0)
 	{
 		VisBuffer = QueryBuffer->AcquireQueryBuffer(RenderPassInfo.NumOcclusionQueries);
@@ -340,7 +340,7 @@ void FMetalRHICommandContext::RHIBeginRenderPass(const FRHIRenderPassInfo& InInf
 	StateCache.SetStateDirty();
 	StateCache.SetRenderTargetsActive(true);
 	
-	FMetalQueryBuffer* VisBuffer = nullptr;
+	FMetalQueryBufferRef VisBuffer = nullptr;
 	if(InInfo.NumOcclusionQueries > 0)
 	{
 		VisBuffer = QueryBuffer->AcquireQueryBuffer(InInfo.NumOcclusionQueries);
diff --git a/Engine/Source/Runtime/Apple/MetalRHI/Private/MetalRHIRenderQuery.cpp b/Engine/Source/Runtime/Apple/MetalRHI/Private/MetalRHIRenderQuery.cpp
--- a/Engine/Source/Runtime/Apple/MetalRHI/Private/MetalRHIRenderQuery.cpp
+++ b/Engine/Source/Runtime/Apple/MetalRHI/Private/MetalRHIRenderQuery.cpp
@@ -54,6 +54,7 @@ uint64 FMetalQueryBuffer::GetResult(uint32 Offset)
 {
 	uint64 Result = 0;
 	
+	check(Buffer.IsValid());
 	if(Buffer.IsValid())
 	{
 		MTL_SCOPED_AUTORELEASE_POOL;
@@ -92,7 +93,7 @@ FMetalQueryBufferPool::~FMetalQueryBufferPool()
 void FMetalQueryBufferPool::Allocate(FMetalQueryResult& NewQuery)
 {
 	check(IsValidRef(CurrentBuffer));
-	FMetalQueryBuffer* QB = CurrentBuffer.GetReference();
+	FMetalQueryBufferRef QB = CurrentBuffer;
 
 	uint32 Offset = Align(QB->WriteOffset, EQueryBufferAlignment);
 	uint32 End = Align(Offset + EQueryResultMaxSize, EQueryBufferAlignment);
@@ -105,7 +106,7 @@ void FMetalQueryBufferPool::Allocate(FMetalQueryResult& NewQuery)
 	}
 }
 
-FMetalQueryBuffer* FMetalQueryBufferPool::AcquireQueryBuffer(uint32_t NumOcclusionQueries)
+FMetalQueryBufferRef FMetalQueryBufferPool::AcquireQueryBuffer(uint32_t NumOcclusionQueries)
 {
 	uint32_t RequiredSize = NumOcclusionQueries * EQueryResultMaxSize;
 	
@@ -160,21 +161,22 @@ FMetalQueryBuffer* FMetalQueryBufferPool::AcquireQueryBuffer(uint32_t NumOcclusi
 	check(Buffer);
 	
 	CurrentBuffer = new FMetalQueryBuffer(this, MoveTemp(Buffer));
-	return CurrentBuffer.GetReference();
+	return CurrentBuffer;
 }
 
-FMetalQueryBuffer* FMetalQueryBufferPool::GetCurrentQueryBuffer()
+FMetalQueryBufferRef FMetalQueryBufferPool::GetCurrentQueryBuffer()
 {
-	return CurrentBuffer.GetReference();
+	return CurrentBuffer;
 }
 
 void FMetalQueryBufferPool::ReleaseCurrentQueryBuffer()
 {
 	if (IsValidRef(CurrentBuffer) && (CurrentBuffer->WriteOffset > 0))
 	{
+		// Keep a reference of this until we know the GPU is finished.
 		FMetalDynamicRHI::Get().DeferredDelete([InBuffer = MoveTemp(CurrentBuffer)]() mutable
 		{
-			InBuffer.SafeRelease();
+			InBuffer = nullptr;
 		});
 	}
 }
@@ -305,7 +307,6 @@ FMetalRHIRenderQuery::FMetalRHIRenderQuery(FMetalDevice& MetalDevice, ERenderQue
 
 FMetalRHIRenderQuery::~FMetalRHIRenderQuery()
 {
-	Buffer.SourceBuffer.SafeRelease();
 	Buffer.Offset = 0;
 }
 
@@ -326,7 +327,7 @@ void FMetalRHIRenderQuery::End_TopOfPipe()
 
 void FMetalRHIRenderQuery::Begin(FMetalRHICommandContext* Context)
 {
-	Buffer.SourceBuffer.SafeRelease();
+	Buffer.SourceBuffer = nullptr;
 	Buffer.Offset = 0;
 
 	Result = 0;
@@ -338,6 +339,10 @@ void FMetalRHIRenderQuery::Begin(FMetalRHICommandContext* Context)
 		{
 			// allocate our space in the current buffer
 			Context->GetQueryBufferPool()->Allocate(Buffer);
+
+			// AddRef to ensure that the object isn't freed until after the results are sampled
+			AddRef();
+
 			Buffer.bCompleted = false;
 
 			if ((GMaxRHIFeatureLevel >= ERHIFeatureLevel::SM5) && Device.SupportsFeature(EMetalFeaturesCountingQueries))
@@ -378,14 +383,13 @@ void FMetalRHIRenderQuery::End(FMetalRHICommandContext* Context)
 			AddRef();
 
 			// Reset the result availability state
-			Buffer.SourceBuffer.SafeRelease();
+			Buffer.SourceBuffer = nullptr;
 			Buffer.Offset = 0;
 			Buffer.bCompleted = false;
 			Result = 0;
 			bAvailable = false;
 			
 			CommandBuffer = Context->GetCurrentCommandBuffer();
-			
 			CommandBuffer->TimestampQueries.Add(this);
             
 			break;
@@ -398,6 +402,15 @@ void FMetalRHIRenderQuery::End(FMetalRHICommandContext* Context)
 	}
 }
 
+void FMetalRHIRenderQuery::SampleOcclusionResult()
+{
+	checkSlow(Type == RQT_Occlusion);
+	Result = Buffer.GetResult();
+
+	// Now we have sampled the occlusion release the reference to this
+	Release();
+}
+
 bool FMetalRHIRenderQuery::GetResult(uint64& OutNumPixels, bool bWait, uint32 GPUIndex)
 {
 	if (!bAvailable)
diff --git a/Engine/Source/Runtime/Apple/MetalRHI/Private/MetalRHIRenderQuery.h b/Engine/Source/Runtime/Apple/MetalRHI/Private/MetalRHIRenderQuery.h
--- a/Engine/Source/Runtime/Apple/MetalRHI/Private/MetalRHIRenderQuery.h
+++ b/Engine/Source/Runtime/Apple/MetalRHI/Private/MetalRHIRenderQuery.h
@@ -40,6 +40,7 @@ public:
 	uint32 WriteOffset;
 };
 
+typedef TRefCountPtr<FMetalQueryBuffer> FMetalQueryBufferRef;
 
 //------------------------------------------------------------------------------
 
@@ -63,14 +64,14 @@ public:
 
 	void Allocate(FMetalQueryResult& NewQuery);
 	
-	FMetalQueryBuffer* AcquireQueryBuffer(uint32_t NumOcclusionQueries);
-	FMetalQueryBuffer* GetCurrentQueryBuffer();
+	FMetalQueryBufferRef AcquireQueryBuffer(uint32_t NumOcclusionQueries);
+	FMetalQueryBufferRef GetCurrentQueryBuffer();
 	void ReleaseCurrentQueryBuffer();
 	void ReleaseQueryBuffer(FMetalBufferPtr Buffer);
 
 	uint32_t BufferSize = 0;
 	
-	TRefCountPtr<FMetalQueryBuffer> CurrentBuffer;
+	FMetalQueryBufferRef CurrentBuffer;
 	TArray<FMetalBufferPtr> Buffers;
 	FMetalDevice& Device;
 };
@@ -90,7 +91,7 @@ public:
 	void Reset();
 	uint64 GetResult();
 
-	TRefCountPtr<FMetalQueryBuffer> SourceBuffer = nullptr;
+	FMetalQueryBufferRef SourceBuffer = nullptr;
 	uint32 Offset = 0;
 	bool bCompleted = false;
 	bool bBatchFence = false;
@@ -101,7 +102,6 @@ public:
 
 #pragma mark - Metal RHI Render Query Class -
 
-
 class FMetalRHIRenderQuery : public FRHIRenderQuery
 {
 	friend class FMetalDynamicRHI;
@@ -122,6 +122,11 @@ public:
 	 */
 	void End(FMetalRHICommandContext* Context);
 
+	/**
+	 * Samples the occlusion result on the submission thread
+	 */
+	void SampleOcclusionResult();
+	
 	/**
 	 * Get the query result
 	 */
diff --git a/Engine/Source/Runtime/Apple/MetalRHI/Private/MetalStateCache.cpp b/Engine/Source/Runtime/Apple/MetalRHI/Private/MetalStateCache.cpp
--- a/Engine/Source/Runtime/Apple/MetalRHI/Private/MetalStateCache.cpp
+++ b/Engine/Source/Runtime/Apple/MetalRHI/Private/MetalStateCache.cpp
@@ -442,7 +442,7 @@ void FMetalStateCache::SetComputeShader(FMetalComputeShader* InComputeShader)
 	}
 }
 
-bool FMetalStateCache::SetRenderPassInfo(FRHIRenderPassInfo const& InRenderTargets, FMetalQueryBuffer* QueryBuffer)
+bool FMetalStateCache::SetRenderPassInfo(FRHIRenderPassInfo const& InRenderTargets, FMetalQueryBufferRef QueryBuffer)
 {
 	bool bNeedsSet = false;
 	
@@ -2137,7 +2137,7 @@ void FMetalStateCache::CommitComputeResources(FMetalCommandEncoder* Compute)
 #endif
 }
 
-void FMetalStateCache::StartRenderPass(const FRHIRenderPassInfo& InInfo, FMetalQueryBuffer* QueryBuffer, MTL::RenderPassDescriptor* InDesc, bool bInIsParallelContext)
+void FMetalStateCache::StartRenderPass(const FRHIRenderPassInfo& InInfo, FMetalQueryBufferRef QueryBuffer, MTL::RenderPassDescriptor* InDesc, bool bInIsParallelContext)
 {
 	bIsParallelContext = bInIsParallelContext;
 	
diff --git a/Engine/Source/Runtime/Apple/MetalRHI/Private/MetalStateCache.h b/Engine/Source/Runtime/Apple/MetalRHI/Private/MetalStateCache.h
--- a/Engine/Source/Runtime/Apple/MetalRHI/Private/MetalStateCache.h
+++ b/Engine/Source/Runtime/Apple/MetalRHI/Private/MetalStateCache.h
@@ -9,10 +9,10 @@
 #include "MetalState.h"
 #include "MetalUniformBuffer.h"
 #include "Shaders/MetalShaderParameterCache.h"
+#include "MetalRHIRenderQuery.h"
 
 class FMetalGraphicsPipelineState;
 class FMetalComputeShader;
-class FMetalQueryBuffer;
 class FMetalDevice;
 
 enum EMetalPipelineFlags
@@ -54,7 +54,7 @@ public:
 	void SetBlendFactor(FLinearColor const& InBlendFactor);
 	void SetStencilRef(uint32 const InStencilRef);
 	void SetComputeShader(FMetalComputeShader* InComputeShader);
-	bool SetRenderPassInfo(FRHIRenderPassInfo const& InRenderTargets, FMetalQueryBuffer* QueryBuffer);
+	bool SetRenderPassInfo(FRHIRenderPassInfo const& InRenderTargets, FMetalQueryBufferRef QueryBuffer);
 	void InvalidateRenderTargets(void);
 	void SetRenderTargetsActive(bool const bActive);
 	void SetViewport(const MTL::Viewport& InViewport);
@@ -158,7 +158,7 @@ public:
 	void CommitRenderResources(FMetalCommandEncoder* Raster);
 	void CommitComputeResources(FMetalCommandEncoder* Compute);
 	void CommitResourceTable(EMetalShaderStages const Frequency, MTL::FunctionType const Type, FMetalCommandEncoder& CommandEncoder);
-	void StartRenderPass(const FRHIRenderPassInfo& Info, FMetalQueryBuffer* QueryBuffer, MTL::RenderPassDescriptor* InDesc, bool bInIsParallelContext);
+	void StartRenderPass(const FRHIRenderPassInfo& Info, FMetalQueryBufferRef QueryBuffer, MTL::RenderPassDescriptor* InDesc, bool bInIsParallelContext);
 	void EndRenderPass();
 	
 	FMetalShaderParameterCache& GetShaderParameters(EMetalShaderStages const Stage) { return ShaderParameters[Stage]; }
@@ -176,7 +176,7 @@ public:
 	const MTL::Viewport& GetViewport(uint32 const Index) const { check(Index < ML_MaxViewports); return Viewport[Index]; }
 	uint32 GetVertexBufferSize(uint32 const Index);
 	uint32 GetRenderTargetArraySize() const { return RenderTargetArraySize; }
-	FMetalQueryBuffer* GetVisibilityResultsBuffer() const { return VisibilityResults; }
+	FMetalQueryBufferRef GetVisibilityResultsBuffer() const { return VisibilityResults; }
 	bool NeedsToSetRenderTarget(const FRHIRenderPassInfo& RenderPassInfo);
 	bool HasValidDepthStencilSurface() const { return IsValidRef(DepthStencilSurface); }
 
@@ -315,7 +315,7 @@ private:
 	TStaticArray<TSet<MTL::Resource*>, EMetalShaderStages::Num> ROResourcesByStage;
 #endif
 
-	FMetalQueryBuffer* VisibilityResults;
+	FMetalQueryBufferRef VisibilityResults;
     MTL::VisibilityResultMode VisibilityMode;
 	NS::UInteger VisibilityOffset;
 	NS::UInteger VisibilityWritten;
diff --git a/Engine/Source/Runtime/Apple/MetalRHI/Private/MetalSubmission.cpp b/Engine/Source/Runtime/Apple/MetalRHI/Private/MetalSubmission.cpp
--- a/Engine/Source/Runtime/Apple/MetalRHI/Private/MetalSubmission.cpp
+++ b/Engine/Source/Runtime/Apple/MetalRHI/Private/MetalSubmission.cpp
@@ -845,7 +845,7 @@ FMetalDynamicRHI::FProcessResult FMetalDynamicRHI::ProcessInterruptQueue()
 			{
 				for (FMetalRHIRenderQuery* OcclusionQuery : Payload->BatchedObjects.OcclusionQueries)
 				{
-					OcclusionQuery->Result = OcclusionQuery->Buffer.GetResult();
+					OcclusionQuery->SampleOcclusionResult();
 				}
 				
 				for (FMetalRHIRenderQuery* TimestampQuery : Payload->BatchedObjects.TimestampQueries)
