Metal RHI

- Stop the submission thread getting stuck when encountering a NotPermitted GPU error.
- Only enqueue commands when processing RHIProcessDeleteQueue if we have objects to delete. This potentially caused a NotPermitted GPU error.

diff --git a/Engine/Source/Runtime/Apple/MetalRHI/Private/MetalRHI.cpp b/Engine/Source/Runtime/Apple/MetalRHI/Private/MetalRHI.cpp
--- a/Engine/Source/Runtime/Apple/MetalRHI/Private/MetalRHI.cpp
+++ b/Engine/Source/Runtime/Apple/MetalRHI/Private/MetalRHI.cpp
@@ -1653,6 +1653,11 @@ void FMetalDynamicRHI::RHIProcessDeleteQueue()
 	FScopeLock Lock(&ObjectsToDeleteCS);
 	TArray<FMetalDeferredDeleteObject> Objects = MoveTemp(ObjectsToDelete);
 	
+	if(!Objects.Num())
+	{
+		return;
+	}
+	
 	EnqueueEndOfPipeTask([this, LocalDeleteObjects = MoveTemp(Objects)]() mutable
 	{
 		for(FMetalDeferredDeleteObject& Object : LocalDeleteObjects)
diff --git a/Engine/Source/Runtime/Apple/MetalRHI/Private/MetalSubmission.cpp b/Engine/Source/Runtime/Apple/MetalRHI/Private/MetalSubmission.cpp
--- a/Engine/Source/Runtime/Apple/MetalRHI/Private/MetalSubmission.cpp
+++ b/Engine/Source/Runtime/Apple/MetalRHI/Private/MetalSubmission.cpp
@@ -816,7 +816,8 @@ FMetalDynamicRHI::FProcessResult FMetalDynamicRHI::ProcessInterruptQueue()
 			}
 			
 			// Remove Completed status check when we remove completion handlers
-			if (Status != MTL::CommandBufferStatusCompleted)
+			bool bCommandBufferFinished = Status == MTL::CommandBufferStatusCompleted || Status == MTL::CommandBufferStatusError;
+			if (!bCommandBufferFinished)
 			{
 				// Skip processing this queue and move on to the next.
 				Result.Status |= EQueueStatus::Pending;
