Reworking Metal Bindless implementation.
- Focus is around making bindless management similar across RHIs.
- Use FRHIHeapDescriptorAllocator to manage descriptor allocations.
- Moving resource->descriptor conversion code out of FMetalDescriptorHeap.
- Renaming descriptor management functions from "reserve" to "allocate".
- Adding allocate+init descriptor allocations.
- Removing unused UpdateDescriptorsWithCopy code.
- Moving descriptor update containers into a single type.

diff --git a/Engine/Source/Runtime/Apple/MetalRHI/Private/MetalBindlessDescriptors.cpp b/Engine/Source/Runtime/Apple/MetalRHI/Private/MetalBindlessDescriptors.cpp
--- a/Engine/Source/Runtime/Apple/MetalRHI/Private/MetalBindlessDescriptors.cpp
+++ b/Engine/Source/Runtime/Apple/MetalRHI/Private/MetalBindlessDescriptors.cpp
@@ -28,14 +28,87 @@ static FAutoConsoleVariableRef CVarBindlessSamplerDescriptorHeapSize(
 	ECVF_ReadOnly
 );
 
-static int32_t GMetalBindlessUpdateDescriptorOnCompute = 1;
-static FAutoConsoleVariableRef CVarBindlessUpdateDescriptorOnCompute(
-	TEXT("Metal.Bindless.UpdateDescriptorOnCompute"),
-	GMetalBindlessUpdateDescriptorOnCompute,
-	TEXT("How to update descriptors on the GPU")
-	TEXT("0 = Use Copy (slow), 1 = Use Compute shader (default)"),
-	ECVF_ReadOnly
-);
+static IRDescriptorTableEntry CreateDescriptor(MTL::SamplerState* SamplerState)
+{
+	IRDescriptorTableEntry DescriptorData{};
+	IRDescriptorTableSetSampler(&DescriptorData, SamplerState, 0.0f);
+	return DescriptorData;
+}
+
+static IRDescriptorTableEntry CreateDescriptor(MTL::Texture* Texture)
+{
+	IRDescriptorTableEntry DescriptorData{};
+	IRDescriptorTableSetTexture(&DescriptorData, Texture, 0.0f, 0u);
+	return DescriptorData;
+}
+
+static IRDescriptorTableEntry CreateDescriptor(FMetalBuffer* Buffer)
+{
+	IRDescriptorTableEntry DescriptorData{};
+	IRDescriptorTableSetBuffer(&DescriptorData, Buffer->GetGPUAddress(), Buffer->GetLength());
+	return DescriptorData;
+}
+
+static IRDescriptorTableEntry CreateDescriptor(const FMetalResourceViewBase::FBufferView& View)
+{
+	IRDescriptorTableEntry DescriptorData{};
+	IRDescriptorTableSetBuffer(&DescriptorData, View.Buffer->GetGPUAddress() + View.Offset, View.Size);
+	return DescriptorData;
+}
+
+static IRDescriptorTableEntry CreateDescriptor(const FMetalResourceViewBase::FTextureBufferBacked& View)
+{
+	IRDescriptorTableEntry DescriptorData{};
+	
+	IRBufferView BufferView;
+	BufferView.buffer = View.Buffer->GetMTLBuffer();
+	BufferView.bufferOffset = View.Buffer->GetOffset() + View.Offset;
+	BufferView.bufferSize = View.Size;
+	BufferView.typedBuffer = true;
+	BufferView.textureBufferView = View.Texture.get();
+	
+	uint32 Stride = GPixelFormats[View.Format].BlockBytes;
+	uint32 FirstElement = View.Offset / Stride;
+	uint32 NumElement = View.Size / Stride;
+	
+	uint64 BufferVA = View.Buffer->GetGPUAddress() + View.Offset;
+	uint64_t ExtraElement = (BufferVA % 16) / Stride;
+	
+	BufferView.textureViewOffsetInElements = ExtraElement;
+	
+	IRDescriptorTableSetBufferView(&DescriptorData, &BufferView);
+	
+	return DescriptorData;
+}
+
+#if METAL_RHI_RAYTRACING
+static IRDescriptorTableEntry CreateDescriptor(const MTL::AccelerationStructure& AccelerationStructure)
+{
+	IRDescriptorTableEntry DescriptorData{};
+	IRDescriptorTableSetAccelerationStructure(&DescriptorData, [AccelerationStructure.GetPtr() gpuResourceID]._impl);
+	return DescriptorData;
+}
+#endif
+
+static IRDescriptorTableEntry CreateDescriptor(FMetalResourceViewBase* Resource)
+{
+	switch (Resource->GetMetalType())
+	{
+	case FMetalResourceViewBase::EMetalType::TextureView:
+		return CreateDescriptor(Resource->GetTextureView().get());
+	case FMetalResourceViewBase::EMetalType::BufferView:
+		return CreateDescriptor(Resource->GetBufferView());
+	case FMetalResourceViewBase::EMetalType::TextureBufferBacked:
+		return CreateDescriptor(Resource->GetTextureBufferBacked());
+#if METAL_RHI_RAYTRACING
+	case FMetalResourceViewBase::EMetalType::AccelerationStructure:
+		return CreateDescriptor(Resource->GetAccelerationStructure());
+#endif
+	default:
+		checkNoEntry();
+		return IRDescriptorTableEntry{};
+	};
+}
 
 FMetalDescriptorHeap::FMetalDescriptorHeap(FMetalDevice& MetalDevice, ERHIDescriptorTypeMask InTypeMask)
 	: Device(MetalDevice)
@@ -53,63 +126,50 @@ void FMetalDescriptorHeap::Init(FRHICommandListBase& RHICmdList, int32 HeapSize)
 		FRHIBufferCreateDesc::CreateStructured(TEXT("ResourceHeap"), HeapSize, 1)
 		.AddUsage(EBufferUsageFlags::Dynamic | EBufferUsageFlags::KeepCPUAccessible | EBufferUsageFlags::UnorderedAccess | EBufferUsageFlags::StructuredBuffer)
 		.DetermineInitialState();
-	
-	ResourceHeapLength = HeapSize;
+
 	ResourceHeap = new FMetalRHIBuffer(RHICmdList, Device, CreateDesc, nullptr);
-	FMetalRHIBuffer* Buffer = ResourceCast(ResourceHeap.GetReference());
 	
-	Descriptors = reinterpret_cast<IRDescriptorTableEntry*>(Buffer->GetCurrentBuffer()->Contents());
+	Manager = new FRHIHeapDescriptorAllocator(TypeMask, HeapSize, {});
 }
 
-void FMetalDescriptorHeap::FreeDescriptor(FRHIDescriptorHandle DescriptorHandle)
+FRHIDescriptorHandle FMetalDescriptorHeap::AllocateDescriptor(ERHIDescriptorType DescriptorType)
 {
-	FMetalDynamicRHI::Get().DeferredDelete([this, DescriptorHandle]()
-	{
-		FScopeLock ScopeLock(&FreeListCS);	
-		FreeList.Enqueue(DescriptorHandle.GetIndex());
-	});
+	return Manager->Allocate(DescriptorType);
 }
 
-uint32 FMetalDescriptorHeap::GetFreeResourceIndex()
+FRHIDescriptorHandle FMetalDescriptorHeap::AllocateDescriptor(ERHIDescriptorType DescriptorType, IRDescriptorTableEntry DescriptorData)
 {
-	{
-		FScopeLock ScopeLock(&FreeListCS);
-		if (!FreeList.IsEmpty())
-		{
-			uint32 FreeIndex;
-			FreeList.Dequeue(FreeIndex);
-			return FreeIndex;
-		}
-	}
-
-	NSUInteger MaxDescriptorCount = ResourceHeapLength / sizeof(IRDescriptorTableEntry);
-	checkf((PeakDescriptorCount + 1) < MaxDescriptorCount, TEXT("Reached Heap Max Capacity (%u/%u)"), PeakDescriptorCount + 1, MaxDescriptorCount);
-
-	const uint32 ResourceIndex = PeakDescriptorCount++;
-	return ResourceIndex;
+	FRHIDescriptorHandle Handle = AllocateDescriptor(DescriptorType);
+	UpdateDescriptorImmediately(Handle, DescriptorData);
+	return Handle;
 }
 
-FRHIDescriptorHandle FMetalDescriptorHeap::ReserveDescriptor(ERHIDescriptorType DescriptorType)
+void FMetalDescriptorHeap::FreeDescriptor(FRHIDescriptorHandle DescriptorHandle)
 {
-	const uint32 ResourceIndex = GetFreeResourceIndex();
-	return FRHIDescriptorHandle(DescriptorType, ResourceIndex);
+	FMetalDynamicRHI::Get().DeferredDelete([this, DescriptorHandle]()
+	{
+		Manager->Free(DescriptorHandle);
+	});
 }
 
-void FMetalDescriptorHeap::UpdateDescriptor(FRHIDescriptorHandle DescriptorHandle, IRDescriptorTableEntry DescriptorData)
+void FMetalDescriptorHeap::UpdateDescriptorImmediately(FRHIDescriptorHandle DescriptorHandle, IRDescriptorTableEntry DescriptorData)
 {
 	checkf(DescriptorHandle.IsValid(), TEXT("Attemping to update invalid descriptor handle!"));
 	
 	uint32 DescriptorIndex = DescriptorHandle.GetIndex();
-	Descriptors[DescriptorIndex] = DescriptorData;
+	GetDescriptorMemory()[DescriptorIndex] = DescriptorData;
 }
 
-void FMetalDescriptorHeap::BindHeap(FMetalCommandEncoder* Encoder, MTL::FunctionType FunctionType, uint32 BindIndex)
+void FMetalDescriptorHeap::BindToEncoder(FMetalCommandEncoder* Encoder, MTL::FunctionType FunctionType, uint32 BindIndex)
 {
-	uint32 DescriptorCount = PeakDescriptorCount.load();
-	const uint64 HeapSize = DescriptorCount * sizeof(IRDescriptorTableEntry);
-
-	FMetalRHIBuffer* Buffer = ResourceCast(ResourceHeap.GetReference());
-	Encoder->SetShaderBuffer(FunctionType, Buffer->GetCurrentBuffer(), 0, HeapSize, BindIndex, MTL::ResourceUsageRead);
+	FRHIDescriptorAllocatorRange AllocatedRange(0, 0);
+	if (Manager && Manager->GetAllocatedRange(AllocatedRange))
+	{
+		uint32 DescriptorCount = AllocatedRange.Last;
+		const uint64 HeapSize = DescriptorCount * sizeof(IRDescriptorTableEntry);
+		
+		Encoder->SetShaderBuffer(FunctionType, ResourceHeap->GetCurrentBuffer(), 0, HeapSize, BindIndex, MTL::ResourceUsageRead);
+	}
 }
 
 FMetalBindlessDescriptorManager::FMetalBindlessDescriptorManager(FMetalDevice& MetalDevice)
@@ -121,7 +181,6 @@ FMetalBindlessDescriptorManager::FMetalBindlessDescriptorManager(FMetalDevice& M
 
 FMetalBindlessDescriptorManager::~FMetalBindlessDescriptorManager()
 {
-
 }
 
 void FMetalBindlessDescriptorManager::Init()
@@ -134,21 +193,39 @@ void FMetalBindlessDescriptorManager::Init()
 	bIsSupported = true;
 }
 
-FRHIDescriptorHandle FMetalBindlessDescriptorManager::ReserveDescriptor(ERHIDescriptorType DescriptorType)
+FRHIDescriptorHandle FMetalBindlessDescriptorManager::AllocateDescriptor(ERHIDescriptorType DescriptorType)
 {
 	if (StandardResources.HandlesAllocation(DescriptorType))
 	{
-		return StandardResources.ReserveDescriptor(DescriptorType);
+		return StandardResources.AllocateDescriptor(DescriptorType);
 	}
-
+	
 	if (SamplerResources.HandlesAllocation(DescriptorType))
 	{
-		return SamplerResources.ReserveDescriptor(DescriptorType);
+		return SamplerResources.AllocateDescriptor(DescriptorType);
 	}
-
+	
 	return FRHIDescriptorHandle();
 }
 
+FRHIDescriptorHandle FMetalBindlessDescriptorManager::AllocateDescriptor(ERHIDescriptorType DescriptorType, MTL::Texture* Texture)
+{
+	IRDescriptorTableEntry Descriptor = CreateDescriptor(Texture);
+	return StandardResources.AllocateDescriptor(DescriptorType, Descriptor);
+}
+
+FRHIDescriptorHandle FMetalBindlessDescriptorManager::AllocateDescriptor(ERHIDescriptorType DescriptorType, FMetalBuffer* Buffer)
+{
+	IRDescriptorTableEntry Descriptor = CreateDescriptor(Buffer);
+	return StandardResources.AllocateDescriptor(DescriptorType, Descriptor);
+}
+
+FRHIDescriptorHandle FMetalBindlessDescriptorManager::AllocateDescriptor(MTL::SamplerState* SamplerState)
+{
+	IRDescriptorTableEntry Descriptor = CreateDescriptor(SamplerState);
+	return SamplerResources.AllocateDescriptor(ERHIDescriptorType::Sampler, Descriptor);
+}
+
 void FMetalBindlessDescriptorManager::FreeDescriptor(FRHIDescriptorHandle DescriptorHandle)
 {
 	check(DescriptorHandle.IsValid());
@@ -157,160 +234,67 @@ void FMetalBindlessDescriptorManager::FreeDescriptor(FRHIDescriptorHandle Descri
 
 	if (StandardResources.HandlesAllocation(DescriptorType))
 	{
-		return StandardResources.FreeDescriptor(DescriptorHandle);
+		StandardResources.FreeDescriptor(DescriptorHandle);
 	}
-
-	if (SamplerResources.HandlesAllocation(DescriptorType))
+	else if (SamplerResources.HandlesAllocation(DescriptorType))
 	{
-		return SamplerResources.FreeDescriptor(DescriptorHandle);
+		SamplerResources.FreeDescriptor(DescriptorHandle);
+	}
+	else
+	{
+		checkNoEntry();
 	}
-
-	checkNoEntry();
-}
-
-void FMetalDescriptorHeap::BindSampler(FRHIDescriptorHandle DescriptorHandle, MTL::SamplerState* Sampler)
-{
-	IRDescriptorTableEntry DescriptorData = {0};
-	IRDescriptorTableSetSampler(&DescriptorData, Sampler, 0.0f);
-
-	UpdateDescriptor(DescriptorHandle, DescriptorData);
 }
 
-void FMetalDescriptorHeap::BindBufferImmediate(FRHIDescriptorHandle DescriptorHandle, FMetalBufferPtr Resource)
+void FMetalBindlessDescriptorManager::UpdateDescriptor(FRHICommandListBase& RHICmdList, FRHIDescriptorHandle DescriptorHandle, MTL::Texture* Texture, EMetalDescriptorUpdateType UpdateType)
 {
-	IRDescriptorTableEntry DescriptorData = {0};
-	
-	IRDescriptorTableSetBuffer(&DescriptorData, Resource->GetGPUAddress(), Resource->GetLength());
+	IRDescriptorTableEntry Descriptor = CreateDescriptor(Texture);
 	
-	UpdateDescriptor(DescriptorHandle, DescriptorData);
-}
-
-void FMetalDescriptorHeap::BindResource(FRHIDescriptorHandle DescriptorHandle, FMetalResourceViewBase* Resource, FMetalRHICommandContext* Context, const EDescriptorUpdateType UpdateType)
-{
-	check(IsInRHIThread() || IsInRenderingThread());
+	UpdateType = !GIsRHIInitialized ? EMetalDescriptorUpdateType::Immediate : UpdateType;
 	
-	IRDescriptorTableEntry DescriptorData = {0};
-
-	switch (Resource->GetMetalType())
-	{
-	case FMetalResourceViewBase::EMetalType::TextureView:
-		{
-			auto const& View = Resource->GetTextureView();
-
-			IRDescriptorTableSetTexture(&DescriptorData, View.get(), 0.0f, 0u);
-		}
-		break;
-	case FMetalResourceViewBase::EMetalType::BufferView:
-		{
-			auto const& View = Resource->GetBufferView();
-
-			IRDescriptorTableSetBuffer(&DescriptorData, View.Buffer->GetGPUAddress() + View.Offset, View.Size);
-		}
-		break;
-	case FMetalResourceViewBase::EMetalType::TextureBufferBacked:
-		{
-			auto const& View = Resource->GetTextureBufferBacked();
-
-			IRBufferView BufferView;
-			BufferView.buffer = View.Buffer->GetMTLBuffer();
-			BufferView.bufferOffset = View.Buffer->GetOffset() + View.Offset;
-			BufferView.bufferSize = View.Size;
-			BufferView.typedBuffer = true;
-			BufferView.textureBufferView = View.Texture.get();
-
-			uint32 Stride = GPixelFormats[View.Format].BlockBytes;
-			uint32 FirstElement = View.Offset / Stride;
-			uint32 NumElement = View.Size / Stride;
-
-			uint64 BufferVA              = View.Buffer->GetGPUAddress() + View.Offset;
-			uint64_t ExtraElement        = (BufferVA % 16) / Stride;
-
-			BufferView.textureViewOffsetInElements = ExtraElement;
-
-			IRDescriptorTableSetBufferView(&DescriptorData, &BufferView);
-		}
-		break;
-#if METAL_RHI_RAYTRACING
-	case FMetalResourceViewBase::EMetalType::AccelerationStructure:
-		{
-			MTL::AccelerationStructure const& AccelerationStructure = Resource->GetAccelerationStructure();
-
-			IRDescriptorTableSetAccelerationStructure(&DescriptorData, [AccelerationStructure.GetPtr() gpuResourceID]._impl);
-		}
-		break;
-#endif
-	default:
-		checkNoEntry();
-		return;
-	};
-
+	RHICmdList.EnqueueLambda([this, UpdateType, Descriptor, DescriptorHandle](FRHICommandListBase& RHICmdList)
 	{
-		if(Context && UpdateType != EDescriptorUpdateType_Immediate)
+		if (UpdateType == EMetalDescriptorUpdateType::Immediate)
 		{
-			check(GIsRHIInitialized);
-			Context->EnqueueDescriptorUpdate(DescriptorData, DescriptorHandle.GetIndex());
+			StandardResources.UpdateDescriptorImmediately(DescriptorHandle, Descriptor);
 		}
 		else
-		{
-			UpdateDescriptor(DescriptorHandle, DescriptorData);
+		{	
+			FMetalRHICommandContext& Context = FMetalRHICommandContext::Get(RHICmdList);
+			Context.EnqueueDescriptorUpdate(DescriptorHandle, Descriptor);
 		}
-	}
+	});
+	
+	RHICmdList.RHIThreadFence(true);
 }
 
-void FMetalBindlessDescriptorManager::UpdateDescriptorsWithGPU(FMetalRHICommandContext* Context, TResourceArray<IRDescriptorTableEntry>&& InComputeDescriptorEntries,
-															   TResourceArray<uint32_t>&& InComputeDescriptorIndices)
+void FMetalBindlessDescriptorManager::UpdateDescriptor(FMetalRHICommandContext* Context, FRHIDescriptorHandle DescriptorHandle, FMetalResourceViewBase* Resource, EMetalDescriptorUpdateType UpdateType)
 {
-	if(GMetalBindlessUpdateDescriptorOnCompute)
+	IRDescriptorTableEntry Descriptor = CreateDescriptor(Resource);
+	
+	if (Context && UpdateType != EMetalDescriptorUpdateType::Immediate)
 	{
-		UpdateDescriptorsWithCompute(Context, MoveTemp(InComputeDescriptorEntries), MoveTemp(InComputeDescriptorIndices));
+		check(GIsRHIInitialized);
+		Context->EnqueueDescriptorUpdate(DescriptorHandle, Descriptor);
 	}
 	else
 	{
-		UpdateDescriptorsWithCopy(Context, MoveTemp(InComputeDescriptorEntries), MoveTemp(InComputeDescriptorIndices));
+		StandardResources.UpdateDescriptorImmediately(DescriptorHandle, Descriptor);
 	}
 }
 
-void FMetalBindlessDescriptorManager::UpdateDescriptorsWithCopy(FMetalRHICommandContext* Context, 
-																   TResourceArray<IRDescriptorTableEntry>&& ComputeDescriptorEntriesCopy,
-																   TResourceArray<uint32_t>&& ComputeDescriptorIndicesCopy)
+void FMetalBindlessDescriptorManager::FlushPendingDescriptorUpdates(FMetalRHICommandContext& Context, const FMetalPendingDescriptorUpdates& PendingDescriptorUpdates)
 {
-	const uint32_t NumDescriptors = ComputeDescriptorIndicesCopy.Num();
+	const uint32_t NumDescriptors = PendingDescriptorUpdates.Num();
 	
-	if(!NumDescriptors)
+	if (!NumDescriptors)
 	{
 		return;
 	}
 	
-	for(uint32_t Idx = 0; Idx < NumDescriptors; ++Idx)
-	{
-		FMetalBufferPtr SourceBuffer = Device.GetTransferAllocator()->Allocate(sizeof(IRDescriptorTableEntry));
-		
-		FMetalRHIBuffer* DestBuffer = ResourceCast(StandardResources.ResourceHeap.GetReference());
-		IRDescriptorTableEntry* DescriptorCopy = (IRDescriptorTableEntry*)SourceBuffer->Contents();
-		
-		IRDescriptorTableEntry Data = ComputeDescriptorEntriesCopy[Idx];
-		FMemory::Memcpy(DescriptorCopy, &Data, sizeof(IRDescriptorTableEntry));
-		
-		uint32_t DescIdx = ComputeDescriptorIndicesCopy[Idx];
-		Context->CopyFromBufferToBuffer(SourceBuffer, 0, DestBuffer->GetCurrentBuffer(), DescIdx*sizeof(IRDescriptorTableEntry), sizeof(IRDescriptorTableEntry));
-	}
-}
-
-
-void FMetalBindlessDescriptorManager::UpdateDescriptorsWithCompute(FMetalRHICommandContext* Context, 
-																   TResourceArray<IRDescriptorTableEntry>&& ComputeDescriptorEntriesCopy,
-																   TResourceArray<uint32_t>&& ComputeDescriptorIndicesCopy)
-{
-	const uint32_t NumDescriptors = ComputeDescriptorIndicesCopy.Num();
-	
-	if(!NumDescriptors)
-	{
-		return;
-	}
-
 	FMetalDescriptorHeap* Heap = new FMetalDescriptorHeap(Device, ERHIDescriptorTypeMask::Standard);
 	
-	TRHICommandList_RecursiveHazardous<FMetalRHICommandContext> RHICmdList(Context);
+	TRHICommandList_RecursiveHazardous<FMetalRHICommandContext> RHICmdList(&Context);
 	Heap->Init(RHICmdList, 256);
 	
 	TShaderMapRef<FUpdateDescriptorHandleCS> ComputeShader(GetGlobalShaderMap(GMaxRHIFeatureLevel));
@@ -320,107 +304,54 @@ void FMetalBindlessDescriptorManager::UpdateDescriptorsWithCompute(FMetalRHIComm
 	
 	// Create the two temporary buffers for descriptor entries and indices
 	FMetalBufferPtr DescriptorEntriesBuffer = Allocator->Allocate(NumDescriptors * sizeof(IRDescriptorTableEntry));
-	FMemory::Memcpy(DescriptorEntriesBuffer->Contents(), ComputeDescriptorEntriesCopy.GetResourceData(), ComputeDescriptorEntriesCopy.GetResourceDataSize());
-
-	FMetalBufferPtr DescriptorIndiciesBuffer = Allocator->Allocate(NumDescriptors * sizeof(uint));
-	FMemory::Memcpy(DescriptorIndiciesBuffer->Contents(), ComputeDescriptorIndicesCopy.GetResourceData(), ComputeDescriptorIndicesCopy.GetResourceDataSize());
-	
-	// Create descriptor handles and bind buffers to descriptor heap
-	FRHIDescriptorHandle DescriptorEntriesHandle = Heap->ReserveDescriptor(ERHIDescriptorType::BufferSRV);
-	Heap->BindBufferImmediate(DescriptorEntriesHandle, DescriptorEntriesBuffer);
+	FMemory::Memcpy(DescriptorEntriesBuffer->Contents(), PendingDescriptorUpdates.Descriptors.GetData(), NumDescriptors * sizeof(IRDescriptorTableEntry));
 	
-	FRHIDescriptorHandle DescriptorIndicesHandle = Heap->ReserveDescriptor(ERHIDescriptorType::BufferSRV);
-	Heap->BindBufferImmediate(DescriptorIndicesHandle, DescriptorIndiciesBuffer);
+	FMetalBufferPtr DescriptorIndiciesBuffer = Allocator->Allocate(NumDescriptors * sizeof(uint32));
+	FMemory::Memcpy(DescriptorIndiciesBuffer->Contents(), PendingDescriptorUpdates.Indices.GetData(), NumDescriptors * sizeof(uint32));
 	
-	FRHIDescriptorHandle DestBufferHandle = Heap->ReserveDescriptor(ERHIDescriptorType::BufferUAV);
-	
-	FMetalRHIBuffer* DestBuffer = (FMetalRHIBuffer*)StandardResources.ResourceHeap.GetReference();
-	Heap->BindBufferImmediate(DestBufferHandle, DestBuffer->GetCurrentBuffer());
+	// Create descriptor handles and bind buffers to descriptor heap
+	FRHIDescriptorHandle DescriptorEntriesHandle = Heap->AllocateDescriptor(ERHIDescriptorType::BufferSRV, CreateDescriptor(DescriptorEntriesBuffer.Get()));
+	FRHIDescriptorHandle DescriptorIndicesHandle = Heap->AllocateDescriptor(ERHIDescriptorType::BufferSRV, CreateDescriptor(DescriptorIndiciesBuffer.Get()));
 	
+	FMetalRHIBuffer* DestBuffer = StandardResources.GetHeap();
+	FRHIDescriptorHandle DestBufferHandle = Heap->AllocateDescriptor(ERHIDescriptorType::BufferUAV, CreateDescriptor(DestBuffer->GetCurrentBuffer().Get()));
+
 	// Create a temporary uniform buffer
-	uint32_t* UBDataPtr = (uint32_t*)FMemory::Malloc(16);
-	
+	uint32* UBDataPtr = reinterpret_cast<uint32*>(FMemory::Malloc(4 * sizeof(uint32)));
+
 	UBDataPtr[0] = NumDescriptors;
 	UBDataPtr[1] = DescriptorIndicesHandle.GetIndex();
 	UBDataPtr[2] = DescriptorEntriesHandle.GetIndex();
 	UBDataPtr[3] = DestBufferHandle.GetIndex();
-	
-	{
-		Context->GetStateCache().SetComputeShader((FMetalComputeShader*)ShaderRHI);
-		
-		FMetalRHIBuffer* Dest = (FMetalRHIBuffer*)StandardResources.ResourceHeap.GetReference();
-		
-		Context->GetStateCache().CacheOrSkipResourceResidencyUpdate(DescriptorEntriesBuffer->GetMTLBuffer(), EMetalShaderStages::Compute, true, true);
-		Context->GetStateCache().CacheOrSkipResourceResidencyUpdate(DescriptorIndiciesBuffer->GetMTLBuffer(), EMetalShaderStages::Compute, true, true);
-		Context->GetStateCache().CacheOrSkipResourceResidencyUpdate(Dest->GetCurrentBuffer()->GetMTLBuffer(), EMetalShaderStages::Compute, false, true);
-		
-		Context->GetStateCache().SetOverriddenDescriptorHeap(Heap);
-		
-		FMetalBufferPtr Buffer;
-		Context->GetStateCache().IRBindPackedUniforms(EMetalShaderStages::Compute, 0, (const uint8*)UBDataPtr, 16, Buffer);
-		Context->GetStateCache().GetShaderParameters(EMetalShaderStages::Compute).MarkAllDirty();
-		
-		Context->RHIDispatchComputeShader(1, 1, 1);
-		Context->GetStateCache().SetOverriddenDescriptorHeap(nullptr);
-		
-		FMetalDynamicRHI::Get().DeferredDelete([DescriptorEntriesBuffer, DescriptorIndiciesBuffer, Heap, UBDataPtr](){
-			delete Heap;
-			FMemory::Free(UBDataPtr);
-		});
-	}
-}
 
-void FMetalDescriptorHeap::BindTexture(FRHICommandListBase* RHICmdList, FRHIDescriptorHandle DescriptorHandle, MTL::Texture* Texture, EDescriptorUpdateType UpdateType)
-{
-	IRDescriptorTableEntry DescriptorData = {0};
-	IRDescriptorTableSetTexture(&DescriptorData, Texture, 0.0f, 0u);
-	
-	UpdateType = !GIsRHIInitialized ? EDescriptorUpdateType_Immediate : UpdateType;
-	
-	if(RHICmdList)
-	{
-		RHICmdList->EnqueueLambda([this, InUpdateType=UpdateType, Data=DescriptorData, Handle=DescriptorHandle](FRHICommandListBase& RHICmdList)
-		{
-			if(InUpdateType == EDescriptorUpdateType_Immediate)
-			{
-				UpdateDescriptor(Handle, Data);
-			}
-			else
-			{	
-				FMetalRHICommandContext* Context = &FMetalRHICommandContext::Get(RHICmdList);
-				Context->EnqueueDescriptorUpdate(Data, Handle.GetIndex());
-			}
-		});
-		
-		RHICmdList->RHIThreadFence(true);
-	}
-	else
-	{
-		check(!IsRunningRHIInSeparateThread() || IsInRHIThread() || IsInParallelRHIThread());
-		UpdateDescriptor(DescriptorHandle, DescriptorData);
-	}
-}
+	Context.GetStateCache().SetComputeShader((FMetalComputeShader*)ShaderRHI);
 
-void FMetalBindlessDescriptorManager::BindDescriptorHeapsToEncoder(FMetalCommandEncoder* Encoder, MTL::FunctionType FunctionType, EMetalShaderStages Frequency)
-{
-	StandardResources.BindHeap(Encoder, FunctionType, kIRStandardHeapBindPoint);
-	SamplerResources.BindHeap(Encoder, FunctionType, kIRSamplerHeapBindPoint);
-}
+	FMetalRHIBuffer* Dest = StandardResources.GetHeap();
 
-void FMetalBindlessDescriptorManager::BindSampler(FRHIDescriptorHandle DescriptorHandle, MTL::SamplerState* Sampler)
-{
-	SamplerResources.BindSampler(DescriptorHandle, Sampler);
-}
+	Context.GetStateCache().CacheOrSkipResourceResidencyUpdate(DescriptorEntriesBuffer->GetMTLBuffer(), EMetalShaderStages::Compute, true, true);
+	Context.GetStateCache().CacheOrSkipResourceResidencyUpdate(DescriptorIndiciesBuffer->GetMTLBuffer(), EMetalShaderStages::Compute, true, true);
+	Context.GetStateCache().CacheOrSkipResourceResidencyUpdate(Dest->GetCurrentBuffer()->GetMTLBuffer(), EMetalShaderStages::Compute, false, true);
 
-void FMetalBindlessDescriptorManager::BindResource(FRHIDescriptorHandle DescriptorHandle, FMetalResourceViewBase* Resource,
-												   FMetalRHICommandContext* Context, const EDescriptorUpdateType UpdateType)
-{
-	StandardResources.BindResource(DescriptorHandle, Resource, Context, UpdateType);
+	Context.GetStateCache().SetOverriddenDescriptorHeap(Heap);
+
+	FMetalBufferPtr Buffer;
+	Context.GetStateCache().IRBindPackedUniforms(EMetalShaderStages::Compute, 0, (const uint8*)UBDataPtr, 16, Buffer);
+	Context.GetStateCache().GetShaderParameters(EMetalShaderStages::Compute).MarkAllDirty();
+
+	Context.RHIDispatchComputeShader(1, 1, 1);
+	Context.GetStateCache().SetOverriddenDescriptorHeap(nullptr);
+
+	FMetalDynamicRHI::Get().DeferredDelete([DescriptorEntriesBuffer, DescriptorIndiciesBuffer, Heap, UBDataPtr]()
+	{
+		delete Heap;
+		FMemory::Free(UBDataPtr);
+	});
 }
 
-void FMetalBindlessDescriptorManager::BindTexture(FRHICommandListBase* RHICmdList, FRHIDescriptorHandle DescriptorHandle, MTL::Texture* Texture, EDescriptorUpdateType UpdateType)
+void FMetalBindlessDescriptorManager::BindDescriptorHeapsToEncoder(FMetalCommandEncoder* Encoder, MTL::FunctionType FunctionType, EMetalShaderStages Frequency)
 {
-	StandardResources.BindTexture(RHICmdList, DescriptorHandle, Texture, UpdateType);
+	StandardResources.BindToEncoder(Encoder, FunctionType, kIRStandardHeapBindPoint);
+	SamplerResources.BindToEncoder(Encoder, FunctionType, kIRSamplerHeapBindPoint);
 }
 
 #endif //PLATFORM_SUPPORTS_BINDLESS_RENDERING
diff --git a/Engine/Source/Runtime/Apple/MetalRHI/Private/MetalBindlessDescriptors.h b/Engine/Source/Runtime/Apple/MetalRHI/Private/MetalBindlessDescriptors.h
--- a/Engine/Source/Runtime/Apple/MetalRHI/Private/MetalBindlessDescriptors.h
+++ b/Engine/Source/Runtime/Apple/MetalRHI/Private/MetalBindlessDescriptors.h
@@ -8,100 +8,122 @@
 #if PLATFORM_SUPPORTS_BINDLESS_RENDERING
 #include "MetalResources.h"
 #include "MetalShaderResources.h"
-#include "Containers/DynamicRHIResourceArray.h"
+#include "RHIDescriptorAllocator.h"
 
 class FMetalCommandEncoder;
 class FMetalDevice;
 
 struct IRDescriptorTableEntry;
 
-enum EDescriptorUpdateType
+enum class EMetalDescriptorUpdateType
 {
-	EDescriptorUpdateType_Immediate,
-	EDescriptorUpdateType_GPU
+	Immediate,
+	GPU,
+};
+
+struct FMetalPendingDescriptorUpdates
+{
+	TArray<uint32> Indices;
+	TArray<IRDescriptorTableEntry> Descriptors;
+
+	int32 Num() const
+	{
+		return Indices.Num();
+	}
+
+	void Add(FRHIDescriptorHandle InHandle, const IRDescriptorTableEntry& InDescriptor)
+	{
+		if (ensure(InHandle.IsValid()))
+		{
+			Indices.Emplace(InHandle.GetIndex());
+			Descriptors.Emplace(InDescriptor);
+		}
+	}
+
+	void Empty()
+	{
+		Indices.Empty();
+		Descriptors.Empty();
+	}
 };
 
 struct FMetalDescriptorHeap
 {
 	FMetalDescriptorHeap(FMetalDevice& MetalDevice, ERHIDescriptorTypeMask InTypeMask);
 	~FMetalDescriptorHeap();
-	
-    void                    Init(FRHICommandListBase& RHICmdList, int32 HeapSize);
-
-    FRHIDescriptorHandle    ReserveDescriptor(ERHIDescriptorType DescriptorType);
-    void                    FreeDescriptor(FRHIDescriptorHandle DescriptorHandle);
-    uint32                  GetFreeResourceIndex();
-
-    void                    UpdateDescriptor(FRHIDescriptorHandle DescriptorHandle, struct IRDescriptorTableEntry DescriptorData);
-    void                    BindHeap(FMetalCommandEncoder* Encoder, MTL::FunctionType FunctionType, uint32 BindIndex);
-    
-	void                   	BindSampler(FRHIDescriptorHandle DescriptorHandle, MTL::SamplerState* Sampler);
-	void                    BindBufferImmediate(FRHIDescriptorHandle DescriptorHandle, FMetalBufferPtr Resource);
-	void                    BindResource(FRHIDescriptorHandle DescriptorHandle, FMetalResourceViewBase* Resource,
-	                                     FMetalRHICommandContext* Context, const EDescriptorUpdateType UpdateType = EDescriptorUpdateType_GPU);
-	void                    BindTexture(FRHICommandListBase* RHICmdList, FRHIDescriptorHandle DescriptorHandle, MTL::Texture* Texture, const EDescriptorUpdateType UpdateType);
-
-	constexpr ERHIDescriptorTypeMask GetTypeMask() const
+
+	void Init(FRHICommandListBase& RHICmdList, int32 HeapSize);
+
+	FRHIDescriptorHandle AllocateDescriptor(ERHIDescriptorType DescriptorType);
+	FRHIDescriptorHandle AllocateDescriptor(ERHIDescriptorType DescriptorType, IRDescriptorTableEntry DescriptorData);
+	void                 FreeDescriptor(FRHIDescriptorHandle DescriptorHandle);
+
+	void UpdateDescriptorImmediately(FRHIDescriptorHandle DescriptorHandle, IRDescriptorTableEntry DescriptorData);
+	void BindToEncoder(FMetalCommandEncoder* Encoder, MTL::FunctionType FunctionType, uint32 BindIndex);
+
+	ERHIDescriptorTypeMask GetTypeMask() const
 	{
 		return TypeMask;
 	}
 
-	constexpr bool HandlesAllocation(ERHIDescriptorType InType) const
+	bool HandlesAllocation(ERHIDescriptorType InType) const
 	{
 		return EnumHasAnyFlags(GetTypeMask(), RHIDescriptorTypeMaskFromType(InType));
 	}
-	
-    //
-	FMetalDevice& 					Device;
-	
-    FCriticalSection                FreeListCS;
-    TQueue<uint32>                  FreeList;
-
-    std::atomic<uint32>             PeakDescriptorCount;
-    IRDescriptorTableEntry*         Descriptors = nullptr;
-    uint32_t                        ResourceHeapLength;
-    FBufferRHIRef			        ResourceHeap;
-    
-    const ERHIDescriptorTypeMask    TypeMask;
+
+	FMetalRHIBuffer* GetHeap() const
+	{
+		return ResourceHeap.GetReference();
+	}
+
+private:
+	IRDescriptorTableEntry* GetDescriptorMemory()
+	{
+		checkSlow(ResourceHeap != nullptr);
+		return reinterpret_cast<IRDescriptorTableEntry*>(ResourceHeap->GetCurrentBuffer()->Contents());
+	}
+
+	FMetalDevice& Device;
+
+	FRHIHeapDescriptorAllocator* Manager = nullptr;
+	TRefCountPtr<FMetalRHIBuffer> ResourceHeap;
+
+	const ERHIDescriptorTypeMask TypeMask;
 };
 
 class FMetalBindlessDescriptorManager
 {
 public:
-                            FMetalBindlessDescriptorManager(FMetalDevice& MetalDevice);
-                            ~FMetalBindlessDescriptorManager();
+	FMetalBindlessDescriptorManager(FMetalDevice& MetalDevice);
+	~FMetalBindlessDescriptorManager();
 
-    void                    Init();
+    void Init();
 
-    FRHIDescriptorHandle    ReserveDescriptor(ERHIDescriptorType DescriptorType);
-    void                    FreeDescriptor(FRHIDescriptorHandle DescriptorHandle);
+	FRHIDescriptorHandle AllocateDescriptor(ERHIDescriptorType DescriptorType);
+	FRHIDescriptorHandle AllocateDescriptor(ERHIDescriptorType DescriptorType, MTL::Texture* Texture);
+	FRHIDescriptorHandle AllocateDescriptor(ERHIDescriptorType DescriptorType, FMetalBuffer* Buffer);
+	FRHIDescriptorHandle AllocateDescriptor(MTL::SamplerState* SamplerState);
 
-    void                    BindSampler(FRHIDescriptorHandle DescriptorHandle, MTL::SamplerState* Sampler);
-    void                    BindResource(FRHIDescriptorHandle DescriptorHandle, FMetalResourceViewBase* Resource,
-										 FMetalRHICommandContext* Context, const EDescriptorUpdateType UpdateType = EDescriptorUpdateType_GPU);
-    void                    BindTexture(FRHICommandListBase* RHICmdList, FRHIDescriptorHandle DescriptorHandle, MTL::Texture* Texture, EDescriptorUpdateType UpdateType);
+	void FreeDescriptor(FRHIDescriptorHandle DescriptorHandle);
+
+	void UpdateDescriptor(FRHICommandListBase& RHICmdList, FRHIDescriptorHandle DescriptorHandle, MTL::Texture* Texture, EMetalDescriptorUpdateType UpdateType);
+	void UpdateDescriptor(FMetalRHICommandContext* Context, FRHIDescriptorHandle DescriptorHandle, FMetalResourceViewBase* Resource, EMetalDescriptorUpdateType UpdateType);
+	
+	void FlushPendingDescriptorUpdates(FMetalRHICommandContext& Context, const FMetalPendingDescriptorUpdates& PendingDescriptorUpdates);
+    void BindDescriptorHeapsToEncoder(FMetalCommandEncoder* Encoder, MTL::FunctionType FunctionType, EMetalShaderStages Frequency);
 
-    void                    BindDescriptorHeapsToEncoder(FMetalCommandEncoder* Encoder, MTL::FunctionType FunctionType, EMetalShaderStages Frequency);
+	bool IsSupported() const
+	{
+		return bIsSupported;
+	}
 
-	bool					IsSupported() {return bIsSupported;}
-    
-	void 					UpdateDescriptorsWithGPU(FMetalRHICommandContext* Context,
-													 TResourceArray<IRDescriptorTableEntry>&& ComputeDescriptorEntries,
-													 TResourceArray<uint32_t>&& ComputeDescriptorIndices);
-    
 private:
-	void                    UpdateDescriptorsWithCompute(FMetalRHICommandContext* Context, TResourceArray<IRDescriptorTableEntry>&& ComputeDescriptorEntries,
-														 TResourceArray<uint32_t>&& ComputeDescriptorIndices);
-	
-	void                    UpdateDescriptorsWithCopy(FMetalRHICommandContext* Context, TResourceArray<IRDescriptorTableEntry>&& ComputeDescriptorEntries,
-														 TResourceArray<uint32_t>&& ComputeDescriptorIndices);
-	
+	FMetalDevice& Device;
 	
+    FMetalDescriptorHeap StandardResources;
+    FMetalDescriptorHeap SamplerResources;
 	
-	bool 					bIsSupported = false;
-	FMetalDevice& 			Device;
-    FMetalDescriptorHeap    StandardResources;
-    FMetalDescriptorHeap    SamplerResources;
+	bool bIsSupported = false;
 };
 
 #endif
diff --git a/Engine/Source/Runtime/Apple/MetalRHI/Private/MetalRHIContext.cpp b/Engine/Source/Runtime/Apple/MetalRHI/Private/MetalRHIContext.cpp
--- a/Engine/Source/Runtime/Apple/MetalRHI/Private/MetalRHIContext.cpp
+++ b/Engine/Source/Runtime/Apple/MetalRHI/Private/MetalRHIContext.cpp
@@ -70,7 +70,7 @@ FMetalRHICommandContext::~FMetalRHICommandContext()
 void FMetalRHICommandContext::ResetContext()
 {
 #if PLATFORM_SUPPORTS_BINDLESS_RENDERING
-	check(!ComputeDescriptorIndices.Num());
+	check(!PendingDescriptorUpdates.Num());
 #endif
 	
 	// Reset cached state in the encoder
@@ -227,12 +227,12 @@ void FMetalRHICommandContext::PushDescriptorUpdates()
 	MTL_SCOPED_AUTORELEASE_POOL;
 
 #if PLATFORM_SUPPORTS_BINDLESS_RENDERING
-	if(IsMetalBindlessEnabled())
+	if (PendingDescriptorUpdates.Num())
 	{
 		check(!bWithinRenderPass);
-		Device.GetBindlessDescriptorManager()->UpdateDescriptorsWithGPU(this, MoveTemp(ComputeDescriptorEntries), MoveTemp(ComputeDescriptorIndices));
-		ComputeDescriptorEntries.Empty();
-		ComputeDescriptorIndices.Empty();
+
+		Device.GetBindlessDescriptorManager()->FlushPendingDescriptorUpdates(*this, PendingDescriptorUpdates);
+		PendingDescriptorUpdates.Empty();
 	}
 #endif
 }
diff --git a/Engine/Source/Runtime/Apple/MetalRHI/Private/MetalState.cpp b/Engine/Source/Runtime/Apple/MetalRHI/Private/MetalState.cpp
--- a/Engine/Source/Runtime/Apple/MetalRHI/Private/MetalState.cpp
+++ b/Engine/Source/Runtime/Apple/MetalRHI/Private/MetalState.cpp
@@ -309,8 +309,7 @@ FMetalSamplerState::FMetalSamplerState(FMetalDevice& MetalDevice, const FSampler
 
 	if(IsMetalBindlessEnabled())
 	{
-		BindlessHandle = BindlessDescriptorManager->ReserveDescriptor(ERHIDescriptorType::Sampler);
-		BindlessDescriptorManager->BindSampler(BindlessHandle, State);
+		BindlessHandle = BindlessDescriptorManager->AllocateDescriptor(State);
 	}
 #endif
 }
diff --git a/Engine/Source/Runtime/Apple/MetalRHI/Private/MetalStateCache.cpp b/Engine/Source/Runtime/Apple/MetalRHI/Private/MetalStateCache.cpp
--- a/Engine/Source/Runtime/Apple/MetalRHI/Private/MetalStateCache.cpp
+++ b/Engine/Source/Runtime/Apple/MetalRHI/Private/MetalStateCache.cpp
@@ -2025,7 +2025,7 @@ void FMetalStateCache::IRBindResourcesToEncoder(ShaderType Shader, FMetalCommand
 	}
 	else
 	{
-		OverriddenDescriptorHeap->BindHeap(Encoder, FunctionType, kIRStandardHeapBindPoint);
+		OverriddenDescriptorHeap->BindToEncoder(Encoder, FunctionType, kIRStandardHeapBindPoint);
 	}
 
     // Bind CBV Table
diff --git a/Engine/Source/Runtime/Apple/MetalRHI/Private/MetalTexture.cpp b/Engine/Source/Runtime/Apple/MetalRHI/Private/MetalTexture.cpp
--- a/Engine/Source/Runtime/Apple/MetalRHI/Private/MetalTexture.cpp
+++ b/Engine/Source/Runtime/Apple/MetalRHI/Private/MetalTexture.cpp
@@ -751,12 +751,12 @@ FMetalSurface::FMetalSurface(FMetalDevice& MetalDevice, FMetalTextureCreateDesc
 	}
 
 #if PLATFORM_SUPPORTS_BINDLESS_RENDERING
-	if(IsMetalBindlessEnabled())
+	if (IsMetalBindlessEnabled())
 	{
 		FMetalBindlessDescriptorManager* BindlessDescriptorManager = Device.GetBindlessDescriptorManager();
 		check(BindlessDescriptorManager);
 
-		BindlessHandle = BindlessDescriptorManager->ReserveDescriptor(ERHIDescriptorType::TextureSRV);
+		BindlessHandle = BindlessDescriptorManager->AllocateDescriptor(ERHIDescriptorType::TextureSRV);
 	}
 #endif
 }
@@ -776,12 +776,12 @@ FMetalExternalSurface::FMetalExternalSurface(FMetalDevice& MetalDevice, MTLTextu
 	UE::RHICore::UpdateGlobalTextureStats(GetDesc(), TotalTextureSize, bOnlyStreamableTexturesInTexturePool, true);
 	
 #if PLATFORM_SUPPORTS_BINDLESS_RENDERING
-	if(IsMetalBindlessEnabled())
+	if (IsMetalBindlessEnabled())
 	{
 		FMetalBindlessDescriptorManager* BindlessDescriptorManager = Device.GetBindlessDescriptorManager();
 		check(BindlessDescriptorManager);
 
-		BindlessHandle = BindlessDescriptorManager->ReserveDescriptor(ERHIDescriptorType::TextureSRV);
+		BindlessHandle = BindlessDescriptorManager->AllocateDescriptor(ERHIDescriptorType::TextureSRV);
 	}
 #endif
 }
@@ -822,8 +822,8 @@ void FMetalSurface::Initialize(FRHICommandListBase& RHICmdList)
 	{
 		FMetalBindlessDescriptorManager* BindlessDescriptorManager = Device.GetBindlessDescriptorManager();
 		check(BindlessDescriptorManager);
-
-		BindlessDescriptorManager->BindTexture(&RHICmdList, BindlessHandle, Texture.get(), EDescriptorUpdateType_Immediate);
+		
+		BindlessDescriptorManager->UpdateDescriptor(RHICmdList, BindlessHandle, Texture.get(), EMetalDescriptorUpdateType::Immediate);
 	}
 #endif
 }
@@ -1298,15 +1298,14 @@ MTLTexturePtr FMetalSurface::GetDrawableTexture()
 	{
 		check(Viewport);
 		Texture = NS::RetainPtr(Viewport->GetDrawableTexture(EMetalViewportAccessRHI));
+
 #if PLATFORM_SUPPORTS_BINDLESS_RENDERING
 		if(IsMetalBindlessEnabled())
 		{
 			FMetalBindlessDescriptorManager* BindlessDescriptorManager = Device.GetBindlessDescriptorManager();
 			check(BindlessDescriptorManager);
 			
-			BindlessHandle = BindlessDescriptorManager->ReserveDescriptor(ERHIDescriptorType::TextureSRV);
-
-			BindlessDescriptorManager->BindTexture(nullptr, BindlessHandle, Texture.get(), EDescriptorUpdateType_Immediate);
+			BindlessHandle = BindlessDescriptorManager->AllocateDescriptor(ERHIDescriptorType::TextureSRV, Texture.get());
 		}
 #endif
 	}
@@ -1316,13 +1315,10 @@ MTLTexturePtr FMetalSurface::GetDrawableTexture()
 void FMetalSurface::ReleaseDrawableTexture()
 {
 #if PLATFORM_SUPPORTS_BINDLESS_RENDERING
-	if(IsMetalBindlessEnabled())
+	if (BindlessHandle.IsValid())
 	{
-		if(BindlessHandle.IsValid())
-		{
-			FMetalDynamicRHI::Get().DeferredDelete(BindlessHandle);
-			BindlessHandle = {};
-		}
+		FMetalDynamicRHI::Get().DeferredDelete(BindlessHandle);
+		BindlessHandle = {};
 	}
 #endif
 
diff --git a/Engine/Source/Runtime/Apple/MetalRHI/Private/MetalUAV.cpp b/Engine/Source/Runtime/Apple/MetalRHI/Private/MetalUAV.cpp
--- a/Engine/Source/Runtime/Apple/MetalRHI/Private/MetalUAV.cpp
+++ b/Engine/Source/Runtime/Apple/MetalRHI/Private/MetalUAV.cpp
@@ -85,12 +85,12 @@ FMetalShaderResourceView::FMetalShaderResourceView(FMetalDevice& InDevice, FRHIC
 	  FMetalResourceViewBase(InDevice)
 {
 #if PLATFORM_SUPPORTS_BINDLESS_RENDERING
-	FMetalBindlessDescriptorManager* BindlessDescriptorManager = Device.GetBindlessDescriptorManager();
-    check(BindlessDescriptorManager);
-
-	if(IsMetalBindlessEnabled())
+	if (IsMetalBindlessEnabled())
 	{
-		BindlessHandle = BindlessDescriptorManager->ReserveDescriptor(ERHIDescriptorType::TextureSRV);
+		FMetalBindlessDescriptorManager* BindlessDescriptorManager = Device.GetBindlessDescriptorManager();
+		check(BindlessDescriptorManager);
+	
+		BindlessHandle = BindlessDescriptorManager->AllocateDescriptor(ERHIDescriptorType::TextureSRV);
 	}
 	
 	SurfaceOverride = nullptr;
@@ -109,11 +109,11 @@ FMetalShaderResourceView::FMetalShaderResourceView(FMetalDevice& InDevice, FRHIC
 FMetalShaderResourceView::~FMetalShaderResourceView()
 {
 #if PLATFORM_SUPPORTS_BINDLESS_RENDERING
-	FMetalBindlessDescriptorManager* BindlessDescriptorManager = Device.GetBindlessDescriptorManager();
-	check(BindlessDescriptorManager);
-
-	if(IsMetalBindlessEnabled())
+	if (IsMetalBindlessEnabled())
 	{
+		FMetalBindlessDescriptorManager* BindlessDescriptorManager = Device.GetBindlessDescriptorManager();
+		check(BindlessDescriptorManager);
+
 		BindlessDescriptorManager->FreeDescriptor(BindlessHandle);
 	}
 #endif
@@ -366,13 +366,12 @@ void FMetalShaderResourceView::UpdateView(FMetalRHICommandContext* Context, cons
 	}
 	
 #if PLATFORM_SUPPORTS_BINDLESS_RENDERING
-	FMetalBindlessDescriptorManager* BindlessDescriptorManager = Device.GetBindlessDescriptorManager();
-    check(BindlessDescriptorManager);
-
-	if(IsMetalBindlessEnabled())
+	if (IsMetalBindlessEnabled())
 	{
-		BindlessDescriptorManager->BindResource(BindlessHandle, this, Context, 
-												bConstructing ? EDescriptorUpdateType_Immediate : EDescriptorUpdateType_GPU);
+		FMetalBindlessDescriptorManager* BindlessDescriptorManager = Device.GetBindlessDescriptorManager();
+		check(BindlessDescriptorManager);
+
+		BindlessDescriptorManager->UpdateDescriptor(Context, BindlessHandle, this, bConstructing ? EMetalDescriptorUpdateType::Immediate : EMetalDescriptorUpdateType::GPU);
 	}
 #endif
 }
@@ -383,12 +382,12 @@ FMetalUnorderedAccessView::FMetalUnorderedAccessView(FMetalDevice& InDevice, FRH
 	, FMetalResourceViewBase(InDevice)
 {
 #if PLATFORM_SUPPORTS_BINDLESS_RENDERING
-    FMetalBindlessDescriptorManager* BindlessDescriptorManager = Device.GetBindlessDescriptorManager();
-	check(BindlessDescriptorManager);
-
-	if(IsMetalBindlessEnabled())
+	if (IsMetalBindlessEnabled())
 	{
-		BindlessHandle = BindlessDescriptorManager->ReserveDescriptor(ERHIDescriptorType::TextureUAV);
+		FMetalBindlessDescriptorManager* BindlessDescriptorManager = Device.GetBindlessDescriptorManager();
+		check(BindlessDescriptorManager);
+
+		BindlessHandle = BindlessDescriptorManager->AllocateDescriptor(ERHIDescriptorType::TextureUAV);
 	}
 #endif
 
@@ -406,11 +405,11 @@ FMetalUnorderedAccessView::FMetalUnorderedAccessView(FMetalDevice& InDevice, FRH
 FMetalUnorderedAccessView::~FMetalUnorderedAccessView()
 {
 #if PLATFORM_SUPPORTS_BINDLESS_RENDERING
-	FMetalBindlessDescriptorManager* BindlessDescriptorManager = Device.GetBindlessDescriptorManager();
-	check(BindlessDescriptorManager);
-
 	if(IsMetalBindlessEnabled())
 	{
+		FMetalBindlessDescriptorManager* BindlessDescriptorManager = Device.GetBindlessDescriptorManager();
+		check(BindlessDescriptorManager);
+
 		BindlessDescriptorManager->FreeDescriptor(BindlessHandle);
 	}
 #endif
@@ -611,13 +610,12 @@ void FMetalUnorderedAccessView::UpdateView(FMetalRHICommandContext* Context, con
         }
 	}
 #if PLATFORM_SUPPORTS_BINDLESS_RENDERING
-    FMetalBindlessDescriptorManager* BindlessDescriptorManager = Device.GetBindlessDescriptorManager();
-	check(BindlessDescriptorManager);
-
 	if(IsMetalBindlessEnabled())
 	{
-		BindlessDescriptorManager->BindResource(BindlessHandle, this, Context, 
-												bConstructing ? EDescriptorUpdateType_Immediate : EDescriptorUpdateType_GPU);
+		FMetalBindlessDescriptorManager* BindlessDescriptorManager = Device.GetBindlessDescriptorManager();
+		check(BindlessDescriptorManager);
+		
+		BindlessDescriptorManager->UpdateDescriptor(Context, BindlessHandle, this, bConstructing ? EMetalDescriptorUpdateType::Immediate : EMetalDescriptorUpdateType::GPU);
 	}
 #endif
 }
diff --git a/Engine/Source/Runtime/Apple/MetalRHI/Public/MetalRHIContext.h b/Engine/Source/Runtime/Apple/MetalRHI/Public/MetalRHIContext.h
--- a/Engine/Source/Runtime/Apple/MetalRHI/Public/MetalRHIContext.h
+++ b/Engine/Source/Runtime/Apple/MetalRHI/Public/MetalRHIContext.h
@@ -12,8 +12,8 @@
 #include "MetalCommandList.h"
 #include "MetalCommandEncoder.h"
 #include "MetalRHIRenderQuery.h"
+#include "MetalBindlessDescriptors.h"
 #include "RHICore.h"
-#include "Containers/DynamicRHIResourceArray.h"
 
 class FMetalEventNode;
 
@@ -212,10 +212,9 @@ public:
 	void Finalize(TArray<FMetalPayload*>& OutPayloads);
 	
 #if PLATFORM_SUPPORTS_BINDLESS_RENDERING
-	void EnqueueDescriptorUpdate(IRDescriptorTableEntry Entry, uint32_t Index)
+	void EnqueueDescriptorUpdate(FRHIDescriptorHandle InHandle, const IRDescriptorTableEntry& InDescriptor)
 	{
-		ComputeDescriptorEntries.Add(Entry);
-		ComputeDescriptorIndices.Add(Index);
+		PendingDescriptorUpdates.Add(InHandle, InDescriptor);
 	}
 #endif
 	
@@ -385,8 +384,7 @@ protected:
 	FMetalParallelRenderPassInfo* ParallelRenderPassInfo;
 	
 #if PLATFORM_SUPPORTS_BINDLESS_RENDERING
-	TResourceArray<IRDescriptorTableEntry>    ComputeDescriptorEntries;
-	TResourceArray<uint32_t>                  ComputeDescriptorIndices;
+	FMetalPendingDescriptorUpdates PendingDescriptorUpdates;
 #endif
 	
 private:
diff --git a/Engine/Source/Runtime/Apple/MetalRHI/Public/MetalResources.h b/Engine/Source/Runtime/Apple/MetalRHI/Public/MetalResources.h
--- a/Engine/Source/Runtime/Apple/MetalRHI/Public/MetalResources.h
+++ b/Engine/Source/Runtime/Apple/MetalRHI/Public/MetalResources.h
@@ -429,40 +429,39 @@ public:
 		FMetalBufferPtr Buffer;
 		uint32 Offset;
 		uint32 Size;
-
+		
 		FBufferView(FMetalBufferPtr Buffer, uint32 Offset, uint32 Size)
-			: Buffer(Buffer)
-			, Offset(Offset)
-			, Size(Size)
+		: Buffer(Buffer)
+		, Offset(Offset)
+		, Size(Size)
+		{}
+	};
+
+	struct FTextureBufferBacked
+	{
+		MTLTexturePtr Texture;
+		FMetalBufferPtr Buffer;
+		uint32 Offset;
+		uint32 Size;
+		EPixelFormat Format;
+		bool bIsBuffer;
+		
+		FTextureBufferBacked(MTLTexturePtr Texture, FMetalBufferPtr Buffer, uint32 Offset, uint32 Size, EPixelFormat Format, bool bIsBuffer)
+		: Texture(Texture)
+		, Buffer(Buffer)
+		, Offset(Offset)
+		, Size(Size)
+		, Format(Format)
+		, bIsBuffer(bIsBuffer)
 		{}
 	};
-    
-    struct FTextureBufferBacked
-    {
-        MTLTexturePtr Texture;
-        FMetalBufferPtr Buffer;
-        uint32 Offset;
-        uint32 Size;
-        EPixelFormat Format;
-        bool bIsBuffer;
-
-        FTextureBufferBacked(MTLTexturePtr Texture, FMetalBufferPtr Buffer, uint32 Offset, uint32 Size, EPixelFormat Format, bool bIsBuffer)
-            :
-              Texture(Texture)
-            , Buffer(Buffer)
-            , Offset(Offset)
-            , Size(Size)
-            , Format(Format)
-            , bIsBuffer(bIsBuffer)
-        {}
-    };
 
 	typedef TVariant<FEmptyVariantState
-		, MTLTexturePtr
-		, FBufferView
-        , FTextureBufferBacked
+	, MTLTexturePtr
+	, FBufferView
+	, FTextureBufferBacked
 #if METAL_RHI_RAYTRACING
-		, MTL::AccelerationStructure
+	, MTL::AccelerationStructure
 #endif
 	> TStorage;
 
@@ -471,16 +470,16 @@ public:
 		Null                  = TStorage::IndexOfType<FEmptyVariantState>(),
 		TextureView           = TStorage::IndexOfType<MTLTexturePtr>(),
 		BufferView            = TStorage::IndexOfType<FBufferView>(),
-        TextureBufferBacked   = TStorage::IndexOfType<FTextureBufferBacked>(),
+		TextureBufferBacked   = TStorage::IndexOfType<FTextureBufferBacked>(),
 #if METAL_RHI_RAYTRACING
 		AccelerationStructure = TStorage::IndexOfType<MTL::AccelerationStructure>()
 #endif
 	};
-	
+
 protected:
 	FMetalResourceViewBase(FMetalDevice& InDevice) : Device(InDevice)
 	{}
-	
+
 public:
 	virtual ~FMetalResourceViewBase();
 
@@ -489,26 +488,26 @@ public:
 		return static_cast<EMetalType>(Storage.GetIndex());
 	}
 
-    MTLTexturePtr const GetTextureView() const
+	const MTLTexturePtr& GetTextureView() const
 	{
 		check(GetMetalType() == EMetalType::TextureView);
 		return Storage.Get<MTLTexturePtr>();
 	}
 
-	FBufferView const& GetBufferView() const
+	const FBufferView& GetBufferView() const
 	{
 		check(GetMetalType() == EMetalType::BufferView);
 		return Storage.Get<FBufferView>();
 	}
-    
-    FTextureBufferBacked const& GetTextureBufferBacked() const
-    {
-        check(GetMetalType() == EMetalType::TextureBufferBacked);
-        return Storage.Get<FTextureBufferBacked>();
-    }
+
+	const FTextureBufferBacked& GetTextureBufferBacked() const
+	{
+		check(GetMetalType() == EMetalType::TextureBufferBacked);
+		return Storage.Get<FTextureBufferBacked>();
+	}
 
 #if METAL_RHI_RAYTRACING
-	MTL::AccelerationStructure const& GetAccelerationStructure() const
+	const MTL::AccelerationStructure& GetAccelerationStructure() const
 	{
 		check(GetMetalType() == EMetalType::AccelerationStructure);
 		return Storage.Get<MTL::AccelerationStructure>();
@@ -523,13 +522,13 @@ public:
 protected:
 	void InitAsTextureView(MTLTexturePtr);
 	void InitAsBufferView(FMetalBufferPtr Buffer, uint32 Offset, uint32 Size);
-    void InitAsTextureBufferBacked(MTLTexturePtr Texture, FMetalBufferPtr Buffer, uint32 Offset, uint32 Size, EPixelFormat Format, bool bIsBuffer);
-
+	void InitAsTextureBufferBacked(MTLTexturePtr Texture, FMetalBufferPtr Buffer, uint32 Offset, uint32 Size, EPixelFormat Format, bool bIsBuffer);
+	
 	void Invalidate();
-
+	
 	FMetalDevice& Device;
 	bool bOwnsResource = true;
-
+	
 private:
 	TStorage Storage;
 };
