Metal Buffer lock cleanup
- FMetalRHIBuffer::Data isn't used
- RHILockBuffer/RHIUnlockBuffer only needs FMetalRHILockData when using a TransferBuffer
- TransferBuffer upload can happen outside of FMetalRHIBuffer Lock/Unlock

diff --git a/Engine/Source/Runtime/Apple/MetalRHI/Private/MetalCommands.cpp b/Engine/Source/Runtime/Apple/MetalRHI/Private/MetalCommands.cpp
--- a/Engine/Source/Runtime/Apple/MetalRHI/Private/MetalCommands.cpp
+++ b/Engine/Source/Runtime/Apple/MetalRHI/Private/MetalCommands.cpp
@@ -98,12 +98,12 @@ void FMetalRHICommandContext::RHISetStreamSource(uint32 StreamIndex, FRHIBuffer*
     FMetalRHIBuffer* VertexBuffer = ResourceCast(VertexBufferRHI);
     
     FMetalBufferPtr TheBuffer = nullptr;
-    if(VertexBuffer && !VertexBuffer->Data)
+    if (VertexBuffer)
     {
         TheBuffer = VertexBuffer->GetCurrentBuffer();
     }
     
-    StateCache.SetVertexStream(StreamIndex, VertexBuffer ? TheBuffer : nullptr, VertexBuffer ? VertexBuffer->Data : nullptr, Offset, VertexBuffer ? VertexBuffer->GetSize() : 0);
+    StateCache.SetVertexStream(StreamIndex, VertexBuffer ? TheBuffer : nullptr, nullptr, Offset, VertexBuffer ? VertexBuffer->GetSize() : 0);
 }
 
 static void SetUniformBufferInternal(FMetalStateCache& StateCache, FMetalShaderData* ShaderData, EMetalShaderStages Stage, uint32 BufferIndex, FRHIUniformBuffer* UBRHI)
diff --git a/Engine/Source/Runtime/Apple/MetalRHI/Private/MetalRHIBuffer.cpp b/Engine/Source/Runtime/Apple/MetalRHI/Private/MetalRHIBuffer.cpp
--- a/Engine/Source/Runtime/Apple/MetalRHI/Private/MetalRHIBuffer.cpp
+++ b/Engine/Source/Runtime/Apple/MetalRHI/Private/MetalRHIBuffer.cpp
@@ -128,48 +128,41 @@ FMetalRHIBuffer::FMetalRHIBuffer(FRHICommandListBase& RHICmdList, FMetalDevice&
 	{
 		check(CreateDesc.Size == InResourceArray->GetResourceDataSize());
 
-		if (Data)
+		if (Mode == MTL::StorageModePrivate)
 		{
-			FMemory::Memcpy(Data->Data, InResourceArray->GetResourceData(), CreateDesc.Size);
-		}
-		else
-		{
-			if (Mode == MTL::StorageModePrivate)
+			if (RHICmdList.IsBottomOfPipe())
 			{
-				if (RHICmdList.IsBottomOfPipe())
-				{
-					void* Backing = this->Lock(RHICmdList, RLM_WriteOnly, 0, CreateDesc.Size);
-					FMemory::Memcpy(Backing, InResourceArray->GetResourceData(), CreateDesc.Size);
-					this->Unlock(RHICmdList);
-				}
-				else
-				{
-					void* Result = FMemory::Malloc(CreateDesc.Size, 16);
-					FMemory::Memcpy(Result, InResourceArray->GetResourceData(), CreateDesc.Size);
-
-					RHICmdList.EnqueueLambda(
-						[this, Result, InSize = CreateDesc.Size](FRHICommandListBase& RHICmdList)
-						{
-							void* Backing = this->Lock(RHICmdList, RLM_WriteOnly, 0, InSize);
-							FMemory::Memcpy(Backing, Result, InSize);
-							this->Unlock(RHICmdList);
-							FMemory::Free(Result);
-						});
-				}
+				void* Backing = this->Lock(RHICmdList, RLM_WriteOnly, 0, CreateDesc.Size);
+				FMemory::Memcpy(Backing, InResourceArray->GetResourceData(), CreateDesc.Size);
+				this->Unlock(RHICmdList);
 			}
 			else
 			{
-				FMetalBufferPtr TheBuffer = GetCurrentBuffer();
-				MTL::Buffer* MTLBuffer = TheBuffer->GetMTLBuffer();
-				FMemory::Memcpy(TheBuffer->Contents(), InResourceArray->GetResourceData(), CreateDesc.Size);
+				void* Result = FMemory::Malloc(CreateDesc.Size, 16);
+				FMemory::Memcpy(Result, InResourceArray->GetResourceData(), CreateDesc.Size);
+
+				RHICmdList.EnqueueLambda(
+					[this, Result, InSize = CreateDesc.Size](FRHICommandListBase& RHICmdList)
+					{
+						void* Backing = this->Lock(RHICmdList, RLM_WriteOnly, 0, InSize);
+						FMemory::Memcpy(Backing, Result, InSize);
+						this->Unlock(RHICmdList);
+						FMemory::Free(Result);
+					});
+			}
+		}
+		else
+		{
+			FMetalBufferPtr TheBuffer = GetCurrentBuffer();
+			MTL::Buffer* MTLBuffer = TheBuffer->GetMTLBuffer();
+			FMemory::Memcpy(TheBuffer->Contents(), InResourceArray->GetResourceData(), CreateDesc.Size);
 #if PLATFORM_MAC 
-				if (Mode == MTL::StorageModeManaged)
-				{
-                    NS::Range ModifyRange = NS::Range(TheBuffer->GetOffset(), TheBuffer->GetLength());
-                    MTLBuffer->didModifyRange(ModifyRange);
-				}
-#endif
+			if (Mode == MTL::StorageModeManaged)
+			{
+				NS::Range ModifyRange = NS::Range(TheBuffer->GetOffset(), TheBuffer->GetLength());
+				MTLBuffer->didModifyRange(ModifyRange);
 			}
+#endif
 		}
 
 		// Discard the resource array's contents.
@@ -321,23 +314,24 @@ bool FMetalRHIBuffer::RequiresTransferBuffer()
 	return (Mode == MTL::StorageModePrivate || (Mode == MTL::StorageModeShared && bIsStatic));
 }
 
+static uint32 AdjustLockSize(uint32 Size, uint32 Offset)
+{
+#if PLATFORM_MAC
+	// Blit encoder validation error, lock size and subsequent blit copy unlock operations need to be in 4 byte multiples on macOS
+	return FMath::Min(Align(Size, 4), Size - Offset);
+#else
+	return Size;
+#endif
+}
+
 void* FMetalRHIBuffer::Lock(FRHICommandListBase& RHICmdList, EResourceLockMode InLockMode, uint32 Offset, uint32 InSize, FMetalBufferPtr InTransferBuffer)
 {
 	check(CurrentLockMode == RLM_Num);
 	check(LockSize == 0 && LockOffset == 0);	
 	check(!TransferBuffer);
 	
-	if (Data)
-	{
-		check(Data->Data);
-		return ((uint8*)Data->Data) + Offset;
-	}
-	
-#if PLATFORM_MAC
-	// Blit encoder validation error, lock size and subsequent blit copy unlock operations need to be in 4 byte multiples on macOS
-	InSize = FMath::Min(Align(InSize, 4), Size - Offset);
-#endif
-	
+	InSize = AdjustLockSize(InSize, Offset);
+
 	const bool bWriteLock = InLockMode == RLM_WriteOnly;
 	const bool bIsStatic = EnumHasAnyFlags(GetUsage(), BUF_Static);
 	const bool bIsDynamic = EnumHasAnyFlags(GetUsage(), BUF_Dynamic);
@@ -451,66 +445,83 @@ void* FMetalRHIBuffer::Lock(FRHICommandListBase& RHICmdList, EResourceLockMode I
 
 void FMetalRHIBuffer::Unlock(FRHICommandListBase& RHICmdList)
 {
-	if (!Data)
+	FMetalBufferPtr CurrentBufferToUnlock = GetCurrentBuffer();
+	
+	check(CurrentBufferToUnlock);
+	check(LockSize > 0);
+	const bool bWriteLock = CurrentLockMode == RLM_WriteOnly;
+	const bool bIsStatic = EnumHasAnyFlags(GetUsage(), BUF_Static);
+	
+	if (bWriteLock)
 	{
-		FMetalBufferPtr CurrentBufferToUnlock = GetCurrentBuffer();
-		
-		check(CurrentBufferToUnlock);
-		check(LockSize > 0);
-		const bool bWriteLock = CurrentLockMode == RLM_WriteOnly;
-		const bool bIsStatic = EnumHasAnyFlags(GetUsage(), BUF_Static);
+		check(LockOffset == 0);
+		check(LockSize <= CurrentBufferToUnlock->GetLength());
+
+		// Use transfer buffer for writing into 'Static' buffers as they could be in use by GPU atm
+		// Initialization of 'Static' buffers still uses direct copy when possible
+		const bool bUseTransferBuffer = RequiresTransferBuffer();
 		
-		if (bWriteLock)
+		if (bUseTransferBuffer)
 		{
-			check(LockOffset == 0);
-            check(LockSize <= CurrentBufferToUnlock->GetLength());
-
-			// Use transfer buffer for writing into 'Static' buffers as they could be in use by GPU atm
-			// Initialization of 'Static' buffers still uses direct copy when possible
-			const bool bUseTransferBuffer = RequiresTransferBuffer();
+			FMetalRHIUploadContext& UploadContext = static_cast<FMetalRHIUploadContext&>(RHICmdList.GetUploadContext());
 			
-			if (bUseTransferBuffer)
+			UploadContext.EnqueueFunction([&InDevice=Device, Size=LockSize, Dest=CurrentBufferToUnlock, InTransferBuffer=TransferBuffer](FMetalRHICommandContext* Context)
 			{
-				FMetalRHIUploadContext& UploadContext = static_cast<FMetalRHIUploadContext&>(RHICmdList.GetUploadContext());
-				
-				UploadContext.EnqueueFunction([&InDevice=Device, Size=LockSize, Dest=CurrentBufferToUnlock, InTransferBuffer=TransferBuffer](FMetalRHICommandContext* Context)
-				{
-					Context->CopyFromBufferToBuffer(InTransferBuffer, 0, Dest, 0, Size);
-					FMetalDynamicRHI::Get().DeferredDelete(InTransferBuffer);
-				});
-				
-				TransferBuffer = nullptr;
-			}
+				Context->CopyFromBufferToBuffer(InTransferBuffer, 0, Dest, 0, Size);
+				FMetalDynamicRHI::Get().DeferredDelete(InTransferBuffer);
+			});
+			
+			TransferBuffer = nullptr;
+		}
 #if PLATFORM_MAC
-			else if (Mode == MTL::StorageModeManaged)
-			{
-				CurrentBufferToUnlock->GetMTLBuffer()->didModifyRange(NS::Range(LockOffset + CurrentBufferToUnlock->GetOffset(), LockSize));
-			}
-#endif //PLATFORM_MAC
-			else
-			{
-				// shared buffers are always mapped so nothing happens
-				check(Mode == MTL::StorageModeShared);
-			}
+		else if (Mode == MTL::StorageModeManaged)
+		{
+			CurrentBufferToUnlock->GetMTLBuffer()->didModifyRange(NS::Range(LockOffset + CurrentBufferToUnlock->GetOffset(), LockSize));
 		}
+#endif //PLATFORM_MAC
 		else
 		{
-			check(CurrentLockMode == RLM_ReadOnly);
-			if(TransferBuffer)
-			{
-				check(Mode == MTL::StorageModePrivate);
-				FMetalDynamicRHI::Get().DeferredDelete(TransferBuffer);
-				TransferBuffer = nullptr;
-			}
+			// shared buffers are always mapped so nothing happens
+			check(Mode == MTL::StorageModeShared);
 		}
 	}
-	
+	else
+	{
+		check(CurrentLockMode == RLM_ReadOnly);
+		if(TransferBuffer)
+		{
+			check(Mode == MTL::StorageModePrivate);
+			FMetalDynamicRHI::Get().DeferredDelete(TransferBuffer);
+			TransferBuffer = nullptr;
+		}
+	}
+
 	check(!TransferBuffer);
 	CurrentLockMode = RLM_Num;
 	LockSize = 0;
 	LockOffset = 0;
 }
 
+void FMetalRHIBuffer::UploadTransferBuffer(FRHICommandListBase& RHICmdList, FMetalBufferPtr&& InTransferBuffer, uint32 UploadSize)
+{
+	check(CurrentBuffer);
+	check(InTransferBuffer);
+	check(UploadSize);
+	
+	RHICmdList.EnqueueLambda(
+		[InTransferBuffer = MoveTemp(InTransferBuffer), CurrentBuffer=CurrentBuffer, UploadSize](FRHICommandListBase& RHICmdList) mutable
+		{
+			FMetalRHIUploadContext& UploadContext = static_cast<FMetalRHIUploadContext&>(RHICmdList.GetUploadContext());
+
+			UploadContext.EnqueueFunction(
+				[InTransferBuffer = MoveTemp(InTransferBuffer), CurrentBuffer, UploadSize](FMetalRHICommandContext* Context)
+				{
+					Context->CopyFromBufferToBuffer(InTransferBuffer, 0, CurrentBuffer, 0, UploadSize);
+					FMetalDynamicRHI::Get().DeferredDelete(InTransferBuffer);
+				});
+		});
+}
+
 void FMetalRHIBuffer::TakeOwnership(FMetalRHIBuffer& Other)
 {
     check(Other.CurrentLockMode == RLM_Num);
@@ -523,7 +534,6 @@ void FMetalRHIBuffer::TakeOwnership(FMetalRHIBuffer& Other)
 
 	CurrentBuffer = Other.CurrentBuffer;
     TransferBuffer = Other.TransferBuffer;
-    Data = Other.Data;
     CurrentLockMode = Other.CurrentLockMode;
     LockOffset = Other.LockOffset;
     LockSize = Other.LockSize;
@@ -532,7 +542,6 @@ void FMetalRHIBuffer::TakeOwnership(FMetalRHIBuffer& Other)
     
 	Other.CurrentBuffer.Reset();
     Other.TransferBuffer = nullptr;
-    Other.Data = nullptr;
     Other.CurrentLockMode = RLM_Num;
     Other.LockOffset = 0;
     Other.LockSize = 0;
@@ -549,15 +558,6 @@ void FMetalRHIBuffer::ReleaseOwnership()
     
 	ReleaseBuffer();
 
-    if (Data)
-    {
-        METAL_INC_DWORD_STAT_BY(MemFreed, Size, GetUsage());
-        auto ReleaseFunction = [ReleaseData=Data](){
-            delete ReleaseData;
-        };
-		FMetalDynamicRHI::Get().DeferredDelete(ReleaseFunction);
-    }
-
 #if METAL_RHI_RAYTRACING
     if (EnumHasAnyFlags(GetUsage(), BUF_AccelerationStructure))
     {
@@ -618,13 +618,14 @@ FRHIBufferInitializer FMetalDynamicRHI::RHICreateBufferInitializer(FRHICommandLi
 	return UE::RHICore::HandleUnknownBufferInitializerInitAction(RHICmdList, CreateDesc);
 }
 
-struct FMetalRHILockData
+struct FMetalRHILockBuffer
 {
-	FMetalRHILockData(FMetalBufferPtr InBuffer, void* InData) : Buffer(InBuffer), Data(InData)
-	{}
+	FMetalRHILockBuffer(FMetalBufferPtr&& InBuffer)
+		: Buffer(MoveTemp(InBuffer))
+	{
+	}
 	
 	FMetalBufferPtr Buffer;
-	void* Data = nullptr;
 };
 
 static FLockTracker GBufferLockTracker;
@@ -632,100 +633,96 @@ static FLockTracker GBufferLockTracker;
 void* FMetalDynamicRHI::RHILockBuffer(class FRHICommandListBase& RHICmdList, FRHIBuffer* BufferRHI, uint32 Offset, uint32 SizeRHI, EResourceLockMode LockMode)
 {
 	MTL_SCOPED_AUTORELEASE_POOL;
-	
+
 	FMetalRHIBuffer* Buffer = ResourceCast(BufferRHI);
-	
+
 	if (RHICmdList.IsTopOfPipe())
 	{
 		void* Result = nullptr;
-		if(LockMode != RLM_WriteOnly)
+		if (LockMode != RLM_WriteOnly)
 		{
 			QUICK_SCOPE_CYCLE_COUNTER(STAT_RHIMETHOD_LockBuffer_FlushAndLock);
 			CSV_SCOPED_TIMING_STAT(RHITFlushes, LockBuffer_BottomOfPipe);
-			
+
 			FRHICommandListScopedFlushAndExecute Flush(RHICmdList.GetAsImmediate());
-			Result = (uint8*)Buffer->Lock(RHICmdList, LockMode, Offset, SizeRHI);
-			
-			FMetalRHILockData* LockData = new FMetalRHILockData(nullptr, Result);
-			GBufferLockTracker.Lock(Buffer, (void*)LockData, Offset, SizeRHI, LockMode);
-			return Result;
+			Result = Buffer->Lock(RHICmdList, LockMode, Offset, SizeRHI);
+
+			GBufferLockTracker.Lock(Buffer, nullptr, Offset, SizeRHI, LockMode);
 		}
 		else
 		{
 			QUICK_SCOPE_CYCLE_COUNTER(STAT_RHIMETHOD_LockBuffer_Malloc);
 			
-			if(Buffer->RequiresTransferBuffer())
+			if (Buffer->RequiresTransferBuffer())
 			{
 				FMetalBufferPtr TempBuffer = Device->GetResourceHeap().CreateBuffer(SizeRHI, BufferBackedLinearTextureOffsetAlignment, BUF_Dynamic, MTL::ResourceCPUCacheModeDefaultCache | MTL::ResourceStorageModeShared, true);
-				
-				Result = (uint8*) TempBuffer->Contents();
-				
-				FMetalRHILockData* LockData = new FMetalRHILockData(TempBuffer, nullptr);
-				GBufferLockTracker.Lock(Buffer, (void*)LockData, Offset, SizeRHI, LockMode);
+
+				Result = TempBuffer->Contents();
+
+				GBufferLockTracker.Lock(Buffer, new FMetalRHILockBuffer(MoveTemp(TempBuffer)), Offset, SizeRHI, LockMode);
 			}
 			else
 			{
 				Result = FMemory::Malloc(SizeRHI, 16);
-				FMetalRHILockData* LockData = new FMetalRHILockData(nullptr, Result);
-				GBufferLockTracker.Lock(Buffer, (void*)LockData, Offset, SizeRHI, LockMode);
+
+				GBufferLockTracker.Lock(Buffer, Result, Offset, SizeRHI, LockMode);
 			}
-			
-			return Result;
 		}
+
+		return Result;
 	}
 	
-	return (uint8*)Buffer->Lock(RHICmdList, LockMode, Offset, SizeRHI);
+	return Buffer->Lock(RHICmdList, LockMode, Offset, SizeRHI);
 }
 
 void FMetalDynamicRHI::RHIUnlockBuffer(class FRHICommandListBase& RHICmdList, FRHIBuffer* BufferRHI)
 {
 	MTL_SCOPED_AUTORELEASE_POOL;
-	
 	QUICK_SCOPE_CYCLE_COUNTER(STAT_FDynamicRHI_UnlockBuffer_RenderThread);
-	
+
 	FMetalRHIBuffer* Buffer = ResourceCast(BufferRHI);
-	
+
 	if (RHICmdList.IsTopOfPipe())
 	{
 		FLockTracker::FLockParams Params = GBufferLockTracker.Unlock(Buffer);
-		FMetalRHILockData* LockData = (FMetalRHILockData*)Params.Buffer;
-		
-		if(Params.LockMode != RLM_WriteOnly)
+
+		if (Params.LockMode != RLM_WriteOnly)
 		{
 			QUICK_SCOPE_CYCLE_COUNTER(STAT_RHIMETHOD_UnlockBuffer_FlushAndUnlock);
 			CSV_SCOPED_TIMING_STAT(RHITFlushes, UnlockBuffer_BottomOfPipe);
 			
 			FRHICommandListScopedFlushAndExecute Flush(RHICmdList.GetAsImmediate());
 			Buffer->Unlock(RHICmdList);
-			
-			delete LockData;
 		}
 		else
 		{
-			RHICmdList.EnqueueLambda([Buffer, Params, LockData](FRHICommandListBase& RHICmdList)
+			if (Buffer->RequiresTransferBuffer())
 			{
-				QUICK_SCOPE_CYCLE_COUNTER(STAT_FRHICommandUpdateBuffer_Execute);
-				
-				bool bRequiresTransferBuffer = Buffer->RequiresTransferBuffer();
-				void* Data = Buffer->Lock(RHICmdList, RLM_WriteOnly, Params.Offset, Params.BufferSize, LockData->Buffer);
+				FMetalRHILockBuffer* LockBuffer = reinterpret_cast<FMetalRHILockBuffer*>(Params.Buffer);
 				
-				if(!bRequiresTransferBuffer)
-				{
-					// If we spend a long time doing this memcpy, it means we got freshly allocated memory from the OS that has never been
-					// initialized and is causing pagefault to bring zeroed pages into our process.
-					{
-						TRACE_CPUPROFILER_EVENT_SCOPE(RHIUnlockBuffer_Memcpy);
-						FMemory::Memcpy(Data, LockData->Data, Params.BufferSize);
-					}
-					
-					FMemory::Free(LockData->Data);
-				}
+				const uint32 UploadSize = AdjustLockSize(Params.BufferSize, Params.Offset);
+				Buffer->UploadTransferBuffer(RHICmdList, MoveTemp(LockBuffer->Buffer), UploadSize);
 				
-				delete LockData;
+				delete LockBuffer;
+			}
+			else
+			{
+				check(Params.Buffer);
 				
-				Buffer->Unlock(RHICmdList);
-			});
-			RHICmdList.RHIThreadFence(true);
+				RHICmdList.EnqueueLambda([Buffer, Params](FRHICommandListBase& RHICmdList)
+				{
+					QUICK_SCOPE_CYCLE_COUNTER(STAT_FRHICommandUpdateBuffer_Execute);
+
+					void* LockData = Buffer->Lock(RHICmdList, Params.LockMode, Params.Offset, Params.BufferSize);
+
+					FMemory::Memcpy(LockData, Params.Buffer, Params.BufferSize);
+
+					Buffer->Unlock(RHICmdList);
+
+					FMemory::Free(Params.Buffer);
+				});
+				RHICmdList.RHIThreadFence(true);
+			}
 		}
 	}
 	else
@@ -741,7 +738,7 @@ void* FMetalDynamicRHI::LockBuffer_BottomOfPipe(FRHICommandListBase& RHICmdList,
 	FMetalRHIBuffer* Buffer = ResourceCast(BufferRHI);
 
 	// default to buffer memory
-	return (uint8*)Buffer->Lock(RHICmdList, LockMode, Offset, Size);
+	return Buffer->Lock(RHICmdList, LockMode, Offset, Size);
 }
 
 void FMetalDynamicRHI::UnlockBuffer_BottomOfPipe(FRHICommandListBase& RHICmdList, FRHIBuffer* BufferRHI)
diff --git a/Engine/Source/Runtime/Apple/MetalRHI/Private/MetalTexture.cpp b/Engine/Source/Runtime/Apple/MetalRHI/Private/MetalTexture.cpp
--- a/Engine/Source/Runtime/Apple/MetalRHI/Private/MetalTexture.cpp
+++ b/Engine/Source/Runtime/Apple/MetalRHI/Private/MetalTexture.cpp
@@ -2104,7 +2104,6 @@ void FMetalRHICommandContext::RHICopyBufferRegion(FRHIBuffer* DstBufferRHI, uint
     FMetalRHIBuffer* SrcBuffer = ResourceCast(SrcBufferRHI);
 
     check(DstBuffer && SrcBuffer);
-    check(!DstBuffer->Data && !SrcBuffer->Data);
     check(DstOffset + NumBytes <= DstBufferRHI->GetSize() && SrcOffset + NumBytes <= SrcBufferRHI->GetSize());
 
     CopyFromBufferToBuffer(SrcBuffer->GetCurrentBuffer(), SrcOffset, DstBuffer->GetCurrentBuffer(), DstOffset, NumBytes);
diff --git a/Engine/Source/Runtime/Apple/MetalRHI/Public/MetalResources.h b/Engine/Source/Runtime/Apple/MetalRHI/Public/MetalResources.h
--- a/Engine/Source/Runtime/Apple/MetalRHI/Public/MetalResources.h
+++ b/Engine/Source/Runtime/Apple/MetalRHI/Public/MetalResources.h
@@ -344,8 +344,10 @@ public:
 	 */
 	void Unlock(FRHICommandListBase& RHICmdList);
 	
+	void UploadTransferBuffer(FRHICommandListBase& RHICmdList, FMetalBufferPtr&& InTransferBuffer, uint32 UploadSize);
+
 	// We need to allocate here because buffer backed textures can be created without an Allocated buffer
-	FMetalBufferPtr GetCurrentBuffer()
+	const FMetalBufferPtr& GetCurrentBuffer()
 	{
 		if(!CurrentBuffer)
 		{
@@ -354,11 +356,17 @@ public:
 		return CurrentBuffer;
 	}
 	
-	FMetalBufferPtr GetCurrentBufferOrNull()
+	const FMetalBufferPtr& GetCurrentBufferOrNull() const
 	{
 		return CurrentBuffer;
 	}
 	
+	// 16- or 32-bit; used for index buffers only.
+	MTL::IndexType GetIndexType() const
+	{
+		return GetStride() == 2 ? MTL::IndexTypeUInt16 : MTL::IndexTypeUInt32;
+	}
+	
 #if METAL_RHI_RAYTRACING
 	bool IsAccelerationStructure() const
 	{
@@ -374,44 +382,30 @@ public:
 	
     void TakeOwnership(FMetalRHIBuffer& Other);
     void ReleaseOwnership();
-    
+
 	FMetalDevice& Device;
-	
+
 	// A temporary shared/CPU accessible buffer for upload/download
-	FMetalBufferPtr TransferBuffer = nullptr;
-	
+	FMetalBufferPtr TransferBuffer;
 	FMetalBufferPtr CurrentBuffer;
-	
-	/** Buffer for small buffers < 4Kb to avoid heap fragmentation. */
-	FMetalBufferData* Data = nullptr;
-	
-	// Current lock mode. RLM_Num indicates this buffer is not locked.
-	uint16 CurrentLockMode = RLM_Num;
-	
+
+	// Initial buffer size.
+	uint32 Size;
+
 	// offset into the buffer (for lock usage)
 	uint32 LockOffset = 0;
-	
+
 	// Sizeof outstanding lock.
 	uint32 LockSize = 0;
-	
-	bool bIsFirstLock = true;
-	
-	// Initial buffer size.
-	uint32 Size;
-	
+
 	// Storage mode
 	MTL::StorageMode Mode;
-	
-	// 16- or 32-bit; used for index buffers only.
-    MTL::IndexType GetIndexType() const
-	{
-		return GetStride() == 2
-			? MTL::IndexTypeUInt16
-			: MTL::IndexTypeUInt32;
-	}
 
-	static_assert((1 << 16) > RLM_Num, "Lock mode does not fit in bitfield");
+	// Current lock mode. RLM_Num indicates this buffer is not locked.
+	EResourceLockMode CurrentLockMode = RLM_Num;
 	
+	bool bIsFirstLock = true;
+
 #if ENABLE_LOW_LEVEL_MEM_TRACKER || UE_MEMORY_TRACE_ENABLED
 	void UpdateAllocationTags(); 
 #endif
