Metal RHI:

- Changed X24_G8 to use the correct platformPixelFormat
- Removed old code which would attempt to use PixelFormatStencil8

diff --git a/Engine/Source/Runtime/Apple/MetalRHI/Private/MetalCommandEncoder.cpp b/Engine/Source/Runtime/Apple/MetalRHI/Private/MetalCommandEncoder.cpp
--- a/Engine/Source/Runtime/Apple/MetalRHI/Private/MetalCommandEncoder.cpp
+++ b/Engine/Source/Runtime/Apple/MetalRHI/Private/MetalCommandEncoder.cpp
@@ -1400,11 +1400,7 @@ void FMetalCommandEncoder::SetShaderTexture(MTL::FunctionType FunctionType, MTL:
 	{
 		uint8 Swizzle[4] = {0,0,0,0};
 		assert(sizeof(Swizzle) == sizeof(uint32));
-		if (Texture->pixelFormat() == MTL::PixelFormatX32_Stencil8
-#if PLATFORM_MAC
-		 ||	Texture->pixelFormat() == MTL::PixelFormatX24_Stencil8
-#endif
-		)
+		if (Texture->pixelFormat() == MTL::PixelFormatX32_Stencil8 || Texture->pixelFormat() == MTL::PixelFormatX24_Stencil8)
 		{
 			Swizzle[0] = Swizzle[1] = Swizzle[2] = Swizzle[3] = 1;
 		}
diff --git a/Engine/Source/Runtime/Apple/MetalRHI/Private/MetalPipeline.cpp b/Engine/Source/Runtime/Apple/MetalRHI/Private/MetalPipeline.cpp
--- a/Engine/Source/Runtime/Apple/MetalRHI/Private/MetalPipeline.cpp
+++ b/Engine/Source/Runtime/Apple/MetalRHI/Private/MetalPipeline.cpp
@@ -130,7 +130,6 @@ struct FMetalGraphicsPipelineKey
 		}
 		
 		uint8 DepthFormatKey = 0;
-		uint8 StencilFormatKey = 0;
 		switch(Init.DepthStencilTargetFormat)
 		{
 			case PF_DepthStencil:
@@ -140,10 +139,6 @@ struct FMetalGraphicsPipelineKey
 				{
 					DepthFormatKey = GetMetalPixelFormatKey(MetalFormat);
 				}
-				if (Init.StencilTargetLoadAction != ERenderTargetLoadAction::ENoAction || Init.StencilTargetStoreAction != ERenderTargetStoreAction::ENoAction)
-				{
-					StencilFormatKey = GetMetalPixelFormatKey(MTL::PixelFormatStencil8);
-				}
 				bHasActiveTargets |= true;
 				break;
 			}
@@ -169,7 +164,6 @@ struct FMetalGraphicsPipelineKey
 		}
 		
 		Key.SetHashValue(Offset_DepthFormat, NumBits_DepthFormat, DepthFormatKey);
-		Key.SetHashValue(Offset_StencilFormat, NumBits_StencilFormat, StencilFormatKey);
 
 		Key.SetHashValue(Offset_SampleCount, NumBits_SampleCount, Init.NumSamples);
 
@@ -706,22 +700,8 @@ static bool ConfigureRenderPipelineDescriptor(FMetalDevice& Device, TDescriptorT
 		case PF_DepthStencil:
 		{
 			MTL::PixelFormat MetalFormat = (MTL::PixelFormat)GPixelFormats[Init.DepthStencilTargetFormat].PlatformFormat;
-			if(MetalFormat == MTL::PixelFormatDepth32Float)
-			{
-				if (Init.DepthTargetLoadAction != ERenderTargetLoadAction::ENoAction || Init.DepthTargetStoreAction != ERenderTargetStoreAction::ENoAction)
-				{
-					RenderPipelineDesc->setDepthAttachmentPixelFormat(MetalFormat);
-				}
-				if (Init.StencilTargetLoadAction != ERenderTargetLoadAction::ENoAction || Init.StencilTargetStoreAction != ERenderTargetStoreAction::ENoAction)
-				{
-					RenderPipelineDesc->setStencilAttachmentPixelFormat(MTL::PixelFormatStencil8);
-				}
-			}
-			else
-			{
-				RenderPipelineDesc->setDepthAttachmentPixelFormat(MetalFormat);
-				RenderPipelineDesc->setStencilAttachmentPixelFormat(MetalFormat);
-			}
+			RenderPipelineDesc->setDepthAttachmentPixelFormat(MetalFormat);
+			RenderPipelineDesc->setStencilAttachmentPixelFormat(MetalFormat);
 			break;
 		}
 		case PF_ShadowDepth:
diff --git a/Engine/Source/Runtime/Apple/MetalRHI/Private/MetalPipeline.h b/Engine/Source/Runtime/Apple/MetalRHI/Private/MetalPipeline.h
--- a/Engine/Source/Runtime/Apple/MetalRHI/Private/MetalPipeline.h
+++ b/Engine/Source/Runtime/Apple/MetalRHI/Private/MetalPipeline.h
@@ -52,7 +52,6 @@ enum EMetalPipelineHashBits
 {
 	NumBits_RenderTargetFormat = 6, //(x8=48),
 	NumBits_DepthFormat = 3, //(x1=3),
-	NumBits_StencilFormat = 3, //(x1=3),
 	NumBits_SampleCount = 3, //(x1=3),
 
 	NumBits_BlendState = 7, //(x8=56),
@@ -82,8 +81,7 @@ enum EMetalPipelineHashOffsets
 	Offset_RenderTargetFormat6 = Offset_RenderTargetFormat5 + NumBits_RenderTargetFormat,
 	Offset_RenderTargetFormat7 = Offset_RenderTargetFormat6 + NumBits_RenderTargetFormat,
 	Offset_DepthFormat = Offset_RenderTargetFormat7 + NumBits_RenderTargetFormat,
-	Offset_StencilFormat = Offset_DepthFormat + NumBits_DepthFormat,
-	Offset_SampleCount = Offset_StencilFormat + NumBits_StencilFormat,
+	Offset_SampleCount = Offset_DepthFormat + NumBits_DepthFormat,
 	Offset_AlphaToCoverage = Offset_SampleCount + NumBits_SampleCount,
 	Offset_End = Offset_AlphaToCoverage + NumBits_AlphaToCoverage
 };
diff --git a/Engine/Source/Runtime/Apple/MetalRHI/Private/MetalRHI.cpp b/Engine/Source/Runtime/Apple/MetalRHI/Private/MetalRHI.cpp
--- a/Engine/Source/Runtime/Apple/MetalRHI/Private/MetalRHI.cpp
+++ b/Engine/Source/Runtime/Apple/MetalRHI/Private/MetalRHI.cpp
@@ -227,6 +227,13 @@ FMetalDynamicRHI::FMetalDynamicRHI(ERHIFeatureLevel::Type RequestedFeatureLevel)
 	// get the device to ask about capabilities?
 	MTL::Device* MTLDevice = Device->GetDevice();
 		
+	// Use D24S8 if supported, else D32S8
+	bool bSupportsD24S8 = false;
+	
+#if PLATFORM_MAC
+	bSupportsD24S8 = MTLDevice->depth24Stencil8PixelFormatSupported();
+#endif
+	
 #if PLATFORM_IOS
     bool bSupportAppleA8 = false;
     GConfig->GetBool(TEXT("/Script/IOSRuntimeSettings.IOSRuntimeSettings"), TEXT("bSupportAppleA8"), bSupportAppleA8, GEngineIni);
@@ -370,7 +377,6 @@ FMetalDynamicRHI::FMetalDynamicRHI(ERHIFeatureLevel::Type RequestedFeatureLevel)
 	check(DeviceIndex < GPUs.Num());
 	FMacPlatformMisc::FGPUDescriptor const& GPUDesc = GPUs[DeviceIndex];
 	
-	bool bSupportsD24S8 = false;
 	bool bSupportsD16 = false;
 	
 	GRHIAdapterName = NSStringToFString(MTLDevice->name());
@@ -532,9 +538,6 @@ FMetalDynamicRHI::FMetalDynamicRHI(ERHIFeatureLevel::Type RequestedFeatureLevel)
 		MemoryStats.SharedSystemMemory = 0;
 	}
 	
-	// Change the support depth format if we can
-	bSupportsD24S8 = MTLDevice->depth24Stencil8PixelFormatSupported();
-	
 	// Disable tiled reflections on Mac Metal for some GPU drivers that ignore the lod-level and so render incorrectly.
 	if (!bSupportsTiledReflections && !FParse::Param(FCommandLine::Get(),TEXT("metaltiledreflections")))
 	{
@@ -863,6 +866,27 @@ FMetalDynamicRHI::FMetalDynamicRHI(ERHIFeatureLevel::Type RequestedFeatureLevel)
 	GPixelFormats[PF_R16G16_SINT		].PlatformFormat	= (uint32)MTL::PixelFormatRG16Sint;
 	GPixelFormats[PF_R32G32_UINT		].PlatformFormat	= (uint32)MTL::PixelFormatRG32Uint;
 
+	// Use Depth28_Stencil8 when it is available for consistency
+	if(bSupportsD24S8)
+	{
+		GPixelFormats[PF_DepthStencil].PlatformFormat				= (uint32)MTL::PixelFormatDepth24Unorm_Stencil8;
+		GPixelFormats[PF_DepthStencil].bIs24BitUnormDepthStencil 	= true;
+		GPixelFormats[PF_DepthStencil].BlockBytes					= 4;
+		GPixelFormats[PF_X24_G8].PlatformFormat						= (uint32)MTL::PixelFormatX24_Stencil8;
+		GPixelFormats[PF_X24_G8].BlockBytes							= 4;
+	}
+	else
+	{
+		GPixelFormats[PF_DepthStencil].PlatformFormat				= (uint32)MTL::PixelFormatDepth32Float_Stencil8;
+		GPixelFormats[PF_DepthStencil].bIs24BitUnormDepthStencil 	= true;
+		GPixelFormats[PF_DepthStencil].BlockBytes					= 5;
+		GPixelFormats[PF_X24_G8].PlatformFormat						= (uint32)MTL::PixelFormatX32_Stencil8;
+		GPixelFormats[PF_X24_G8].BlockBytes							= 5;
+	}
+
+	GPixelFormats[PF_DepthStencil].Supported = true;
+	GPixelFormats[PF_X24_G8].Supported = true;
+
 #if PLATFORM_IOS
     GPixelFormats[PF_DXT1				].PlatformFormat	= (uint32)MTL::PixelFormatInvalid;
 	GPixelFormats[PF_DXT1				].Supported			= false;
@@ -950,10 +974,6 @@ FMetalDynamicRHI::FMetalDynamicRHI(ERHIFeatureLevel::Type RequestedFeatureLevel)
 		GPixelFormats[PF_FloatR11G11B10		].Supported			= true;
 	}
 	
-	GPixelFormats[PF_DepthStencil		].PlatformFormat	= (uint32)MTL::PixelFormatDepth32Float_Stencil8;
-	GPixelFormats[PF_DepthStencil		].BlockBytes		= 4;
-
-	GPixelFormats[PF_DepthStencil		].Supported			= true;
 	GPixelFormats[PF_ShadowDepth		].PlatformFormat	= (uint32)MTL::PixelFormatDepth16Unorm;
 	GPixelFormats[PF_ShadowDepth		].BlockBytes		= 2;
 	GPixelFormats[PF_ShadowDepth		].Supported			= true;
@@ -986,20 +1006,7 @@ FMetalDynamicRHI::FMetalDynamicRHI(ERHIFeatureLevel::Type RequestedFeatureLevel)
 	GPixelFormats[PF_PLATFORM_HDR_0		].NumComponents		= 4;
 	GPixelFormats[PF_PLATFORM_HDR_0		].PlatformFormat	= (uint32)MTL::PixelFormatRGBA16Float;
 	GPixelFormats[PF_PLATFORM_HDR_0		].Supported			= GRHISupportsHDROutput;
-		
-	// Use Depth28_Stencil8 when it is available for consistency
-	if(bSupportsD24S8)
-	{
-		GPixelFormats[PF_DepthStencil	].PlatformFormat	= (uint32)MTL::PixelFormatDepth24Unorm_Stencil8;
-		GPixelFormats[PF_DepthStencil	].bIs24BitUnormDepthStencil = true;
-	}
-	else
-	{
-		GPixelFormats[PF_DepthStencil	].PlatformFormat	= (uint32)MTL::PixelFormatDepth32Float_Stencil8;
-		GPixelFormats[PF_DepthStencil	].bIs24BitUnormDepthStencil = false;
-	}
-	GPixelFormats[PF_DepthStencil		].BlockBytes		= 4;
-	GPixelFormats[PF_DepthStencil		].Supported			= true;
+
 	if (bSupportsD16)
 	{
 		GPixelFormats[PF_ShadowDepth		].PlatformFormat	= (uint32)MTL::PixelFormatDepth16Unorm;
@@ -1034,9 +1041,6 @@ FMetalDynamicRHI::FMetalDynamicRHI(ERHIFeatureLevel::Type RequestedFeatureLevel)
 	GPixelFormats[PF_UYVY				].PlatformFormat	= (uint32)MTL::PixelFormatInvalid;
 	GPixelFormats[PF_FloatRGBA			].PlatformFormat	= (uint32)MTL::PixelFormatRGBA16Float;
 	GPixelFormats[PF_FloatRGBA			].BlockBytes		= 8;
-    GPixelFormats[PF_X24_G8				].PlatformFormat	= (uint32)MTL::PixelFormatStencil8;
-    GPixelFormats[PF_X24_G8				].BlockBytes		= 1;
-	GPixelFormats[PF_X24_G8             ].Supported         = true;
 	
     GPixelFormats[PF_R32_FLOAT			].PlatformFormat	= (uint32)MTL::PixelFormatR32Float;
 #if PLATFORM_MAC
diff --git a/Engine/Source/Runtime/Apple/MetalRHI/Private/MetalStateCache.cpp b/Engine/Source/Runtime/Apple/MetalRHI/Private/MetalStateCache.cpp
--- a/Engine/Source/Runtime/Apple/MetalRHI/Private/MetalStateCache.cpp
+++ b/Engine/Source/Runtime/Apple/MetalRHI/Private/MetalStateCache.cpp
@@ -784,42 +784,9 @@ bool FMetalStateCache::SetRenderPassInfo(FRHIRenderPassInfo const& InRenderTarge
 				}
             }
 
-			switch (DepthStencilPixelFormat)
+			if(DepthStencilPixelFormat == PF_X24_G8 || DepthStencilPixelFormat == PF_DepthStencil)
 			{
-				case PF_X24_G8:
-				case PF_DepthStencil:
-				case PF_D24:
-				{
-                    MTL::PixelFormat DepthStencilFormat = Surface.Texture ? (MTL::PixelFormat)Surface.Texture->pixelFormat() : MTL::PixelFormatInvalid;
-					
-					switch(DepthStencilFormat)
-					{
-						case MTL::PixelFormatDepth32Float:
-							StencilTexture =  nullptr;
-							break;
-						case MTL::PixelFormatStencil8:
-							StencilTexture = DepthTexture;
-							break;
-						case MTL::PixelFormatDepth32Float_Stencil8:
-							StencilTexture = DepthTexture;
-							break;
-#if PLATFORM_MAC
-						case MTL::PixelFormatDepth24Unorm_Stencil8:
-							StencilTexture = DepthTexture;
-							break;
-#endif
-						default:
-							break;
-					}
-					
-					break;
-				}
-				case PF_ShadowDepth:
-				{
-					break;
-				}
-				default:
-					break;
+				StencilTexture = DepthTexture;
 			}
 			
 			float DepthClearValue = 0.0f;
@@ -920,9 +887,9 @@ bool FMetalStateCache::SetRenderPassInfo(FRHIRenderPassInfo const& InRenderTarge
 			
             //if we're dealing with a samplecount mismatch we just bail on stencil entirely as stencil
             //doesn't have an autoresolve target to use.
-			
-			bool const bCombinedDepthStencilUsingDepth = (StencilTexture && StencilTexture->pixelFormat() != MTL::PixelFormatStencil8 && RenderPassInfo.DepthStencilRenderTarget.ExclusiveDepthStencil.IsUsingDepth());
+			bool const bCombinedDepthStencilUsingDepth =  RenderPassInfo.DepthStencilRenderTarget.ExclusiveDepthStencil.IsUsingDepth();
 			bool const bUsingStencil = RenderPassInfo.DepthStencilRenderTarget.ExclusiveDepthStencil.IsUsingStencil() || (bCombinedDepthStencilUsingDepth);
+			
 			if (StencilTexture && bUsingStencil)
 			{
                 MTL::RenderPassStencilAttachmentDescriptor* StencilAttachment = MTL::RenderPassStencilAttachmentDescriptor::alloc()->init();
diff --git a/Engine/Source/Runtime/Apple/MetalRHI/Private/MetalTexture.cpp b/Engine/Source/Runtime/Apple/MetalRHI/Private/MetalTexture.cpp
--- a/Engine/Source/Runtime/Apple/MetalRHI/Private/MetalTexture.cpp
+++ b/Engine/Source/Runtime/Apple/MetalRHI/Private/MetalTexture.cpp
@@ -321,7 +321,6 @@ uint8 GetMetalPixelFormatKey(MTL::PixelFormat Format)
 		{
 			// Add depth stencil formats first, so we don't have to use 6 bits for them in the pipeline hash
 			Get(MTL::PixelFormatDepth32Float);
-			Get(MTL::PixelFormatStencil8);
 			Get(MTL::PixelFormatDepth32Float_Stencil8);
 	#if PLATFORM_MAC
 			Get(MTL::PixelFormatDepth24Unorm_Stencil8);
diff --git a/Engine/Source/Runtime/Apple/MetalRHI/Private/MetalUAV.cpp b/Engine/Source/Runtime/Apple/MetalRHI/Private/MetalUAV.cpp
--- a/Engine/Source/Runtime/Apple/MetalRHI/Private/MetalUAV.cpp
+++ b/Engine/Source/Runtime/Apple/MetalRHI/Private/MetalUAV.cpp
@@ -284,20 +284,6 @@ void FMetalShaderResourceView::UpdateView(FMetalRHICommandContext* Context, cons
             }
 #endif
         }
-
-        if (Info.Format == PF_X24_G8)
-        {
-            // Stencil buffer view of a depth texture
-            check(Texture->GetDesc().Format == PF_DepthStencil);
-            switch (Texture->Texture->pixelFormat())
-            {
-            default: checkNoEntry(); break;
-#if PLATFORM_MAC
-            case MTL::PixelFormatDepth24Unorm_Stencil8: MetalFormat = MTL::PixelFormatX24_Stencil8; break;
-#endif
-            case MTL::PixelFormatDepth32Float_Stencil8: MetalFormat = MTL::PixelFormatX32_Stencil8; break;
-            }
-        }
         
         bool bUseSourceTexture = Info.bAllMips && Info.bAllSlices && MetalFormat == Texture->Texture->pixelFormat() &&
                                 SRVDimensionToMetalTextureType(Device, Info.Dimension) == TextureType;
@@ -498,20 +484,6 @@ void FMetalUnorderedAccessView::UpdateView(FMetalRHICommandContext* Context, con
 #endif
         }
         
-        if (Info.Format == PF_X24_G8)
-        {
-            // Stencil buffer view of a depth texture
-            check(Texture->GetDesc().Format == PF_DepthStencil);
-            switch (Texture->Texture->pixelFormat())
-            {
-            default: checkNoEntry(); break;
-#if PLATFORM_MAC
-            case MTL::PixelFormatDepth24Unorm_Stencil8: MetalFormat = MTL::PixelFormatX24_Stencil8; break;
-#endif
-            case MTL::PixelFormatDepth32Float_Stencil8: MetalFormat = MTL::PixelFormatX32_Stencil8; break;
-            }
-        }
-        
         bool bUseSourceTexture = Info.bAllMips && Info.bAllSlices &&
                                 UAVDimensionToMetalTextureType(Info.Dimension) == TextureType && MetalFormat == Texture->Texture->pixelFormat();
         
