Metal base vertex support fix

GRHISupportsBaseVertexIndex & GRHISupportsFirstInstance should be set based of GPU family.

MTLFeatureSet_iOS_GPUFamily3_v1 -> This feature set corresponds with Apple A9 devices running iOS 9.
Reference: https://developer.apple.com/metal/Metal-Feature-Set-Tables.pdf

Also added more logging to plugin code, which now checks GRHISupportsBaseVertexIndex

#rnx

diff --git a/Engine/Source/Runtime/Apple/MetalRHI/Private/MetalRHI.cpp b/Engine/Source/Runtime/Apple/MetalRHI/Private/MetalRHI.cpp
--- a/Engine/Source/Runtime/Apple/MetalRHI/Private/MetalRHI.cpp
+++ b/Engine/Source/Runtime/Apple/MetalRHI/Private/MetalRHI.cpp
@@ -715,8 +715,8 @@ FMetalDynamicRHI::FMetalDynamicRHI(ERHIFeatureLevel::Type RequestedFeatureLevel)
 	GRHIHDRDisplayOutputFormat = (GRHISupportsHDROutput) ? PF_PLATFORM_HDR_0 : PF_B8G8R8A8;
 	GMaxWorkGroupInvocations = 512;
 #else
-	// Only A9+ can support this, so for now we need to limit this to the desktop-forward renderer only.
-	GRHISupportsBaseVertexIndex = MTLDevice->supportsFeatureSet(MTL::FeatureSet_iOS_GPUFamily3_v1) && (GMaxRHIFeatureLevel >= ERHIFeatureLevel::SM5);
+	// "Base vertex/instance drawing" is Apple3 family, which is A9 and newer:
+	GRHISupportsBaseVertexIndex = MTLDevice->supportsFeatureSet(MTL::FeatureSet_iOS_GPUFamily3_v1);
 	GRHISupportsFirstInstance = GRHISupportsBaseVertexIndex;
 	
 	// TODO: Move this into IOSPlatform
