MetalRHIBuffer changes
- Fixing transfer buffer locks. They weren't advancing to the next buffer correctly.
- Doing a small cleanup of FMetalRHIBuffer::Lock.

diff --git a/Engine/Source/Runtime/Apple/MetalRHI/Private/MetalRHIBuffer.cpp b/Engine/Source/Runtime/Apple/MetalRHI/Private/MetalRHIBuffer.cpp
--- a/Engine/Source/Runtime/Apple/MetalRHI/Private/MetalRHIBuffer.cpp
+++ b/Engine/Source/Runtime/Apple/MetalRHI/Private/MetalRHIBuffer.cpp
@@ -324,7 +324,7 @@ static uint32 AdjustLockSize(uint32 Size, uint32 Offset)
 #endif
 }
 
-void* FMetalRHIBuffer::Lock(FRHICommandListBase& RHICmdList, EResourceLockMode InLockMode, uint32 Offset, uint32 InSize, FMetalBufferPtr InTransferBuffer)
+void* FMetalRHIBuffer::Lock(FRHICommandListBase& RHICmdList, EResourceLockMode InLockMode, uint32 Offset, uint32 InSize)
 {
 	check(CurrentLockMode == RLM_Num);
 	check(LockSize == 0 && LockOffset == 0);	
@@ -333,42 +333,32 @@ void* FMetalRHIBuffer::Lock(FRHICommandListBase& RHICmdList, EResourceLockMode I
 	InSize = AdjustLockSize(InSize, Offset);
 
 	const bool bWriteLock = InLockMode == RLM_WriteOnly;
-	const bool bIsStatic = EnumHasAnyFlags(GetUsage(), BUF_Static);
-	const bool bIsDynamic = EnumHasAnyFlags(GetUsage(), BUF_Dynamic);
-	const bool bIsVolatile = EnumHasAnyFlags(GetUsage(), BUF_Volatile);
-	
+	const bool bIsStatic = EnumHasAnyFlags(GetUsage(), EBufferUsageFlags::Static);
+
 	void* ReturnPointer = nullptr;
-	
+
 	uint32 Len = GetCurrentBuffer()->GetLength(); // all buffers should have the same length or we are in trouble.
 	check(Len >= InSize);
-	
-	if(bWriteLock)
+
+	if (bWriteLock)
 	{
 		const bool bUseTransferBuffer = RequiresTransferBuffer();
-		
+
 		// If we are locking for the first time then use the current buffer
 		bool bValidFirstLock = bIsFirstLock && CurrentBuffer;
-		
-		if(!bValidFirstLock &&
-		   (!bIsStatic || bUseTransferBuffer))
+
+		if (!bValidFirstLock && (!bIsStatic || bUseTransferBuffer))
 		{
 			SwitchBuffer(RHICmdList);
 		}
-		
+
 		bIsFirstLock = false;
-		
+
 		if(bUseTransferBuffer)
 		{			
-			if(InTransferBuffer)
-			{
-				TransferBuffer = InTransferBuffer;
-			}
-			else
-			{
-				TransferBuffer = Device.GetTransferAllocator()->Allocate(Len);
-				ReturnPointer = (uint8*) TransferBuffer->Contents();
-				check(ReturnPointer != nullptr);
-			}
+			TransferBuffer = Device.GetTransferAllocator()->Allocate(Len);
+			ReturnPointer = TransferBuffer->Contents();
+			check(ReturnPointer != nullptr);
 		}
 		else
 		{
@@ -382,8 +372,8 @@ void* FMetalRHIBuffer::Lock(FRHICommandListBase& RHICmdList, EResourceLockMode I
 		check(InLockMode == EResourceLockMode::RLM_ReadOnly);
 		// assumes offset is 0 for reads.
 		check(Offset == 0);
-		
-		if(Mode == MTL::StorageModePrivate)
+
+		if (Mode == MTL::StorageModePrivate)
 		{
 			check(!TransferBuffer);
 			SCOPE_CYCLE_COUNTER(STAT_MetalBufferPageOffTime);
@@ -401,7 +391,7 @@ void* FMetalRHIBuffer::Lock(FRHICommandListBase& RHICmdList, EResourceLockMode I
 			
 			ReturnPointer = TransferBuffer->Contents();
 		}
-		#if PLATFORM_MAC
+#if PLATFORM_MAC
 		else if(Mode == MTL::StorageModeManaged)
 		{
 			SCOPE_CYCLE_COUNTER(STAT_MetalBufferPageOffTime);
@@ -417,16 +407,16 @@ void* FMetalRHIBuffer::Lock(FRHICommandListBase& RHICmdList, EResourceLockMode I
 			
 			ReturnPointer = GetCurrentBuffer()->Contents();
 		}
-		#endif
+#endif
 		else
 		{
 			// Shared
             ReturnPointer = GetCurrentBuffer()->Contents();
 		}
-		
+
 		check(ReturnPointer);
 	} // Read Path
-	
+
 	check(GetCurrentBuffer());
 	check((!GetCurrentBuffer()->GetMTLBuffer()->heap() && !GetCurrentBuffer()->GetMTLBuffer()->isAliasable()) || GetCurrentBuffer()->GetMTLBuffer()->heap() != nullptr);
 	
@@ -504,22 +494,30 @@ void FMetalRHIBuffer::Unlock(FRHICommandListBase& RHICmdList)
 
 void FMetalRHIBuffer::UploadTransferBuffer(FRHICommandListBase& RHICmdList, FMetalBufferPtr&& InTransferBuffer, uint32 UploadSize)
 {
-	check(CurrentBuffer);
 	check(InTransferBuffer);
 	check(UploadSize);
 	
 	RHICmdList.EnqueueLambda(
-		[InTransferBuffer = MoveTemp(InTransferBuffer), CurrentBuffer=CurrentBuffer, UploadSize](FRHICommandListBase& RHICmdList) mutable
+		[this, InTransferBuffer = MoveTemp(InTransferBuffer), UploadSize](FRHICommandListBase& RHICmdList) mutable
 		{
+			if (!(bIsFirstLock && CurrentBuffer))
+			{
+				SwitchBuffer(RHICmdList);
+			}
+			
+			bIsFirstLock = false;
+				
 			FMetalRHIUploadContext& UploadContext = static_cast<FMetalRHIUploadContext&>(RHICmdList.GetUploadContext());
 
 			UploadContext.EnqueueFunction(
-				[InTransferBuffer = MoveTemp(InTransferBuffer), CurrentBuffer, UploadSize](FMetalRHICommandContext* Context)
+				[InTransferBuffer = MoveTemp(InTransferBuffer), CurrentBuffer=CurrentBuffer, UploadSize](FMetalRHICommandContext* Context)
 				{
 					Context->CopyFromBufferToBuffer(InTransferBuffer, 0, CurrentBuffer, 0, UploadSize);
 					FMetalDynamicRHI::Get().DeferredDelete(InTransferBuffer);
 				});
 		});
+
+	RHICmdList.RHIThreadFence(true);
 }
 
 void FMetalRHIBuffer::TakeOwnership(FMetalRHIBuffer& Other)
@@ -699,10 +697,10 @@ void FMetalDynamicRHI::RHIUnlockBuffer(class FRHICommandListBase& RHICmdList, FR
 			if (Buffer->RequiresTransferBuffer())
 			{
 				FMetalRHILockBuffer* LockBuffer = reinterpret_cast<FMetalRHILockBuffer*>(Params.Buffer);
-				
+
 				const uint32 UploadSize = AdjustLockSize(Params.BufferSize, Params.Offset);
 				Buffer->UploadTransferBuffer(RHICmdList, MoveTemp(LockBuffer->Buffer), UploadSize);
-				
+
 				delete LockBuffer;
 			}
 			else
diff --git a/Engine/Source/Runtime/Apple/MetalRHI/Public/MetalResources.h b/Engine/Source/Runtime/Apple/MetalRHI/Public/MetalResources.h
--- a/Engine/Source/Runtime/Apple/MetalRHI/Public/MetalResources.h
+++ b/Engine/Source/Runtime/Apple/MetalRHI/Public/MetalResources.h
@@ -337,7 +337,7 @@ public:
 	/**
 	 * Prepare a CPU accessible buffer for uploading to GPU memory
 	 */
-	void* Lock(FRHICommandListBase& RHICmdList, EResourceLockMode LockMode, uint32 Offset, uint32 Size=0, FMetalBufferPtr InTransferBuffer = nullptr);
+	void* Lock(FRHICommandListBase& RHICmdList, EResourceLockMode LockMode, uint32 Offset, uint32 Size);
 	
 	/**
 	 * Prepare a CPU accessible buffer for uploading to GPU memory
