Metal - Fix bindless MSAA in ForwardShading
In bindless, we were setting the texture type to TextureType2DMultisampleArray too early.
This caused CreateTexture to fail when creating the MSAAResolveTexture, because the number of samples was unset (== 1), which is invalid for a multisample texture type.

#tests Run locally with SM5 bindless and r.ForwardShading=True && r.AntiAliasingMethod=3

diff --git a/Engine/Source/Runtime/Apple/MetalRHI/Private/MetalTexture.cpp b/Engine/Source/Runtime/Apple/MetalRHI/Private/MetalTexture.cpp
--- a/Engine/Source/Runtime/Apple/MetalRHI/Private/MetalTexture.cpp
+++ b/Engine/Source/Runtime/Apple/MetalRHI/Private/MetalTexture.cpp
@@ -395,15 +395,7 @@ FMetalTextureCreateDesc::FMetalTextureCreateDesc(FMetalDevice& Device, FRHITextu
 		{
 			if (InDesc.IsTexture2D())
 			{
-				if (InDesc.NumSamples > 1)
-				{
-					Desc->setTextureType(MTL::TextureType2DMultisampleArray);
-				}
-				else
-				{
-					Desc->setTextureType(MTL::TextureType2DArray);
-				}
-				
+				Desc->setTextureType(MTL::TextureType2DArray);
 				Desc->setArrayLength(1);
 			}
 			else if (InDesc.IsTextureCube())
@@ -564,6 +556,7 @@ FMetalSurface::FMetalSurface(FMetalDevice& MetalDevice, FMetalTextureCreateDesc
 	}
     
     bool bIsMSAARequired = CreateDesc.NumSamples > 1 && AllowMSAA();
+	bool bIsBindless = IsMetalBindlessEnabled();
 
 #if RHI_USE_RESOURCE_DEBUG_NAME
     NS::String* DebugName = FStringToNSString(CreateDesc.DebugName);
@@ -597,8 +590,6 @@ FMetalSurface::FMetalSurface(FMetalDevice& MetalDevice, FMetalTextureCreateDesc
 
 		const bool bAtomicCompatible = EnumHasAllFlags(CreateDesc.Flags, TexCreate_AtomicCompatible) || EnumHasAllFlags(CreateDesc.Flags, ETextureCreateFlags::Atomic64Compatible);
 		
-		bool bIsBindless = IsMetalBindlessEnabled();
-		
 		bool bBufferBacked = EnumHasAllFlags(CreateDesc.Flags, TexCreate_UAV | TexCreate_NoTiling);
 		if (bIsBindless)
 		{
@@ -669,7 +660,9 @@ FMetalSurface::FMetalSurface(FMetalDevice& MetalDevice, FMetalTextureCreateDesc
 	{
 		MTLTextureDescriptorPtr Desc = CreateDesc.Desc;
 		check(CreateDesc.bIsRenderTarget);
-		Desc->setTextureType(MTL::TextureType2DMultisample);
+		
+		// In bindless texture types need to be converted to Array Types to match the generated AIR
+		Desc->setTextureType(bIsBindless ? MTL::TextureType2DMultisampleArray : MTL::TextureType2DMultisample);
 
 		// allow commandline to override
 		uint32 NewNumSamples;
diff --git a/Engine/Source/Runtime/Apple/MetalRHI/Private/MetalUAV.cpp b/Engine/Source/Runtime/Apple/MetalRHI/Private/MetalUAV.cpp
--- a/Engine/Source/Runtime/Apple/MetalRHI/Private/MetalUAV.cpp
+++ b/Engine/Source/Runtime/Apple/MetalRHI/Private/MetalUAV.cpp
@@ -311,7 +311,7 @@ void FMetalShaderResourceView::UpdateView(FMetalRHICommandContext* Context, cons
                 ArraySize = Info.ArrayRange.Num * 6;	
             }
             
-            if(TextureType != MTL::TextureType2DMultisample)
+            if(TextureType != MTL::TextureType2DMultisample && TextureType != MTL::TextureType2DMultisampleArray)
             {
                 TextureType = SRVDimensionToMetalTextureType(Device, Info.Dimension);
             }
