Fixed issue with bindings reset when using bindful, now only bindless resets the uniform buffer bindings.

diff --git a/Engine/Source/Runtime/Apple/MetalRHI/Private/MetalStateCache.cpp b/Engine/Source/Runtime/Apple/MetalRHI/Private/MetalStateCache.cpp
--- a/Engine/Source/Runtime/Apple/MetalRHI/Private/MetalStateCache.cpp
+++ b/Engine/Source/Runtime/Apple/MetalRHI/Private/MetalStateCache.cpp
@@ -268,11 +268,6 @@ void FMetalStateCache::Reset()
 	RasterBits = EMetalRenderFlagMask;
 	PipelineBits = EMetalPipelineFlagMask;
 	
-	ResetBindings();
-}
-
-void FMetalStateCache::ResetBindings()
-{
 	FMemory::Memzero(DirtyUniformBuffers);
 	FMemory::Memzero(BoundUniformBuffers);
 	ActiveUniformBuffers.Empty();
@@ -287,6 +282,7 @@ void FMetalStateCache::ResetBindings()
 		VertexBuffers[i].Length = 0;
 		VertexBuffers[i].Offset = 0;
 	}
+	
 	for (uint32 Frequency = 0; Frequency < EMetalShaderStages::Num; Frequency++)
 	{
 		ShaderSamplers[Frequency].Bound = 0;
@@ -313,16 +309,10 @@ void FMetalStateCache::ResetBindings()
 			ShaderTextures[Frequency].Usage[i] = MTL::ResourceUsage(0);
 		}
 		ShaderTextures[Frequency].Bound = 0;
-		
-#if PLATFORM_SUPPORTS_BINDLESS_RENDERING
-		HeapsUsedByStage[Frequency].Reset();
-		ROResourcesByStage[Frequency].Reset();
-		RWResourcesByStage[Frequency].Reset();
-#endif
-
-		SetRenderTargetsActive(false);
 	}
 
+	SetRenderTargetsActive(false);
+	
 #if PLATFORM_SUPPORTS_BINDLESS_RENDERING
 	if(GIsMetalInitialized)
 	{
@@ -348,6 +338,33 @@ void FMetalStateCache::ResetBindings()
 		}
 	}
 #endif
+	
+	ResetBindings();
+}
+
+void FMetalStateCache::ResetBindings()
+{
+	if(GIsMetalInitialized)
+	{
+#if PLATFORM_SUPPORTS_BINDLESS_RENDERING
+		FMetalBindlessDescriptorManager* BindlessDescriptorManager = Device.GetBindlessDescriptorManager();
+		
+		if(BindlessDescriptorManager->IsSupported())
+		{
+			FMemory::Memzero(DirtyUniformBuffers);
+			FMemory::Memzero(BoundUniformBuffers);
+			ActiveUniformBuffers.Empty();
+			
+			for (uint32 Frequency = 0; Frequency < EMetalShaderStages::Num; Frequency++)
+			{
+				HeapsUsedByStage[Frequency].Reset();
+				ROResourcesByStage[Frequency].Reset();
+				RWResourcesByStage[Frequency].Reset();
+			}
+		}
+#endif
+	}
+	SetRenderTargetsActive(false);
 }
 
 static bool MTLScissorRectEqual(MTL::ScissorRect const& Left, MTL::ScissorRect const& Right)
