Metal - Disabling EMetalFeaturesCubemapArrays for all iOS target
Because we currently support SM5 on iOS device with A9 and above, we set emulate_cube_array on the SPIRV compiler.
This and the EMetalFeaturesCubemapArrays can cause a missmatch between the compiled shader and the resulting shader code

#tests local test on Vision Pro & iOS

diff --git a/Engine/Source/Runtime/Apple/MetalRHI/Private/MetalDevice.cpp b/Engine/Source/Runtime/Apple/MetalRHI/Private/MetalDevice.cpp
--- a/Engine/Source/Runtime/Apple/MetalRHI/Private/MetalDevice.cpp
+++ b/Engine/Source/Runtime/Apple/MetalRHI/Private/MetalDevice.cpp
@@ -468,10 +468,12 @@ void FMetalDevice::EnumerateFeatureSupport()
 		Features |= EMetalFeaturesCountingQueries | EMetalFeaturesBaseVertexInstance | EMetalFeaturesIndirectBuffer | EMetalFeaturesMSAADepthResolve | EMetalFeaturesMSAAStoreAndResolve;
 	}
 	
+#if PLATFORM_MAC // on iOS we use emulate_cube_array with SPIRV
 	if(bSupportsMac2Minimum || Device->supportsFamily(MTL::GPUFamilyApple4))
 	{
 		Features |= EMetalFeaturesCubemapArrays;
 	}
+#endif // PLATFORM_MAC
 	
 	if(bSupportsMac2Minimum || Device->supportsFamily(MTL::GPUFamilyApple5))
 	{
