Metal RHI

Fix page fault when resetting bindings, now just mark flags dirty

diff --git a/Engine/Source/Runtime/Apple/MetalRHI/Private/MetalStateCache.cpp b/Engine/Source/Runtime/Apple/MetalRHI/Private/MetalStateCache.cpp
--- a/Engine/Source/Runtime/Apple/MetalRHI/Private/MetalStateCache.cpp
+++ b/Engine/Source/Runtime/Apple/MetalRHI/Private/MetalStateCache.cpp
@@ -351,16 +351,15 @@ void FMetalStateCache::ResetBindings()
 		
 		if(BindlessDescriptorManager->IsSupported())
 		{
-			FMemory::Memzero(DirtyUniformBuffers);
-			FMemory::Memzero(BoundUniformBuffers);
-			ActiveUniformBuffers.Empty();
-			
-			for (uint32 Frequency = 0; Frequency < EMetalShaderStages::Num; Frequency++)
-			{
-				HeapsUsedByStage[Frequency].Reset();
-				ROResourcesByStage[Frequency].Reset();
-				RWResourcesByStage[Frequency].Reset();
-			}
+			DirtyUniformBuffers[EMetalShaderStages::Vertex] = 0xffffffff;
+			DirtyUniformBuffers[EMetalShaderStages::Pixel] = 0xffffffff;
+	#if PLATFORM_SUPPORTS_GEOMETRY_SHADERS
+			DirtyUniformBuffers[EMetalShaderStages::Geometry] = 0xffffffff;
+	#endif
+	#if PLATFORM_SUPPORTS_MESH_SHADERS
+			DirtyUniformBuffers[EMetalShaderStages::Mesh] = 0xffffffff;
+			DirtyUniformBuffers[EMetalShaderStages::Amplification] = 0xffffffff;
+	#endif
 		}
 #endif
 	}
