Metal RHI

- Temporarily disable Compaction as it causes crashes on GPU

diff --git a/Engine/Source/Runtime/Apple/MetalRHI/Private/MetalRayTracing.cpp b/Engine/Source/Runtime/Apple/MetalRHI/Private/MetalRayTracing.cpp
--- a/Engine/Source/Runtime/Apple/MetalRHI/Private/MetalRayTracing.cpp
+++ b/Engine/Source/Runtime/Apple/MetalRHI/Private/MetalRayTracing.cpp
@@ -16,7 +16,7 @@
 #include "RayTracingValidationShaders.h"
 #include "Async/ParallelFor.h"
 
-static int32 GMetalRayTracingAllowCompaction = 1;
+static int32 GMetalRayTracingAllowCompaction = 0;
 static FAutoConsoleVariableRef CVarMetalRayTracingAllowCompaction(
 	TEXT("r.Metal.RayTracing.AllowCompaction"),
 	GMetalRayTracingAllowCompaction,
