Metal RHI:

- Enable Parallel translate by default

diff --git a/Engine/Source/Runtime/Apple/MetalRHI/Private/MetalRHI.cpp b/Engine/Source/Runtime/Apple/MetalRHI/Private/MetalRHI.cpp
--- a/Engine/Source/Runtime/Apple/MetalRHI/Private/MetalRHI.cpp
+++ b/Engine/Source/Runtime/Apple/MetalRHI/Private/MetalRHI.cpp
@@ -52,15 +52,14 @@ static TAutoConsoleVariable<int32> CVarUseIOSRHIThread(
 
 static TAutoConsoleVariable<int32> CVarMetalParallel(
 													TEXT("r.Metal.Parallel"),
-													0,
+													PLATFORM_MAC,
 													TEXT("Controls Parallel Translate support for MacOS/IOS:\n")
 													TEXT("\t0: No Parallel support.\n")
 													TEXT("\t1: Parallel enabled.\n")
-													TEXT("Default is 0."),
+													TEXT("Default is 1 on Mac, 0 on iOS/VisionOS."),
 													ECVF_Default | ECVF_RenderThreadSafe
 													);
 
-
 // If precaching is active we should not need the file cache.
 // however, precaching and filecache are compatible with each other, there maybe some scenarios in which both could be used.
 static TAutoConsoleVariable<bool> CVarEnableMetalPSOFileCacheWhenPrecachingActive(
