MetalRHI

- Add concurrent dispatch with conservative memory barrier

diff --git a/Engine/Source/Runtime/Apple/MetalRHI/Private/MetalRHIContext.cpp b/Engine/Source/Runtime/Apple/MetalRHI/Private/MetalRHIContext.cpp
--- a/Engine/Source/Runtime/Apple/MetalRHI/Private/MetalRHIContext.cpp
+++ b/Engine/Source/Runtime/Apple/MetalRHI/Private/MetalRHIContext.cpp
@@ -17,14 +17,20 @@ static int32 GMetalSampleBlitEncoderTimings = 0;
 static FAutoConsoleVariableRef CVarMetalSampleBlitEncoderTimings(
 	TEXT("rhi.Metal.SampleBlitEncoderTimings"),
 	GMetalSampleBlitEncoderTimings,
-	TEXT("Whether to include Blit Encoder Timings in GPU Stats, can cause poor perf due to splitting encoders: on"));
+	TEXT("Whether to include Blit Encoder Timings in GPU Stats, can cause poor perf due to splitting encoders: off"));
 
-static int32 GMetalSampleComputeEncoderTimings = 0;
+static int32 GMetalSampleComputeEncoderTimings = 1;
 static FAutoConsoleVariableRef CVarMetalSampleComputeEncoderTimings(
 	TEXT("rhi.Metal.SampleComputeEncoderTimings"),
 	GMetalSampleComputeEncoderTimings,
 	TEXT("Whether to include Compute Encoder Timings in GPU Stats, can cause poor perf due to splitting encoders: on"));
 
+static int32 GMetalConcurrentDispatch = 1;
+static FAutoConsoleVariableRef CVarMetalConcurrentDispatch(
+	TEXT("rhi.Metal.ConcurrentDispatch"),
+	GMetalConcurrentDispatch,
+	TEXT("Whether concurrent dispatch is on or off: default on"));
+
 void METALRHI_API SafeReleaseMetalObject(NS::Object* Object)
 {
 	if(GIsMetalInitialized && GDynamicRHI && Object)
@@ -148,7 +154,9 @@ void FMetalRHICommandContext::BeginComputeEncoder()
 			StateCache.ResetBindings();
 		}
 		bool bUseStageCounterSamples = Device.SupportsFeature(EMetalFeaturesStageCounterSampling) && GMetalSampleComputeEncoderTimings;
-		CurrentEncoder.BeginComputeCommandEncoding(MTL::DispatchTypeSerial,
+		
+		MTL::DispatchType DispatchType = GMetalConcurrentDispatch ? MTL::DispatchTypeConcurrent : MTL::DispatchTypeSerial;
+		CurrentEncoder.BeginComputeCommandEncoding(DispatchType,
 												   bUseStageCounterSamples ? Device.GetCounterSampler() : nullptr);
 	}
 	
diff --git a/Engine/Source/Runtime/Apple/MetalRHI/Private/MetalTransitionData.cpp b/Engine/Source/Runtime/Apple/MetalRHI/Private/MetalTransitionData.cpp
--- a/Engine/Source/Runtime/Apple/MetalRHI/Private/MetalTransitionData.cpp
+++ b/Engine/Source/Runtime/Apple/MetalRHI/Private/MetalTransitionData.cpp
@@ -35,7 +35,7 @@ void FMetalTransitionData::BeginResourceTransitions() const
 {
 }
 
-void FMetalTransitionData::EndResourceTransitions(FMetalCommandEncoder& CurrentEncoder) const
+void FMetalTransitionData::EndResourceTransitions(FMetalCommandEncoder& CurrentEncoder, bool& UAVBarrier) const
 {
 	check(SrcPipelines == DstPipelines);
 
@@ -54,6 +54,8 @@ void FMetalTransitionData::EndResourceTransitions(FMetalCommandEncoder& CurrentE
 
 		checkf(Info.AccessAfter != ERHIAccess::Unknown, TEXT("Transitioning a resource to an unknown state is not allowed."));
 
+		UAVBarrier = UAVBarrier || EnumHasAnyFlags(Info.AccessBefore, ERHIAccess::UAVCompute) || EnumHasAnyFlags(Info.AccessAfter, ERHIAccess::UAVCompute);
+		
 		switch (Info.Type)
 		{
 			case FRHITransitionInfo::EType::UAV:
@@ -121,8 +123,18 @@ void FMetalRHICommandContext::RHIBeginTransitions(TArrayView<const FRHITransitio
 
 void FMetalRHICommandContext::RHIEndTransitions(TArrayView<const FRHITransition*> Transitions)
 {
+	bool UAVBarrier = false;
+	
 	for (auto Transition : Transitions)
 	{
-		Transition->GetPrivateData<FMetalTransitionData>()->EndResourceTransitions(CurrentEncoder);
+		Transition->GetPrivateData<FMetalTransitionData>()->EndResourceTransitions(CurrentEncoder, UAVBarrier);
+	}
+	
+	if (UAVBarrier)
+	{
+		if (CurrentEncoder.IsComputeCommandEncoderActive())
+		{
+			CurrentEncoder.GetComputeCommandEncoder()->memoryBarrier(MTL::BarrierScopeBuffers | MTL::BarrierScopeTextures);
+		}
 	}
 }
diff --git a/Engine/Source/Runtime/Apple/MetalRHI/Private/MetalTransitionData.h b/Engine/Source/Runtime/Apple/MetalRHI/Private/MetalTransitionData.h
--- a/Engine/Source/Runtime/Apple/MetalRHI/Private/MetalTransitionData.h
+++ b/Engine/Source/Runtime/Apple/MetalRHI/Private/MetalTransitionData.h
@@ -42,7 +42,7 @@ public:
 	void BeginResourceTransitions() const;
 
 	// End resource transitions.
-	void EndResourceTransitions(FMetalCommandEncoder& CurrentEncoder) const;
+	void EndResourceTransitions(FMetalCommandEncoder& CurrentEncoder, bool& UAVBarrier) const;
 
 private:
 	ERHIPipeline              SrcPipelines   = ERHIPipeline::Num;
