Metal RHI

- Fixed side table binding relying on reflection of binding name spvBufferSizeConstants
- Striped dead StreamShader code

diff --git a/Engine/Source/Runtime/Apple/MetalRHI/Private/MetalPipeline.cpp b/Engine/Source/Runtime/Apple/MetalRHI/Private/MetalPipeline.cpp
--- a/Engine/Source/Runtime/Apple/MetalRHI/Private/MetalPipeline.cpp
+++ b/Engine/Source/Runtime/Apple/MetalRHI/Private/MetalPipeline.cpp
@@ -396,39 +396,34 @@ FMetalShaderPipeline::~FMetalShaderPipeline()
     
 }
 
-void FMetalShaderPipeline::InitResourceMask()
+void FMetalShaderPipeline::InitResourceMask(const FGraphicsPipelineStateInitializer& Init)
 {
 	if (RenderPipelineReflection)
 	{
-		InitResourceMask(EMetalShaderVertex);
-		InitResourceMask(EMetalShaderFragment);
+		InitResourceMask(EMetalShaderVertex, ResourceCast(Init.BoundShaderState.GetVertexShader())->SideTableBinding);
+		InitResourceMask(EMetalShaderFragment, Init.BoundShaderState.GetPixelShader() ? ResourceCast(Init.BoundShaderState.GetPixelShader())->SideTableBinding : -1);
 		
 		if (Device.GetRuntimeDebuggingLevel() < EMetalDebugLevelValidation)
 		{
 			RenderPipelineReflection.reset();
 		}
 	}
+}
+
+void FMetalShaderPipeline::InitResourceMask(const FComputePipelineStateInitializer& Init)
+{
 	if (ComputePipelineReflection)
 	{
-		InitResourceMask(EMetalShaderCompute);
+		InitResourceMask(EMetalShaderCompute, ResourceCast(Init.ComputeShader)->SideTableBinding);
 		
 		if (Device.GetRuntimeDebuggingLevel() < EMetalDebugLevelValidation)
 		{
 			ComputePipelineReflection.reset();
 		}
 	}
-	if (StreamPipelineReflection)
-	{
-		InitResourceMask(EMetalShaderStream);
-		
-		if (Device.GetRuntimeDebuggingLevel() < EMetalDebugLevelValidation)
-		{
-			StreamPipelineReflection.reset();
-		}
-	}
 }
 
-void FMetalShaderPipeline::InitResourceMask(EMetalShaderFrequency Frequency)
+void FMetalShaderPipeline::InitResourceMask(EMetalShaderFrequency Frequency, uint8 SideTableBinding)
 {
     if (@available(macOS 13.0, iOS 16.0, *))
     {
@@ -453,12 +448,6 @@ void FMetalShaderPipeline::InitResourceMask(EMetalShaderFrequency Frequency)
                 Bindings = ComputePipelineReflection->bindings();
                 break;
             }
-            case EMetalShaderStream:
-            {
-                check(StreamPipelineReflection);
-                Bindings = StreamPipelineReflection->vertexBindings();
-                break;
-            }
             default:
                 check(false);
                 break;
@@ -480,7 +469,7 @@ void FMetalShaderPipeline::InitResourceMask(EMetalShaderFrequency Frequency)
                 {
                     MTL::BufferBinding* BufferBinding = (MTL::BufferBinding*)Binding;
                     checkf(Binding->index() < ML_MaxBuffers, TEXT("Metal buffer index exceeded!"));
-                    if (NSStringToFString(Binding->name()) != TEXT("BufferSizes") && NSStringToFString(Binding->name()) != TEXT("spvBufferSizeConstants"))
+                    if (NSStringToFString(Binding->name()) != TEXT("BufferSizes") && SideTableBinding != Binding->index())
                     {
                         ResourceMask[Frequency].BufferMask |= (1 << Binding->index());
                         
@@ -538,12 +527,6 @@ void FMetalShaderPipeline::InitResourceMask(EMetalShaderFrequency Frequency)
                 Arguments = ComputePipelineReflection->arguments();
                 break;
             }
-            case EMetalShaderStream:
-            {
-                check(StreamPipelineReflection);
-                Arguments = StreamPipelineReflection->vertexArguments();
-                break;
-            }
             default:
                 check(false);
                 break;
@@ -564,7 +547,7 @@ void FMetalShaderPipeline::InitResourceMask(EMetalShaderFrequency Frequency)
                 case MTL::ArgumentTypeBuffer:
                 {
                     checkf(Argument->index() < ML_MaxBuffers, TEXT("Metal buffer index exceeded!"));
-                    if (NSStringToFString(Argument->name()) != TEXT("BufferSizes") && NSStringToFString(Argument->name()) != TEXT("spvBufferSizeConstants"))
+                    if (NSStringToFString(Argument->name()) != TEXT("BufferSizes") && SideTableBinding != Argument->index())
                     {
                         ResourceMask[Frequency].BufferMask |= (1 << Argument->index());
                         
@@ -1292,7 +1275,7 @@ static FMetalShaderPipelinePtr CreateMTLRenderPipeline(FMetalDevice& Device, boo
 
 	if (Pipeline && Device.GetRuntimeDebuggingLevel() >= EMetalDebugLevelFastValidation)
 	{
-		Pipeline->InitResourceMask();
+		Pipeline->InitResourceMask(Init);
 	}
 	
     return !bSync ? nullptr : FMetalShaderPipelinePtr(Pipeline);
diff --git a/Engine/Source/Runtime/Apple/MetalRHI/Private/MetalPipeline.h b/Engine/Source/Runtime/Apple/MetalRHI/Private/MetalPipeline.h
--- a/Engine/Source/Runtime/Apple/MetalRHI/Private/MetalPipeline.h
+++ b/Engine/Source/Runtime/Apple/MetalRHI/Private/MetalPipeline.h
@@ -35,9 +35,8 @@ enum EMetalShaderFrequency
     EMetalShaderVertex = 0,
     EMetalShaderFragment = 1,
     EMetalShaderCompute = 2,
-    EMetalShaderStream = 3,
     EMetalShaderRenderNum = 2,
-    EMetalShaderStagesNum = 4
+    EMetalShaderStagesNum = 3
 };
 
 enum EMetalLimits
@@ -112,18 +111,17 @@ public:
     ~FMetalShaderPipeline();
     
     void Init();
-    void InitResourceMask();
-    void InitResourceMask(EMetalShaderFrequency Frequency);
+    void InitResourceMask(const FGraphicsPipelineStateInitializer& Init);
+	void InitResourceMask(const FComputePipelineStateInitializer& Init);
+    void InitResourceMask(EMetalShaderFrequency Frequency, uint8 SideTableBinding);
     
 	FMetalDevice& Device;
 	MTLRenderPipelineStatePtr RenderPipelineState;
     MTLComputePipelineStatePtr ComputePipelineState;
-    MTLRenderPipelineStatePtr StreamPipelineState;
 	TArray<uint32> BufferDataSizes[EMetalShaderStagesNum];
 	TMap<uint8, uint8> TextureTypes[EMetalShaderStagesNum];
 	FMetalShaderResourceMask ResourceMask[EMetalShaderStagesNum];
     MTLRenderPipelineReflectionPtr RenderPipelineReflection;
-    MTLRenderPipelineReflectionPtr StreamPipelineReflection;
     MTLComputePipelineReflectionPtr ComputePipelineReflection;
 #if METAL_DEBUG_OPTIONS
 	NS::String* VertexSource = nullptr;
diff --git a/Engine/Source/Runtime/Apple/MetalRHI/Private/MetalStateCache.cpp b/Engine/Source/Runtime/Apple/MetalRHI/Private/MetalStateCache.cpp
--- a/Engine/Source/Runtime/Apple/MetalRHI/Private/MetalStateCache.cpp
+++ b/Engine/Source/Runtime/Apple/MetalRHI/Private/MetalStateCache.cpp
@@ -2352,7 +2352,7 @@ void FMetalStateCache::EnsureTextureAndType(EMetalShaderStages Stage, uint32 Ind
 }
 
 /** Validates the pipeline/binding state */
-bool FMetalStateCache::ValidateFunctionBindings(FMetalShaderPipeline* Pipeline, EMetalShaderFrequency Frequency)
+bool FMetalStateCache::ValidateFunctionBindings(FMetalShaderPipeline* Pipeline, EMetalShaderFrequency Frequency, uint8 SideTableIndex)
 {
     bool bOK = true;
     
@@ -2476,7 +2476,7 @@ bool FMetalStateCache::ValidateFunctionBindings(FMetalShaderPipeline* Pipeline,
                     case MTL::ArgumentTypeBuffer:
                     {
                         checkf(Arg->index() < ML_MaxBuffers, TEXT("Metal buffer index exceeded!"));
-                        if (NSStringToFString(Arg->name()) != TEXT("BufferSizes") && NSStringToFString(Arg->name()) != TEXT("spvBufferSizeConstants"))
+                        if (NSStringToFString(Arg->name()) != TEXT("BufferSizes") && Arg->index() != SideTableIndex)
                         {
                             if (ShaderBuffers[Frequency].Buffers[Arg->index()].Buffer == nullptr && ShaderBuffers[Frequency].Buffers[Arg->index()].Bytes == nullptr)
                             {
@@ -2530,16 +2530,17 @@ void FMetalStateCache::Validate()
 {
 #if METAL_DEBUG_OPTIONS
     FMetalShaderPipeline* Pipeline = GetPipelineState();
-    bool bOK = ValidateFunctionBindings(Pipeline, EMetalShaderVertex);
+	
+    bool bOK = ValidateFunctionBindings(Pipeline, EMetalShaderVertex, GraphicsPSO->VertexShader->SideTableBinding);
     if (!bOK)
     {
         UE_LOG(LogMetal, Error, TEXT("Metal Validation failures for vertex shader:\n%s"), Pipeline->VertexSource ? *NSStringToFString(Pipeline->VertexSource) : TEXT("nullptr"));
     }
     
-    bOK = ValidateFunctionBindings(GetPipelineState(), EMetalShaderFragment);
+    bOK = ValidateFunctionBindings(GetPipelineState(), EMetalShaderFragment, GraphicsPSO->PixelShader ? GraphicsPSO->PixelShader->SideTableBinding : -1);
     if (!bOK)
     {
-        UE_LOG(LogMetal, Error, TEXT("Metal Validation failures for fragment shader:\n%s"), Pipeline->FragmentSource ? *NSStringToFString(Pipeline->VertexSource) : TEXT("nullptr"));
+        UE_LOG(LogMetal, Error, TEXT("Metal Validation failures for fragment shader:\n%s"), Pipeline->FragmentSource ? *NSStringToFString(Pipeline->FragmentSource) : TEXT("nullptr"));
     }
 #endif
 }
diff --git a/Engine/Source/Runtime/Apple/MetalRHI/Private/MetalStateCache.h b/Engine/Source/Runtime/Apple/MetalRHI/Private/MetalStateCache.h
--- a/Engine/Source/Runtime/Apple/MetalRHI/Private/MetalStateCache.h
+++ b/Engine/Source/Runtime/Apple/MetalRHI/Private/MetalStateCache.h
@@ -217,7 +217,7 @@ private:
 	void SetScissorRect(uint32 Index, bool const bEnable, MTL::ScissorRect const& Rect);
     
     void Validate();
-    bool ValidateFunctionBindings(FMetalShaderPipeline* Pipeline, EMetalShaderFrequency Frequency);
+    bool ValidateFunctionBindings(FMetalShaderPipeline* Pipeline, EMetalShaderFrequency Frequency, uint8 SideTableIndex);
 
 private:
 
