MetalViewport SubmitImmersive fix
-[IOGPUMetalCommandBuffer setProtectionOptions:]:585: failed assertion `setProtectionOptions: with uncommitted encoder'

diff --git a/Engine/Source/Runtime/Apple/MetalRHI/Private/MetalViewport.cpp b/Engine/Source/Runtime/Apple/MetalRHI/Private/MetalViewport.cpp
--- a/Engine/Source/Runtime/Apple/MetalRHI/Private/MetalViewport.cpp
+++ b/Engine/Source/Runtime/Apple/MetalRHI/Private/MetalViewport.cpp
@@ -779,6 +779,11 @@ void FMetalViewport::PresentImmersive(const MetalRHIVisionOS::PresentImmersivePa
 					}
 				}
 				
+				// We need to make sure that any outstanding encoders have been completed before
+				// We add our completion handler and encode_present.
+				Context.EndCommandBuffer();
+				Context.StartCommandBuffer();
+
 				// We need to attach the completion handler and the present signal to the final
 				// command buffer
 				FMetalCommandBuffer* FinalCommandBuffer = Context.GetCurrentCommandBuffer();;
