Metal RHI

Fixed bug with index buffer offset when creating BLAS

diff --git a/Engine/Source/Runtime/Apple/MetalRHI/Private/MetalRayTracing.cpp b/Engine/Source/Runtime/Apple/MetalRHI/Private/MetalRayTracing.cpp
--- a/Engine/Source/Runtime/Apple/MetalRHI/Private/MetalRayTracing.cpp
+++ b/Engine/Source/Runtime/Apple/MetalRHI/Private/MetalRayTracing.cpp
@@ -346,7 +346,10 @@ static void FillPrimitiveAccelerationStructureDesc(FMetalDevice& Device, MTL::Pr
 			
 			GeometryDescriptor->setIndexType(IndexBuffer->GetIndexType());
 			GeometryDescriptor->setIndexBuffer(IndexBufferRes ? IndexBufferRes->GetMTLBuffer() : nullptr);
-			GeometryDescriptor->setIndexBufferOffset(IndexBufferRes ? IndexBufferRes->GetOffset() + Initializer.IndexBufferOffset : 0);
+			
+			const uint32 IndexStride = IndexBuffer ? IndexBuffer->GetStride() : 0;
+			uint32_t IndexBufferOffset = Initializer.IndexBufferOffset + IndexStride * Segment.FirstPrimitive * FMetalRayTracingGeometry::IndicesPerPrimitive;
+			GeometryDescriptor->setIndexBufferOffset(IndexBufferRes ? IndexBufferRes->GetOffset() + IndexBufferOffset : 0);
 		}
 
 		const FMetalBufferPtr VertexBufferRes = VertexBuffer->GetCurrentBufferOrNull();
