Metal - Improve fatal error in ConvertSurfaceDataToFColor to display the unsupported format

#tests Local run on macOS

diff --git a/Engine/Source/Runtime/Apple/MetalRHI/Private/MetalRenderTarget.cpp b/Engine/Source/Runtime/Apple/MetalRHI/Private/MetalRenderTarget.cpp
--- a/Engine/Source/Runtime/Apple/MetalRHI/Private/MetalRenderTarget.cpp
+++ b/Engine/Source/Runtime/Apple/MetalRHI/Private/MetalRenderTarget.cpp
@@ -120,7 +120,7 @@ static void ConvertSurfaceDataToFColor(EPixelFormat Format, uint32 Width, uint32
 	else
 	{
 		// not supported yet
-		NOT_SUPPORTED("RHIReadSurfaceData Format");
+		UE_LOG(LogMetal, Fatal, TEXT("ConvertSurfaceDataToFColor doesn't support format %s"), GetPixelFormatString(Format));
 	}
 }
 
