Metal RHI

- Fixed potential bug where bindless handle could be null when retrieving viewport texture
- Removed IR and cleaned up bindless parameter loop

diff --git a/Engine/Source/Runtime/Apple/MetalRHI/Private/MetalCommands.cpp b/Engine/Source/Runtime/Apple/MetalRHI/Private/MetalCommands.cpp
--- a/Engine/Source/Runtime/Apple/MetalRHI/Private/MetalCommands.cpp
+++ b/Engine/Source/Runtime/Apple/MetalRHI/Private/MetalCommands.cpp
@@ -577,16 +577,26 @@ static void SetShaderParameters(
 #if PLATFORM_SUPPORTS_BINDLESS_RENDERING
     for (const FRHIShaderParameterResource& Parameter : InBindlessParameters)
     {
-        const FRHIDescriptorHandle Handle = UE::RHICore::GetBindlessParameterHandle(Parameter);
-        if (Handle.IsValid())
+        FRHIDescriptorHandle Handle = UE::RHICore::GetBindlessParameterHandle(Parameter);
+        if (!Handle.IsValid())
         {
-            checkf(Handle.IsValid(), TEXT("Metal resource did not provide a valid descriptor handle. Please validate that all Metal types can provide this or that the resource is still valid."));
-            Binder.SetBindlessHandle(Handle, Parameter.Index);
+			if (Parameter.Type == FRHIShaderParameterResource::EType::Texture)
+			{
+				if (FRHIResource* Resource = Parameter.Resource)
+				{
+					FMetalSurface* Surface = static_cast<FMetalSurface*>(Resource);
+					if(Surface->Viewport)
+					{
+						// We need to update the drawable texture on the surface and create a bindless handle, we can ignore the return value.
+						Surface->GetDrawableTexture();
+						Handle = UE::RHICore::GetBindlessParameterHandle(Parameter);
+					}
+				}
+			}
         }
-    }
-	
-	for (const FRHIShaderParameterResource& Parameter : InBindlessParameters)
-	{
+				
+		Binder.SetBindlessHandle(Handle, Parameter.Index);
+			
 		if (Parameter.Type == FRHIShaderParameterResource::EType::UnorderedAccessView)
 		{
 			if (ShaderFrequency == SF_Pixel || ShaderFrequency == SF_Compute)
@@ -602,7 +612,11 @@ static void SetShaderParameters(
 		{
 			Binder.StateCache.IRMakeSRVResident(GetShaderStage(ShaderFrequency), static_cast<FMetalShaderResourceView*>(Parameter.Resource));
 		}
-	}
+		else if(Parameter.Type == FRHIShaderParameterResource::EType::Texture)
+		{
+			Binder.StateCache.IRMakeTextureResident(GetShaderStage(ShaderFrequency), static_cast<FMetalSurface*>(Parameter.Resource)->Texture.get());
+		}
+    }
 #endif 
 	
     for (const FRHIShaderParameterResource& Parameter : InResourceParameters)
@@ -648,14 +662,6 @@ static void SetShaderParameters(
             break;
         }
     }
-    
-    
-#if PLATFORM_SUPPORTS_BINDLESS_RENDERING
-    if(IsMetalBindlessEnabled())
-    {
-        Binder.StateCache.IRForwardBindlessParameters(Binder.Stage, InBindlessParameters);
-    }
-#endif
 }
 
 void FMetalRHICommandContext::RHISetShaderParameters(FRHIGraphicsShader* Shader, TConstArrayView<uint8> InParametersData, TConstArrayView<FRHIShaderParameter> InParameters, TConstArrayView<FRHIShaderParameterResource> InResourceParameters, TConstArrayView<FRHIShaderParameterResource> InBindlessParameters)
diff --git a/Engine/Source/Runtime/Apple/MetalRHI/Private/MetalStateCache.cpp b/Engine/Source/Runtime/Apple/MetalRHI/Private/MetalStateCache.cpp
--- a/Engine/Source/Runtime/Apple/MetalRHI/Private/MetalStateCache.cpp
+++ b/Engine/Source/Runtime/Apple/MetalRHI/Private/MetalStateCache.cpp
@@ -1642,55 +1642,6 @@ void FMetalStateCache::IRMakeTextureResident(EMetalShaderStages const Frequency,
 	}
 }
 
-void FMetalStateCache::IRForwardBindlessParameters(EMetalShaderStages const Frequency, TConstArrayView<FRHIShaderParameterResource> InBindlessParameters)
-{
-    FMetalBindlessDescriptorManager* BindlessDescriptorManager = Device.GetBindlessDescriptorManager();
-
-    // Collect resources we need to map for this frame.
-    for (FRHIShaderParameterResource const& Parameter : InBindlessParameters)
-    {
-		if(Parameter.Type == FRHIShaderParameterResource::EType::UniformBuffer)
-		{
-			continue;
-		}
-		
-        const FRHIDescriptorHandle Handle = UE::RHICore::GetBindlessParameterHandle(Parameter);
-        if (Handle.IsValid())
-        {
-            if (Parameter.Type != FRHIShaderParameterResource::EType::Sampler)
-            {
-                switch (Parameter.Type)
-                {
-                    case FRHIShaderParameterResource::EType::Texture:
-                    {
-                        FRHITexture* Texture = static_cast<FRHITexture*>(Parameter.Resource);
-                        IRMakeTextureResident(Frequency, GetMetalSurfaceFromRHITexture(Texture)->Texture.get());
-                        break;
-                    }
-
-                    case FRHIShaderParameterResource::EType::ResourceView:
-                    {
-                        FMetalShaderResourceView* SRV = static_cast<FMetalShaderResourceView*>(Parameter.Resource);
-                        IRMakeSRVResident(Frequency, SRV);
-                        break;
-                    }
-
-                    case FRHIShaderParameterResource::EType::UnorderedAccessView:
-                    {
-                        FMetalUnorderedAccessView* UAV = static_cast<FMetalUnorderedAccessView*>(Parameter.Resource);
-                        IRMakeUAVResident(Frequency, UAV);
-                        break;
-                    }
-
-                    default:
-                        checkNoEntry();
-                        break;
-                }
-            }
-        }
-    }
-}
-
 void FMetalStateCache::IRBindPackedUniforms(EMetalShaderStages const Frequency, int32 Index, uint8 const* Bytes, const uint32 Size, FMetalBufferPtr& Buffer)
 {
 	uint64 PackedUniformsVA;
diff --git a/Engine/Source/Runtime/Apple/MetalRHI/Private/MetalStateCache.h b/Engine/Source/Runtime/Apple/MetalRHI/Private/MetalStateCache.h
--- a/Engine/Source/Runtime/Apple/MetalRHI/Private/MetalStateCache.h
+++ b/Engine/Source/Runtime/Apple/MetalRHI/Private/MetalStateCache.h
@@ -117,8 +117,6 @@ public:
     void IRMakeUAVResident(EMetalShaderStages const Frequency, FMetalUnorderedAccessView* UAV);
     void IRMakeTextureResident(EMetalShaderStages const Frequency, MTL::Texture* Texture);
     
-    void IRForwardBindlessParameters(EMetalShaderStages const Frequency, TConstArrayView<FRHIShaderParameterResource> InBindlessParameters);
-
     void IRBindUniformBuffer(EMetalShaderStages const Frequency, int32 Index, FMetalUniformBuffer* UB);
     void IRBindPackedUniforms(EMetalShaderStages const Frequency, int32 Index, uint8 const* Bytes, const uint32 Size, FMetalBufferPtr& Buffer);
 
