Fixed sync point not being reset in MetalRHIContext

diff --git a/Engine/Source/Runtime/Apple/MetalRHI/Private/MetalRHIContext.cpp b/Engine/Source/Runtime/Apple/MetalRHI/Private/MetalRHIContext.cpp
--- a/Engine/Source/Runtime/Apple/MetalRHI/Private/MetalRHIContext.cpp
+++ b/Engine/Source/Runtime/Apple/MetalRHI/Private/MetalRHIContext.cpp
@@ -636,6 +636,8 @@ void FMetalRHICommandContext::Finalize(TArray<FMetalPayload*>& OutPayloads)
 		BatchedSyncPoints.ToSignal.Reset();
 	}
 	
+	ContextSyncPoint = nullptr;
+	
 	OutPayloads.Append(MoveTemp(Payloads));
 }
 
