Metal RHI

- Removed check when getting bindless handle

#DisabledValidations Preflight

diff --git a/Engine/Source/Runtime/Apple/MetalRHI/Public/MetalResources.h b/Engine/Source/Runtime/Apple/MetalRHI/Public/MetalResources.h
--- a/Engine/Source/Runtime/Apple/MetalRHI/Public/MetalResources.h
+++ b/Engine/Source/Runtime/Apple/MetalRHI/Public/MetalResources.h
@@ -293,9 +293,6 @@ public:
 #if PLATFORM_SUPPORTS_BINDLESS_RENDERING
     virtual FRHIDescriptorHandle GetDefaultBindlessHandle() const override final
     {
-#if PLATFORM_SUPPORTS_BINDLESS_RENDERING
-		check(!IsMetalBindlessEnabled() || BindlessHandle.IsValid());
-#endif
         return BindlessHandle;
     }
 #endif // PLATFORM_SUPPORTS_BINDLESS_RENDERING
