Metal RHI

- Disable breaking encoders as the default behaviour when profiling

diff --git a/Engine/Source/Runtime/Apple/MetalRHI/Private/MetalRHIContext.cpp b/Engine/Source/Runtime/Apple/MetalRHI/Private/MetalRHIContext.cpp
--- a/Engine/Source/Runtime/Apple/MetalRHI/Private/MetalRHIContext.cpp
+++ b/Engine/Source/Runtime/Apple/MetalRHI/Private/MetalRHIContext.cpp
@@ -19,7 +19,7 @@ static FAutoConsoleVariableRef CVarMetalSampleBlitEncoderTimings(
 	GMetalSampleBlitEncoderTimings,
 	TEXT("Whether to include Blit Encoder Timings in GPU Stats, can cause poor perf due to splitting encoders: off"));
 
-static int32 GMetalSampleComputeEncoderTimings = 1;
+static int32 GMetalSampleComputeEncoderTimings = 0;
 static FAutoConsoleVariableRef CVarMetalSampleComputeEncoderTimings(
 	TEXT("rhi.Metal.SampleComputeEncoderTimings"),
 	GMetalSampleComputeEncoderTimings,
@@ -153,7 +153,7 @@ void FMetalRHICommandContext::BeginComputeEncoder()
 			CurrentEncoderFence = CurrentEncoder.EndEncoding();
 			StateCache.ResetBindings();
 		}
-		bool bUseStageCounterSamples = Device.SupportsFeature(EMetalFeaturesStageCounterSampling) && GMetalSampleComputeEncoderTimings;
+		bool bUseStageCounterSamples = Device.SupportsFeature(EMetalFeaturesStageCounterSampling);
 		
 		MTL::DispatchType DispatchType = GMetalConcurrentDispatch ? MTL::DispatchTypeConcurrent : MTL::DispatchTypeSerial;
 		CurrentEncoder.BeginComputeCommandEncoding(DispatchType,
@@ -205,7 +205,7 @@ void FMetalRHICommandContext::BeginBlitEncoder()
 			StateCache.ResetBindings();
 		}
 
-		bool bUseStageCounterSamples = Device.SupportsFeature(EMetalFeaturesStageCounterSampling) && GMetalSampleBlitEncoderTimings;
+		bool bUseStageCounterSamples = Device.SupportsFeature(EMetalFeaturesStageCounterSampling);
 		CurrentEncoder.BeginBlitCommandEncoding(bUseStageCounterSamples ? Device.GetCounterSampler() : nullptr);
 	}
 	
