Metal RHI

- Fix bugs in old compaction code and enable again by default

diff --git a/Engine/Source/Runtime/Apple/MetalRHI/Private/MetalRayTracing.cpp b/Engine/Source/Runtime/Apple/MetalRHI/Private/MetalRayTracing.cpp
--- a/Engine/Source/Runtime/Apple/MetalRHI/Private/MetalRayTracing.cpp
+++ b/Engine/Source/Runtime/Apple/MetalRHI/Private/MetalRayTracing.cpp
@@ -16,7 +16,7 @@
 #include "RayTracingValidationShaders.h"
 #include "Async/ParallelFor.h"
 
-static int32 GMetalRayTracingAllowCompaction = 0;
+static int32 GMetalRayTracingAllowCompaction = 1;
 static FAutoConsoleVariableRef CVarMetalRayTracingAllowCompaction(
 	TEXT("r.Metal.RayTracing.AllowCompaction"),
 	GMetalRayTracingAllowCompaction,
@@ -110,33 +110,25 @@ private:
 	TArray<FMetalRayTracingGeometry*> PendingRequests;
 
 	/** Enqueued compaction requests (submitted compaction size requests; waiting on readback and actual compaction). */
-	TQueue<FMetalRayTracingGeometry*> ActiveRequests;
-
-	/** Num of active requests enqueued. */
-	uint32 NumActiveRequests;
+	TArray<FMetalRayTracingGeometry*> ActiveRequests;
 
 	/** Buffer used for compacted size readback. */
 	FMetalBufferPtr CompactedStructureSizeBuffer;
 
 	/** Size entry allocated in the CompactedStructureSizeBuffer (in element count). */
 	uint32 SizeBufferMaxCapacity;
-
-	/** CompactedStructureSizeBuffer write index. Wrap once equal to SizeBufferMaxCapacity. */
-	uint32 WriteIndex;
+	
+	FCriticalSection CS;
+	FMetalSyncPoint* CompactionQuerySyncPoint = nullptr;
 };
 
 FMetalRayTracingCompactionRequestHandler::FMetalRayTracingCompactionRequestHandler(FMetalDevice& Device)
 	: SizeBufferMaxCapacity(GMetalRayTracingMaxBatchedCompaction)
-	, WriteIndex(0u)
 {
-	PendingRequests.Reserve(GMetalRayTracingMaxBatchedCompaction);
-
 	MTL::Buffer* BufferPtr = Device.GetDevice()->newBuffer(GMetalRayTracingMaxBatchedCompaction * sizeof(uint32), MTL::ResourceStorageModeShared);
 	CompactedStructureSizeBuffer = FMetalBufferPtr(new FMetalBuffer(BufferPtr, FMetalBuffer::FreePolicy::Owner));
 
 	check(CompactedStructureSizeBuffer);
-
-	NumActiveRequests = 0;
 }
 
 FMetalRayTracingCompactionRequestHandler::~FMetalRayTracingCompactionRequestHandler()
@@ -144,10 +136,18 @@ FMetalRayTracingCompactionRequestHandler::~FMetalRayTracingCompactionRequestHand
 	check(PendingRequests.IsEmpty());
 	
 	FMetalDynamicRHI::Get().DeferredDelete(CompactedStructureSizeBuffer);
+	
+	if(CompactionQuerySyncPoint)
+	{
+		CompactionQuerySyncPoint->Release();
+		CompactionQuerySyncPoint = nullptr;
+	}
 }
 
 void FMetalRayTracingCompactionRequestHandler::RequestCompact(FMetalRayTracingGeometry* InRTGeometry)
 {
+	FScopeLock Lock(&CS);
+	
 	check(InRTGeometry->GetAccelerationStructure());
 	ERayTracingAccelerationStructureFlags GeometryBuildFlags = GetRayTracingAccelerationStructureBuildFlags(InRTGeometry->Initializer);
 	check(EnumHasAllFlags(GeometryBuildFlags, ERayTracingAccelerationStructureFlags::AllowCompaction) &&
@@ -159,7 +159,28 @@ void FMetalRayTracingCompactionRequestHandler::RequestCompact(FMetalRayTracingGe
 
 bool FMetalRayTracingCompactionRequestHandler::ReleaseRequest(FMetalRayTracingGeometry* InRTGeometry)
 {
-	PendingRequests.Remove(InRTGeometry);
+	FScopeLock Lock(&CS);
+
+	// Remove from pending list, not found then try active requests
+	if (PendingRequests.Remove(InRTGeometry) <= 0)
+	{
+		// If currently enqueued, then clear pointer to not handle the compaction request anymore			
+		for (int32 BLASIndex = 0; BLASIndex < ActiveRequests.Num(); ++BLASIndex)
+		{
+			if (ActiveRequests[BLASIndex] == InRTGeometry)
+			{
+				ActiveRequests[BLASIndex] = nullptr;
+				return true;
+			}
+		}
+
+		return false;
+	}
+	else
+	{
+		return true;
+	}
+	
 	return true;
 }
 
@@ -213,80 +234,84 @@ void FMetalRHICommandContext::CopyAndCompactAccelerationStructure(MTLAcceleratio
 
 void FMetalRayTracingCompactionRequestHandler::Update(FMetalRHICommandContext& Context)
 {
+	FScopeLock Lock(&CS);
+	
 	// Early exit to avoid unecessary encoding breaks.
 	if (PendingRequests.IsEmpty() && ActiveRequests.IsEmpty())
 	{
 		return;
 	}
-
-	// Process pending requests.
-	for (FMetalRayTracingGeometry* Geometry : PendingRequests)
+	
+	// Process active requests.
+	if(CompactionQuerySyncPoint && CompactionQuerySyncPoint->IsComplete())
 	{
-		Geometry->CompactionSizeIndex = WriteIndex;
-		WriteIndex = (++WriteIndex % SizeBufferMaxCapacity);
-
+		// Try to readback active requests.
 		uint32* CompactedSizes = (uint32*)CompactedStructureSizeBuffer->Contents();
-		CompactedSizes[Geometry->CompactionSizeIndex] = 0;
 		
-		FMetalAccelerationStructure* AccelerationStructure = Geometry->GetAccelerationStructure();
-		Context.WriteCompactedAccelerationStructureSize(AccelerationStructure->GetAccelerationStructure(), CompactedStructureSizeBuffer, Geometry->CompactionSizeIndex * sizeof(uint32));
-
-		ActiveRequests.Enqueue(Geometry);
-		NumActiveRequests++;
-
-		// enqueued enough requests for this update round
-		if (NumActiveRequests >= GMetalRayTracingMaxBatchedCompaction)
-		{
-			break;
+		for (FMetalRayTracingGeometry* ActiveRequestsTail : ActiveRequests)
+		{	
+			if(!ActiveRequestsTail)
+			{
+				continue;
+			}
+			
+			uint32 CompactedSize = CompactedSizes[ActiveRequestsTail->CompactionSizeIndex];
+			check(CompactedSize != 0);
+			
+			FMetalAccelerationStructure* SrcBuffer = ActiveRequestsTail->GetAccelerationStructure();
+			MTLAccelerationStructurePtr SrcBLAS = SrcBuffer->GetAccelerationStructure();
+			
+			// Allocate new acceleration structure with compacted size
+			FMetalAccelerationStructure* DestAccelerationStructure = Context.GetDevice().GetResourceHeap().CreateAccelerationStructure(CompactedSize, false);
+			
+			if(GetEmitDrawEvents())
+			{
+				FString DebugNameString = ActiveRequestsTail->Initializer.DebugName.ToString();
+				DestAccelerationStructure->GetAccelerationStructure()->setLabel(FStringToNSString(DebugNameString));
+			}
+			check(DestAccelerationStructure);
+			
+			MTLAccelerationStructurePtr DestBLAS = DestAccelerationStructure->GetAccelerationStructure();
+			Context.CopyAndCompactAccelerationStructure(SrcBLAS, DestBLAS);
+			
+			// Swap acceleration structure buffers 
+			ActiveRequestsTail->SetAccelerationStructure(DestAccelerationStructure);	
 		}
+		
+		ActiveRequests.Empty(ActiveRequests.Num());
+		
+		CompactionQuerySyncPoint->Release();
+		CompactionQuerySyncPoint = nullptr;
 	}
 	
-	PendingRequests.Empty();
-
-	// Try to readback active requests.
-	uint32* CompactedSizes = (uint32*)CompactedStructureSizeBuffer->Contents();
-
-	// Process active requests.
-	while (!ActiveRequests.IsEmpty())
+	// Process pending requests.
+	if(ActiveRequests.IsEmpty())
 	{
-		FMetalRayTracingGeometry* ActiveRequestsTail = *ActiveRequests.Peek();
-		check(ActiveRequestsTail);
-
-		if (!ActiveRequestsTail->bHasPendingCompactionRequests)
-		{
-			ActiveRequests.Pop();
-			continue;
-		}
-		
-		// Break whenever we encounter the first incomplete request (i.e. size == 0).
-		uint32 CompactedSize = CompactedSizes[ActiveRequestsTail->CompactionSizeIndex];
-		if (CompactedSize == 0u)
-		{
-			break;
-		}
-
-		FMetalAccelerationStructure* SrcBuffer = ActiveRequestsTail->GetAccelerationStructure();
-		MTLAccelerationStructurePtr SrcBLAS = SrcBuffer->GetAccelerationStructure();
-		
-		FString DebugNameString = ActiveRequestsTail->Initializer.DebugName.ToString();
-
-		// Allocate new acceleration structure with compacted size
-		FMetalAccelerationStructure* DestAccelerationStructure = Context.GetDevice().GetResourceHeap().CreateAccelerationStructure(CompactedSize, false);
+		uint32* CompactedSizes = (uint32*)CompactedStructureSizeBuffer->Contents();
+		uint32_t CompactionIndex = 0;
 		
-		if(GetEmitDrawEvents())
+		while (!PendingRequests.IsEmpty())
 		{
-			DestAccelerationStructure->GetAccelerationStructure()->setLabel(FStringToNSString(DebugNameString));
+			FMetalRayTracingGeometry* Geometry = PendingRequests[0];
+			
+			Geometry->CompactionSizeIndex = CompactionIndex++;
+			CompactedSizes[Geometry->CompactionSizeIndex] = 0;
+			
+			FMetalAccelerationStructure* AccelerationStructure = Geometry->GetAccelerationStructure();
+			Context.WriteCompactedAccelerationStructureSize(AccelerationStructure->GetAccelerationStructure(), CompactedStructureSizeBuffer, Geometry->CompactionSizeIndex * sizeof(uint32));
+			
+			ActiveRequests.Add(Geometry);			
+			PendingRequests.Remove(Geometry);
+			
+			// enqueued enough requests for this update round
+			if (ActiveRequests.Num() >= GMetalRayTracingMaxBatchedCompaction)
+			{
+				break;
+			}
 		}
-		check(DestAccelerationStructure);
-		
-		MTLAccelerationStructurePtr DestBLAS = DestAccelerationStructure->GetAccelerationStructure();
-		Context.CopyAndCompactAccelerationStructure(SrcBLAS, DestBLAS);
 		
-		// Swap acceleration structure buffers 
-		ActiveRequestsTail->SetAccelerationStructure(DestAccelerationStructure);
-		
-		ActiveRequests.Pop();
-		NumActiveRequests--;
+		CompactionQuerySyncPoint = Context.GetContextSyncPoint();
+		CompactionQuerySyncPoint->AddRef();
 	}
 }
 
@@ -959,7 +984,6 @@ void FMetalRayTracingScene::BuildAccelerationStructure(FMetalRHICommandContext&
 
 	FMetalBufferPtr CurScratchBuffer = InScratchBuffer->GetCurrentBuffer();
 	check(CurScratchBuffer);
-	
 
 	CommandContext.BuildAccelerationStructure(CurInstanceBuffer, InstanceBufferOffset, CurScratchBuffer, CurScratchBuffer->GetOffset() + ScratchOffset,
 											  CurHitGroupContributionsBuffer, HitGroupContributionsBufferOffset, NumInstances, AccelerationStructure);
