Metal RHI

Guard for crash when Begin Breadcrumb is called within a parallel pass and end is not.

diff --git a/Engine/Source/Runtime/Apple/MetalRHI/Private/MetalProfiler.h b/Engine/Source/Runtime/Apple/MetalRHI/Private/MetalProfiler.h
--- a/Engine/Source/Runtime/Apple/MetalRHI/Private/MetalProfiler.h
+++ b/Engine/Source/Runtime/Apple/MetalRHI/Private/MetalProfiler.h
@@ -478,21 +478,24 @@ public:
 		uint64_t* Start = Event.TimestampTOP;
 		uint64_t* End = Event.TimestampBOP;
 		
-		// Ignore samples whe
-		if(Event.bWithinRenderPass || Event.Samples.Num() == 0)
+		// TODO: Carl - We have a rare bug where BeginbreadCrumb is called in a parallel pass, but end is not, we cannot cleanly handle this atm
+		if(Start && End)
 		{
-			*Start = 0; 
-			*End = 0;
-		}
-		else
-		{
-			for(FMetalCounterSamplePtr Sample : Event.Samples)
+			if(Event.bWithinRenderPass || Event.Samples.Num() == 0)
+			{
+				*Start = 0; 
+				*End = 0;
+			}
+			else
 			{
-				uint64_t StartTime, EndTime;
-				Sample->ResolveStageCounters(StartTime, EndTime);
-				
-				*Start = *Start > 0 ? FMath::Min(StartTime, *Start) : StartTime;
-				*End = *End > 0 ? FMath::Max(EndTime, *End) : EndTime;
+				for(FMetalCounterSamplePtr Sample : Event.Samples)
+				{
+					uint64_t StartTime, EndTime;
+					Sample->ResolveStageCounters(StartTime, EndTime);
+					
+					*Start = *Start > 0 ? FMath::Min(StartTime, *Start) : StartTime;
+					*End = *End > 0 ? FMath::Max(EndTime, *End) : EndTime;
+				}
 			}
 		}
 		
