Metal RHI

- Move Indirect Args for acceleration structure outside of buffer as we were updating data on the CPU being used on the GPU
- FMetalSubRingBuffer was being created with MTLResourceStorageModeManaged which is incompatible with heaps. If heaps are enabled we default to Managed

diff --git a/Engine/Source/Runtime/Apple/MetalRHI/Private/MetalBindlessDescriptors.cpp b/Engine/Source/Runtime/Apple/MetalRHI/Private/MetalBindlessDescriptors.cpp
--- a/Engine/Source/Runtime/Apple/MetalRHI/Private/MetalBindlessDescriptors.cpp
+++ b/Engine/Source/Runtime/Apple/MetalRHI/Private/MetalBindlessDescriptors.cpp
@@ -85,7 +85,12 @@ static IRDescriptorTableEntry CreateDescriptor(const FMetalResourceViewBase::FTe
 static IRDescriptorTableEntry CreateDescriptor(FMetalAccelerationStructure* AccelerationStructure)
 {
 	IRDescriptorTableEntry DescriptorData{};
-	IRDescriptorTableSetAccelerationStructure(&DescriptorData, AccelerationStructure->GetIndirectArgumentBuffer()->GetGPUAddress());
+	
+	// We may not have an indirect argument buffer so just generate an empty descriptor
+	if(AccelerationStructure->GetIndirectArgumentBuffer())
+	{
+		IRDescriptorTableSetAccelerationStructure(&DescriptorData, AccelerationStructure->GetIndirectArgumentBuffer()->GetGPUAddress());
+	}
 	return DescriptorData;
 }
 #endif
diff --git a/Engine/Source/Runtime/Apple/MetalRHI/Private/MetalBuffer.cpp b/Engine/Source/Runtime/Apple/MetalRHI/Private/MetalBuffer.cpp
--- a/Engine/Source/Runtime/Apple/MetalRHI/Private/MetalBuffer.cpp
+++ b/Engine/Source/Runtime/Apple/MetalRHI/Private/MetalBuffer.cpp
@@ -958,13 +958,23 @@ FMetalSubBufferRing::FMetalSubBufferRing(FMetalDevice& MetalDevice, NS::UInteger
 , Options(InOptions)
 , Storage((MTL::StorageMode)((Options & MTL::ResourceStorageModeMask) >> MTL::ResourceStorageModeShift))
 {
-	Options = (MTL::ResourceOptions)FMetalCommandQueue::GetCompatibleResourceOptions(Options);
+	static bool bSupportsHeaps = Device.SupportsFeature(EMetalFeaturesHeaps);
+	
+	// If we are using heaps use shared storage
+	if(bSupportsHeaps)
+	{
+		Storage = MTL::StorageModeShared;
+		Options = (Options & ~MTL::ResourceStorageModeMask) | (MTL::ResourceStorageModeShared << MTL::ResourceStorageModeShift);
+	}
+	else
+	{
+		Options = (MTL::ResourceOptions)FMetalCommandQueue::GetCompatibleResourceOptions(Options);
+	}
 #if !defined(WITH_IOS_SIMULATOR) || !WITH_IOS_SIMULATOR
 	check(Storage != MTL::StorageModePrivate /* Private memory requires command-buffers and encoders to properly marshal! */);
 #endif
 	FMemory::Memzero(FrameSize);
 }
-
 FMetalSubBufferRing::~FMetalSubBufferRing()
 {
 }
@@ -2084,7 +2094,7 @@ void FMetalResourceHeap::ReleaseTexture(FMetalSurface* Surface, MTLTexturePtr Te
 }
 
 #if METAL_RHI_RAYTRACING
-FMetalAccelerationStructure* FMetalResourceHeap::CreateAccelerationStructure(uint32 Size, bool bIsTLAS)
+FMetalAccelerationStructure* FMetalResourceHeap::CreateAccelerationStructure(uint32 Size)
 {
 	LLM_SCOPE_METAL(ELLMTagMetal::AccelerationStructures);
 	LLM_PLATFORM_SCOPE_METAL(ELLMTagMetal::AccelerationStructures);
@@ -2130,17 +2140,7 @@ FMetalAccelerationStructure* FMetalResourceHeap::CreateAccelerationStructure(uin
 		Block.Options = ResourceOptions;
 	}
 	
-	FMetalBufferPtr IndirectArgumentBuffer;
-	if(bIsTLAS)
-	{
-		IndirectArgumentBuffer = Device.GetResourceHeap().CreateBuffer(sizeof(IRRaytracingAccelerationStructureGPUHeader),
-																	   sizeof(IRRaytracingAccelerationStructureGPUHeader),
-																	   BUF_Dynamic,
-																	   MTL::ResourceCPUCacheModeDefaultCache | MTL::ResourceStorageModeShared,
-																	   true);
-	}
-	
-	FMetalAccelerationStructure* AS = new FMetalAccelerationStructure(IndirectArgumentBuffer, AccelerationStructure, Size);
+	FMetalAccelerationStructure* AS = new FMetalAccelerationStructure(AccelerationStructure, Size);
 	
 	return AS;
 }
@@ -2162,11 +2162,6 @@ void FMetalResourceHeap::ReleaseAccelerationStructure(FMetalAccelerationStructur
 			}
 		}
 		
-		if(AccelerationStructure->GetIndirectArgumentBuffer())
-		{
-			AccelerationStructure->GetIndirectArgumentBuffer()->MarkDeleted();
-		}
-		
 		delete AccelerationStructure;
 		AccelerationStructure = nullptr;
 	}
diff --git a/Engine/Source/Runtime/Apple/MetalRHI/Private/MetalBuffer.h b/Engine/Source/Runtime/Apple/MetalRHI/Private/MetalBuffer.h
--- a/Engine/Source/Runtime/Apple/MetalRHI/Private/MetalBuffer.h
+++ b/Engine/Source/Runtime/Apple/MetalRHI/Private/MetalBuffer.h
@@ -469,7 +469,7 @@ public:
 	void ReleaseTexture(FMetalSurface* Surface, MTLTexturePtr Texture);
 	
 #if METAL_RHI_RAYTRACING
-	FMetalAccelerationStructure* CreateAccelerationStructure(uint32 Size, bool bIsTLAS);
+	FMetalAccelerationStructure* CreateAccelerationStructure(uint32 Size);
 	void ReleaseAccelerationStructure(FMetalAccelerationStructure* AccelerationStructure);
 #endif
 	
diff --git a/Engine/Source/Runtime/Apple/MetalRHI/Private/MetalRHIBuffer.cpp b/Engine/Source/Runtime/Apple/MetalRHI/Private/MetalRHIBuffer.cpp
--- a/Engine/Source/Runtime/Apple/MetalRHI/Private/MetalRHIBuffer.cpp
+++ b/Engine/Source/Runtime/Apple/MetalRHI/Private/MetalRHIBuffer.cpp
@@ -91,7 +91,7 @@ FMetalRHIBuffer::FMetalRHIBuffer(FRHICommandListBase& RHICmdList, FMetalDevice&
 #if METAL_RHI_RAYTRACING
 	if (EnumHasAnyFlags(CreateDesc.Usage, BUF_AccelerationStructure))
 	{
-		AccelerationStructure = Device.GetResourceHeap().CreateAccelerationStructure(Size, true);
+		AccelerationStructure = Device.GetResourceHeap().CreateAccelerationStructure(Size);
 		return;
 	}
 #endif
diff --git a/Engine/Source/Runtime/Apple/MetalRHI/Private/MetalRayTracing.cpp b/Engine/Source/Runtime/Apple/MetalRHI/Private/MetalRayTracing.cpp
--- a/Engine/Source/Runtime/Apple/MetalRHI/Private/MetalRayTracing.cpp
+++ b/Engine/Source/Runtime/Apple/MetalRHI/Private/MetalRayTracing.cpp
@@ -45,9 +45,22 @@ static FAutoConsoleVariableRef CVarMetalRayTracingMaxBatchedCompaction(
 
 FMetalAccelerationStructure::~FMetalAccelerationStructure()
 {
+	if(IndirectArgumentBuffer)
+	{
+		FMetalDynamicRHI::Get().DeferredDelete(IndirectArgumentBuffer);
+	}
 	AccelerationStructure.reset();
 }
 
+void FMetalAccelerationStructure::SetIndirectArgumentBuffer(FMetalBufferPtr IndirectArgs)
+{
+	if(IndirectArgumentBuffer)
+	{
+		FMetalDynamicRHI::Get().DeferredDelete(IndirectArgumentBuffer);
+	}
+	IndirectArgumentBuffer = IndirectArgs;
+}
+
 static ERayTracingAccelerationStructureFlags GetRayTracingAccelerationStructureBuildFlags(const FRayTracingGeometryInitializer& Initializer)
 {
 	ERayTracingAccelerationStructureFlags BuildFlags = ERayTracingAccelerationStructureFlags::None;
@@ -262,7 +275,7 @@ void FMetalRayTracingCompactionRequestHandler::Update(FMetalRHICommandContext& C
 			MTLAccelerationStructurePtr SrcBLAS = SrcBuffer->GetAccelerationStructure();
 			
 			// Allocate new acceleration structure with compacted size
-			FMetalAccelerationStructure* DestAccelerationStructure = Context.GetDevice().GetResourceHeap().CreateAccelerationStructure(CompactedSize, false);
+			FMetalAccelerationStructure* DestAccelerationStructure = Context.GetDevice().GetResourceHeap().CreateAccelerationStructure(CompactedSize);
 			
 			if(GetEmitDrawEvents())
 			{
@@ -560,7 +573,7 @@ FMetalRayTracingGeometry::FMetalRayTracingGeometry(FMetalDevice& InDevice, FRHIC
 
 	FString DebugNameString = Initializer.DebugName.ToString();
 	
-	AccelerationStructure = Device.GetResourceHeap().CreateAccelerationStructure(SizeInfo.ResultSize, false);
+	AccelerationStructure = Device.GetResourceHeap().CreateAccelerationStructure(SizeInfo.ResultSize);
 	if(GetEmitDrawEvents())
 	{
 		AccelerationStructure->GetAccelerationStructure()->setLabel(FStringToNSString(DebugNameString));
@@ -840,16 +853,15 @@ void FMetalRayTracingScene::BindBuffer(FRHIBuffer* InBuffer, uint32 InBufferOffs
 	check(IsInRHIThread() || !IsRunningRHIInSeparateThread());
 	check(SizeInfo.ResultSize + InBufferOffset <= InBuffer->GetSize());
 
-	FMetalRHIBuffer* Buffer = ResourceCast(InBuffer);
+	AccelerationStructureBuffer = ResourceCast(InBuffer);
 	
-	check(Buffer->IsAccelerationStructure());
-	AccelerationStructure = Buffer->AccelerationStructure;
+	check(AccelerationStructureBuffer->IsAccelerationStructure());
 	
 	check(InBufferOffset % GRHIRayTracingAccelerationStructureAlignment == 0);
 	
 	if(GetEmitDrawEvents())
 	{
-		AccelerationStructure->GetAccelerationStructure()->setLabel(FStringToNSString(Initializer.DebugName.ToString()));
+		GetAccelerationStructure()->GetAccelerationStructure()->setLabel(FStringToNSString(Initializer.DebugName.ToString()));
 	}
 }
 
@@ -894,9 +906,12 @@ void FMetalRHICommandContext::BuildAccelerationStructure(FMetalBufferPtr CurInst
 	MTLAccelerationStructurePtr NativeAS = AS->GetAccelerationStructure();
 
 	IRRaytracingAccelerationStructureGPUHeader* Header = (IRRaytracingAccelerationStructureGPUHeader*)IndirectArgs->Contents();
+	memset(Header, 0, sizeof(IRRaytracingAccelerationStructureGPUHeader));
+	
 	Header->accelerationStructureID = NativeAS->gpuResourceID()._impl;
 	Header->addressOfInstanceContributions = (uint64_t)HitGroupContributionsBuffer->GetGPUAddress() + HitGroupContributionsBufferOffset;
 
+	check(Header->addressOfInstanceContributions);
 	CommandEncoder->buildAccelerationStructure(NativeAS.get(), InstanceDescriptor, ScratchBuffer->GetMTLBuffer(), ScratchBufferOffset);
 }
 
@@ -964,7 +979,6 @@ void FMetalRayTracingScene::BuildAccelerationStructure(FMetalRHICommandContext&
 	check(CurInstanceBuffer);
 
 	uint32 InstanceBufferOffset = InstanceOffset + static_cast<uint32>(CurInstanceBuffer->GetOffset());
-	uint32 HitGroupContributionsBufferOffset = HitGroupContributionOffset + static_cast<uint32>(CurHitGroupContributionsBuffer->GetOffset());
 
 	TRefCountPtr<FMetalRHIBuffer> ScratchBuffer;
 	if (InScratchBuffer == nullptr)
@@ -981,12 +995,23 @@ void FMetalRayTracingScene::BuildAccelerationStructure(FMetalRHICommandContext&
 			ScratchOffset = 0;
 		}
 	}
-
+	
+	// Create new Indirect Args buffer
+	FMetalBufferPtr IndirectArgumentBuffer = Device.GetResourceHeap().CreateBuffer(sizeof(IRRaytracingAccelerationStructureGPUHeader),
+																   sizeof(IRRaytracingAccelerationStructureGPUHeader),
+																   BUF_Dynamic,
+																   MTL::ResourceCPUCacheModeDefaultCache | MTL::ResourceStorageModeShared,
+																   true);
+	
+	
+	GetAccelerationStructure()->SetIndirectArgumentBuffer(IndirectArgumentBuffer);
+	AccelerationStructureBuffer->UpdateLinkedViews(&CommandContext);
+	
 	FMetalBufferPtr CurScratchBuffer = InScratchBuffer->GetCurrentBuffer();
 	check(CurScratchBuffer);
 
 	CommandContext.BuildAccelerationStructure(CurInstanceBuffer, InstanceBufferOffset, CurScratchBuffer, CurScratchBuffer->GetOffset() + ScratchOffset,
-											  CurHitGroupContributionsBuffer, HitGroupContributionsBufferOffset, NumInstances, AccelerationStructure);
+											  CurHitGroupContributionsBuffer, HitGroupContributionOffset, NumInstances, GetAccelerationStructure());
 }
 
 void FMetalRHICommandContext::RHIBuildAccelerationStructures(TConstArrayView<FRayTracingSceneBuildParams> Params)
diff --git a/Engine/Source/Runtime/Apple/MetalRHI/Private/MetalRayTracing.h b/Engine/Source/Runtime/Apple/MetalRHI/Private/MetalRayTracing.h
--- a/Engine/Source/Runtime/Apple/MetalRHI/Private/MetalRayTracing.h
+++ b/Engine/Source/Runtime/Apple/MetalRHI/Private/MetalRayTracing.h
@@ -93,6 +93,12 @@ public:
 	// Any referenced geometry is kept alive while the scene is alive.
 	TArray<TRefCountPtr<FRHIRayTracingGeometry>> ReferencedGeometries;
 	
+	FMetalAccelerationStructure* GetAccelerationStructure()
+	{
+		check(AccelerationStructureBuffer && AccelerationStructureBuffer->AccelerationStructure);
+		return AccelerationStructureBuffer->AccelerationStructure;
+	}
+	
 private:
 	friend class FMetalRHICommandContext;
 
@@ -103,6 +109,6 @@ private:
 	const FRayTracingSceneInitializer Initializer;
 
 	/** Acceleration Structure for the whole scene. */
-	FMetalAccelerationStructure* AccelerationStructure; 
+	FMetalRHIBuffer* AccelerationStructureBuffer; 
 };
 #endif // METAL_RHI_RAYTRACING
diff --git a/Engine/Source/Runtime/Apple/MetalRHI/Public/MetalResources.h b/Engine/Source/Runtime/Apple/MetalRHI/Public/MetalResources.h
--- a/Engine/Source/Runtime/Apple/MetalRHI/Public/MetalResources.h
+++ b/Engine/Source/Runtime/Apple/MetalRHI/Public/MetalResources.h
@@ -120,11 +120,10 @@ typedef TSharedPtr<FMetalBuffer, ESPMode::ThreadSafe> FMetalBufferPtr;
 class FMetalAccelerationStructure
 {
 public:
-	FMetalAccelerationStructure(FMetalBufferPtr InIndirectArgumentBuffer, MTLAccelerationStructurePtr InAccelerationStructure, uint32_t InSize) :
-		IndirectArgumentBuffer(InIndirectArgumentBuffer),
+	FMetalAccelerationStructure(MTLAccelerationStructurePtr InAccelerationStructure, uint32_t InSize) :
 		AccelerationStructure(InAccelerationStructure),
 		Size(InSize)
-	{	
+	{
 	}
 	
 	~FMetalAccelerationStructure();
@@ -140,6 +139,8 @@ public:
 		AccelerationStructure->setLabel(FStringToNSString(Label));
 	}
 	
+	void SetIndirectArgumentBuffer(FMetalBufferPtr IndirectArgs);
+	
 	FMetalBufferPtr GetIndirectArgumentBuffer()
 	{
 		return IndirectArgumentBuffer;
@@ -151,10 +152,10 @@ public:
 	}
 	
 private:
-	FMetalBufferPtr IndirectArgumentBuffer;
 	MTLAccelerationStructurePtr AccelerationStructure; 
-	
 	uint32_t Size;
+	
+	FMetalBufferPtr IndirectArgumentBuffer;
 };
 #endif
 
