Metal - Fix Metal 3.1 on iOS when running from shader source
Use bSharedMaterialNativeLibraries=False & r.Shaders.Symbols=1 & r.Shaders.ExtraData=1 & -metalshaderdebug to ensure this specific case

#tests Local test & Preflight

diff --git a/Engine/Source/Runtime/Apple/MetalRHI/Private/Shaders/Types/Templates/MetalBaseShader.h b/Engine/Source/Runtime/Apple/MetalRHI/Private/Shaders/Types/Templates/MetalBaseShader.h
--- a/Engine/Source/Runtime/Apple/MetalRHI/Private/Shaders/Types/Templates/MetalBaseShader.h
+++ b/Engine/Source/Runtime/Apple/MetalRHI/Private/Shaders/Types/Templates/MetalBaseShader.h
@@ -353,6 +353,9 @@ void TMetalBaseShader<BaseResourceType, ShaderType>::Init(TArrayView<const uint8
 				case 0:
                     MetalVersion = MTL::LanguageVersion2_4;
                     break;
+				case 9:
+					MetalVersion = MTL::LanguageVersion3_1;
+					break;
 				default:
 					UE_LOG(LogRHI, Fatal, TEXT("Failed to create shader with unknown version %d: %s"), Header.Version, *NSStringToFString(NewShaderString));
 					MetalVersion = MTL::LanguageVersion2_4;
