Metal RHI

- Temporarily re-enable command encoder splitting

diff --git a/Engine/Source/Runtime/Apple/MetalRHI/Private/MetalRHIContext.cpp b/Engine/Source/Runtime/Apple/MetalRHI/Private/MetalRHIContext.cpp
--- a/Engine/Source/Runtime/Apple/MetalRHI/Private/MetalRHIContext.cpp
+++ b/Engine/Source/Runtime/Apple/MetalRHI/Private/MetalRHIContext.cpp
@@ -19,7 +19,7 @@ static FAutoConsoleVariableRef CVarMetalSampleBlitEncoderTimings(
 	GMetalSampleBlitEncoderTimings,
 	TEXT("Whether to include Blit Encoder Timings in GPU Stats, can cause poor perf due to splitting encoders: off"));
 
-static int32 GMetalSampleComputeEncoderTimings = 0;
+static int32 GMetalSampleComputeEncoderTimings = 1;
 static FAutoConsoleVariableRef CVarMetalSampleComputeEncoderTimings(
 	TEXT("rhi.Metal.SampleComputeEncoderTimings"),
 	GMetalSampleComputeEncoderTimings,
