Metal RHI

- Added memory barrier after write when updating descriptors
- Disabled splitting compute encoders

diff --git a/Engine/Source/Runtime/Apple/MetalRHI/Private/MetalBindlessDescriptors.cpp b/Engine/Source/Runtime/Apple/MetalRHI/Private/MetalBindlessDescriptors.cpp
--- a/Engine/Source/Runtime/Apple/MetalRHI/Private/MetalBindlessDescriptors.cpp
+++ b/Engine/Source/Runtime/Apple/MetalRHI/Private/MetalBindlessDescriptors.cpp
@@ -339,11 +339,15 @@ void FMetalBindlessDescriptorManager::FlushPendingDescriptorUpdates(FMetalRHICom
 
 	Context.GetStateCache().SetOverriddenDescriptorHeap(Heap);
 
-	FMetalBufferPtr Buffer;
-	Context.GetStateCache().IRBindPackedUniforms(EMetalShaderStages::Compute, 0, (const uint8*)UBDataPtr, 16, Buffer);
+	FMetalBufferPtr NullBuffer;
+	Context.GetStateCache().IRBindPackedUniforms(EMetalShaderStages::Compute, 0, (const uint8*)UBDataPtr, 16, NullBuffer);
 	Context.GetStateCache().GetShaderParameters(EMetalShaderStages::Compute).MarkAllDirty();
 
 	Context.RHIDispatchComputeShader(1, 1, 1);
+	
+	// Insert a compute memory barrier after this write for subsequent compute shader reads of bindless heap
+	Context.InsertComputeMemoryBarrier();
+	
 	Context.GetStateCache().SetOverriddenDescriptorHeap(nullptr);
 
 	FMetalDynamicRHI::Get().DeferredDelete([DescriptorEntriesBuffer, DescriptorIndiciesBuffer, Heap, UBDataPtr]()
diff --git a/Engine/Source/Runtime/Apple/MetalRHI/Private/MetalRHIContext.cpp b/Engine/Source/Runtime/Apple/MetalRHI/Private/MetalRHIContext.cpp
--- a/Engine/Source/Runtime/Apple/MetalRHI/Private/MetalRHIContext.cpp
+++ b/Engine/Source/Runtime/Apple/MetalRHI/Private/MetalRHIContext.cpp
@@ -19,7 +19,7 @@ static FAutoConsoleVariableRef CVarMetalSampleBlitEncoderTimings(
 	GMetalSampleBlitEncoderTimings,
 	TEXT("Whether to include Blit Encoder Timings in GPU Stats, can cause poor perf due to splitting encoders: off"));
 
-static int32 GMetalSampleComputeEncoderTimings = 1;
+static int32 GMetalSampleComputeEncoderTimings = 0;
 static FAutoConsoleVariableRef CVarMetalSampleComputeEncoderTimings(
 	TEXT("rhi.Metal.SampleComputeEncoderTimings"),
 	GMetalSampleComputeEncoderTimings,
diff --git a/Engine/Source/Runtime/Apple/MetalRHI/Private/MetalTransitionData.cpp b/Engine/Source/Runtime/Apple/MetalRHI/Private/MetalTransitionData.cpp
--- a/Engine/Source/Runtime/Apple/MetalRHI/Private/MetalTransitionData.cpp
+++ b/Engine/Source/Runtime/Apple/MetalRHI/Private/MetalTransitionData.cpp
@@ -54,7 +54,7 @@ void FMetalTransitionData::EndResourceTransitions(FMetalCommandEncoder& CurrentE
 
 		checkf(Info.AccessAfter != ERHIAccess::Unknown, TEXT("Transitioning a resource to an unknown state is not allowed."));
 
-		UAVBarrier = UAVBarrier || EnumHasAnyFlags(Info.AccessBefore, ERHIAccess::UAVCompute) || EnumHasAnyFlags(Info.AccessAfter, ERHIAccess::UAVCompute);
+		UAVBarrier = UAVBarrier || EnumHasAnyFlags(Info.AccessBefore, ERHIAccess::UAVMask) || EnumHasAnyFlags(Info.AccessAfter, ERHIAccess::UAVMask);
 		
 		switch (Info.Type)
 		{
@@ -132,9 +132,14 @@ void FMetalRHICommandContext::RHIEndTransitions(TArrayView<const FRHITransition*
 	
 	if (UAVBarrier)
 	{
-		if (CurrentEncoder.IsComputeCommandEncoderActive())
-		{
-			CurrentEncoder.GetComputeCommandEncoder()->memoryBarrier(MTL::BarrierScopeBuffers | MTL::BarrierScopeTextures);
-		}
+		InsertComputeMemoryBarrier();
+	}
+}
+
+void FMetalRHICommandContext::InsertComputeMemoryBarrier()
+{
+	if (CurrentEncoder.IsComputeCommandEncoderActive())
+	{
+		CurrentEncoder.GetComputeCommandEncoder()->memoryBarrier(MTL::BarrierScopeBuffers | MTL::BarrierScopeTextures);
 	}
 }
diff --git a/Engine/Source/Runtime/Apple/MetalRHI/Public/MetalRHIContext.h b/Engine/Source/Runtime/Apple/MetalRHI/Public/MetalRHIContext.h
--- a/Engine/Source/Runtime/Apple/MetalRHI/Public/MetalRHIContext.h
+++ b/Engine/Source/Runtime/Apple/MetalRHI/Public/MetalRHIContext.h
@@ -220,6 +220,7 @@ public:
 #endif
 	
 	void PushDescriptorUpdates();
+	void InsertComputeMemoryBarrier();
 	
 	enum class EPhase
 	{
