Metal RHI

- Fix crash in compaction when removing geometry

diff --git a/Engine/Source/Runtime/Apple/MetalRHI/Private/MetalRayTracing.cpp b/Engine/Source/Runtime/Apple/MetalRHI/Private/MetalRayTracing.cpp
--- a/Engine/Source/Runtime/Apple/MetalRHI/Private/MetalRayTracing.cpp
+++ b/Engine/Source/Runtime/Apple/MetalRHI/Private/MetalRayTracing.cpp
@@ -589,6 +589,8 @@ FMetalRayTracingGeometry::~FMetalRayTracingGeometry()
 
 void FMetalRayTracingGeometry::ReleaseUnderlyingResource()
 {
+	RemoveCompactionRequest();
+	
 	if(AccelerationStructure)
 	{
 		FMetalDynamicRHI::Get().DeferredDelete(AccelerationStructure);
@@ -629,7 +631,6 @@ void FMetalRayTracingGeometry::RemoveCompactionRequest()
 	{
 		check(GetAccelerationStructure());
 		bool bRequestFound = Device.GetRayTracingCompactionRequestHandler()->ReleaseRequest(this);
-		check(bRequestFound);
 		bHasPendingCompactionRequests = false;
 	}
 }
