Metal RHI

- Add support for Shader validation to Metal, can be enabled with -metalshadervalidation
- rhi.Metal.ShaderValidation.Type can be used to narrow down specific pipelines
- Also added support for enabling Metal api validation with -metalapivalidation

diff --git a/Engine/Source/Developer/Apple/MetalShaderFormat/Private/MetalShaderFormat.cpp b/Engine/Source/Developer/Apple/MetalShaderFormat/Private/MetalShaderFormat.cpp
--- a/Engine/Source/Developer/Apple/MetalShaderFormat/Private/MetalShaderFormat.cpp
+++ b/Engine/Source/Developer/Apple/MetalShaderFormat/Private/MetalShaderFormat.cpp
@@ -246,8 +246,10 @@ public:
 			// Enqueue the library compilation as a task so we can go wide
 			FGraphEventRef CompletionFence = FFunctionGraphTask::CreateAndDispatchWhenReady([ShaderFormatName, ArchivePath, LibraryPath, PartialShaders, DebugOutputDir, &CompiledLibraries, Platform]()
 				{
-					static const auto CVar = IConsoleManager::Get().FindConsoleVariable(TEXT("r.Shaders.Symbols"));
-					bool bStripLibrary = FDataDrivenShaderPlatformInfo::GetIsMobile(Platform) && CVar->GetInt() == 0;
+					static const auto CVarSymbols = IConsoleManager::Get().FindConsoleVariable(TEXT("r.Shaders.Symbols"));
+					static const auto CVarExtraData = IConsoleManager::Get().FindConsoleVariable(TEXT("r.Shaders.ExtraData"));
+				
+					bool bStripLibrary = FDataDrivenShaderPlatformInfo::GetIsMobile(Platform) && CVarSymbols->GetInt() == 0 && CVarExtraData->GetInt() == 0;
 				
 					if (FinalizeLibrary_Metal(ShaderFormatName, ArchivePath, LibraryPath, PartialShaders, DebugOutputDir, bStripLibrary))
 					{
diff --git a/Engine/Source/Runtime/Apple/MetalRHI/Private/MetalCommandQueue.cpp b/Engine/Source/Runtime/Apple/MetalRHI/Private/MetalCommandQueue.cpp
--- a/Engine/Source/Runtime/Apple/MetalRHI/Private/MetalCommandQueue.cpp
+++ b/Engine/Source/Runtime/Apple/MetalRHI/Private/MetalCommandQueue.cpp
@@ -45,7 +45,11 @@ FMetalCommandQueue::FMetalCommandQueue(FMetalDevice& MetalDevice, uint32 const M
 #endif
     GConfig->GetInt(Settings, TEXT("MetalLanguageVersion"), MetalShaderVersion, GEngineIni);
 	
-	ValidateVersion(MetalShaderVersion);
+	MTL::LanguageVersion ShaderVersion = ValidateVersion(MetalShaderVersion);
+	if(Device.IsShaderValidationEnabled() && ShaderVersion < MTL::LanguageVersion3_0)
+	{
+		UE_LOG(LogMetal, Warning, TEXT("Shader validation is enabled, but ShaderVersion is less than 3.0, change MetalLanguageVersion to 3.0+ for Line/Trace info"));
+	}
 
 	if(MaxNumCommandBuffers == 0)
 	{
diff --git a/Engine/Source/Runtime/Apple/MetalRHI/Private/MetalDevice.cpp b/Engine/Source/Runtime/Apple/MetalRHI/Private/MetalDevice.cpp
--- a/Engine/Source/Runtime/Apple/MetalRHI/Private/MetalDevice.cpp
+++ b/Engine/Source/Runtime/Apple/MetalRHI/Private/MetalDevice.cpp
@@ -145,6 +145,25 @@ static FAutoConsoleVariableRef CVarMetalForceNoHeap(
 	TEXT("[IOS] When enabled, act as if -nometalheap was on the commandline\n")
 	TEXT("(On by default (1))"));
 
+int32 GMetalShaderValidationType = 0;
+static FAutoConsoleVariableRef CVarMetalShaderValidationAll(
+	TEXT("rhi.Metal.ShaderValidation.Type"),
+	GMetalShaderValidationType,
+	TEXT("Enable to set shader validation on specific types\n")
+	TEXT("0: All shaders (slow, default) \n")
+	TEXT("1: All compute shaders \n")
+	TEXT("2: All render pipeline shaders \n")
+	TEXT("3: Match shader name \n")
+	TEXT("Enable to set shader validation on specific types\n"),
+	ECVF_ReadOnly);
+
+FString GMetalShaderValidationShaderName = "";
+static FAutoConsoleVariableRef CVarMetalShaderValidationShaderName(
+	TEXT("rhi.Metal.ShaderValidation.ShaderName"),
+	GMetalShaderValidationShaderName,
+	TEXT("Enable to set shader validation on compute shaders with a name\n"), ECVF_ReadOnly);
+
+
 #if PLATFORM_MAC
 static NS::Object* GMetalDeviceObserver;
 static MTL::Device* GetMTLDevice(uint32& DeviceIndex)
@@ -330,6 +349,11 @@ FMetalDevice* FMetalDevice::CreateDevice()
 	
 	FMetalDevice* MetalDevice = new FMetalDevice(Device, DeviceIndex);
 	
+#if !UE_BUILD_SHIPPING
+	bool bShaderValidationEnabled = FParse::Param(FCommandLine::Get(), TEXT("metalshadervalidation"));
+	MetalDevice->bShaderValidationEnabled = bShaderValidationEnabled;
+#endif
+	
 	if (MetalDevice->SupportsFeature(EMetalFeaturesFences))
 	{
 		FMetalFencePool::Get().Initialise(Device);
diff --git a/Engine/Source/Runtime/Apple/MetalRHI/Private/MetalDevice.h b/Engine/Source/Runtime/Apple/MetalRHI/Private/MetalDevice.h
--- a/Engine/Source/Runtime/Apple/MetalRHI/Private/MetalDevice.h
+++ b/Engine/Source/Runtime/Apple/MetalRHI/Private/MetalDevice.h
@@ -83,6 +83,17 @@ enum EMetalDebugLevel
 	EMetalDebugLevelWaitForComplete,
 };
 
+/**
+ * EMetalShaderValidationType: Type of Metal Shader Validation to perform
+ */
+enum class EMetalShaderValidationType
+{
+	All,
+	ComputeOnly,
+	RenderPipelineOnly,
+	ShaderName,
+};
+
 class FMetalBindlessDescriptorManager;
 
 class FMetalRayTracingCompactionRequestHandler;
@@ -169,6 +180,11 @@ public:
 		return CounterSampler;
 	}
 	
+	bool IsShaderValidationEnabled()
+	{
+		return bShaderValidationEnabled;
+	}
+	
 private:
 	FMetalDevice(MTL::Device* MetalDevice, uint32 DeviceIndex);
 	
@@ -226,6 +242,8 @@ private:
 	
 	FMetalCounterSampler* CounterSampler = nullptr;
 	
+	bool bShaderValidationEnabled = false;
+	
 #if PLATFORM_SUPPORTS_BINDLESS_RENDERING
     /** Bindless Descriptor Heaps manager. */
     FMetalBindlessDescriptorManager* BindlessDescriptorManager;
diff --git a/Engine/Source/Runtime/Apple/MetalRHI/Private/MetalPipeline.cpp b/Engine/Source/Runtime/Apple/MetalRHI/Private/MetalPipeline.cpp
--- a/Engine/Source/Runtime/Apple/MetalRHI/Private/MetalPipeline.cpp
+++ b/Engine/Source/Runtime/Apple/MetalRHI/Private/MetalPipeline.cpp
@@ -40,6 +40,9 @@ static FAutoConsoleVariableRef CVarMetalBinaryCacheDebugOutput(
     GMetalBinaryCacheDebugOutput,
     TEXT("Enable to output logging information for PSO Binary cache default(0) \n"), ECVF_ReadOnly);
 
+extern int32 GMetalShaderValidationType;
+extern FString GMetalShaderValidationShaderName;
+
 static uint32 BlendBitOffsets[] = { Offset_BlendState0, Offset_BlendState1, Offset_BlendState2, Offset_BlendState3, Offset_BlendState4, Offset_BlendState5, Offset_BlendState6, Offset_BlendState7 };
 static uint32 RTBitOffsets[] = { Offset_RenderTargetFormat0, Offset_RenderTargetFormat1, Offset_RenderTargetFormat2, Offset_RenderTargetFormat3, Offset_RenderTargetFormat4, Offset_RenderTargetFormat5, Offset_RenderTargetFormat6, Offset_RenderTargetFormat7 };
 static_assert(Offset_RasterEnd < 64 && Offset_End < 128, "Offset_RasterEnd must be < 64 && Offset_End < 128");
@@ -624,6 +627,8 @@ static bool ConfigureRenderPipelineDescriptor(FMetalDevice& Device, TDescriptorT
 	constexpr bool bIsRenderPipelineDesc = std::is_same_v<TDescriptorType, MTL::RenderPipelineDescriptor>;
 	
 	FMetalPixelShader* PixelShader = (FMetalPixelShader*)Init.BoundShaderState.PixelShaderRHI;
+	FMetalVertexShader* VertexShader = Init.BoundShaderState.VertexShaderRHI ? (FMetalVertexShader*)Init.BoundShaderState.VertexShaderRHI : nullptr;
+	
 	uint32 const NumActiveTargets = Init.ComputeNumValidRenderTargets();
 	check(NumActiveTargets <= MaxSimultaneousRenderTargets);
 	if (PixelShader)
@@ -632,6 +637,36 @@ static bool ConfigureRenderPipelineDescriptor(FMetalDevice& Device, TDescriptorT
 		UE_CLOG((NumActiveTargets < MaxNumActiveTargets), LogMetal, Verbose, TEXT("NumActiveTargets doesn't match pipeline's pixel shader output mask: %u, %hx"), NumActiveTargets, PixelShader->Bindings.InOutMask.Bitmask);
 	}
 	
+#if !UE_BUILD_SHIPPING
+	bool bEnableValidation = false;
+	if(Device.IsShaderValidationEnabled())
+	{
+		switch((EMetalShaderValidationType)GMetalShaderValidationType)
+		{
+			case EMetalShaderValidationType::All:
+			case EMetalShaderValidationType::RenderPipelineOnly:
+				bEnableValidation = true;
+				break;
+			case EMetalShaderValidationType::ShaderName:
+			{
+				if(PixelShader && FString(PixelShader->GetShaderName()).Contains(GMetalShaderValidationShaderName) ||
+				   VertexShader && FString(VertexShader->GetShaderName()).Contains(GMetalShaderValidationShaderName)) 
+				{
+					bEnableValidation = true;
+				}
+			}
+			case EMetalShaderValidationType::ComputeOnly:
+			default:
+				break;
+		}
+		
+		if(bEnableValidation)
+		{
+			RenderPipelineDesc->setShaderValidation(MTL::ShaderValidationEnabled);
+		}
+	}
+#endif
+	
 	FMetalBlendState* BlendState = (FMetalBlendState*)Init.BlendState;
 	
 	MTL::RenderPipelineColorAttachmentDescriptorArray* ColorAttachments = RenderPipelineDesc->colorAttachments();
@@ -734,8 +769,6 @@ static bool ConfigureRenderPipelineDescriptor(FMetalDevice& Device, TDescriptorT
 	
 	if constexpr(bIsRenderPipelineDesc)
 	{
-		FMetalVertexShader* VertexShader = (FMetalVertexShader*)Init.BoundShaderState.VertexShaderRHI;
-		
 		MTL::PipelineBufferDescriptorArray* VertexPipelineBuffers = RenderPipelineDesc->vertexBuffers();
 		FMetalShaderBindings& VertexBindings = VertexShader->Bindings;
 		int8 VertexSideTable = VertexShader->SideTableBinding;
diff --git a/Engine/Source/Runtime/Apple/MetalRHI/Private/Shaders/Types/MetalComputeShader.cpp b/Engine/Source/Runtime/Apple/MetalRHI/Private/Shaders/Types/MetalComputeShader.cpp
--- a/Engine/Source/Runtime/Apple/MetalRHI/Private/Shaders/Types/MetalComputeShader.cpp
+++ b/Engine/Source/Runtime/Apple/MetalRHI/Private/Shaders/Types/MetalComputeShader.cpp
@@ -13,6 +13,9 @@
 
 #pragma mark - Metal RHI Compute Shader Class
 
+extern int32 GMetalShaderValidationType;
+extern FString GMetalShaderValidationShaderName;
+
 FMetalComputeShader::FMetalComputeShader(FMetalDevice& Device, TArrayView<const uint8> InCode, MTLLibraryPtr InLibrary)
 	: TMetalBaseShader<FRHIComputeShader, SF_Compute>(Device)
 	, NumThreadsX(0)
@@ -54,7 +57,38 @@ FMetalShaderPipelinePtr FMetalComputeShader::GetPipeline()
 		NS::Error* Error;
 		MTL::ComputePipelineDescriptor* Descriptor = MTL::ComputePipelineDescriptor::alloc()->init();
         check(Descriptor);
-        
+		
+#if !UE_BUILD_SHIPPING
+		if(Device.IsShaderValidationEnabled())
+		{
+			bool bEnableValidation = false;
+			
+			// Handle Validation
+			switch((EMetalShaderValidationType)GMetalShaderValidationType)
+			{
+				case EMetalShaderValidationType::All:
+				case EMetalShaderValidationType::ComputeOnly:
+					bEnableValidation = true;
+					break;
+				case EMetalShaderValidationType::ShaderName:
+				{
+					if(NSStringToFString(Func->name()).Contains(GMetalShaderValidationShaderName)) 
+					{
+						bEnableValidation = true;
+					}
+				}
+				case EMetalShaderValidationType::RenderPipelineOnly:
+				default:
+					break;
+			}
+			
+			if(bEnableValidation)
+			{
+				Descriptor->setShaderValidation(MTL::ShaderValidationEnabled);
+			}
+		}
+#endif 
+		
 		Descriptor->setLabel(Func->name());
 		Descriptor->setComputeFunction(Func.get());
         
diff --git a/Engine/Source/Runtime/ApplicationCore/Private/IOS/IOSAppDelegate.cpp b/Engine/Source/Runtime/ApplicationCore/Private/IOS/IOSAppDelegate.cpp
--- a/Engine/Source/Runtime/ApplicationCore/Private/IOS/IOSAppDelegate.cpp
+++ b/Engine/Source/Runtime/ApplicationCore/Private/IOS/IOSAppDelegate.cpp
@@ -1784,6 +1784,19 @@ FCriticalSection RenderSuspend;
 extern double GCStartTime;
 - (void)applicationDidBecomeActive:(UIApplication *)application
 {
+#if !UE_BUILD_SHIPPING
+	if (FParse::Param(FCommandLine::Get(), TEXT("metalshadervalidation")))
+	{
+		FPlatformMisc::SetEnvironmentVar(TEXT("MTL_SHADER_VALIDATION"), TEXT("1"));
+		FPlatformMisc::SetEnvironmentVar(TEXT("MTL_SHADER_VALIDATION_DEFAULT_STATE"), TEXT("none"));
+	}
+	
+	if (FParse::Param(FCommandLine::Get(), TEXT("metalapivalidation")))
+	{
+		FPlatformMisc::SetEnvironmentVar(TEXT("MTL_DEBUG_LAYER"), TEXT("1"));
+	}
+#endif
+	
 	FIOSCoreDelegates::OnDidBecomeActive.Broadcast();
 	
 	UE::GC::Private::SetSkipDestroy(false);
diff --git a/Engine/Source/Runtime/Core/Private/Apple/ApplePlatformMisc.cpp b/Engine/Source/Runtime/Core/Private/Apple/ApplePlatformMisc.cpp
--- a/Engine/Source/Runtime/Core/Private/Apple/ApplePlatformMisc.cpp
+++ b/Engine/Source/Runtime/Core/Private/Apple/ApplePlatformMisc.cpp
@@ -387,4 +387,18 @@ void FApplePlatformMisc::EndNamedEvent()
 #endif // FRAMEPRO_ENABLED
 }
 
+void FApplePlatformMisc::SetEnvironmentVar(const TCHAR* InVariableName, const TCHAR* Value)
+{
+	FString VariableName = InVariableName;
+	VariableName.ReplaceInline(TEXT("-"), TEXT("_"));
+	if (Value == NULL || Value[0] == TEXT('\0'))
+	{
+		unsetenv(TCHAR_TO_ANSI(*VariableName));
+	}
+	else
+	{
+		setenv(TCHAR_TO_ANSI(*VariableName), TCHAR_TO_ANSI(Value), 1);
+	}
+}
+
 #endif // APPLE_PROFILING_ENABLED
diff --git a/Engine/Source/Runtime/Core/Private/Mac/MacPlatformMisc.cpp b/Engine/Source/Runtime/Core/Private/Mac/MacPlatformMisc.cpp
--- a/Engine/Source/Runtime/Core/Private/Mac/MacPlatformMisc.cpp
+++ b/Engine/Source/Runtime/Core/Private/Mac/MacPlatformMisc.cpp
@@ -642,21 +642,6 @@ void FMacPlatformMisc::PlatformTearDown()
 	FPlatformProcess::CeaseBeingFirstInstance();
 }
 
-void FMacPlatformMisc::SetEnvironmentVar(const TCHAR* InVariableName, const TCHAR* Value)
-{
-	FString VariableName = InVariableName;
-	VariableName.ReplaceInline(TEXT("-"), TEXT("_"));
-	if (Value == NULL || Value[0] == TEXT('\0'))
-	{
-		unsetenv(TCHAR_TO_ANSI(*VariableName));
-	}
-	else
-	{
-		setenv(TCHAR_TO_ANSI(*VariableName), TCHAR_TO_ANSI(Value), 1);
-	}
-}
-
-
 TArray<uint8> FMacPlatformMisc::GetMacAddress()
 {
 	TArray<uint8> Result;
diff --git a/Engine/Source/Runtime/Core/Public/Apple/ApplePlatformMisc.h b/Engine/Source/Runtime/Core/Public/Apple/ApplePlatformMisc.h
--- a/Engine/Source/Runtime/Core/Public/Apple/ApplePlatformMisc.h
+++ b/Engine/Source/Runtime/Core/Public/Apple/ApplePlatformMisc.h
@@ -101,4 +101,6 @@ struct CORE_API FApplePlatformMisc : public FGenericPlatformMisc
 	//////// Platform specific
 	static void* CreateAutoreleasePool();
 	static void ReleaseAutoreleasePool(void *Pool);
+	
+	static void SetEnvironmentVar(const TCHAR* InVariableName, const TCHAR* Value);
 };
diff --git a/Engine/Source/Runtime/Core/Public/Mac/MacPlatformMisc.h b/Engine/Source/Runtime/Core/Public/Mac/MacPlatformMisc.h
--- a/Engine/Source/Runtime/Core/Public/Mac/MacPlatformMisc.h
+++ b/Engine/Source/Runtime/Core/Public/Mac/MacPlatformMisc.h
@@ -17,7 +17,6 @@ struct CORE_API FMacPlatformMisc : public FApplePlatformMisc
 	static void PlatformPreInit();
 	static void PlatformInit();
 	static void PlatformTearDown();
-	static void SetEnvironmentVar(const TCHAR* VariableName, const TCHAR* Value);
 
 	UE_FORCEINLINE_HINT static constexpr int32 GetMaxPathLength()
 	{
diff --git a/Engine/Source/Runtime/Launch/Private/Mac/LaunchMac.cpp b/Engine/Source/Runtime/Launch/Private/Mac/LaunchMac.cpp
--- a/Engine/Source/Runtime/Launch/Private/Mac/LaunchMac.cpp
+++ b/Engine/Source/Runtime/Launch/Private/Mac/LaunchMac.cpp
@@ -92,6 +92,19 @@ void EngineFullCrashHandler(const FGenericCrashContext& GenericContext)
 		}
 		GSavedCommandLine += Argument;
 	}
+	
+#if !UE_BUILD_SHIPPING
+	if (GSavedCommandLine.Contains(TEXT("metalshadervalidation")))
+	{
+		FPlatformMisc::SetEnvironmentVar(TEXT("MTL_SHADER_VALIDATION"), TEXT("1"));
+		FPlatformMisc::SetEnvironmentVar(TEXT("MTL_SHADER_VALIDATION_DEFAULT_STATE"), TEXT("none"));
+	}
+	
+	if (GSavedCommandLine.Contains(TEXT("metalapivalidation")))
+	{
+		FPlatformMisc::SetEnvironmentVar(TEXT("MTL_DEBUG_LAYER"), TEXT("1"));
+	}
+#endif
 
 #if UE_GAME
 	// On Mac we always want games to save files to user dir instead of inside the app bundle
