Metal RHI

- Fix Usage not being correctly set when binding buffers and cleaned up some code.

diff --git a/Engine/Source/Runtime/Apple/MetalRHI/Private/MetalStateCache.cpp b/Engine/Source/Runtime/Apple/MetalRHI/Private/MetalStateCache.cpp
--- a/Engine/Source/Runtime/Apple/MetalRHI/Private/MetalStateCache.cpp
+++ b/Engine/Source/Runtime/Apple/MetalRHI/Private/MetalStateCache.cpp
@@ -176,6 +176,13 @@ FMetalStateCache::~FMetalStateCache()
 {
 	RenderPassDesc = nullptr;
 	
+	InternalResetStageBindings();
+	
+	VisibilityResults = nullptr;
+}
+
+void FMetalStateCache::InternalResetStageBindings()
+{
 	for (uint32 i = 0; i < MaxVertexElementCount; i++)
 	{
 		VertexBuffers[i].Buffer = nullptr;
@@ -186,37 +193,43 @@ FMetalStateCache::~FMetalStateCache()
 		VertexBuffers[i].Length = 0;
 		VertexBuffers[i].Offset = 0;
 	}
+	
 	for (uint32 Frequency = 0; Frequency < EMetalShaderStages::Num; Frequency++)
 	{
-		ShaderSamplers[Frequency].Bound = 0;
+		FMetalBufferBindings& StageBufferBindings = ShaderBuffers[Frequency];
+		FMetalSamplerBindings& StageSamplerBindings = ShaderSamplers[Frequency];
+		FMetalTextureBindings& StageTextureBindings = ShaderTextures[Frequency];
+		
+		StageBufferBindings.Bound = 0;
+		
 		for (uint32 i = 0; i < ML_MaxSamplers; i++)
 		{
-			ShaderSamplers[Frequency].Samplers[i] = nullptr;
+			StageSamplerBindings.Samplers[i] = nullptr;
 		}
 		for (uint32 i = 0; i < ML_MaxBuffers; i++)
 		{
-			BoundUniformBuffers[Frequency][i] = nullptr;
-			ShaderBuffers[Frequency].Buffers[i].Buffer = nullptr;
+			FMetalBufferBinding& CurrentBufferBinding = StageBufferBindings.Buffers[i];
+			
+			CurrentBufferBinding.Buffer = nullptr;
 #if METAL_RHI_RAYTRACING
-			ShaderBuffers[Frequency].Buffers[i].AccelerationStructure = nullptr;
+			CurrentBufferBinding.AccelerationStructure = nullptr;
 #endif
-			ShaderBuffers[Frequency].Buffers[i].Bytes = nullptr;
-			ShaderBuffers[Frequency].Buffers[i].Length = 0;
-			ShaderBuffers[Frequency].Buffers[i].ElementRowPitch = 0;
-			ShaderBuffers[Frequency].Buffers[i].Offset = 0;
-			ShaderBuffers[Frequency].Buffers[i].Usage = MTL::ResourceUsage(0);
-			ShaderBuffers[Frequency].Formats[i] = PF_Unknown;
+			CurrentBufferBinding.Bytes = nullptr;
+			CurrentBufferBinding.Length = 0;
+			CurrentBufferBinding.ElementRowPitch = 0;
+			CurrentBufferBinding.Offset = 0;
+			CurrentBufferBinding.Usage = MTL::ResourceUsage(0);
+			StageBufferBindings.Formats[i] = PF_Unknown;
 		}
 		ShaderBuffers[Frequency].Bound = 0;
+		
 		for (uint32 i = 0; i < ML_MaxTextures; i++)
 		{
-			ShaderTextures[Frequency].Textures[i] = nullptr;
-			ShaderTextures[Frequency].Usage[i] = MTL::ResourceUsage(0);
+			StageTextureBindings.Textures[i] = nullptr;
+			StageTextureBindings.Usage[i] = MTL::ResourceUsage(0);
 		}
-		ShaderTextures[Frequency].Bound = 0;
+		StageTextureBindings.Bound = 0;
 	}
-	
-	VisibilityResults = nullptr;
 }
 
 void FMetalStateCache::Reset()
@@ -272,45 +285,8 @@ void FMetalStateCache::Reset()
 	FMemory::Memzero(BoundUniformBuffers);
 	ActiveUniformBuffers.Empty();
 	
-	for (uint32 i = 0; i < MaxVertexElementCount; i++)
-	{
-		VertexBuffers[i].Buffer = nullptr;
-		VertexBuffers[i].Bytes = nullptr;
-#if METAL_RHI_RAYTRACING
-		VertexBuffers[i].AccelerationStructure = nullptr;
-#endif
-		VertexBuffers[i].Length = 0;
-		VertexBuffers[i].Offset = 0;
-	}
+	InternalResetStageBindings();
 	
-	for (uint32 Frequency = 0; Frequency < EMetalShaderStages::Num; Frequency++)
-	{
-		ShaderSamplers[Frequency].Bound = 0;
-		for (uint32 i = 0; i < ML_MaxSamplers; i++)
-		{
-			ShaderSamplers[Frequency].Samplers[i] = nullptr;
-		}
-		for (uint32 i = 0; i < ML_MaxBuffers; i++)
-		{
-			ShaderBuffers[Frequency].Buffers[i].Buffer = nullptr;
-#if METAL_RHI_RAYTRACING
-			ShaderBuffers[Frequency].Buffers[i].AccelerationStructure = nullptr;
-#endif
-			ShaderBuffers[Frequency].Buffers[i].Bytes = nullptr;
-			ShaderBuffers[Frequency].Buffers[i].Length = 0;
-			ShaderBuffers[Frequency].Buffers[i].ElementRowPitch = 0;
-			ShaderBuffers[Frequency].Buffers[i].Offset = 0;
-			ShaderBuffers[Frequency].Formats[i] = PF_Unknown;
-		}
-		ShaderBuffers[Frequency].Bound = 0;
-		for (uint32 i = 0; i < ML_MaxTextures; i++)
-		{
-			ShaderTextures[Frequency].Textures[i] = nullptr;
-			ShaderTextures[Frequency].Usage[i] = MTL::ResourceUsage(0);
-		}
-		ShaderTextures[Frequency].Bound = 0;
-	}
-
 	SetRenderTargetsActive(false);
 	
 #if PLATFORM_SUPPORTS_BINDLESS_RENDERING
@@ -1426,34 +1402,37 @@ void FMetalStateCache::SetShaderBuffer(
 	check(Frequency < EMetalShaderStages::Num);
 	check(Index < ML_MaxBuffers);
 	
-	if (ShaderBuffers[Frequency].Buffers[Index].Buffer != Buffer ||
-		ShaderBuffers[Frequency].Buffers[Index].Bytes != Bytes ||
-		ShaderBuffers[Frequency].Buffers[Index].Offset != Offset ||
-		ShaderBuffers[Frequency].Buffers[Index].Length != Length ||
-		ShaderBuffers[Frequency].Buffers[Index].ElementRowPitch != ElementRowPitch ||
-		!(ShaderBuffers[Frequency].Buffers[Index].Usage & Usage) ||
-		ShaderBuffers[Frequency].Formats[Index] != Format)
+	FMetalBufferBindings& StageBufferBindings = ShaderBuffers[Frequency]; 
+	FMetalBufferBinding& CurrentBufferBinding = StageBufferBindings.Buffers[Index];
+	
+	if (CurrentBufferBinding.Buffer != Buffer ||
+		CurrentBufferBinding.Bytes != Bytes ||
+		CurrentBufferBinding.Offset != Offset ||
+		CurrentBufferBinding.Length != Length ||
+		CurrentBufferBinding.ElementRowPitch != ElementRowPitch ||
+		CurrentBufferBinding.Usage != Usage ||
+		StageBufferBindings.Formats[Index] != Format)
 	{
-		ShaderBuffers[Frequency].Buffers[Index].Buffer = Buffer;
-		ShaderBuffers[Frequency].Buffers[Index].Bytes = Bytes;
+		CurrentBufferBinding.Buffer = Buffer;
+		CurrentBufferBinding.Bytes = Bytes;
 #if METAL_RHI_RAYTRACING
-		ShaderBuffers[Frequency].Buffers[Index].AccelerationStructure = nullptr;
+		CurrentBufferBinding.AccelerationStructure = nullptr;
 #endif
-		ShaderBuffers[Frequency].Buffers[Index].ReferencedResources = ReferencedResources;
-		ShaderBuffers[Frequency].Buffers[Index].Offset = Offset;
-		ShaderBuffers[Frequency].Buffers[Index].Length = Length;
-		ShaderBuffers[Frequency].Buffers[Index].ElementRowPitch = ElementRowPitch;
-		ShaderBuffers[Frequency].Buffers[Index].Usage = Usage;
+		CurrentBufferBinding.ReferencedResources = ReferencedResources;
+		CurrentBufferBinding.Offset = Offset;
+		CurrentBufferBinding.Length = Length;
+		CurrentBufferBinding.ElementRowPitch = ElementRowPitch;
+		CurrentBufferBinding.Usage = Usage;
 		
-		ShaderBuffers[Frequency].Formats[Index] = Format;
+		StageBufferBindings.Formats[Index] = Format;
 		
 		if (Buffer || Bytes)
 		{
-			ShaderBuffers[Frequency].Bound |= (1 << Index);
+			StageBufferBindings.Bound |= (1 << Index);
 		}
 		else
 		{
-			ShaderBuffers[Frequency].Bound &= ~(1 << Index);
+			StageBufferBindings.Bound &= ~(1 << Index);
 		}
 	}
 }
@@ -1464,25 +1443,28 @@ void FMetalStateCache::SetShaderBuffer(EMetalShaderStages const Frequency, FMeta
 	check(Frequency < EMetalShaderStages::Num);
 	check(Index < ML_MaxBuffers);
 
-	if (ShaderBuffers[Frequency].Buffers[Index].AccelerationStructure != AccelerationStructure)
+	FMetalBufferBindings& StageBufferBindings = ShaderBuffers[Frequency]; 
+	FMetalBufferBinding& CurrentBufferBinding = StageBufferBindings.Buffers[Index];
+	
+	if (CurrentBufferBinding.AccelerationStructure != AccelerationStructure)
 	{
-		ShaderBuffers[Frequency].Buffers[Index].AccelerationStructure = AccelerationStructure;
-		ShaderBuffers[Frequency].Buffers[Index].Buffer = nullptr;
-		ShaderBuffers[Frequency].Buffers[Index].Bytes = nullptr;
-		ShaderBuffers[Frequency].Buffers[Index].ReferencedResources = ReferencedResources;
-		ShaderBuffers[Frequency].Buffers[Index].Offset = 0;
-		ShaderBuffers[Frequency].Buffers[Index].Length = 0;
-		ShaderBuffers[Frequency].Buffers[Index].Usage = MTL::ResourceUsage(0);
+		CurrentBufferBinding.AccelerationStructure = AccelerationStructure;
+		CurrentBufferBinding.Buffer = nullptr;
+		CurrentBufferBinding.Bytes = nullptr;
+		CurrentBufferBinding.ReferencedResources = ReferencedResources;
+		CurrentBufferBinding.Offset = 0;
+		CurrentBufferBinding.Length = 0;
+		CurrentBufferBinding.Usage = MTL::ResourceUsage(0);
 
-		ShaderBuffers[Frequency].Formats[Index] = PF_Unknown;
+		StageBufferBindings.Formats[Index] = PF_Unknown;
 
 		if (AccelerationStructure)
 		{
-			ShaderBuffers[Frequency].Bound |= (1 << Index);
+			StageBufferBindings.Bound |= (1 << Index);
 		}
 		else
 		{
-			ShaderBuffers[Frequency].Bound &= ~(1 << Index);
+			StageBufferBindings.Bound &= ~(1 << Index);
 		}
 	}
 }
diff --git a/Engine/Source/Runtime/Apple/MetalRHI/Private/MetalStateCache.h b/Engine/Source/Runtime/Apple/MetalRHI/Private/MetalStateCache.h
--- a/Engine/Source/Runtime/Apple/MetalRHI/Private/MetalStateCache.h
+++ b/Engine/Source/Runtime/Apple/MetalRHI/Private/MetalStateCache.h
@@ -281,6 +281,8 @@ private:
 	};
     
 private:
+	void InternalResetStageBindings();
+	
 	FMetalDevice& Device;
 	
 	FMetalShaderParameterCache ShaderParameters[EMetalShaderStages::Num];
