Metal RHI

Added log dump when erroring under shader validation

diff --git a/Engine/Source/Runtime/Apple/MetalRHI/Private/MetalSubmission.cpp b/Engine/Source/Runtime/Apple/MetalRHI/Private/MetalSubmission.cpp
--- a/Engine/Source/Runtime/Apple/MetalRHI/Private/MetalSubmission.cpp
+++ b/Engine/Source/Runtime/Apple/MetalRHI/Private/MetalSubmission.cpp
@@ -419,6 +419,25 @@ FMetalDynamicRHI::FProcessResult FMetalDynamicRHI::ProcessSubmissionQueue()
 					{
 						FMetalCommandBuffer* CurrentCommandBuffer = Payload->CommandBuffersToExecute[Index];
 						
+#if !UE_BUILD_SHIPPING
+						MTL::CommandBuffer* MTLCmdBuffer = CurrentCommandBuffer->GetMTLCmdBuffer();
+						
+						id<MTLLogContainer> logs = (__bridge id<MTLLogContainer>)MTLCmdBuffer->logs();
+						for (id<MTLFunctionLog> log : logs)
+						{
+							MTL::FunctionLog* FunctionLog = (__bridge MTL::FunctionLog*)log;
+							if(FunctionLog)
+							{
+								MTL::FunctionLogDebugLocation* DebugLocation = FunctionLog->debugLocation();
+								
+								if(DebugLocation)
+								{
+									UE_LOG(LogMetal, Warning, TEXT("Shader Validation Log:\n functionName: %s\nencoder: %s\nline: %d\ncolumn: %d"), *NSStringToFString(DebugLocation->functionName()), *NSStringToFString(FunctionLog->encoderLabel()), DebugLocation->line(), DebugLocation->column());
+								}
+							}
+						}
+#endif
+						
 						CurrentQueue.OcclusionQueries.Append(MoveTemp(CurrentCommandBuffer->OcclusionQueries));
 						CurrentQueue.TimestampQueries.Append(MoveTemp(CurrentCommandBuffer->TimestampQueries));
 #if RHI_NEW_GPU_PROFILER == 0
@@ -533,7 +552,6 @@ void FMetalDynamicRHI::FlushBatchedPayloads(FMetalCommandQueue::FPayloadArray& P
 #if RHI_NEW_GPU_PROFILER
 				CommandBuffer->FlushProfilerEvents(Payload->EventStream, Time);
 #endif
-				
 				Payload->Queue.CommitCommandBuffer(CommandBuffer);
 			}
 		}
