Metal RHI

- Enabled some missing features from Metal SupportsAsyncGetRenderQueryResult, SupportsGPUTimestampBubblesRemoval and SupportsParallelOcclusionQueries
- Enabled HDR format on Apple silicon Macs

diff --git a/Engine/Source/Runtime/Apple/MetalRHI/Private/MetalDynamicRHI.cpp b/Engine/Source/Runtime/Apple/MetalRHI/Private/MetalDynamicRHI.cpp
--- a/Engine/Source/Runtime/Apple/MetalRHI/Private/MetalDynamicRHI.cpp
+++ b/Engine/Source/Runtime/Apple/MetalRHI/Private/MetalDynamicRHI.cpp
@@ -175,7 +175,6 @@ bool FMetalDynamicRHI::RHIGetRenderQueryResult(FRHIRenderQuery* QueryRHI, uint64
 {
     MTL_SCOPED_AUTORELEASE_POOL;
     
-	check(IsInRenderingThread());
 	FMetalRHIRenderQuery* Query = ResourceCast(QueryRHI);
 	return Query->GetResult(OutNumPixels, bWait, GPUIndex);
 }
diff --git a/Engine/Source/Runtime/Apple/MetalRHI/Private/MetalRHI.cpp b/Engine/Source/Runtime/Apple/MetalRHI/Private/MetalRHI.cpp
--- a/Engine/Source/Runtime/Apple/MetalRHI/Private/MetalRHI.cpp
+++ b/Engine/Source/Runtime/Apple/MetalRHI/Private/MetalRHI.cpp
@@ -214,6 +214,8 @@ FMetalDynamicRHI::FMetalDynamicRHI(ERHIFeatureLevel::Type RequestedFeatureLevel)
 	
 	GRHISupportsMultithreading = true;
 	GRHISupportsMultithreadedResources = true;
+	GRHISupportsAsyncGetRenderQueryResult = true;
+	GRHISupportsGPUTimestampBubblesRemoval = true;
 	
 	// we cannot render to a volume texture without geometry shader or vertex-shader-layer support, so initialise to false and enable based on platform feature availability
 	GSupportsVolumeTextureRendering = false;
@@ -434,6 +436,7 @@ FMetalDynamicRHI::FMetalDynamicRHI(ERHIFeatureLevel::Type RequestedFeatureLevel)
 		bSupportsTiledReflections = true;
 		bSupportsDistanceFields = true;
 		GSupportsTimestampRenderQueries = true;
+		GSupportsParallelOcclusionQueries = true;
 		
 		GRHISupportsWaveOperations = true;
 		GRHIMinimumWaveSize = 32;
@@ -1001,14 +1004,26 @@ FMetalDynamicRHI::FMetalDynamicRHI(ERHIFeatureLevel::Type RequestedFeatureLevel)
 	GPixelFormats[PF_FloatR11G11B10		].BlockBytes		= 4;
 	GPixelFormats[PF_FloatR11G11B10		].Supported			= true;
 	
-	// Only one HDR format for OSX.
-	GPixelFormats[PF_PLATFORM_HDR_0		].BlockSizeX		= 1;
-	GPixelFormats[PF_PLATFORM_HDR_0		].BlockSizeY		= 1;
-	GPixelFormats[PF_PLATFORM_HDR_0		].BlockSizeZ		= 1;
-	GPixelFormats[PF_PLATFORM_HDR_0		].BlockBytes		= 8;
-	GPixelFormats[PF_PLATFORM_HDR_0		].NumComponents		= 4;
-	GPixelFormats[PF_PLATFORM_HDR_0		].PlatformFormat	= (uint32)MTL::PixelFormatRGBA16Float;
-	GPixelFormats[PF_PLATFORM_HDR_0		].Supported			= GRHISupportsHDROutput;
+	if(GRHIVendorId == (uint32)EGpuVendorId::Apple)
+	{
+		GPixelFormats[PF_PLATFORM_HDR_0		].BlockSizeX		= 1;
+		GPixelFormats[PF_PLATFORM_HDR_0		].BlockSizeY		= 1;
+		GPixelFormats[PF_PLATFORM_HDR_0		].BlockSizeZ		= 1;
+		GPixelFormats[PF_PLATFORM_HDR_0		].BlockBytes		= 4;
+		GPixelFormats[PF_PLATFORM_HDR_0		].NumComponents		= 3;
+		GPixelFormats[PF_PLATFORM_HDR_0		].PlatformFormat	= (uint32)MTL::PixelFormatBGR10_XR_sRGB;
+		GPixelFormats[PF_PLATFORM_HDR_0		].Supported			= GRHISupportsHDROutput;
+	}
+	else
+	{
+		GPixelFormats[PF_PLATFORM_HDR_0		].BlockSizeX		= 1;
+		GPixelFormats[PF_PLATFORM_HDR_0		].BlockSizeY		= 1;
+		GPixelFormats[PF_PLATFORM_HDR_0		].BlockSizeZ		= 1;
+		GPixelFormats[PF_PLATFORM_HDR_0		].BlockBytes		= 8;
+		GPixelFormats[PF_PLATFORM_HDR_0		].NumComponents		= 4;
+		GPixelFormats[PF_PLATFORM_HDR_0		].PlatformFormat	= (uint32)MTL::PixelFormatRGBA16Float;
+		GPixelFormats[PF_PLATFORM_HDR_0		].Supported			= GRHISupportsHDROutput;
+	}
 
 	if (bSupportsD16)
 	{
