Metal RHI

- Fixed bug where textures were not set for read correctly.

diff --git a/Engine/Source/Runtime/Apple/MetalRHI/Private/MetalStateCache.cpp b/Engine/Source/Runtime/Apple/MetalRHI/Private/MetalStateCache.cpp
--- a/Engine/Source/Runtime/Apple/MetalRHI/Private/MetalStateCache.cpp
+++ b/Engine/Source/Runtime/Apple/MetalRHI/Private/MetalStateCache.cpp
@@ -1470,12 +1470,6 @@ void FMetalStateCache::SetShaderBuffer(EMetalShaderStages const Frequency, FMeta
 }
 #endif // METAL_RHI_RAYTRACING
 
-static bool CanMakeTextureResidentViaHeaps(const MTL::Texture* Texture)
-{
-    return !(Texture->usage() & MTL::TextureUsageRenderTarget)
-        && !(Texture->usage() & MTL::TextureUsageShaderWrite);
-}
-
 #if METAL_USE_METAL_SHADER_CONVERTER
 
 void FMetalStateCache::CacheOrSkipResourceResidencyUpdate(MTL::Resource* InResource, EMetalShaderStages const Frequency, bool bReadOnly, bool bForceUseResource)
@@ -1616,11 +1610,11 @@ void FMetalStateCache::IRMakeTextureResident(EMetalShaderStages const Frequency,
 	}
 	else if (Texture->parentTexture())
 	{
-		CacheOrSkipResourceResidencyUpdate(Texture->parentTexture(), Frequency, true, CanMakeTextureResidentViaHeaps(Texture->parentTexture()));
+		CacheOrSkipResourceResidencyUpdate(Texture->parentTexture(), Frequency, true, true);
 	}
 	
 	{
-		CacheOrSkipResourceResidencyUpdate(Texture, Frequency, true, CanMakeTextureResidentViaHeaps(Texture));
+		CacheOrSkipResourceResidencyUpdate(Texture, Frequency, true, true);
 	}
 }
 
