Metal RHI

- Fixed missing retain when generating stage in functions

diff --git a/Engine/Source/Runtime/Apple/MetalRHI/Private/MetalPipeline.cpp b/Engine/Source/Runtime/Apple/MetalRHI/Private/MetalPipeline.cpp
--- a/Engine/Source/Runtime/Apple/MetalRHI/Private/MetalPipeline.cpp
+++ b/Engine/Source/Runtime/Apple/MetalRHI/Private/MetalPipeline.cpp
@@ -1058,7 +1058,7 @@ static FMetalShaderPipelinePtr CreateMTLRenderPipeline(FMetalDevice& Device, boo
 				dispatch_data_t LibraryData = dispatch_data_create(State->StageInFunctionBytecode.GetData(), State->StageInFunctionBytecode.Num(), nil, DISPATCH_DATA_DESTRUCTOR_DEFAULT);
 				
 				NS::Error* error = nullptr;
-				MTLLibraryPtr StageInLib = NS::TransferPtr(MetalDevice->newLibrary(LibraryData, &error));
+				MTLLibraryPtr StageInLib = NS::RetainPtr(MetalDevice->newLibrary(LibraryData, &error));
 				MTL::LinkedFunctions* StageInFunction = MTL::LinkedFunctions::alloc()->init();
 				StageInFunction->setFunctions(NS::Array::array(
 															   StageInLib->newFunction(NS::String::string("irconverter_stage_in_shader", NS::UTF8StringEncoding))
