Metal RHI

- Fixed flickering on Metal due to presentable textures being invalid when copying
- Fixed another flickering bug where the Load action was clear due to a written atomic not being thread safe and unnecessary

diff --git a/Engine/Source/Runtime/Apple/MetalRHI/Private/MetalRHIContext.cpp b/Engine/Source/Runtime/Apple/MetalRHI/Private/MetalRHIContext.cpp
--- a/Engine/Source/Runtime/Apple/MetalRHI/Private/MetalRHIContext.cpp
+++ b/Engine/Source/Runtime/Apple/MetalRHI/Private/MetalRHIContext.cpp
@@ -59,7 +59,7 @@ FMetalRHICommandContext::FMetalRHICommandContext(FMetalDevice& MetalDevice, clas
 	, CommandQueue(Device.GetCommandQueue(EMetalQueueType::Direct))
 	, CommandList(Device.GetCommandQueue(EMetalQueueType::Direct))
 	, CurrentEncoder(MetalDevice, CommandList)
-	, StateCache(MetalDevice, true)
+	, StateCache(MetalDevice)
 	, QueryBuffer(new FMetalQueryBufferPool(MetalDevice))
 	, RenderPassDesc(nullptr)
 	, Profiler(InProfiler)
diff --git a/Engine/Source/Runtime/Apple/MetalRHI/Private/MetalStateCache.cpp b/Engine/Source/Runtime/Apple/MetalRHI/Private/MetalStateCache.cpp
--- a/Engine/Source/Runtime/Apple/MetalRHI/Private/MetalStateCache.cpp
+++ b/Engine/Source/Runtime/Apple/MetalRHI/Private/MetalStateCache.cpp
@@ -116,7 +116,7 @@ FORCEINLINE MTL::StoreAction GetConditionalMetalRTStoreAction(FMetalDevice& Devi
 	}	
 }
 
-FMetalStateCache::FMetalStateCache(FMetalDevice& MetalDevice, bool const bInImmediate)
+FMetalStateCache::FMetalStateCache(FMetalDevice& MetalDevice)
 : Device(MetalDevice)
 , DepthStore(MTL::StoreActionUnknown)
 , StencilStore(MTL::StoreActionUnknown)
@@ -136,7 +136,6 @@ FMetalStateCache::FMetalStateCache(FMetalDevice& MetalDevice, bool const bInImme
 , bIsRenderTargetActive(false)
 , bHasValidRenderTarget(false)
 , bHasValidColorTarget(false)
-, bImmediate(bInImmediate)
 {
 	FMemory::Memzero(Viewport);
 	FMemory::Memzero(Scissor);
@@ -664,8 +663,7 @@ bool FMetalStateCache::SetRenderPassInfo(FRHIRenderPassInfo const& InRenderTarge
 					ColorAttachment->setSlice(ArraySliceIndex);
 				}
 				
-				ColorAttachment->setLoadAction((Surface.Written || !bImmediate) ? GetMetalRTLoadAction(HighLevelLoadAction) : MTL::LoadActionClear);
-				FPlatformAtomics::InterlockedExchange(&Surface.Written, 1);
+				ColorAttachment->setLoadAction(GetMetalRTLoadAction(HighLevelLoadAction));
 				
 				bNeedsClear |= (ColorAttachment->loadAction() == MTL::LoadActionClear);
 				
@@ -1592,12 +1590,6 @@ void FMetalStateCache::IRMakeUAVResident(EMetalShaderStages const Frequency, FMe
             checkNoEntry();
             break;
     }
-	
-	if (UAV->IsTexture())
-	{
-		// @TODO: this needs refactoring.
-		FPlatformAtomics::InterlockedExchange(&ResourceCast(UAV->GetTexture())->Written, 1);
-	}
 }
 
 void FMetalStateCache::IRMakeTextureResident(EMetalShaderStages const Frequency, MTL::Texture* Texture)
@@ -1785,12 +1777,6 @@ void FMetalStateCache::SetShaderUnorderedAccessView(EMetalShaderStages ShaderSta
 			}
 	#endif
 		}
-
-		if (UAV->IsTexture())
-		{
-			// @todo this needs refactoring.
-			FPlatformAtomics::InterlockedExchange(&ResourceCast(UAV->GetTexture())->Written, 1);
-		}
 	}
 }
 
diff --git a/Engine/Source/Runtime/Apple/MetalRHI/Private/MetalStateCache.h b/Engine/Source/Runtime/Apple/MetalRHI/Private/MetalStateCache.h
--- a/Engine/Source/Runtime/Apple/MetalRHI/Private/MetalStateCache.h
+++ b/Engine/Source/Runtime/Apple/MetalRHI/Private/MetalStateCache.h
@@ -44,7 +44,7 @@ struct FMetalDescriptorHeap;
 class FMetalStateCache
 {
 public:
-	FMetalStateCache(FMetalDevice& Device, bool const bInImmediate);
+	FMetalStateCache(FMetalDevice& Device);
 	~FMetalStateCache();
 	
 	/** Reset cached state for reuse */
@@ -202,9 +202,8 @@ public:
 		OverriddenDescriptorHeap = InHeap;
 	}
 	
-private:
 	void ConditionalUpdateBackBuffer(FMetalSurface& Surface);
-	
+private:
 	void SetDepthStencilState(FMetalDepthStencilState* InDepthStencilState);
 	void SetRasterizerState(FMetalRasterizerState* InRasterizerState);
 	
@@ -354,7 +353,6 @@ private:
 	bool bIsRenderTargetActive;
 	bool bHasValidRenderTarget;
 	bool bHasValidColorTarget;
-    bool bImmediate;
 	bool bIsParallelContext = false;
 };
 
diff --git a/Engine/Source/Runtime/Apple/MetalRHI/Private/MetalTexture.cpp b/Engine/Source/Runtime/Apple/MetalRHI/Private/MetalTexture.cpp
--- a/Engine/Source/Runtime/Apple/MetalRHI/Private/MetalTexture.cpp
+++ b/Engine/Source/Runtime/Apple/MetalRHI/Private/MetalTexture.cpp
@@ -546,7 +546,6 @@ FMetalSurface::FMetalSurface(FMetalDevice& MetalDevice, FMetalTextureCreateDesc
 	, Viewport          (nullptr)
 	, TotalTextureSize  (0)
 {
-	FPlatformAtomics::InterlockedExchange(&Written, 0);
 	check(CreateDesc.Extent.X > 0 && CreateDesc.Extent.Y > 0 && CreateDesc.NumMips > 0);
     
 	// the special back buffer surface will be updated in BeginDrawingViewport - no need to set the texture here
@@ -1033,8 +1032,6 @@ void FMetalSurface::UpdateSurfaceAndDestroySourceBuffer(FMetalRHICommandContext*
 	
 	FMetalDynamicRHI::Get().DeferredDelete(Source);
 	
-	FPlatformAtomics::InterlockedExchange(&Written, 1);
-	
 #if STATS
 	GMetalTexturePageOnTime += (FPlatformTime::Cycles64() - Start);
 #endif
@@ -1788,8 +1785,6 @@ static void InternalUpdateTexture2D(FMetalRHICommandContext& Context, FRHITextur
 	{
         Tex->replaceRegion(Region, MipIndex, 0, Buffer->Contents(), SourcePitch, 0);
 	}
-
-	FPlatformAtomics::InterlockedExchange(&Texture->Written, 1);
 }
 
 void FMetalDynamicRHI::RHIUpdateTexture2D(FRHICommandListBase& RHICmdList, FRHITexture* TextureRHI, uint32 MipIndex, const struct FUpdateTextureRegion2D& UpdateRegion, uint32 SourcePitch, const uint8* SourceData)
@@ -1894,8 +1889,6 @@ static void InternalUpdateTexture3D(FMetalRHICommandContext& Context, FRHITextur
 	{
         Tex->replaceRegion(Region, MipIndex, 0, (uint8*)Buffer->Contents(), SourceRowPitch, SourceDepthPitch);
 	}
-
-	FPlatformAtomics::InterlockedExchange(&Texture->Written, 1);
 }
 
 FUpdateTexture3DData FMetalDynamicRHI::RHIBeginUpdateTexture3D(FRHICommandListBase& RHICmdList, FRHITexture* Texture, uint32 MipIndex, const struct FUpdateTextureRegion3D& UpdateRegion)
@@ -1981,6 +1974,10 @@ void FMetalRHICommandContext::RHICopyTexture(FRHITexture* SourceTextureRHI, FRHI
     FMetalSurface* MetalSrcTexture = GetMetalSurfaceFromRHITexture(SourceTextureRHI);
     FMetalSurface* MetalDestTexture = GetMetalSurfaceFromRHITexture(DestTextureRHI);
     
+	// Need to ensure that the src/dest texture are available if they are presentable
+	StateCache.ConditionalUpdateBackBuffer(*MetalSrcTexture);
+	StateCache.ConditionalUpdateBackBuffer(*MetalDestTexture);
+	
     const bool TextureFormatExactMatch = (SourceTextureRHI->GetFormat() == DestTextureRHI->GetFormat());
     const bool TextureFormatCompatible = MetalRHICopyTexutre_IsTextureFormatCompatible(SourceTextureRHI->GetFormat(), DestTextureRHI->GetFormat());
     
diff --git a/Engine/Source/Runtime/Apple/MetalRHI/Public/MetalResources.h b/Engine/Source/Runtime/Apple/MetalRHI/Public/MetalResources.h
--- a/Engine/Source/Runtime/Apple/MetalRHI/Public/MetalResources.h
+++ b/Engine/Source/Runtime/Apple/MetalRHI/Public/MetalResources.h
@@ -299,7 +299,6 @@ public:
 #endif // PLATFORM_SUPPORTS_BINDLESS_RENDERING
 
 	FMetalDevice& Device;
-	int16 volatile Written;
 	uint8 FormatKey;
 
 	//texture used for store actions and binding to shader params
