Disable occlusion feedback and fall back to queries on A8/A9 devices as these GPUs are missing required support

diff --git a/Engine/Source/Runtime/Apple/MetalRHI/Private/MetalRHI.cpp b/Engine/Source/Runtime/Apple/MetalRHI/Private/MetalRHI.cpp
--- a/Engine/Source/Runtime/Apple/MetalRHI/Private/MetalRHI.cpp
+++ b/Engine/Source/Runtime/Apple/MetalRHI/Private/MetalRHI.cpp
@@ -1227,6 +1227,21 @@ FMetalDynamicRHI::FMetalDynamicRHI(ERHIFeatureLevel::Type RequestedFeatureLevel)
 		GRHIGlobals.RayTracing.RequiresInlineRayTracingSBT = true;
 	}
 #endif
+
+#if PLATFORM_IOS || PLATFORM_TVOS
+	// Disable occlusion feedback on devices older than FeatureSet_iOS_GPUFamily3_v2 / MTLFeatureSet_tvOS_GPUFamily2_v1
+	// as buffer writes are not supported from fragment shaders until A10. See UE-348147.
+	static auto CVarOcclusionFeedback = IConsoleManager::Get().FindConsoleVariable(TEXT("r.OcclusionFeedback.Enable"));
+	if (CVarOcclusionFeedback && CVarOcclusionFeedback->GetInt() != 0 && 
+#if PLATFORM_TVOS		
+		!MTLDevice->supportsFeatureSet(MTL::FeatureSet_tvOS_GPUFamily2_v1))
+#else
+		!MTLDevice->supportsFeatureSet(MTL::FeatureSet_iOS_GPUFamily3_v2))
+#endif
+	{
+		CVarOcclusionFeedback->Set(0);
+	}
+#endif
 		
 	GDynamicRHI = this;
 	
