Metal RHI

- Bumping the number of Simultaneous UAV's on Metal platforms

diff --git a/Engine/Source/Runtime/Apple/MetalRHI/Private/MetalPipeline.h b/Engine/Source/Runtime/Apple/MetalRHI/Private/MetalPipeline.h
--- a/Engine/Source/Runtime/Apple/MetalRHI/Private/MetalPipeline.h
+++ b/Engine/Source/Runtime/Apple/MetalRHI/Private/MetalPipeline.h
@@ -19,6 +19,9 @@ typedef __uint128_t FMetalTextureMask;
 #else
 #error "Unsupported Platform!"
 #endif
+
+#define METAL_MAX_BUFFERS 31
+
 typedef uint32 FMetalBufferMask;
 typedef uint16 FMetalSamplerMask;
 
diff --git a/Engine/Source/Runtime/Apple/MetalRHI/Private/MetalRHI.cpp b/Engine/Source/Runtime/Apple/MetalRHI/Private/MetalRHI.cpp
--- a/Engine/Source/Runtime/Apple/MetalRHI/Private/MetalRHI.cpp
+++ b/Engine/Source/Runtime/Apple/MetalRHI/Private/MetalRHI.cpp
@@ -518,9 +518,8 @@ FMetalDynamicRHI::FMetalDynamicRHI(ERHIFeatureLevel::Type RequestedFeatureLevel)
 	}
 	
 	GRHIGlobals.bSupportsBindless = RHIGetRuntimeBindlessConfiguration(GMaxRHIShaderPlatform) != ERHIBindlessConfiguration::Disabled;
-
-	// Bindless is technically unlimited so we set 32 as Max UAV's, < SM5 8 
-	GRHIGlobals.MaxSimultaneousUAVs = GMaxRHIFeatureLevel >= ERHIFeatureLevel::SM6 ? 32 : 8;
+	
+	GRHIGlobals.MaxSimultaneousUAVs = FMath::Min(METAL_MAX_TEXTURES, METAL_MAX_BUFFERS);
 	
 	ValidateTargetedRHIFeatureLevelExists(GMaxRHIShaderPlatform);
 	
diff --git a/Engine/Source/Runtime/Apple/MetalRHI/Public/MetalShaderResources.h b/Engine/Source/Runtime/Apple/MetalRHI/Public/MetalShaderResources.h
--- a/Engine/Source/Runtime/Apple/MetalRHI/Public/MetalShaderResources.h
+++ b/Engine/Source/Runtime/Apple/MetalRHI/Public/MetalShaderResources.h
@@ -8,18 +8,6 @@
 
 #include "CrossCompilerCommon.h"
 
-/**
-* Shader related constants.
-*/
-enum
-{
-	METAL_MAX_UNIFORM_BUFFER_BINDINGS = 12,	// @todo-mobile: Remove me
-	METAL_FIRST_UNIFORM_BUFFER = 0,			// @todo-mobile: Remove me
-	METAL_MAX_COMPUTE_STAGE_UAV_UNITS = 8,	// @todo-mobile: Remove me
-	METAL_UAV_NOT_SUPPORTED_FOR_GRAPHICS_UNIT = -1, // for now, only CS supports UAVs/ images
-	METAL_MAX_BUFFERS = 31,
-};
-
 /**
 * Buffer data-types for MetalRHI & MetalSL
 */
