Fixed a bug in FMetalExternalSurface where the bindless handle was not bound to the resource

diff --git a/Engine/Source/Runtime/Apple/MetalRHI/Private/MetalTexture.cpp b/Engine/Source/Runtime/Apple/MetalRHI/Private/MetalTexture.cpp
--- a/Engine/Source/Runtime/Apple/MetalRHI/Private/MetalTexture.cpp
+++ b/Engine/Source/Runtime/Apple/MetalRHI/Private/MetalTexture.cpp
@@ -777,7 +777,13 @@ FMetalExternalSurface::FMetalExternalSurface(FMetalDevice& MetalDevice, MTLTextu
 		FMetalBindlessDescriptorManager* BindlessDescriptorManager = Device.GetBindlessDescriptorManager();
 		check(BindlessDescriptorManager);
 
-		BindlessHandle = BindlessDescriptorManager->AllocateDescriptor(ERHIDescriptorType::TextureSRV);
+		check(InTexture->textureType() == MTL::TextureType2D);
+		
+		// Only a single slice as we expect Texture2D only
+		NS::Range Slices(0, 1);
+		Texture = NS::TransferPtr(InTexture->newTextureView(InTexture->pixelFormat(), MTL::TextureType2DArray, NS::Range(0, InTexture->mipmapLevelCount()), Slices));
+		
+		BindlessHandle = BindlessDescriptorManager->AllocateDescriptor(ERHIDescriptorType::TextureSRV, Texture.get());
 	}
 #endif
 }
