Fix for ConvertSurfaceDataToFColor missing F16F and fixing linear support

diff --git a/Engine/Source/Runtime/Apple/MetalRHI/Private/MetalRenderTarget.cpp b/Engine/Source/Runtime/Apple/MetalRHI/Private/MetalRenderTarget.cpp
--- a/Engine/Source/Runtime/Apple/MetalRHI/Private/MetalRenderTarget.cpp
+++ b/Engine/Source/Runtime/Apple/MetalRHI/Private/MetalRenderTarget.cpp
@@ -60,74 +60,13 @@ struct FMetalRGBA16
 	uint16 A;
 };
 
-void FMetalDynamicRHI::RHIReadSurfaceData(FRHITexture* TextureRHI, FIntRect InRect, TArray<FLinearColor>& OutData, FReadSurfaceDataFlags InFlags)
-{
-	// Use our current surface read implementation and convert to linear - should refactor to make optimal
-	TArray<FColor> OutDataUnConverted;
-	RHIReadSurfaceData(TextureRHI, InRect, OutDataUnConverted, InFlags);
-
-	OutData.SetNumUninitialized(OutDataUnConverted.Num());
-
-	for (uint32 i = 0; i < OutDataUnConverted.Num(); ++i)
-	{
-		OutData[i] = OutDataUnConverted[i].ReinterpretAsLinear();
-	}
-}
+template <typename...> constexpr std::false_type AlwaysFalse;
 
-static void ConvertSurfaceDataToFColor(EPixelFormat Format, uint32 Width, uint32 Height, uint8 *In, uint32 SrcPitch, FColor* Out, FReadSurfaceDataFlags InFlags)
+template<typename ColorType>
+void ReadSurfaceDataInternal(FMetalDevice* Device, FRHITexture* TextureRHI, FIntRect Rect, TArray<ColorType>& OutData, FReadSurfaceDataFlags InFlags)
 {
-	bool bLinearToGamma = InFlags.GetLinearToGamma();
-	if (Format == PF_G16 || Format == PF_R16_UINT || Format == PF_R16_SINT)
-	{
-		ConvertRawR16DataToFColor(Width, Height, In, SrcPitch, Out);
-	}
-	else if (Format == PF_R8G8B8A8)
-	{
-		ConvertRawR8G8B8A8DataToFColor(Width, Height, In, SrcPitch, Out);
-	}
-	else if (Format == PF_B8G8R8A8)
-	{
-		ConvertRawB8G8R8A8DataToFColor(Width, Height, In, SrcPitch, Out);
-	}
-	else if (Format == PF_A2B10G10R10)
-	{
-		ConvertRawR10G10B10A2DataToFColor(Width, Height, In, SrcPitch, Out);
-	}
-	else if (Format == PF_FloatRGBA || Format == PF_PLATFORM_HDR_0)
-	{
-		ConvertRawR16G16B16A16FDataToFColor(Width, Height, In, SrcPitch, Out, bLinearToGamma);
-	}
-	else if (Format == PF_FloatR11G11B10)
-	{
-		ConvertRawR11G11B10DataToFColor(Width, Height, In, SrcPitch, Out, bLinearToGamma);
-	}
-	else if (Format == PF_A32B32G32R32F)
-	{
-		ConvertRawR32G32B32A32DataToFColor(Width, Height, In, SrcPitch, Out, bLinearToGamma);
-	}
-	else if (Format == PF_A16B16G16R16)
-	{
-		ConvertRawR16G16B16A16DataToFColor(Width, Height, In, SrcPitch, Out, bLinearToGamma);
-	}
-	else if (Format == PF_G16R16)
-	{
-		ConvertRawR16G16DataToFColor(Width, Height, In, SrcPitch, Out);
-	}
-	else if (Format == PF_DepthStencil)
-	{
-		ConvertRawD32S8DataToFColor(Width, Height, In, SrcPitch, Out, InFlags);
-	}
-	else
-	{
-		// not supported yet
-		UE_LOG(LogMetal, Fatal, TEXT("ConvertSurfaceDataToFColor doesn't support format %s"), GetPixelFormatString(Format));
-	}
-}
-
-void FMetalDynamicRHI::RHIReadSurfaceData(FRHITexture* TextureRHI, FIntRect Rect, TArray<FColor>& OutData, FReadSurfaceDataFlags InFlags)
-{
-    MTL_SCOPED_AUTORELEASE_POOL;
-    
+	MTL_SCOPED_AUTORELEASE_POOL;
+	
 	// allocate output space
 	const uint32 SizeX = Rect.Width();
 	const uint32 SizeY = Rect.Height();
@@ -141,20 +80,20 @@ void FMetalDynamicRHI::RHIReadSurfaceData(FRHITexture* TextureRHI, FIntRect Rect
 
 	FMetalSurface* Surface = GetMetalSurfaceFromRHITexture(TextureRHI);
 
-	FColor* OutDataPtr = OutData.GetData();
+	ColorType* OutDataPtr = OutData.GetData();
 	MTL::Region Region(Rect.Min.X, Rect.Min.Y, SizeX, SizeY);
-    
+	
 	MTLTexturePtr Texture = Surface->Texture;
-    if(!Texture && EnumHasAnyFlags(Surface->GetDesc().Flags, TexCreate_Presentable))
-    {
-        Texture = Surface->GetCurrentTexture();
-    }
+	if(!Texture && EnumHasAnyFlags(Surface->GetDesc().Flags, TexCreate_Presentable))
+	{
+		Texture = Surface->GetCurrentTexture();
+	}
 	
-    if(!Texture)
-    {
-        UE_LOG(LogRHI, Error, TEXT("Trying to read from an uninitialised texture."));
-        return;
-    }
+	if(!Texture)
+	{
+		UE_LOG(LogRHI, Error, TEXT("Trying to read from an uninitialised texture."));
+		return;
+	}
 
 	if (GMetalUseTexGetBytes && Surface->GetDesc().Format != PF_DepthStencil && Surface->GetDesc().Format != PF_ShadowDepth)
 	{
@@ -163,7 +102,7 @@ void FMetalDynamicRHI::RHIReadSurfaceData(FRHITexture* TextureRHI, FIntRect Rect
 		FRHICommandListImmediate& RHICmdList = FRHICommandListImmediate::Get();
 		MTLTexturePtr TempTexture;
 		
-		RHICmdList.EnqueueLambda([this, &Texture, &TempTexture, SizeX, SizeY, &Region, Surface, InFlags](FRHICommandListImmediate& RHICmdList)
+		RHICmdList.EnqueueLambda([Device, &Texture, &TempTexture, SizeX, SizeY, &Region, Surface, InFlags](FRHICommandListImmediate& RHICmdList)
 		{
 			FMetalRHICommandContext& Context = FMetalRHICommandContext::Get(RHICmdList);
 			
@@ -226,8 +165,20 @@ void FMetalDynamicRHI::RHIReadSurfaceData(FRHITexture* TextureRHI, FIntRect Rect
 		
 		Texture->getBytes(Data.GetData(), Stride, BytesPerImage, Region, 0, 0);
 		
-		ConvertSurfaceDataToFColor(Surface->GetDesc().Format, SizeX, SizeY, (uint8*)Data.GetData(), Stride, OutDataPtr, InFlags);
-		
+		if constexpr (std::is_same_v<ColorType, FColor>)
+		{
+			ConvertRAWSurfaceDataToFColor(TextureRHI->GetFormat(), SizeX, SizeY, (uint8*)Data.GetData(), Stride, OutDataPtr, InFlags);
+		}
+		else if constexpr (std::is_same_v<ColorType, FLinearColor>)
+		{
+			ConvertRAWSurfaceDataToFLinearColor(TextureRHI->GetFormat(), SizeX, SizeY, (uint8*)Data.GetData(), Stride, OutDataPtr, InFlags);
+		}
+		else
+		{
+			// Color type not recognized.
+			static_assert(AlwaysFalse<ColorType>);
+		}
+
 		if (TempTexture)
 		{
 			FMetalDynamicRHI::Get().DeferredDelete(TempTexture);
@@ -245,7 +196,7 @@ void FMetalDynamicRHI::RHIReadSurfaceData(FRHITexture* TextureRHI, FIntRect Rect
 		
 		FMetalBufferPtr Buffer = Device->CreatePooledBuffer(FMetalPooledBufferArgs(Device, BytesPerImage, BUF_Dynamic, MTL::StorageModeShared));
 		
-		RHICmdList.EnqueueLambda([this, &Buffer, &Texture, &Region, Surface, AlignedStride, BytesPerImage, InFlags](FRHICommandListImmediate& RHICmdList)
+		RHICmdList.EnqueueLambda([Device, &Buffer, &Texture, &Region, Surface, AlignedStride, BytesPerImage, InFlags](FRHICommandListImmediate& RHICmdList)
 		{
 			// Synchronise the texture with the CPU
 			SCOPE_CYCLE_COUNTER(STAT_MetalTexturePageOffTime);
@@ -272,12 +223,34 @@ void FMetalDynamicRHI::RHIReadSurfaceData(FRHITexture* TextureRHI, FIntRect Rect
 		//kick the current command buffer.
 		RHICmdList.SubmitAndBlockUntilGPUIdle();
 		
-		ConvertSurfaceDataToFColor(Surface->GetDesc().Format, SizeX, SizeY, (uint8*)Buffer->Contents(), AlignedStride, OutDataPtr, InFlags);
+		if constexpr (std::is_same_v<ColorType, FColor>)
+		{
+			ConvertRAWSurfaceDataToFColor(TextureRHI->GetFormat(), SizeX, SizeY, (uint8*)Buffer->Contents(), AlignedStride, OutDataPtr, InFlags);
+		}
+		else if constexpr (std::is_same_v<ColorType, FLinearColor>)
+		{
+			ConvertRAWSurfaceDataToFLinearColor(TextureRHI->GetFormat(), SizeX, SizeY, (uint8*)Buffer->Contents(), AlignedStride, OutDataPtr, InFlags);
+		}
+		else
+		{
+			// Color type not recognized.
+			static_assert(AlwaysFalse<ColorType>);
+		}
 		
 		FMetalDynamicRHI::Get().DeferredDelete(Buffer);
 	}
 }
 
+void FMetalDynamicRHI::RHIReadSurfaceData(FRHITexture* TextureRHI, FIntRect InRect, TArray<FLinearColor>& OutData, FReadSurfaceDataFlags InFlags)
+{
+	ReadSurfaceDataInternal(Device, TextureRHI, InRect, OutData, InFlags);
+}
+
+void FMetalDynamicRHI::RHIReadSurfaceData(FRHITexture* TextureRHI, FIntRect Rect, TArray<FColor>& OutData, FReadSurfaceDataFlags InFlags)
+{
+	ReadSurfaceDataInternal(Device, TextureRHI, Rect, OutData, InFlags);
+}
+
 #if PLATFORM_IOS
 void FMetalDynamicRHI::RHIReadSurfaceDataDirect(FRHITexture* TextureRHI, FIntRect Rect, TArray<FColor>& OutData)
 {   
