Adjust GRHIMaxDispatchThreadGroupsPerDimension to limits from specifications.

#tests verified ensure is not hit anymore with specific model ran with NNERuntimeIREERdg

#rnx

diff --git a/Engine/Source/Runtime/Apple/MetalRHI/Private/MetalRHI.cpp b/Engine/Source/Runtime/Apple/MetalRHI/Private/MetalRHI.cpp
--- a/Engine/Source/Runtime/Apple/MetalRHI/Private/MetalRHI.cpp
+++ b/Engine/Source/Runtime/Apple/MetalRHI/Private/MetalRHI.cpp
@@ -750,9 +750,9 @@ FMetalDynamicRHI::FMetalDynamicRHI(ERHIFeatureLevel::Type RequestedFeatureLevel)
         GRHISupportsArrayIndexFromAnyShader = true;
     }
             
-	GRHIMaxDispatchThreadGroupsPerDimension.X = MAX_uint16;
-	GRHIMaxDispatchThreadGroupsPerDimension.Y = MAX_uint16;
-	GRHIMaxDispatchThreadGroupsPerDimension.Z = MAX_uint16;
+	GRHIMaxDispatchThreadGroupsPerDimension.X = MAX_int32;
+	GRHIMaxDispatchThreadGroupsPerDimension.Y = MAX_int32;
+	GRHIMaxDispatchThreadGroupsPerDimension.Z = MAX_int32;
 
 	GMaxTextureMipCount = FPlatformMath::CeilLogTwo( GMaxTextureDimensions ) + 1;
 	GMaxTextureMipCount = FPlatformMath::Min<int32>( MAX_TEXTURE_MIP_COUNT, GMaxTextureMipCount );
