Remove RHIDrawIndexedIndirect command
- This was not implemented on some platforms, and not exposed in the RHI command list API.

diff --git a/Engine/Source/Runtime/Apple/MetalRHI/Private/MetalCommands.cpp b/Engine/Source/Runtime/Apple/MetalRHI/Private/MetalCommands.cpp
--- a/Engine/Source/Runtime/Apple/MetalRHI/Private/MetalCommands.cpp
+++ b/Engine/Source/Runtime/Apple/MetalRHI/Private/MetalCommands.cpp
@@ -1172,73 +1172,6 @@ void FMetalRHICommandContext::RHIDrawIndexedPrimitive(FRHIBuffer* IndexBufferRHI
 	}
 }
 
-void FMetalRHICommandContext::RHIDrawIndexedIndirect(FRHIBuffer* IndexBufferRHI, FRHIBuffer* ArgumentsBufferRHI, int32 DrawArgumentsIndex, uint32 /*NumInstances*/)
-{
-    MTL_SCOPED_AUTORELEASE_POOL;
-    
-	if (Device.SupportsFeature(EMetalFeaturesIndirectBuffer))
-	{
-		SCOPE_CYCLE_COUNTER(STAT_MetalDrawCallTime);
-		check(CurrentEncoder.IsRenderCommandEncoderActive());
-		
-		EPrimitiveType PrimitiveType = StateCache.GetPrimitiveType();
-		
-		RHI_DRAW_CALL_INC();
-		FMetalRHIBuffer* IndexBuffer = ResourceCast(IndexBufferRHI);
-		FMetalRHIBuffer* ArgumentsBuffer = ResourceCast(ArgumentsBufferRHI);
-		
-		if(!PrepareToDraw(PrimitiveType))
-		{
-			return;
-		}
-
-		FMetalBufferPtr TheBackingIndexBuffer = IndexBuffer->GetCurrentBuffer();
-		FMetalBufferPtr TheBackingBuffer = ArgumentsBuffer->GetCurrentBuffer();
-		
-		check(TheBackingIndexBuffer);
-		check(TheBackingBuffer);
-		
-		// finalize any pending state
-		PrepareToRender(PrimitiveType);
-		
-#if METAL_USE_METAL_SHADER_CONVERTER
-		if(IsMetalBindlessEnabled())
-		{			
-			// TODO: Carl - Remove this when API validation is fixed
-			// Binding to uniforms bind point to work around error in API validation
-			uint32 NullBuffer = 0x0;
-			CurrentEncoder.GetRenderCommandEncoder()->setVertexBytes(&NullBuffer, sizeof(uint32), kIRArgumentBufferUniformsBindPoint);
-			
-			StateCache.IRMapVertexBuffers(CurrentEncoder.GetRenderCommandEncoder());
-			CurrentEncoder.GetRenderCommandEncoder()->useResource(TheBackingBuffer->GetMTLBuffer(), MTL::ResourceUsageRead);
-			CurrentEncoder.GetRenderCommandEncoder()->useResource(TheBackingIndexBuffer->GetMTLBuffer(), MTL::ResourceUsageRead);
-			IRRuntimeDrawIndexedPrimitives(CurrentEncoder.GetRenderCommandEncoder(),
-										   TranslatePrimitiveType(PrimitiveType),
-										   IndexBuffer->GetIndexType(),
-										   TheBackingIndexBuffer->GetMTLBuffer(),
-										   TheBackingIndexBuffer->GetOffset(),
-										   TheBackingBuffer->GetMTLBuffer(),
-										   TheBackingBuffer->GetOffset() + (DrawArgumentsIndex * 5 * sizeof(uint32)));
-		}
-		else
-#endif
-		{
-			METAL_GPUPROFILE(FMetalProfiler::GetProfiler()->EncodeDraw(CurrentEncoder.GetCommandBufferStats(), __FUNCTION__, 1, 1, 1));
-			
-			CurrentEncoder.GetRenderCommandEncoder()->drawIndexedPrimitives(TranslatePrimitiveType(PrimitiveType),
-																			IndexBuffer->GetIndexType(), 
-																			TheBackingIndexBuffer->GetMTLBuffer(), 
-																			TheBackingIndexBuffer->GetOffset(),
-																			TheBackingBuffer->GetMTLBuffer(), 
-																			TheBackingBuffer->GetOffset() + (DrawArgumentsIndex * 5 * sizeof(uint32)));
-		}
-	}
-	else
-	{
-		NOT_SUPPORTED("RHIDrawIndexedIndirect");
-	}
-}
-
 void FMetalRHICommandContext::RHIDrawIndexedPrimitiveIndirect(FRHIBuffer* IndexBufferRHI, FRHIBuffer* ArgumentBufferRHI, uint32 ArgumentOffset)
 {
     MTL_SCOPED_AUTORELEASE_POOL;
diff --git a/Engine/Source/Runtime/Apple/MetalRHI/Public/MetalRHIContext.h b/Engine/Source/Runtime/Apple/MetalRHI/Public/MetalRHIContext.h
--- a/Engine/Source/Runtime/Apple/MetalRHI/Public/MetalRHIContext.h
+++ b/Engine/Source/Runtime/Apple/MetalRHI/Public/MetalRHIContext.h
@@ -129,8 +129,6 @@ public:
 	
 	virtual void RHIDrawPrimitiveIndirect(FRHIBuffer* ArgumentBuffer, uint32 ArgumentOffset) final override;
 	
-	virtual void RHIDrawIndexedIndirect(FRHIBuffer* IndexBufferRHI, FRHIBuffer* ArgumentsBufferRHI, int32 DrawArgumentsIndex, uint32 NumInstances) final override;
-	
 	// @param NumPrimitives need to be >0
 	virtual void RHIDrawIndexedPrimitive(FRHIBuffer* IndexBuffer, int32 BaseVertexIndex, uint32 FirstInstance, uint32 NumVertices, uint32 StartIndex, uint32 NumPrimitives, uint32 NumInstances) final override;
 	
diff --git a/Engine/Source/Runtime/D3D12RHI/Private/D3D12CommandContext.h b/Engine/Source/Runtime/D3D12RHI/Private/D3D12CommandContext.h
--- a/Engine/Source/Runtime/D3D12RHI/Private/D3D12CommandContext.h
+++ b/Engine/Source/Runtime/D3D12RHI/Private/D3D12CommandContext.h
@@ -702,7 +702,6 @@ public:
 
 	virtual void RHIDrawPrimitive(uint32 BaseVertexIndex, uint32 NumPrimitives, uint32 NumInstances) final override;
 	virtual void RHIDrawPrimitiveIndirect(FRHIBuffer* ArgumentBuffer, uint32 ArgumentOffset) final override;
-	virtual void RHIDrawIndexedIndirect(FRHIBuffer* IndexBufferRHI, FRHIBuffer* ArgumentsBufferRHI, int32 DrawArgumentsIndex, uint32 NumInstances) final override;
 	virtual void RHIDrawIndexedPrimitive(FRHIBuffer* IndexBuffer, int32 BaseVertexIndex, uint32 FirstInstance, uint32 NumVertices, uint32 StartIndex, uint32 NumPrimitives, uint32 NumInstances) final override;
 	virtual void RHIDrawIndexedPrimitiveIndirect(FRHIBuffer* IndexBuffer, FRHIBuffer* ArgumentBuffer, uint32 ArgumentOffset) final override;
 	virtual void RHIMultiDrawIndexedPrimitiveIndirect(FRHIBuffer* IndexBuffer, FRHIBuffer* ArgumentBuffer, uint32 ArgumentOffset, FRHIBuffer* CountBuffer, uint32 CountBufferOffset, uint32 MaxDrawArguments) final override;
@@ -1046,10 +1045,6 @@ public:
 	{
 		ContextRedirect(RHIDrawPrimitiveIndirect(ArgumentBuffer, ArgumentOffset));
 	}
-	FORCEINLINE virtual void RHIDrawIndexedIndirect(FRHIBuffer* IndexBufferRHI, FRHIBuffer* ArgumentsBufferRHI, int32 DrawArgumentsIndex, uint32 NumInstances) final override
-	{
-		ContextRedirect(RHIDrawIndexedIndirect(IndexBufferRHI, ArgumentsBufferRHI, DrawArgumentsIndex, NumInstances));
-	}
 	FORCEINLINE virtual void RHIDrawIndexedPrimitive(FRHIBuffer* IndexBuffer, int32 BaseVertexIndex, uint32 FirstInstance, uint32 NumVertices, uint32 StartIndex, uint32 NumPrimitives, uint32 NumInstances) final override
 	{
 		ContextRedirect(RHIDrawIndexedPrimitive(IndexBuffer, BaseVertexIndex, FirstInstance, NumVertices, StartIndex, NumPrimitives, NumInstances));
diff --git a/Engine/Source/Runtime/D3D12RHI/Private/D3D12Commands.cpp b/Engine/Source/Runtime/D3D12RHI/Private/D3D12Commands.cpp
--- a/Engine/Source/Runtime/D3D12RHI/Private/D3D12Commands.cpp
+++ b/Engine/Source/Runtime/D3D12RHI/Private/D3D12Commands.cpp
@@ -1249,26 +1249,6 @@ void FD3D12CommandContext::RHIDrawPrimitiveIndirect(FRHIBuffer* ArgumentBufferRH
 	PostGpuEvent();
 }
 
-void FD3D12CommandContext::RHIDrawIndexedIndirect(FRHIBuffer* IndexBufferRHI, FRHIBuffer* ArgumentBufferRHI, int32 DrawArgumentsIndex, uint32 /*NumInstances*/)
-{
-	RHI_DRAW_CALL_INC();
-
-	SetupDraw(IndexBufferRHI, 1);
-
-	FD3D12ResourceLocation& ArgumentBufferLocation = SetupIndirectArgument(ArgumentBufferRHI);
-
-	GraphicsCommandList()->ExecuteIndirect(
-		GetParentDevice()->GetParentAdapter()->GetDrawIndexedIndirectCommandSignature(),
-		1,
-		ArgumentBufferLocation.GetResource()->GetResource(),
-		ArgumentBufferLocation.GetOffsetFromBaseOfResource() + DrawArgumentsIndex * ArgumentBufferRHI->GetStride(),
-		NULL,
-		0
-	);
-
-	PostGpuEvent();
-}
-
 void FD3D12CommandContext::RHIDrawIndexedPrimitive(FRHIBuffer* IndexBufferRHI, int32 BaseVertexIndex, uint32 FirstInstance, uint32 NumVertices, uint32 StartIndex, uint32 NumPrimitives, uint32 NumInstances)
 {
 	FD3D12Buffer* IndexBuffer = RetrieveObject<FD3D12Buffer>(IndexBufferRHI);
diff --git a/Engine/Source/Runtime/NullDrv/Public/NullRHI.h b/Engine/Source/Runtime/NullDrv/Public/NullRHI.h
--- a/Engine/Source/Runtime/NullDrv/Public/NullRHI.h
+++ b/Engine/Source/Runtime/NullDrv/Public/NullRHI.h
@@ -517,32 +517,22 @@ public:
 
 	virtual void RHIDrawPrimitive(uint32 BaseVertexIndex, uint32 NumPrimitives, uint32 NumInstances) final override
 	{
-
 	}
+
 	virtual void RHIDrawPrimitiveIndirect(FRHIBuffer* ArgumentBuffer, uint32 ArgumentOffset) final override
 	{
-
 	}
 
-	virtual void RHIDrawIndexedIndirect(FRHIBuffer* IndexBufferRHI, FRHIBuffer* ArgumentsBufferRHI, int32 DrawArgumentsIndex, uint32 NumInstances) final override
-	{
-
-	}
-
-
 	virtual void RHIDrawIndexedPrimitive(FRHIBuffer* IndexBuffer, int32 BaseVertexIndex, uint32 FirstInstance, uint32 NumVertices, uint32 StartIndex, uint32 NumPrimitives, uint32 NumInstances) final override
 	{
-
 	}
 
 	virtual void RHIDrawIndexedPrimitiveIndirect(FRHIBuffer* IndexBuffer, FRHIBuffer* ArgumentBuffer, uint32 ArgumentOffset) final override
 	{
-
 	}
 
 	virtual void RHIMultiDrawIndexedPrimitiveIndirect(FRHIBuffer* IndexBuffer, FRHIBuffer* ArgumentBuffer, uint32 ArgumentOffset, FRHIBuffer* CountBuffer, uint32 CountBufferOffset, uint32 MaxDrawArguments) final override
 	{
-
 	}
 
 	virtual void RHIBlockUntilGPUIdle() final override
diff --git a/Engine/Source/Runtime/OpenGLDrv/Private/OpenGLCommands.cpp b/Engine/Source/Runtime/OpenGLDrv/Private/OpenGLCommands.cpp
--- a/Engine/Source/Runtime/OpenGLDrv/Private/OpenGLCommands.cpp
+++ b/Engine/Source/Runtime/OpenGLDrv/Private/OpenGLCommands.cpp
@@ -2279,44 +2279,6 @@ void FOpenGLDynamicRHI::RHIDrawPrimitiveIndirect(FRHIBuffer* ArgumentBufferRHI,
 	KickHint.OnDrawCall();
 }
 
-void FOpenGLDynamicRHI::RHIDrawIndexedIndirect(FRHIBuffer* IndexBufferRHI, FRHIBuffer* ArgumentsBufferRHI, int32 DrawArgumentsIndex, uint32 /*NumInstances*/)
-{
-	if (!FOpenGL::SupportsDrawIndirect())
-	{
-		UE_LOG(LogRHI, Fatal, TEXT("OpenGL RHI does not yet support indirect draw calls."));
-		return;
-	}
-
-	VERIFY_GL_SCOPE();
-#if (RHI_NEW_GPU_PROFILER == 0)
-	RegisterGPUWork(1);
-#endif
-
-	FOpenGLBuffer* IndexBuffer = ResourceCast(IndexBufferRHI);
-	SetupDraw(IndexBuffer, 0, 1);
-
-	GLenum DrawMode = GL_TRIANGLES;
-	GLsizei NumElements = 0;
-	FindPrimitiveType(PrimitiveType, 0, DrawMode, NumElements);
-
-	GLenum IndexType = IndexBuffer->GetStride() == sizeof(uint32) ? GL_UNSIGNED_INT : GL_UNSIGNED_SHORT;
-
-	check(ArgumentsBufferRHI);
-	FOpenGLBuffer* ArgumentsBuffer = ResourceCast(ArgumentsBufferRHI);
-	glBindBuffer(GL_DRAW_INDIRECT_BUFFER, ArgumentsBuffer->Resource);
-
-	{
-		CONDITIONAL_SCOPE_CYCLE_COUNTER(STAT_OpenGLShaderFirstDrawTime, PendingState.BoundShaderState->RequiresDriverInstantiation());
-	
-		// Offset is based on an index into the list of structures
-		FOpenGL::DrawElementsIndirect(DrawMode, IndexType, INDEX_TO_VOID(DrawArgumentsIndex * 5 * sizeof(uint32)));
-	}
-
-	glBindBuffer(GL_DRAW_INDIRECT_BUFFER, 0);
-
-	KickHint.OnDrawCall();
-}
-
 void FOpenGLDynamicRHI::RHIDrawIndexedPrimitive(FRHIBuffer* IndexBufferRHI, int32 BaseVertexIndex, uint32 FirstInstance, uint32 NumVertices, uint32 StartIndex, uint32 NumPrimitives, uint32 NumInstances)
 {
 	SCOPE_CYCLE_COUNTER_DETAILED(STAT_OpenGLDrawPrimitiveTime);
diff --git a/Engine/Source/Runtime/OpenGLDrv/Private/OpenGLDrv.h b/Engine/Source/Runtime/OpenGLDrv/Private/OpenGLDrv.h
--- a/Engine/Source/Runtime/OpenGLDrv/Private/OpenGLDrv.h
+++ b/Engine/Source/Runtime/OpenGLDrv/Private/OpenGLDrv.h
@@ -437,7 +437,6 @@ public:
 
 	UE_API virtual void RHIDrawPrimitive(uint32 BaseVertexIndex, uint32 NumPrimitives, uint32 NumInstances) final override;
 	UE_API virtual void RHIDrawPrimitiveIndirect(FRHIBuffer* ArgumentBuffer, uint32 ArgumentOffset) final override;
-	UE_API virtual void RHIDrawIndexedIndirect(FRHIBuffer* IndexBufferRHI, FRHIBuffer* ArgumentsBufferRHI, int32 DrawArgumentsIndex, uint32 NumInstances) final override;
 	UE_API virtual void RHIDrawIndexedPrimitive(FRHIBuffer* IndexBuffer, int32 BaseVertexIndex, uint32 FirstInstance, uint32 NumVertices, uint32 StartIndex, uint32 NumPrimitives, uint32 NumInstances) final override;
 	UE_API virtual void RHIDrawIndexedPrimitiveIndirect(FRHIBuffer* IndexBuffer, FRHIBuffer* ArgumentBuffer, uint32 ArgumentOffset) final override;
 	UE_API virtual void RHIEnableDepthBoundsTest(bool bEnable) final override;
diff --git a/Engine/Source/Runtime/RHI/Public/RHICommandList.h b/Engine/Source/Runtime/RHI/Public/RHICommandList.h
--- a/Engine/Source/Runtime/RHI/Public/RHICommandList.h
+++ b/Engine/Source/Runtime/RHI/Public/RHICommandList.h
@@ -2166,23 +2166,6 @@ FRHICOMMAND_MACRO(FRHICommandDrawPrimitiveIndirect)
 	RHI_EXECUTE_API void Execute(FRHICommandListBase& CmdList);
 };
 
-FRHICOMMAND_MACRO(FRHICommandDrawIndexedIndirect)
-{
-	FRHIBuffer* IndexBufferRHI;
-	FRHIBuffer* ArgumentsBufferRHI;
-	uint32 DrawArgumentsIndex;
-	uint32 NumInstances;
-
-	inline FRHICommandDrawIndexedIndirect(FRHIBuffer* InIndexBufferRHI, FRHIBuffer* InArgumentsBufferRHI, uint32 InDrawArgumentsIndex, uint32 InNumInstances)
-		: IndexBufferRHI(InIndexBufferRHI)
-		, ArgumentsBufferRHI(InArgumentsBufferRHI)
-		, DrawArgumentsIndex(InDrawArgumentsIndex)
-		, NumInstances(InNumInstances)
-	{
-	}
-	RHI_EXECUTE_API void Execute(FRHICommandListBase& CmdList);
-};
-
 FRHICOMMAND_MACRO(FRHICommandDrawIndexedPrimitiveIndirect)
 {
 	FRHIBuffer* IndexBuffer;
diff --git a/Engine/Source/Runtime/RHI/Public/RHICommandListCommandExecutes.inl b/Engine/Source/Runtime/RHI/Public/RHICommandListCommandExecutes.inl
--- a/Engine/Source/Runtime/RHI/Public/RHICommandListCommandExecutes.inl
+++ b/Engine/Source/Runtime/RHI/Public/RHICommandListCommandExecutes.inl
@@ -22,7 +22,6 @@ struct FRHICommandBeginRenderQuery;
 struct FRHICommandBuildLocalBoundShaderState;
 struct FRHICommandBuildLocalGraphicsPipelineState;
 struct FRHICommandBuildLocalUniformBuffer;
-struct FRHICommandDrawIndexedIndirect;
 struct FRHICommandDrawIndexedPrimitive;
 struct FRHICommandDrawIndexedPrimitiveIndirect;
 struct FRHICommandDrawPrimitive;
@@ -305,12 +304,6 @@ void FRHICommandDrawPrimitiveIndirect::Execute(FRHICommandListBase& CmdList)
 	INTERNAL_DECORATOR(RHIDrawPrimitiveIndirect)(ArgumentBuffer, ArgumentOffset);
 }
 
-void FRHICommandDrawIndexedIndirect::Execute(FRHICommandListBase& CmdList)
-{
-	RHISTAT(DrawIndexedIndirect);
-	INTERNAL_DECORATOR(RHIDrawIndexedIndirect)(IndexBufferRHI, ArgumentsBufferRHI, DrawArgumentsIndex, NumInstances);
-}
-
 void FRHICommandDrawIndexedPrimitiveIndirect::Execute(FRHICommandListBase& CmdList)
 {
 	RHISTAT(DrawIndexedPrimitiveIndirect);
diff --git a/Engine/Source/Runtime/RHI/Public/RHIContext.h b/Engine/Source/Runtime/RHI/Public/RHIContext.h
--- a/Engine/Source/Runtime/RHI/Public/RHIContext.h
+++ b/Engine/Source/Runtime/RHI/Public/RHIContext.h
@@ -791,8 +791,6 @@ public:
 
 	virtual void RHIDrawPrimitiveIndirect(FRHIBuffer* ArgumentBuffer, uint32 ArgumentOffset) = 0;
 
-	virtual void RHIDrawIndexedIndirect(FRHIBuffer* IndexBufferRHI, FRHIBuffer* ArgumentsBufferRHI, int32 DrawArgumentsIndex, uint32 NumInstances) = 0;
-
 	// @param NumPrimitives need to be >0 
 	virtual void RHIDrawIndexedPrimitive(FRHIBuffer* IndexBuffer, int32 BaseVertexIndex, uint32 FirstInstance, uint32 NumVertices, uint32 StartIndex, uint32 NumPrimitives, uint32 NumInstances) = 0;
 
diff --git a/Engine/Source/Runtime/RHI/Public/RHIValidationContext.h b/Engine/Source/Runtime/RHI/Public/RHIValidationContext.h
--- a/Engine/Source/Runtime/RHI/Public/RHIValidationContext.h
+++ b/Engine/Source/Runtime/RHI/Public/RHIValidationContext.h
@@ -952,16 +952,6 @@ public:
 		Tracker->Draw();
 	}
 
-	virtual void RHIDrawIndexedIndirect(FRHIBuffer* IndexBufferRHI, FRHIBuffer* ArgumentsBufferRHI, int32 DrawArgumentsIndex, uint32 NumInstances) override final
-	{
-		ValidateDrawing();
-		FValidationRHI::ValidateIndirectArgsBuffer(ArgumentsBufferRHI, DrawArgumentsIndex * ArgumentsBufferRHI->GetStride(), sizeof(FRHIDrawIndexedIndirectParameters), 0);
-		Tracker->Assert(ArgumentsBufferRHI->GetWholeResourceIdentity(), ERHIAccess::IndirectArgs);
-		Tracker->Assert(IndexBufferRHI->GetWholeResourceIdentity(), ERHIAccess::VertexOrIndexBuffer);
-		RHIContext->RHIDrawIndexedIndirect(IndexBufferRHI, ArgumentsBufferRHI, DrawArgumentsIndex, NumInstances);
-		Tracker->Draw();
-	}
-
 	// @param NumPrimitives need to be >0 
 	virtual void RHIDrawIndexedPrimitive(FRHIBuffer* IndexBuffer, int32 BaseVertexIndex, uint32 FirstInstance, uint32 NumVertices, uint32 StartIndex, uint32 NumPrimitives, uint32 NumInstances) override final
 	{
diff --git a/Engine/Source/Runtime/VulkanRHI/Private/VulkanCommands.cpp b/Engine/Source/Runtime/VulkanRHI/Private/VulkanCommands.cpp
--- a/Engine/Source/Runtime/VulkanRHI/Private/VulkanCommands.cpp
+++ b/Engine/Source/Runtime/VulkanRHI/Private/VulkanCommands.cpp
@@ -715,34 +715,6 @@ void FVulkanCommandListContext::RHIDrawIndexedPrimitive(FRHIBuffer* IndexBufferR
 #endif
 }
 
-void FVulkanCommandListContext::RHIDrawIndexedIndirect(FRHIBuffer* IndexBufferRHI, FRHIBuffer* ArgumentsBufferRHI, int32 DrawArgumentsIndex, uint32 /*NumInstances*/)
-{
-#if VULKAN_ENABLE_AGGRESSIVE_STATS
-	SCOPE_CYCLE_COUNTER(STAT_VulkanDrawCallTime);
-#endif
-	RHI_DRAW_CALL_INC();
-
-	CommitGraphicsResourceTables();
-
-	PendingGfxState->PrepareForDraw(*this);
-
-	FVulkanBuffer* IndexBuffer = ResourceCast(IndexBufferRHI);
-	const VkCommandBuffer CommandBufferHandle = GetCommandBuffer().GetHandle();
-	VulkanRHI::vkCmdBindIndexBuffer(CommandBufferHandle, IndexBuffer->GetHandle(), IndexBuffer->GetOffset(), IndexBuffer->GetIndexType());
-
-	FVulkanBuffer* ArgumentBuffer = ResourceCast(ArgumentsBufferRHI);
-	VkDeviceSize ArgumentOffset = DrawArgumentsIndex * sizeof(VkDrawIndexedIndirectCommand);
-
-	VulkanRHI::vkCmdDrawIndexedIndirect(CommandBufferHandle, ArgumentBuffer->GetHandle(), ArgumentBuffer->GetOffset() + ArgumentOffset, 1, sizeof(VkDrawIndexedIndirectCommand));
-
-#if (RHI_NEW_GPU_PROFILER == 0)
-	if (FVulkanPlatform::RegisterGPUWork() && IsImmediate())
-	{
-		RegisterGPUWork(1);
-	}
-#endif
-}
-
 void FVulkanCommandListContext::RHIDrawIndexedPrimitiveIndirect(FRHIBuffer* IndexBufferRHI, FRHIBuffer* ArgumentBufferRHI, uint32 ArgumentOffset)
 {
 #if VULKAN_ENABLE_AGGRESSIVE_STATS
diff --git a/Engine/Source/Runtime/VulkanRHI/Private/VulkanContext.h b/Engine/Source/Runtime/VulkanRHI/Private/VulkanContext.h
--- a/Engine/Source/Runtime/VulkanRHI/Private/VulkanContext.h
+++ b/Engine/Source/Runtime/VulkanRHI/Private/VulkanContext.h
@@ -291,7 +291,6 @@ public:
 	virtual void RHISetStencilRef(uint32 StencilRef) final override;
 	virtual void RHIDrawPrimitive(uint32 BaseVertexIndex, uint32 NumPrimitives, uint32 NumInstances) final override;
 	virtual void RHIDrawPrimitiveIndirect(FRHIBuffer* ArgumentBuffer, uint32 ArgumentOffset) final override;
-	virtual void RHIDrawIndexedIndirect(FRHIBuffer* IndexBufferRHI, FRHIBuffer* ArgumentsBufferRHI, int32 DrawArgumentsIndex, uint32 NumInstances) final override;
 	virtual void RHIDrawIndexedPrimitive(FRHIBuffer* IndexBuffer, int32 BaseVertexIndex, uint32 FirstInstance, uint32 NumVertices, uint32 StartIndex, uint32 NumPrimitives, uint32 NumInstances) final override;
 	virtual void RHIDrawIndexedPrimitiveIndirect(FRHIBuffer* IndexBuffer, FRHIBuffer* ArgumentBuffer, uint32 ArgumentOffset) final override;
 #if PLATFORM_SUPPORTS_MESH_SHADERS
diff --git a/Engine/Source/Runtime/Windows/D3D11RHI/Private/D3D11Commands.cpp b/Engine/Source/Runtime/Windows/D3D11RHI/Private/D3D11Commands.cpp
--- a/Engine/Source/Runtime/Windows/D3D11RHI/Private/D3D11Commands.cpp
+++ b/Engine/Source/Runtime/Windows/D3D11RHI/Private/D3D11Commands.cpp
@@ -1146,32 +1146,6 @@ void FD3D11DynamicRHI::RHIDrawPrimitiveIndirect(FRHIBuffer* ArgumentBufferRHI, u
 	EnableUAVOverlap();
 }
 
-void FD3D11DynamicRHI::RHIDrawIndexedIndirect(FRHIBuffer* IndexBufferRHI, FRHIBuffer* ArgumentsBufferRHI, int32 DrawArgumentsIndex, uint32 NumInstances)
-{
-	FD3D11Buffer* IndexBuffer = ResourceCast(IndexBufferRHI);
-	FD3D11Buffer* ArgumentsBuffer = ResourceCast(ArgumentsBufferRHI);
-
-	RHI_DRAW_CALL_INC();
-
-#if (RHI_NEW_GPU_PROFILER == 0)
-	RegisterGPUWork(1);
-#endif
-
-	CommitGraphicsResourceTables();
-	CommitNonComputeShaderConstants();
-
-	// determine 16bit vs 32bit indices
-	const DXGI_FORMAT Format = (IndexBuffer->GetStride() == sizeof(uint16) ? DXGI_FORMAT_R16_UINT : DXGI_FORMAT_R32_UINT);
-
-	TrackResourceBoundAsIB(IndexBuffer);
-	StateCache.SetIndexBuffer(IndexBuffer->Resource, Format, 0);
-	StateCache.SetPrimitiveTopology(GetD3D11PrimitiveType(PrimitiveType));
-
-	Direct3DDeviceIMContext->DrawIndexedInstancedIndirect(ArgumentsBuffer->Resource, DrawArgumentsIndex * 5 * sizeof(uint32));
-
-	EnableUAVOverlap();
-}
-
 void FD3D11DynamicRHI::RHIDrawIndexedPrimitive(FRHIBuffer* IndexBufferRHI, int32 BaseVertexIndex, uint32 FirstInstance, uint32 NumVertices, uint32 StartIndex, uint32 NumPrimitives, uint32 NumInstances)
 {
 	RHI_DRAW_CALL_STATS(PrimitiveType, NumVertices, NumPrimitives, NumInstances);
diff --git a/Engine/Source/Runtime/Windows/D3D11RHI/Private/D3D11RHIPrivate.h b/Engine/Source/Runtime/Windows/D3D11RHI/Private/D3D11RHIPrivate.h
--- a/Engine/Source/Runtime/Windows/D3D11RHI/Private/D3D11RHIPrivate.h
+++ b/Engine/Source/Runtime/Windows/D3D11RHI/Private/D3D11RHIPrivate.h
@@ -562,7 +562,6 @@ public:
 	D3D11RHI_API void SetRenderTargetsAndClear(const FRHISetRenderTargetsInfo& RenderTargetsInfo);
 	D3D11RHI_API virtual void RHIDrawPrimitive(uint32 BaseVertexIndex, uint32 NumPrimitives, uint32 NumInstances) final override;
 	D3D11RHI_API virtual void RHIDrawPrimitiveIndirect(FRHIBuffer* ArgumentBuffer, uint32 ArgumentOffset) final override;
-	D3D11RHI_API virtual void RHIDrawIndexedIndirect(FRHIBuffer* IndexBufferRHI, FRHIBuffer* ArgumentsBufferRHI, int32 DrawArgumentsIndex, uint32 NumInstances) final override;
 	D3D11RHI_API virtual void RHIDrawIndexedPrimitive(FRHIBuffer* IndexBuffer, int32 BaseVertexIndex, uint32 FirstInstance, uint32 NumVertices, uint32 StartIndex, uint32 NumPrimitives, uint32 NumInstances) final override;
 	D3D11RHI_API virtual void RHIDrawIndexedPrimitiveIndirect(FRHIBuffer* IndexBuffer, FRHIBuffer* ArgumentBuffer, uint32 ArgumentOffset) final override;
 	virtual void RHIEnableDepthBoundsTest(bool bEnable) final override
