Don't attempt to memory map lazy Metal shader libraries if USE_MMAPPED_SHADERARCHIVE=0

diff --git a/Engine/Source/Runtime/Apple/MetalRHI/Private/MetalDynamicRHI_Shaders.cpp b/Engine/Source/Runtime/Apple/MetalRHI/Private/MetalDynamicRHI_Shaders.cpp
--- a/Engine/Source/Runtime/Apple/MetalRHI/Private/MetalDynamicRHI_Shaders.cpp
+++ b/Engine/Source/Runtime/Apple/MetalRHI/Private/MetalDynamicRHI_Shaders.cpp
@@ -213,8 +213,9 @@ FRHIShaderLibraryRef FMetalDynamicRHI::RHICreateShaderLibrary(EShaderPlatform Pl
 				TUniquePtr<FMetalShaderLibrary::FLazyMetalLib> LazyLibrary = MakeUnique<FMetalShaderLibrary::FLazyMetalLib>();
 				LazyLibrary->MetalLibraryFilePath = MetalLibraryFilePath;
 				LazyLibrary->Data = MakeUnique<FMetalShaderLibrary::FShaderLibDataOwner>();
+#if USE_MMAPPED_SHADERARCHIVE
 				MapShaderLibraryFile(MetalLibraryFilePath, *LazyLibrary->Data);
-				
+#endif				
 				LazyLibraries.Add(MoveTemp(LazyLibrary));
 			}
 
