Resolve static code analysis warning about registering callbacks without unregistering

#rnx

diff --git a/Engine/Source/Runtime/Apple/MetalRHI/Private/MetalPipeline.cpp b/Engine/Source/Runtime/Apple/MetalRHI/Private/MetalPipeline.cpp
--- a/Engine/Source/Runtime/Apple/MetalRHI/Private/MetalPipeline.cpp
+++ b/Engine/Source/Runtime/Apple/MetalRHI/Private/MetalPipeline.cpp
@@ -1329,8 +1329,8 @@ FMetalPipelineStateCacheManager::FMetalPipelineStateCacheManager(FMetalDevice& D
 	FMetalShaderPipelineCache::Create(Device);
 	
 #if PLATFORM_IOS
-	FCoreDelegates::ApplicationWillDeactivateDelegate.AddStatic(&FShaderPipelineCache::PauseBatching);
-	FCoreDelegates::ApplicationHasReactivatedDelegate.AddStatic(&FShaderPipelineCache::ResumeBatching);
+	FCoreDelegates::ApplicationWillDeactivateDelegate.AddRaw(this, &FShaderPipelineCache::PauseBatching);
+	FCoreDelegates::ApplicationHasReactivatedDelegate.AddRaw(this, &FShaderPipelineCache::ResumeBatching);
 	
 	OnShaderPipelineCachePreOpenDelegate = FShaderPipelineCache::GetCachePreOpenDelegate().AddRaw(this, &FMetalPipelineStateCacheManager::OnShaderPipelineCachePreOpen);
 	OnShaderPipelineCacheOpenedDelegate = FShaderPipelineCache::GetCacheOpenedDelegate().AddRaw(this, &FMetalPipelineStateCacheManager::OnShaderPipelineCacheOpened);
@@ -1340,6 +1340,9 @@ FMetalPipelineStateCacheManager::FMetalPipelineStateCacheManager(FMetalDevice& D
 
 FMetalPipelineStateCacheManager::~FMetalPipelineStateCacheManager()
 {
+	FCoreDelegates::ApplicationWillDeactivateDelegate.RemoveAll(this);
+	FCoreDelegates::ApplicationHasReactivatedDelegate.RemoveAll(this);
+	
 	if (OnShaderPipelineCacheOpenedDelegate.IsValid())
 	{
 		FShaderPipelineCache::GetCacheOpenedDelegate().Remove(OnShaderPipelineCacheOpenedDelegate);
