Metal RHI

- Bumped timeout limit on GPU Fence

diff --git a/Engine/Source/Runtime/Apple/MetalRHI/Private/MetalUAV.cpp b/Engine/Source/Runtime/Apple/MetalRHI/Private/MetalUAV.cpp
--- a/Engine/Source/Runtime/Apple/MetalRHI/Private/MetalUAV.cpp
+++ b/Engine/Source/Runtime/Apple/MetalRHI/Private/MetalUAV.cpp
@@ -810,7 +810,7 @@ void FMetalGPUFence::Wait(FRHICommandListImmediate& RHICmdList, FRHIGPUMask GPUM
 		bool bSuccess = false;
 		if (SharedEvent)
 		{
-			static const uint64_t FENCE_EVENT_WAIT_TIMEOUT = 100;
+			static const uint64_t FENCE_EVENT_WAIT_TIMEOUT = 1000000;
 			bSuccess = SharedEvent->waitUntilSignaledValue(1, FENCE_EVENT_WAIT_TIMEOUT);
 		}
 		if(!bSuccess && SyncPoint)
