Metal RHI

- Fixed potential use after Free on GPU with MTLSharedEvent

diff --git a/Engine/Source/Runtime/Apple/MetalRHI/Private/MetalUAV.cpp b/Engine/Source/Runtime/Apple/MetalRHI/Private/MetalUAV.cpp
--- a/Engine/Source/Runtime/Apple/MetalRHI/Private/MetalUAV.cpp
+++ b/Engine/Source/Runtime/Apple/MetalRHI/Private/MetalUAV.cpp
@@ -775,6 +775,14 @@ FMetalGPUFence::FMetalGPUFence(FMetalDevice& InDevice, FName InName)
 	}
 }
 
+FMetalGPUFence::~FMetalGPUFence()
+{
+	if(GMetalUseEventForGPUFence)
+	{
+		FMetalDynamicRHI::Get().DeferredDelete([Event=SharedEvent]{});
+	}
+}
+
 void FMetalGPUFence::Clear()
 {
 	SyncPoint = nullptr;
diff --git a/Engine/Source/Runtime/Apple/MetalRHI/Public/MetalResources.h b/Engine/Source/Runtime/Apple/MetalRHI/Public/MetalResources.h
--- a/Engine/Source/Runtime/Apple/MetalRHI/Public/MetalResources.h
+++ b/Engine/Source/Runtime/Apple/MetalRHI/Public/MetalResources.h
@@ -624,7 +624,8 @@ class FMetalGPUFence final : public FRHIGPUFence
 {
 public:
 	FMetalGPUFence(FMetalDevice& Device, FName InName);
-
+	~FMetalGPUFence();
+	
 	virtual void Clear() override;
 	virtual bool Poll() const override;
 	virtual void Wait(FRHICommandListImmediate& RHICmdList, FRHIGPUMask GPUMask) const override;
