Resolve static code analysis warning about registering callbacks without unregistering

#rnx

diff --git a/Engine/Source/Runtime/Apple/MetalRHI/Private/MetalPipeline.cpp b/Engine/Source/Runtime/Apple/MetalRHI/Private/MetalPipeline.cpp
--- a/Engine/Source/Runtime/Apple/MetalRHI/Private/MetalPipeline.cpp
+++ b/Engine/Source/Runtime/Apple/MetalRHI/Private/MetalPipeline.cpp
@@ -1329,8 +1329,8 @@ FMetalPipelineStateCacheManager::FMetalPipelineStateCacheManager(FMetalDevice& D
 	FMetalShaderPipelineCache::Create(Device);
 	
 #if PLATFORM_IOS
-	FCoreDelegates::ApplicationWillDeactivateDelegate.AddStatic(&FShaderPipelineCache::PauseBatching);
-	FCoreDelegates::ApplicationHasReactivatedDelegate.AddStatic(&FShaderPipelineCache::ResumeBatching);
+	ApplicationWillDeactivateDelegateHandle = FCoreDelegates::ApplicationWillDeactivateDelegate.AddStatic(&FShaderPipelineCache::PauseBatching);
+	ApplicationHasReactivatedDelegateHandle = FCoreDelegates::ApplicationHasReactivatedDelegate.AddStatic(&FShaderPipelineCache::ResumeBatching);
 	
 	OnShaderPipelineCachePreOpenDelegate = FShaderPipelineCache::GetCachePreOpenDelegate().AddRaw(this, &FMetalPipelineStateCacheManager::OnShaderPipelineCachePreOpen);
 	OnShaderPipelineCacheOpenedDelegate = FShaderPipelineCache::GetCacheOpenedDelegate().AddRaw(this, &FMetalPipelineStateCacheManager::OnShaderPipelineCacheOpened);
@@ -1340,6 +1340,9 @@ FMetalPipelineStateCacheManager::FMetalPipelineStateCacheManager(FMetalDevice& D
 
 FMetalPipelineStateCacheManager::~FMetalPipelineStateCacheManager()
 {
+	FCoreDelegates::ApplicationWillDeactivateDelegate.Remove(ApplicationWillDeactivateDelegateHandle);
+	FCoreDelegates::ApplicationHasReactivatedDelegate.Remove(ApplicationHasReactivatedDelegateHandle);
+	
 	if (OnShaderPipelineCacheOpenedDelegate.IsValid())
 	{
 		FShaderPipelineCache::GetCacheOpenedDelegate().Remove(OnShaderPipelineCacheOpenedDelegate);
diff --git a/Engine/Source/Runtime/Apple/MetalRHI/Private/MetalPipeline.h b/Engine/Source/Runtime/Apple/MetalRHI/Private/MetalPipeline.h
--- a/Engine/Source/Runtime/Apple/MetalRHI/Private/MetalPipeline.h
+++ b/Engine/Source/Runtime/Apple/MetalRHI/Private/MetalPipeline.h
@@ -93,6 +93,8 @@ public:
 	~FMetalPipelineStateCacheManager();
 	
 private:
+	FDelegateHandle ApplicationWillDeactivateDelegateHandle;
+	FDelegateHandle ApplicationHasReactivatedDelegateHandle;
 	FDelegateHandle OnShaderPipelineCachePreOpenDelegate;
 	FDelegateHandle OnShaderPipelineCacheOpenedDelegate;
 	FDelegateHandle OnShaderPipelineCachePrecompilationCompleteDelegate;
