Metal RHI

- Fix bug where we run out of BlendState keys

diff --git a/Engine/Source/Runtime/Apple/MetalRHI/Private/MetalPipeline.cpp b/Engine/Source/Runtime/Apple/MetalRHI/Private/MetalPipeline.cpp
--- a/Engine/Source/Runtime/Apple/MetalRHI/Private/MetalPipeline.cpp
+++ b/Engine/Source/Runtime/Apple/MetalRHI/Private/MetalPipeline.cpp
@@ -45,7 +45,7 @@ extern FString GMetalShaderValidationShaderName;
 
 static uint32 BlendBitOffsets[] = { Offset_BlendState0, Offset_BlendState1, Offset_BlendState2, Offset_BlendState3, Offset_BlendState4, Offset_BlendState5, Offset_BlendState6, Offset_BlendState7 };
 static uint32 RTBitOffsets[] = { Offset_RenderTargetFormat0, Offset_RenderTargetFormat1, Offset_RenderTargetFormat2, Offset_RenderTargetFormat3, Offset_RenderTargetFormat4, Offset_RenderTargetFormat5, Offset_RenderTargetFormat6, Offset_RenderTargetFormat7 };
-static_assert(Offset_RasterEnd < 64 && Offset_End < 128, "Offset_RasterEnd must be < 64 && Offset_End < 128");
+static_assert(Offset_BlendStateEnd < 64 && Offset_End < 128, "Offset_BlendStateEnd must be <= 64 && Offset_End < 128");
 
 static float RoundUpNearestEven(const float f)
 {
@@ -68,7 +68,7 @@ struct FMetalGraphicsPipelineKey
 	template<typename Type>
 	inline void SetHashValue(uint32 Offset, uint32 NumBits, Type Value)
 	{
-		if (Offset < Offset_RasterEnd)
+		if (Offset < Offset_BlendStateEnd)
 		{
 			uint64 BitMask = ((((uint64)1ULL) << NumBits) - 1) << Offset;
 			RenderPipelineHash.RasterBits = (RenderPipelineHash.RasterBits & ~BitMask) | (((uint64)Value << Offset) & BitMask);
diff --git a/Engine/Source/Runtime/Apple/MetalRHI/Private/MetalPipeline.h b/Engine/Source/Runtime/Apple/MetalRHI/Private/MetalPipeline.h
--- a/Engine/Source/Runtime/Apple/MetalRHI/Private/MetalPipeline.h
+++ b/Engine/Source/Runtime/Apple/MetalRHI/Private/MetalPipeline.h
@@ -49,11 +49,13 @@ enum EMetalLimits
 
 enum EMetalPipelineHashBits
 {
+	// Blend States
+	NumBits_BlendState = 8, //(x8=64),
+	
+	// Raster States
 	NumBits_RenderTargetFormat = 6, //(x8=48),
 	NumBits_DepthFormat = 3, //(x1=3),
 	NumBits_SampleCount = 3, //(x1=3),
-
-	NumBits_BlendState = 7, //(x8=56),
 	NumBits_PrimitiveTopology = 2, //(x1=2)
 	NumBits_AlphaToCoverage = 1, //(x1=1)
 };
@@ -68,8 +70,7 @@ enum EMetalPipelineHashOffsets
 	Offset_BlendState5 = Offset_BlendState4 + NumBits_BlendState,
 	Offset_BlendState6 = Offset_BlendState5 + NumBits_BlendState,
 	Offset_BlendState7 = Offset_BlendState6 + NumBits_BlendState,
-	Offset_PrimitiveTopology = Offset_BlendState7 + NumBits_BlendState,
-	Offset_RasterEnd = Offset_PrimitiveTopology + NumBits_PrimitiveTopology,
+	Offset_BlendStateEnd = Offset_BlendState7 + NumBits_BlendState - 1, // 63
 
 	Offset_RenderTargetFormat0 = 64,
 	Offset_RenderTargetFormat1 = Offset_RenderTargetFormat0 + NumBits_RenderTargetFormat,
@@ -82,7 +83,8 @@ enum EMetalPipelineHashOffsets
 	Offset_DepthFormat = Offset_RenderTargetFormat7 + NumBits_RenderTargetFormat,
 	Offset_SampleCount = Offset_DepthFormat + NumBits_DepthFormat,
 	Offset_AlphaToCoverage = Offset_SampleCount + NumBits_SampleCount,
-	Offset_End = Offset_AlphaToCoverage + NumBits_AlphaToCoverage
+	Offset_PrimitiveTopology = Offset_AlphaToCoverage + NumBits_AlphaToCoverage,
+	Offset_End = Offset_PrimitiveTopology + NumBits_PrimitiveTopology
 };
 
 class FMetalDevice;
