Fix errors compiling MetalRHI.cpp.

diff --git a/Engine/Source/Runtime/Apple/MetalRHI/Private/MetalRHI.cpp b/Engine/Source/Runtime/Apple/MetalRHI/Private/MetalRHI.cpp
--- a/Engine/Source/Runtime/Apple/MetalRHI/Private/MetalRHI.cpp
+++ b/Engine/Source/Runtime/Apple/MetalRHI/Private/MetalRHI.cpp
@@ -1911,7 +1911,7 @@ void FMetalDynamicRHI::RHIWriteResourceGPUAddresses(FRHICommandListBase& RHICmdL
 				else if (Resource->GetType() == RRT_RayTracingGeometry)
 				{
 					const FMetalRayTracingGeometry* RayTracingGeometry = ResourceCast(static_cast<const FRHIRayTracingGeometry*>(Resource));
-					AddressesPtr[ResourceIndex] = RayTracingGeometry->GetGPUVirtualAddress();
+					AddressesPtr[ResourceIndex] = RayTracingGeometry->GetGPUAddress();
 				}
 #endif
 				else
diff --git a/Engine/Source/Runtime/Apple/MetalRHI/Private/MetalRHIPrivate.h b/Engine/Source/Runtime/Apple/MetalRHI/Private/MetalRHIPrivate.h
--- a/Engine/Source/Runtime/Apple/MetalRHI/Private/MetalRHIPrivate.h
+++ b/Engine/Source/Runtime/Apple/MetalRHI/Private/MetalRHIPrivate.h
@@ -210,11 +210,23 @@ static FORCEINLINE typename TMetalResourceTraits<TRHIType>::TConcreteType* Resou
 	return static_cast<typename TMetalResourceTraits<TRHIType>::TConcreteType*>(Resource);
 }
 
+template<typename TRHIType>
+static FORCEINLINE const typename TMetalResourceTraits<TRHIType>::TConcreteType* ResourceCast(const TRHIType* Resource)
+{
+	return static_cast<const typename TMetalResourceTraits<TRHIType>::TConcreteType*>(Resource);
+}
+
 static FORCEINLINE FMetalSurface* ResourceCast(FRHITexture* Texture)
 {
 	return GetMetalSurfaceFromRHITexture(Texture);
 }
 
+static FORCEINLINE const FMetalSurface* ResourceCast(const FRHITexture* Texture)
+{
+	unimplemented();
+	return nullptr;
+}
+
 MTL::LanguageVersion ValidateVersion(uint32 Version);
 
 // Needs to be the same as EShaderFrequency when all stages are supported, but unlike EShaderFrequency you can compile out stages.
