Metal - Remove support for Apple A8

#tests Local & preflight

diff --git a/Engine/Source/Developer/Apple/MetalShaderFormat/Private/MetalCompileShaderSPIRV.cpp b/Engine/Source/Developer/Apple/MetalShaderFormat/Private/MetalCompileShaderSPIRV.cpp
--- a/Engine/Source/Developer/Apple/MetalShaderFormat/Private/MetalCompileShaderSPIRV.cpp
+++ b/Engine/Source/Developer/Apple/MetalShaderFormat/Private/MetalCompileShaderSPIRV.cpp
@@ -194,7 +194,6 @@ void FMetalCompileShaderSPIRV::DoCompileMetalShader(
 	bool const bZeroInitialise = Input.Environment.CompilerFlags.Contains(CFLAG_ZeroInitialise);
 	bool const bBoundsChecks = Input.Environment.CompilerFlags.Contains(CFLAG_BoundsChecking);
 
-	bool bSupportAppleA8 = Input.Environment.GetCompileArgument(TEXT("SUPPORT_APPLE_A8"), false);
 	bool bAllowFastIntrinsics = Input.Environment.GetCompileArgument(TEXT("METAL_USE_FAST_INTRINSICS"), false);
 
 	// WPO requires that we make all multiply/sincos instructions invariant :(
@@ -972,14 +971,14 @@ void FMetalCompileShaderSPIRV::DoCompileMetalShader(
 				break;
 			case EMetalGPUSemanticsTBDRDesktop:
 				TargetDesc.Language = CrossCompiler::EShaderConductorLanguage::Metal_iOS;
-				TargetDesc.CompileFlags->SetDefine(TEXT("ios_support_base_vertex_instance"), !bSupportAppleA8);
+				TargetDesc.CompileFlags->SetDefine(TEXT("ios_support_base_vertex_instance"), 1);
 				TargetDesc.CompileFlags->SetDefine(TEXT("use_framebuffer_fetch_subpasses"), 1);
 				TargetDesc.CompileFlags->SetDefine(TEXT("emulate_cube_array"), 1);
 				break;
 			case EMetalGPUSemanticsMobile:
 			default:
 				TargetDesc.Language = CrossCompiler::EShaderConductorLanguage::Metal_iOS;
-				TargetDesc.CompileFlags->SetDefine(TEXT("ios_support_base_vertex_instance"), !bSupportAppleA8);
+				TargetDesc.CompileFlags->SetDefine(TEXT("ios_support_base_vertex_instance"), 1);
 				TargetDesc.CompileFlags->SetDefine(TEXT("use_framebuffer_fetch_subpasses"), 1);
 				TargetDesc.CompileFlags->SetDefine(TEXT("emulate_cube_array"), 1);
 				break;
diff --git a/Engine/Source/Developer/IOS/IOSTargetPlatformControls/Private/IOSTargetPlatformControls.cpp b/Engine/Source/Developer/IOS/IOSTargetPlatformControls/Private/IOSTargetPlatformControls.cpp
--- a/Engine/Source/Developer/IOS/IOSTargetPlatformControls/Private/IOSTargetPlatformControls.cpp
+++ b/Engine/Source/Developer/IOS/IOSTargetPlatformControls/Private/IOSTargetPlatformControls.cpp
@@ -445,14 +445,6 @@ static bool SupportsMetalMRT()
 	return bSupportsMetalMRT;
 }
 
-static bool SupportsA8Devices()
-{
-    // default to NOT supporting A8 devices
-    bool bSupportAppleA8 = false;
-    GConfig->GetBool(TEXT("/Script/IOSRuntimeSettings.IOSRuntimeSettings"), TEXT("bSupportAppleA8"), bSupportAppleA8, GEngineIni);
-    return bSupportAppleA8;
-}
-
 void FIOSTargetPlatformControls::GetPlatformSpecificProjectAnalytics( TArray<FAnalyticsEventAttribute>& AnalyticsParamArray ) const
 {
 	TNonDesktopTargetPlatformControlsBase<FIOSPlatformProperties>::GetPlatformSpecificProjectAnalytics( AnalyticsParamArray );
@@ -512,9 +504,8 @@ void FIOSTargetPlatformControls::GetTextureFormats( const UTexture* Texture, TAr
 	if (TextureFormatNames.Num() == 0)
 	{
         int32 BlockSize = 1;
-		// Compressed volume textures require MTLGPUFamilyApple3 or later
-		// min spec for TVOS is AppleTV HD which is MTLGPUFamilyApple2 (A8)
-		bool bSupportCompressedVolumeTexture = !bIsTVOS && !SupportsA8Devices();
+		// Compressed volume textures require MTLGPUFamilyApple3 or later which is our minimal supported
+		bool bSupportCompressedVolumeTexture = true;
 		bool bSupportFilteredFloat32Textures = false;
 		GetDefaultTextureFormatNamePerLayer(TextureFormatNames, this->GetTargetPlatformSettings(), this, Texture, bSupportCompressedVolumeTexture, BlockSize, bSupportFilteredFloat32Textures);
 	}
diff --git a/Engine/Source/Developer/IOS/IOSTargetPlatformSettings/Private/IOSTargetPlatformSettings.cpp b/Engine/Source/Developer/IOS/IOSTargetPlatformSettings/Private/IOSTargetPlatformSettings.cpp
--- a/Engine/Source/Developer/IOS/IOSTargetPlatformSettings/Private/IOSTargetPlatformSettings.cpp
+++ b/Engine/Source/Developer/IOS/IOSTargetPlatformSettings/Private/IOSTargetPlatformSettings.cpp
@@ -112,14 +112,6 @@ static bool SupportsMetalMRT()
 	return bSupportsMetalMRT;
 }
 
-static bool SupportsA8Devices()
-{
-    // default to NOT supporting A8 devices
-    bool bSupportAppleA8 = false;
-    GConfig->GetBool(TEXT("/Script/IOSRuntimeSettings.IOSRuntimeSettings"), TEXT("bSupportAppleA8"), bSupportAppleA8, GEngineIni);
-    return bSupportAppleA8;
-}
-
 bool FIOSTargetPlatformSettings::SupportsFeature( ETargetPlatformFeatures Feature ) const
 {
 	switch (Feature)
diff --git a/Engine/Source/Runtime/Apple/MetalRHI/Private/MetalRHI.cpp b/Engine/Source/Runtime/Apple/MetalRHI/Private/MetalRHI.cpp
--- a/Engine/Source/Runtime/Apple/MetalRHI/Private/MetalRHI.cpp
+++ b/Engine/Source/Runtime/Apple/MetalRHI/Private/MetalRHI.cpp
@@ -21,6 +21,7 @@
 #include "HAL/FileManager.h"
 #include "MetalProfiler.h"
 #include "GenericPlatform/GenericPlatformDriver.h"
+#include "Math/UnrealMathUtility.h"
 #include "MetalShaderResources.h"
 #include "MetalLLM.h"
 #include "Engine/RendererSettings.h"
@@ -241,29 +242,17 @@ FMetalDynamicRHI::FMetalDynamicRHI(ERHIFeatureLevel::Type RequestedFeatureLevel)
 #endif
 	
 #if PLATFORM_IOS
-    bool bSupportAppleA8 = false;
-    GConfig->GetBool(TEXT("/Script/IOSRuntimeSettings.IOSRuntimeSettings"), TEXT("bSupportAppleA8"), bSupportAppleA8, GEngineIni);
-        
-    bool bIsA8FeatureSet = false;
-        
-#if PLATFORM_TVOS
-	GRHISupportsDrawIndirect = MTLDevice->supportsFeatureSet(MTL::FeatureSet_tvOS_GPUFamily2_v1);
-	GRHISupportsPixelShaderUAVs = MTLDevice->supportsFeatureSet(MTL::FeatureSet_tvOS_GPUFamily2_v1);
-        
-    if (!MTLDevice->supportsFeatureSet(MTL::FeatureSet_tvOS_GPUFamily2_v1))
-    {
-        bIsA8FeatureSet = true;
-    }
-        
-#else
-	if (!MTLDevice->supportsFeatureSet(MTL::FeatureSet_iOS_GPUFamily3_v1))
+	if (!MTLDevice->supportsFamily(MTL::GPUFamilyApple3))
 	{
-        bIsA8FeatureSet = true;
-    }
-    
-	GRHISupportsRWTextureBuffers = MTLDevice->supportsFeatureSet(MTL::FeatureSet_iOS_GPUFamily4_v1);
-	GRHISupportsDrawIndirect = MTLDevice->supportsFeatureSet(MTL::FeatureSet_iOS_GPUFamily3_v1);
-	GRHISupportsPixelShaderUAVs = MTLDevice->supportsFeatureSet(MTL::FeatureSet_iOS_GPUFamily3_v1);
+		static const FText UnsupportedGPUFamilyApple9 = NSLOCTEXT("MetalRHI", "UnsupportedA8FeatureSet", "This device does not support the Apple family 3 GPU features (Apple A9+), which is required as the minimum for this engine version.");
+		FMessageDialog::Open(EAppMsgType::Ok, UnsupportedGPUFamilyApple9); 
+		UE_LOG(LogMetal, Fatal, TEXT("%s"), *UnsupportedGPUFamilyApple9.ToString());
+		FPlatformMisc::RequestExit(true);
+	}
+        
+	GRHISupportsDrawIndirect = true;
+	GRHISupportsPixelShaderUAVs = true;
+	GRHISupportsRWTextureBuffers = MTLDevice->supportsFamily(MTL::GPUFamilyApple4);
 
 	const MTL::FeatureSet FeatureSets[] = {
         MTL::FeatureSet_iOS_GPUFamily1_v1,
@@ -290,21 +279,6 @@ FMetalDynamicRHI::FMetalDynamicRHI(ERHIFeatureLevel::Type RequestedFeatureLevel)
 		
 	GSupportsVolumeTextureRendering = Device->SupportsFeature(EMetalFeaturesLayeredRendering);
 	bSupportsPointLights = GSupportsVolumeTextureRendering;
-#endif
-
-    if(bIsA8FeatureSet)
-    {
-        if(!bSupportAppleA8)
-        {
-            UE_LOG(LogMetal, Fatal, TEXT("This device does not supports the Apple A8x or above feature set which is the minimum for this build. Please check the Support Apple A8 checkbox in the IOS Project Settings."));
-        }
-        
-        static auto* CVarMobileVirtualTextures = IConsoleManager::Get().FindTConsoleVariableDataInt(TEXT("r.Mobile.VirtualTextures"));
-        if(CVarMobileVirtualTextures->GetValueOnAnyThread() != 0)
-        {
-            UE_LOG(LogMetal, Warning, TEXT("Mobile Virtual Textures require a minimum of the Apple A9 feature set."));
-        }
-    }
         
     bool bProjectSupportsMRTs = false;
     GConfig->GetBool(TEXT("/Script/IOSRuntimeSettings.IOSRuntimeSettings"), TEXT("bSupportsMetalMRT"), bProjectSupportsMRTs, GEngineIni);
@@ -313,7 +287,7 @@ FMetalDynamicRHI::FMetalDynamicRHI(ERHIFeatureLevel::Type RequestedFeatureLevel)
 	bool const bForceES3_1 = FParse::Param(FCommandLine::Get(),TEXT("es31"));
 
     // Only allow SM5 MRT on A9 or above devices
-    if (bProjectSupportsMRTs && bRequestedMetalMRT && !bIsA8FeatureSet && !bForceES3_1)
+    if (bProjectSupportsMRTs && bRequestedMetalMRT && !bForceES3_1)
     {
 #if PLATFORM_TVOS
 		ValidateTargetedRHIFeatureLevelExists(SP_METAL_SM5_IOS);
@@ -326,11 +300,6 @@ FMetalDynamicRHI::FMetalDynamicRHI(ERHIFeatureLevel::Type RequestedFeatureLevel)
     }
     else
 	{
-		if (bRequestedMetalMRT && !bForceES3_1)
-		{
-			UE_LOG(LogMetal, Warning, TEXT("Metal MRT support requires an iOS or tvOS device with an A8 processor or later. Falling back to Metal ES 3.1."));
-		}
-		
 #if PLATFORM_TVOS
 		ValidateTargetedRHIFeatureLevelExists(SP_METAL_ES3_1_TVOS);
 		GMaxRHIShaderPlatform = SP_METAL_ES3_1_TVOS;
@@ -721,9 +690,8 @@ FMetalDynamicRHI::FMetalDynamicRHI(ERHIFeatureLevel::Type RequestedFeatureLevel)
 	GRHIHDRDisplayOutputFormat = (GRHISupportsHDROutput) ? PF_PLATFORM_HDR_0 : PF_B8G8R8A8;
 	GMaxWorkGroupInvocations = 512;
 #else
-	// "Base vertex/instance drawing" is Apple3 family, which is A9 and newer:
-	GRHISupportsBaseVertexIndex = MTLDevice->supportsFeatureSet(MTL::FeatureSet_iOS_GPUFamily3_v1);
-	GRHISupportsFirstInstance = GRHISupportsBaseVertexIndex;
+	GRHISupportsBaseVertexIndex = true;
+	GRHISupportsFirstInstance = true;
 	
 	// TODO: Move this into IOSPlatform
     {
@@ -959,26 +927,11 @@ FMetalDynamicRHI::FMetalDynamicRHI(ERHIFeatureLevel::Type RequestedFeatureLevel)
 	GPixelFormats[PF_PLATFORM_HDR_0		].PlatformFormat	= (uint32)MTL::PixelFormatBGR10_XR_sRGB;
 	GPixelFormats[PF_PLATFORM_HDR_0		].Supported			= GRHISupportsHDROutput;
 		
-#if PLATFORM_TVOS
-    if (!MTLDevice->supportsFeatureSet(MTL::FeatureSet_tvOS_GPUFamily2_v1))
-#else
-	if (!MTLDevice->supportsFeatureSet(MTL::FeatureSet_iOS_GPUFamily3_v2))
-#endif
-	{
-		GPixelFormats[PF_FloatRGB			].PlatformFormat 	= (uint32)MTL::PixelFormatRGBA16Float;
-		GPixelFormats[PF_FloatRGBA			].BlockBytes		= 8;
-		GPixelFormats[PF_FloatR11G11B10		].PlatformFormat	= (uint32)MTL::PixelFormatRGBA16Float;
-		GPixelFormats[PF_FloatR11G11B10		].BlockBytes		= 8;
-		GPixelFormats[PF_FloatR11G11B10		].Supported			= true;
-	}
-	else
-	{
-		GPixelFormats[PF_FloatRGB			].PlatformFormat	= (uint32)MTL::PixelFormatRG11B10Float;
-		GPixelFormats[PF_FloatRGB			].BlockBytes		= 4;
-		GPixelFormats[PF_FloatR11G11B10		].PlatformFormat	= (uint32)MTL::PixelFormatRG11B10Float;
-		GPixelFormats[PF_FloatR11G11B10		].BlockBytes		= 4;
-		GPixelFormats[PF_FloatR11G11B10		].Supported			= true;
-	}
+	GPixelFormats[PF_FloatRGB			].PlatformFormat	= (uint32)MTL::PixelFormatRG11B10Float;
+	GPixelFormats[PF_FloatRGB			].BlockBytes		= 4;
+	GPixelFormats[PF_FloatR11G11B10		].PlatformFormat	= (uint32)MTL::PixelFormatRG11B10Float;
+	GPixelFormats[PF_FloatR11G11B10		].BlockBytes		= 4;
+	GPixelFormats[PF_FloatR11G11B10		].Supported			= true;
 	
 	GPixelFormats[PF_ShadowDepth		].PlatformFormat	= (uint32)MTL::PixelFormatDepth16Unorm;
 	GPixelFormats[PF_ShadowDepth		].BlockBytes		= 2;
@@ -1236,21 +1189,6 @@ FMetalDynamicRHI::FMetalDynamicRHI(ERHIFeatureLevel::Type RequestedFeatureLevel)
 		GRHIGlobals.RayTracing.RequiresInlineRayTracingSBT = true;
 	}
 #endif
-
-#if PLATFORM_IOS || PLATFORM_TVOS
-	// Disable occlusion feedback on devices older than FeatureSet_iOS_GPUFamily3_v2 / MTLFeatureSet_tvOS_GPUFamily2_v1
-	// as buffer writes are not supported from fragment shaders until A10. See UE-348147.
-	static auto CVarOcclusionFeedback = IConsoleManager::Get().FindConsoleVariable(TEXT("r.OcclusionFeedback.Enable"));
-	if (CVarOcclusionFeedback && CVarOcclusionFeedback->GetInt() != 0 && 
-#if PLATFORM_TVOS		
-		!MTLDevice->supportsFeatureSet(MTL::FeatureSet_tvOS_GPUFamily2_v1))
-#else
-		!MTLDevice->supportsFeatureSet(MTL::FeatureSet_iOS_GPUFamily3_v2))
-#endif
-	{
-		CVarOcclusionFeedback->Set(0);
-	}
-#endif
 		
 	GDynamicRHI = this;
 	
diff --git a/Engine/Source/Runtime/Apple/MetalRHI/Private/MetalState.cpp b/Engine/Source/Runtime/Apple/MetalRHI/Private/MetalState.cpp
--- a/Engine/Source/Runtime/Apple/MetalRHI/Private/MetalState.cpp
+++ b/Engine/Source/Runtime/Apple/MetalRHI/Private/MetalState.cpp
@@ -270,13 +270,7 @@ static MTL::SamplerState* FindOrCreateSamplerState(FMetalDevice& Device, const F
 		Desc->setRAddressMode(TranslateWrapMode(Initializer.AddressW));
 		Desc->setLodMinClamp(Initializer.MinMipLevel);
 		Desc->setLodMaxClamp(Initializer.MaxMipLevel);
-#if PLATFORM_TVOS
-		Desc->setCompareFunction(MTL::CompareFunctionNever);
-#elif PLATFORM_IOS
-		Desc->setCompareFunction(Device.GetDevice()->supportsFeatureSet(MTL::FeatureSet_iOS_GPUFamily3_v1) ? TranslateSamplerCompareFunction(Initializer.SamplerComparisonFunction) : MTL::CompareFunctionNever);
-#else
 		Desc->setCompareFunction(TranslateSamplerCompareFunction(Initializer.SamplerComparisonFunction));
-#endif
 #if PLATFORM_MAC
 		Desc->setBorderColor(Initializer.BorderColor == 0 ? MTL::SamplerBorderColorTransparentBlack : MTL::SamplerBorderColorOpaqueWhite);
 #endif
diff --git a/Engine/Source/Runtime/Engine/Private/ShaderCompiler/ShaderCompiler.cpp b/Engine/Source/Runtime/Engine/Private/ShaderCompiler/ShaderCompiler.cpp
--- a/Engine/Source/Runtime/Engine/Private/ShaderCompiler/ShaderCompiler.cpp
+++ b/Engine/Source/Runtime/Engine/Private/ShaderCompiler/ShaderCompiler.cpp
@@ -3717,7 +3717,6 @@ void GlobalBeginCompileShader(
 			bool bAllowFastIntrinsics = false;
 			bool bForceFloats = false;
 			bool bEnableMathOptimisations = true;
-            bool bSupportAppleA8 = false;
 			bool bMetalOptimizeForSize = false;
 
 			if (IsPCPlatform(EShaderPlatform(Target.Platform)))
@@ -3734,7 +3733,6 @@ void GlobalBeginCompileShader(
 				GConfig->GetBool(TEXT("/Script/IOSRuntimeSettings.IOSRuntimeSettings"), TEXT("UseFastIntrinsics"), bAllowFastIntrinsics, GEngineIni);
 				GConfig->GetBool(TEXT("/Script/IOSRuntimeSettings.IOSRuntimeSettings"), TEXT("EnableMathOptimisations"), bEnableMathOptimisations, GEngineIni);
 				GConfig->GetBool(TEXT("/Script/IOSRuntimeSettings.IOSRuntimeSettings"), TEXT("ForceFloats"), bForceFloats, GEngineIni);
-                GConfig->GetBool(TEXT("/Script/IOSRuntimeSettings.IOSRuntimeSettings"), TEXT("bSupportAppleA8"), bSupportAppleA8, GEngineIni);
 				GConfig->GetBool(TEXT("/Script/IOSRuntimeSettings.IOSRuntimeSettings"), TEXT("MetalOptimizeForSize"), bMetalOptimizeForSize, GEngineIni);
                 
 				// Force no development shaders on iOS
@@ -3744,7 +3742,6 @@ void GlobalBeginCompileShader(
             Input.Environment.FullPrecisionInPS |= bForceFloats;
             
 			Input.Environment.SetCompileArgument(TEXT("METAL_USE_FAST_INTRINSICS"), bAllowFastIntrinsics);
-            Input.Environment.SetCompileArgument(TEXT("SUPPORT_APPLE_A8"), bSupportAppleA8);
             Input.Environment.SetCompileArgument(TEXT("METAL_OPTIMIZE_FOR_SIZE"), bMetalOptimizeForSize);
 			
 			// Same as console-variable above, but that's global and this is per-platform, per-project
diff --git a/Engine/Source/Runtime/IOS/IOSRuntimeSettings/Classes/IOSRuntimeSettings.h b/Engine/Source/Runtime/IOS/IOSRuntimeSettings/Classes/IOSRuntimeSettings.h
--- a/Engine/Source/Runtime/IOS/IOSRuntimeSettings/Classes/IOSRuntimeSettings.h
+++ b/Engine/Source/Runtime/IOS/IOSRuntimeSettings/Classes/IOSRuntimeSettings.h
@@ -504,14 +504,6 @@ public:
 	 */
 	UPROPERTY(EditAnywhere, config, Category=Rendering, meta = (DisplayName = "Enable Fast-Math optimisations", ConfigRestartRequired = true))
 	bool EnableMathOptimisations;
-	
-    /** Supports Apple A8 devices.
-     * Disables 3d texture compression
-     * Virtual Textures are not supported on A8 devices due to missing hardware features
-     * This will disable also Base_Vertex semantics in the msl which may have negative consequences if shaders rely on it
-     */
-    UPROPERTY(config, EditAnywhere, Category = Rendering, Meta = (DisplayName = "Support Apple A8", ConfigRestartRequired = true))
-    bool bSupportAppleA8;
     
 	// Whether or not the keyboard should be usable on it's own without a UITextField
 	UPROPERTY(GlobalConfig, EditAnywhere, Category = Input)
diff --git a/Engine/Source/Runtime/RenderCore/Private/Shader.cpp b/Engine/Source/Runtime/RenderCore/Private/Shader.cpp
--- a/Engine/Source/Runtime/RenderCore/Private/Shader.cpp
+++ b/Engine/Source/Runtime/RenderCore/Private/Shader.cpp
@@ -2176,7 +2176,6 @@ void ShaderMapAppendKey(EShaderPlatform Platform, FShaderKeyGenerator& KeyGen)
 		bool bAllowFastIntrinsics = false;
 		bool bEnableMathOptimisations = true;
 		bool bForceFloats = false;
-        bool bSupportAppleA8 = false;
 		bool bMetalOptimizeBySize = false;
 
 		if (IsPCPlatform(Platform))
@@ -2190,7 +2189,6 @@ void ShaderMapAppendKey(EShaderPlatform Platform, FShaderKeyGenerator& KeyGen)
 			GConfig->GetBool(TEXT("/Script/IOSRuntimeSettings.IOSRuntimeSettings"), TEXT("UseFastIntrinsics"), bAllowFastIntrinsics, GEngineIni);
 			GConfig->GetBool(TEXT("/Script/IOSRuntimeSettings.IOSRuntimeSettings"), TEXT("EnableMathOptimisations"), bEnableMathOptimisations, GEngineIni);
 			GConfig->GetBool(TEXT("/Script/IOSRuntimeSettings.IOSRuntimeSettings"), TEXT("ForceFloats"), bForceFloats, GEngineIni);
-            GConfig->GetBool(TEXT("/Script/IOSRuntimeSettings.IOSRuntimeSettings"), TEXT("bSupportAppleA8"), bSupportAppleA8, GEngineIni);
 			GConfig->GetBool(TEXT("/Script/IOSRuntimeSettings.IOSRuntimeSettings"), TEXT("MetalOptimizeBySize"), bMetalOptimizeBySize, GEngineIni);
 		}
 		
@@ -2213,12 +2211,6 @@ void ShaderMapAppendKey(EShaderPlatform Platform, FShaderKeyGenerator& KeyGen)
 			KeyGen.Append(TEXT("FP32"));
 		}
 		
-        if(bSupportAppleA8)
-        {
-			KeyGen.AppendSeparator();
-			KeyGen.Append(TEXT("A8GPU"));
-        }
-		
 		if(bMetalOptimizeBySize)
         {
 			KeyGen.AppendSeparator();
