Remove immediate context from Metal RHI when the new GPU profiler is in use

diff --git a/Engine/Source/Runtime/Apple/MetalRHI/Private/MetalDynamicRHI.h b/Engine/Source/Runtime/Apple/MetalRHI/Private/MetalDynamicRHI.h
--- a/Engine/Source/Runtime/Apple/MetalRHI/Private/MetalDynamicRHI.h
+++ b/Engine/Source/Runtime/Apple/MetalRHI/Private/MetalDynamicRHI.h
@@ -306,7 +306,9 @@ public:
 private:
 	FMetalDevice* Device;
 	FTextureMemoryStats MemoryStats;
+#if (RHI_NEW_GPU_PROFILER == 0)
 	FMetalRHICommandContext ImmediateContext;
+#endif
 	TMap<uint32, FVertexDeclarationRHIRef> VertexDeclarationCache;
 	TLockFreePointerListUnordered<FMetalRHICommandContext, PLATFORM_CACHE_LINE_SIZE> MetalCommandContextPool;
 	
diff --git a/Engine/Source/Runtime/Apple/MetalRHI/Private/MetalRHI.cpp b/Engine/Source/Runtime/Apple/MetalRHI/Private/MetalRHI.cpp
--- a/Engine/Source/Runtime/Apple/MetalRHI/Private/MetalRHI.cpp
+++ b/Engine/Source/Runtime/Apple/MetalRHI/Private/MetalRHI.cpp
@@ -194,13 +194,10 @@ static void VerifyMetalCompiler()
 
 FMetalDynamicRHI::FMetalDynamicRHI(ERHIFeatureLevel::Type RequestedFeatureLevel)
 	: Device(FMetalDevice::CreateDevice())
+#if (RHI_NEW_GPU_PROFILER == 0)
 	, ImmediateContext(*Device, nullptr)
+#endif
 {
-	FMetalRHICommandContext* RHICommandContext = static_cast<FMetalRHICommandContext*>(RHIGetDefaultContext());
-	METAL_GPUPROFILE(FMetalProfiler::CreateProfiler(*RHICommandContext));
-	
-	RHICommandContext->ResetContext();
-	
 	check(Singleton == nullptr);
 	Singleton = this;
 
@@ -1197,7 +1194,6 @@ FMetalDynamicRHI::FMetalDynamicRHI(ERHIFeatureLevel::Type RequestedFeatureLevel)
 	
 	GIsMetalInitialized = true;
 
-	ImmediateContext.SetProfiler(nullptr);
 #if ENABLE_METAL_GPUPROFILE && RHI_NEW_GPU_PROFILER == 0
 	FMetalProfiler* Profiler = FMetalProfiler::CreateProfiler(ImmediateContext);
 	ImmediateContext.SetProfiler(Profiler);
@@ -1554,12 +1550,20 @@ MTL::Device* FMetalDynamicRHI::RHIGetDevice()
 
 void* FMetalDynamicRHI::RHIGetNativeGraphicsQueue()
 {
+#if RHI_NEW_GPU_PROFILER
+	return nullptr;
+#else
 	return ImmediateContext.GetCommandQueue().GetQueue();
+#endif
 }
 
 void* FMetalDynamicRHI::RHIGetNativeComputeQueue()
 {
+#if RHI_NEW_GPU_PROFILER
+	return nullptr;
+#else
 	return ImmediateContext.GetCommandQueue().GetQueue();
+#endif
 }
 
 uint16 FMetalDynamicRHI::RHIGetPlatformTextureMaxSampleCount()
@@ -1613,7 +1617,13 @@ void FMetalDynamicRHI::RHIBlockUntilGPUIdle()
 
 IRHICommandContext* FMetalDynamicRHI::RHIGetDefaultContext()
 {
+#if RHI_NEW_GPU_PROFILER
+	// There is no "default"/"immediate" context in Metal now, just return a regular (i.e. parallel) RHI context.
+	// Forward this call through the GDynamicRHI pointer so the validation layer can still do its per-context setup.
+	return static_cast<IRHICommandContext*>(&GDynamicRHI->RHIGetCommandContext(ERHIPipeline::Graphics)->GetLowestLevelContext());
+#else
 	return &ImmediateContext;
+#endif
 }
 
 IRHIComputeContext* FMetalDynamicRHI::RHIGetCommandContext(ERHIPipeline Pipeline)
diff --git a/Engine/Source/Runtime/Apple/MetalRHI/Private/MetalTexture.cpp b/Engine/Source/Runtime/Apple/MetalRHI/Private/MetalTexture.cpp
--- a/Engine/Source/Runtime/Apple/MetalRHI/Private/MetalTexture.cpp
+++ b/Engine/Source/Runtime/Apple/MetalRHI/Private/MetalTexture.cpp
@@ -1590,8 +1590,10 @@ FTextureRHIRef FMetalDynamicRHI::RHIAsyncReallocateTexture2D(FRHITexture* OldTex
 	NewTexture->Initialize(RHICmdList);
 
     // Copy shared mips
-    RHICmdList.EnqueueLambda([this, OldTexture, NewSizeX, NewSizeY, NewTexture, RequestStatus](FRHICommandListImmediate& RHICmdList)
+    RHICmdList.EnqueueLambda([OldTexture, NewSizeX, NewSizeY, NewTexture, RequestStatus](FRHICommandListImmediate& ExecutingCmdList)
     {
+		FMetalRHICommandContext& Context = FMetalRHICommandContext::Get(ExecutingCmdList);
+
         // figure out what mips to schedule
         const uint32 NumSharedMips = FMath::Min(OldTexture->GetNumMips(), NewTexture->GetNumMips());
         const uint32 SourceMipOffset = OldTexture->GetNumMips() - NumSharedMips;
@@ -1616,7 +1618,10 @@ FTextureRHIRef FMetalDynamicRHI::RHIAsyncReallocateTexture2D(FRHITexture* OldTex
             const uint32 MipSizeX = FMath::Max<uint32>(1, NewSizeX >> (MipIndex + DestMipOffset));
             const uint32 MipSizeY = FMath::Max<uint32>(1, NewSizeY >> (MipIndex + DestMipOffset));
 
-            ImmediateContext.CopyFromTextureToTexture(OldTexture->Texture.get(), SliceIndex, MipIndex + SourceMipOffset, Origin, MTL::Size(MipSizeX, MipSizeY, 1), NewTexture->Texture.get(), SliceIndex, MipIndex + DestMipOffset, Origin);
+			Context.CopyFromTextureToTexture(
+				OldTexture->Texture.get(), SliceIndex, MipIndex + SourceMipOffset, Origin, MTL::Size(MipSizeX, MipSizeY, 1),
+				NewTexture->Texture.get(), SliceIndex, MipIndex + DestMipOffset  , Origin
+			);
         }
 
         // Like D3D mark this as complete immediately.
