Remove legacy GPU profiler code paths from Metal RHI

diff --git a/Engine/Source/Runtime/Apple/MetalRHI/Private/MetalBuffer.cpp b/Engine/Source/Runtime/Apple/MetalRHI/Private/MetalBuffer.cpp
--- a/Engine/Source/Runtime/Apple/MetalRHI/Private/MetalBuffer.cpp
+++ b/Engine/Source/Runtime/Apple/MetalRHI/Private/MetalBuffer.cpp
@@ -106,7 +106,6 @@ FMetalSubBufferHeap::FMetalSubBufferHeap(FMetalDevice& MetalDevice, NS::UInteger
 	Options = (MTL::ResourceOptions)FMetalCommandQueue::GetCompatibleResourceOptions(Options);
 	static bool bSupportsHeaps = Device.SupportsFeature(EMetalFeaturesHeaps);
 	NS::UInteger FullSize = Align(Size, Alignment);
-	METAL_GPUPROFILE(FScopedMetalCPUStats CPUStat(FString::Printf(TEXT("AllocBuffer: %llu, %llu"), FullSize, Options)));
 	
 	MTL::StorageMode Storage = (MTL::StorageMode)((Options & MTL::ResourceStorageModeMask) >> MTL::ResourceStorageModeShift);
 #if PLATFORM_MAC
@@ -503,7 +502,6 @@ FMetalSubBufferLinear::FMetalSubBufferLinear(FMetalDevice& MetalDevice, NS::UInt
 {
 	Options = (MTL::ResourceOptions)FMetalCommandQueue::GetCompatibleResourceOptions(Options);
 	NS::UInteger FullSize = Align(Size, Alignment);
-	METAL_GPUPROFILE(FScopedMetalCPUStats CPUStat(FString::Printf(TEXT("AllocBuffer: %llu, %llu"), FullSize, Options)));
 	
 	MTL::StorageMode Storage = (MTL::StorageMode)((Options & MTL::ResourceStorageModeMask) >> MTL::ResourceStorageModeShift);
 	
@@ -668,7 +666,6 @@ FMetalSubBufferMagazine::FMetalSubBufferMagazine(FMetalDevice& MetalDevice, NS::
     }
     
     NS::UInteger FullSize = Align(Size, MinAlign);
-	METAL_GPUPROFILE(FScopedMetalCPUStats CPUStat(FString::Printf(TEXT("AllocBuffer: %llu, %llu"), FullSize, Options)));
 	
 #if PLATFORM_MAC
 	check(Storage != MTL::StorageModeManaged /* Managed memory cannot be safely suballocated! When you overwrite existing data the GPU buffer is immediately disposed of! */);
@@ -1223,8 +1220,6 @@ MTL::Buffer* FMetalBufferPoolPolicyData::CreateResource(FRHICommandListBase&, Cr
 	NS::UInteger StorageMode = (NS::UInteger)Args.Storage << MTL::ResourceStorageModeShift;
 	MTL::ResourceOptions ResourceOptions = FMetalCommandQueue::GetCompatibleResourceOptions(MTL::ResourceOptions(CpuCacheMode | StorageMode | MTL::ResourceHazardTrackingModeUntracked));
 	
-	METAL_GPUPROFILE(FScopedMetalCPUStats CPUStat(FString::Printf(TEXT("AllocBuffer: %llu, %llu"), BufferSize, ResourceOptions)));
-    
 	MTL::Buffer* NewBuf = Args.Device->GetDevice()->newBuffer(BufferSize, ResourceOptions);
 
 	check(!CreationArgumentMap.Contains(NewBuf));
@@ -1300,7 +1295,6 @@ MTLTexturePtr FMetalTexturePool::CreateTexture(FMetalDevice& Device, MTL::Textur
 	}
 	else
 	{
-		METAL_GPUPROFILE(FScopedMetalCPUStats CPUStat(FString::Printf(TEXT("AllocTexture: %s"), TEXT("")/**FString([Desc.GetPtr() description])*/)));
         Texture = NS::TransferPtr(Device.GetDevice()->newTexture(Desc));
 #if STATS || ENABLE_LOW_LEVEL_MEM_TRACKER
 		MetalLLM::LogAllocTexture(Texture.get());
@@ -1941,7 +1935,6 @@ FMetalBufferPtr FMetalResourceHeap::CreateBuffer(uint32 Size, uint32 Alignment,
 	}
 	else
 	{
-		METAL_GPUPROFILE(FScopedMetalCPUStats CPUStat(FString::Printf(TEXT("AllocBuffer: %llu, %llu"), BlockSize, Options)));
         Buffer = FMetalBufferPtr(new FMetalBuffer(Device.GetDevice()->newBuffer(BlockSize, Options), FMetalBuffer::FreePolicy::Owner));
 	}
 	
@@ -2053,7 +2046,6 @@ MTLTexturePtr FMetalResourceHeap::CreateTexture(MTL::TextureDescriptor* Desc, FM
     MTLHeapPtr Heap = GetTextureHeap(Desc, Res);
 	if (Heap)
 	{
-		METAL_GPUPROFILE(FScopedMetalCPUStats CPUStat(FString::Printf(TEXT("AllocTexture: %s"), TEXT("")/**FString([Desc.GetPtr() description])*/)));
 		MTLTexturePtr Texture = NS::TransferPtr(Heap->newTexture(Desc));
 #if STATS || ENABLE_LOW_LEVEL_MEM_TRACKER
 		MetalLLM::LogAllocTexture(Texture.get());
diff --git a/Engine/Source/Runtime/Apple/MetalRHI/Private/MetalCommandBuffer.h b/Engine/Source/Runtime/Apple/MetalRHI/Private/MetalCommandBuffer.h
--- a/Engine/Source/Runtime/Apple/MetalRHI/Private/MetalCommandBuffer.h
+++ b/Engine/Source/Runtime/Apple/MetalRHI/Private/MetalCommandBuffer.h
@@ -9,17 +9,11 @@
 #include "MetalProfiler.h"
 #include "MetalCommandQueue.h"
 
-class FMetalCommandBufferTimer;
-struct FMetalCounterSample;
-class FMetalEventNode;
-
 class FMetalCommandBuffer
 {
 public:
     FMetalCommandBuffer(MTL::CommandBuffer* InCommandBuffer, FMetalCommandQueue& Queue) 
-#if RHI_NEW_GPU_PROFILER
 	: EventStream(Queue.GetProfilerQueue())
-#endif
     {
         CommandBuffer = InCommandBuffer;
 		CommandBuffer->retain();
@@ -30,22 +24,9 @@ public:
 		CommandBuffer->release();
 	}
     
-    FORCEINLINE MTL::CommandBuffer*& GetMTLCmdBuffer() {return CommandBuffer;}
-	
-#if RHI_NEW_GPU_PROFILER == 0
-	FORCEINLINE void SetTimer(FMetalCommandBufferTimer* InTimer)
-	{
-		Timer = InTimer;
-	}
-	
-	FORCEINLINE FMetalCommandBufferTimer* GetTimer()
-	{
-		return Timer;
-	}
-#endif
+    FORCEINLINE MTL::CommandBuffer*& GetMTLCmdBuffer() { return CommandBuffer; }
 	
 #if WITH_RHI_BREADCRUMBS
-#if RHI_NEW_GPU_PROFILER
 	void BeginBreadcrumb(FRHIBreadcrumbNode* Node)
 	{
 		FMetalBreadcrumbTrackerObject TrackerObject;
@@ -63,7 +44,6 @@ public:
 		TrackerObject.Node = Node;
 		BreadcrumbTrackerStream.Add(TrackerObject);
 	}
-#endif
 #endif
     
 	void AddCounterSample(FMetalCounterSamplePtr CounterSample)
@@ -73,30 +53,17 @@ public:
 			return;
 		}
 		
-#if RHI_NEW_GPU_PROFILER == 0
-		for(FMetalEventNode* Node : ActiveEventNodes)
-		{
-			if(EventSampleCounters.Contains(Node))
-			{
-				EventSampleCounters[Node].Add(CounterSample);
-			}
-			else
-			{
-				EventSampleCounters.Add(Node, {CounterSample});
-			}
-		}
-#elif WITH_RHI_BREADCRUMBS
+	#if WITH_RHI_BREADCRUMBS
 		FMetalBreadcrumbTrackerObject TrackerObject;
 		TrackerObject.Type = EMetalBreadcrumbTrackerType::Encode;
 		TrackerObject.CmdBuffer = this;
 		TrackerObject.CounterSample = CounterSample;
 		BreadcrumbTrackerStream.Add(TrackerObject);
-#endif
+	#endif
 		
 		CounterSamples.Add(CounterSample);
 	}
 	
-#if RHI_NEW_GPU_PROFILER
 	void SetBeginWorkTimestamp(uint64_t* Timestamp)
 	{
 		BeginWorkTimestamp = Timestamp;
@@ -131,16 +98,11 @@ public:
 		BeginEvents.Reset();
 		Destination.Append(MoveTemp(EventStream));
 	}
-#endif
 	
 	TArray<FMetalRHIRenderQuery*> TimestampQueries;
 	TArray<FMetalRHIRenderQuery*> OcclusionQueries;
 	TArray<FMetalCounterSamplePtr> CounterSamples;
-	
-#if RHI_NEW_GPU_PROFILER == 0
-	TMap<FMetalEventNode*, TArray<FMetalCounterSamplePtr>> EventSampleCounters;
-	TArray<FMetalEventNode*> ActiveEventNodes;
-#else
+
 	UE::RHI::GPUProfiler::FEventStream EventStream;
 	TArray<UE::RHI::GPUProfiler::FEvent::FBeginWork*, TInlineAllocator<8>> BeginEvents;
 	
@@ -150,9 +112,7 @@ public:
 	
 	uint64_t* BeginWorkTimestamp;
 	uint64_t* EndWorkTimestamp;
-#endif
 	
 private:
-    MTL::CommandBuffer* CommandBuffer;	
-	FMetalCommandBufferTimer* Timer = nullptr;
+    MTL::CommandBuffer* CommandBuffer;
 };
diff --git a/Engine/Source/Runtime/Apple/MetalRHI/Private/MetalCommandEncoder.cpp b/Engine/Source/Runtime/Apple/MetalRHI/Private/MetalCommandEncoder.cpp
--- a/Engine/Source/Runtime/Apple/MetalRHI/Private/MetalCommandEncoder.cpp
+++ b/Engine/Source/Runtime/Apple/MetalRHI/Private/MetalCommandEncoder.cpp
@@ -6,7 +6,6 @@
 
 #include "MetalCommandEncoder.h"
 #include "MetalCommandBuffer.h"
-#include "MetalCommandList.h"
 #include "MetalRHIPrivate.h"
 #include "MetalShaderTypes.h"
 #include "MetalGraphicsPipelineState.h"
@@ -66,17 +65,14 @@ FString FMetalCommandData::ToString() const
 
 #pragma mark - Public C++ Boilerplate -
 
-FMetalCommandEncoder::FMetalCommandEncoder(FMetalDevice& MetalDevice, FMetalCommandList& CmdList)
-: Device(MetalDevice)
-, CommandList(CmdList)
-, RingBuffer(Device, EncoderRingBufferSize, BufferOffsetAlignment, FMetalCommandQueue::GetCompatibleResourceOptions((MTL::ResourceOptions)(MTL::ResourceHazardTrackingModeUntracked | BUFFER_RESOURCE_STORAGE_MANAGED)))
-, RenderPassDesc(nullptr)
-, EncoderFence(nullptr)
-#if ENABLE_METAL_GPUPROFILE
-, CommandBufferStats(nullptr)
-#endif
-, EncoderNum(0)
-, CmdBufIndex(0)
+FMetalCommandEncoder::FMetalCommandEncoder(FMetalDevice& MetalDevice, FMetalCommandQueue& CommandQueue)
+	: Device(MetalDevice)
+	, CommandQueue(CommandQueue)
+	, RingBuffer(Device, EncoderRingBufferSize, BufferOffsetAlignment, FMetalCommandQueue::GetCompatibleResourceOptions((MTL::ResourceOptions)(MTL::ResourceHazardTrackingModeUntracked | BUFFER_RESOURCE_STORAGE_MANAGED)))
+	, RenderPassDesc(nullptr)
+	, EncoderFence(nullptr)
+	, EncoderNum(0)
+	, CmdBufIndex(0)
 {
 	for (uint32 Frequency = 0; Frequency < uint32(MTL::FunctionTypeObject)+1; Frequency++)
 	{
@@ -242,20 +238,12 @@ void FMetalCommandEncoder::StartCommandBuffer(void)
 	if (!CommandBuffer)
 	{
 		CmdBufIndex++;
-		CommandBuffer = CommandList.GetCommandQueue().CreateCommandBuffer();
+		CommandBuffer = CommandQueue.CreateCommandBuffer();
 		
 		if (DebugGroups.Num())
 		{
 			CommandBuffer->GetMTLCmdBuffer()->setLabel(DebugGroups.Last());
 		}
-		
-	#if ENABLE_METAL_GPUPROFILE && RHI_NEW_GPU_PROFILER == 0
-		FMetalProfiler* Profiler = FMetalProfiler::GetProfiler();
-		if (Profiler)
-		{
-			CommandBufferStats = Profiler->AllocateCommandBuffer(CommandBuffer->GetMTLCmdBuffer(), 0);
-		}
-	#endif
 	}
 }
 	
@@ -299,16 +287,6 @@ void FMetalCommandEncoder::EndCommandBuffer(FMetalRHICommandContext* Context)
 	}
 #endif
 	
-#if ENABLE_METAL_GPUPROFILE && RHI_NEW_GPU_PROFILER == 0
-	if(CommandBufferStats)
-	{
-		CommandBufferStats->End(CommandBuffer->GetMTLCmdBuffer());
-		CommandBufferStats = nullptr;
-	}
-#endif
-	
-	CommandList.FinalizeCommandBuffer(CommandBuffer);
-	
 	CommandBuffer = nullptr;
 	EncoderNum = 0;
 }
@@ -421,7 +399,7 @@ void FMetalCommandEncoder::BeginRenderCommandEncoding(FMetalCounterSampler* Samp
         }
 	}
 
-	EncoderFence = CommandList.GetCommandQueue().CreateFence(Label);
+	EncoderFence = CommandQueue.CreateFence(Label);
 }
 
 void FMetalCommandEncoder::BeginRenderCommandEncoding(MTL::RenderPassDescriptor* InRenderPassDesc, MTLParallelRenderCommandEncoderPtr ParallelEncoder)
@@ -448,7 +426,7 @@ void FMetalCommandEncoder::BeginRenderCommandEncoding(MTL::RenderPassDescriptor*
 		}
 	}
 
-	//EncoderFence = CommandList.GetCommandQueue().CreateFence(Label);
+	//EncoderFence = CommandQueue.CreateFence(Label);
 }
 
 MTLParallelRenderCommandEncoderPtr FMetalCommandEncoder::BeginParallelRenderCommandEncoding(FMetalCounterSampler* Sampler)
@@ -507,7 +485,7 @@ MTLParallelRenderCommandEncoderPtr FMetalCommandEncoder::BeginParallelRenderComm
 		}
 	}
 
-	EncoderFence = CommandList.GetCommandQueue().CreateFence(Label);
+	EncoderFence = CommandQueue.CreateFence(Label);
 	
 	return ParallelRenderCommandEncoder;
 }
@@ -555,7 +533,7 @@ void FMetalCommandEncoder::BeginComputeCommandEncoding(MTL::DispatchType Dispatc
         }
 	}
 	
-	EncoderFence = CommandList.GetCommandQueue().CreateFence(Label);
+	EncoderFence = CommandQueue.CreateFence(Label);
 }
 
 void FMetalCommandEncoder::BeginBlitCommandEncoding(FMetalCounterSampler* Sampler)
@@ -598,7 +576,7 @@ void FMetalCommandEncoder::BeginBlitCommandEncoding(FMetalCounterSampler* Sample
         }
 	}
 	
-	EncoderFence = CommandList.GetCommandQueue().CreateFence(Label);
+	EncoderFence = CommandQueue.CreateFence(Label);
 }
 
 #if METAL_RHI_RAYTRACING
@@ -631,7 +609,7 @@ void FMetalCommandEncoder::BeginAccelerationStructureCommandEncoding(void)
 		}
 	}
 
-	EncoderFence = CommandList.GetCommandQueue().CreateFence(Label);
+	EncoderFence = CommandQueue.CreateFence(Label);
 }
 #endif // METAL_RHI_RAYTRACING
 
@@ -1022,13 +1000,6 @@ void FMetalCommandEncoder::PopDebugGroup(void)
 	}
 }
 
-#if ENABLE_METAL_GPUPROFILE
-FMetalCommandBufferStats* FMetalCommandEncoder::GetCommandBufferStats(void)
-{
-	return CommandBufferStats;
-}
-#endif
-
 #pragma mark - Public Render State Mutators -
 
 void FMetalCommandEncoder::SetRenderPassDescriptor(MTL::RenderPassDescriptor* RenderPass)
diff --git a/Engine/Source/Runtime/Apple/MetalRHI/Private/MetalCommandEncoder.h b/Engine/Source/Runtime/Apple/MetalRHI/Private/MetalCommandEncoder.h
--- a/Engine/Source/Runtime/Apple/MetalRHI/Private/MetalCommandEncoder.h
+++ b/Engine/Source/Runtime/Apple/MetalRHI/Private/MetalCommandEncoder.h
@@ -9,11 +9,9 @@
 #include "MetalProfiler.h"
 #include "MetalCounterSampler.h"
 
-class FMetalCommandList;
 class FMetalCommandQueue;
 class FMetalGraphicsPipelineState;
 class FMetalDevice;
-struct FMetalCommandBufferStats;
 class FMetalRHICommandContext;
 class FMetalShaderPipeline;
 typedef TSharedPtr<FMetalShaderPipeline, ESPMode::ThreadSafe> FMetalShaderPipelinePtr;
@@ -151,7 +149,7 @@ public:
 #pragma mark - Public C++ Boilerplate -
 
 	/** Default constructor */
-	FMetalCommandEncoder(FMetalDevice& Device, FMetalCommandList& CmdList);
+	FMetalCommandEncoder(FMetalDevice& Device, FMetalCommandQueue& CommandQueue);
 	
 	/** Release Function */
 	void Release(void);
@@ -322,11 +320,6 @@ public:
 	
 	/* Pop the latest named string off of the stack. */
 	void PopDebugGroup(void);
-	
-#if ENABLE_METAL_GPUPROFILE
-	/* Get the command-buffer stats object. */
-	FMetalCommandBufferStats* GetCommandBufferStats(void);
-#endif
 
 #pragma mark - Public Render State Mutators -
 
@@ -602,7 +595,7 @@ public:
 	
 #pragma mark - Private Member Variables -
 	FMetalDevice& Device;
-	FMetalCommandList& CommandList;
+	FMetalCommandQueue& CommandQueue;
 
 	FMetalBufferBindings ShaderBuffers[int(MTL::FunctionTypeObject)+1];
 	
@@ -625,9 +618,6 @@ public:
 #endif // METAL_RHI_RAYTRACING
 	
 	TRefCountPtr<FMetalFence> EncoderFence;
-#if ENABLE_METAL_GPUPROFILE
-	FMetalCommandBufferStats* CommandBufferStats;
-#endif
 
 	TArray<NS::String*> DebugGroups;
     
diff --git a/Engine/Source/Runtime/Apple/MetalRHI/Private/MetalCommandList.cpp b/Engine/Source/Runtime/Apple/MetalRHI/Private/MetalCommandList.cpp
--- a/Engine/Source/Runtime/Apple/MetalRHI/Private/MetalCommandList.cpp
+++ b/Engine/Source/Runtime/Apple/MetalRHI/Private/MetalCommandList.cpp
@@ -4,7 +4,6 @@
 	MetalCommandList.cpp: Metal command buffer list wrapper.
 =============================================================================*/
 
-#include "MetalCommandList.h"
 #include "MetalRHIPrivate.h"
 #include "MetalShaderTypes.h"
 #include "MetalGraphicsPipelineState.h"
@@ -17,14 +16,6 @@
 #if PLATFORM_IOS
 extern bool GIsSuspended;
 #endif
-
-FMetalCommandList::FMetalCommandList(FMetalCommandQueue& InCommandQueue)
-	: CommandQueue(InCommandQueue)
-{}
-
-FMetalCommandList::~FMetalCommandList(void)
-{
-}
 	
 #pragma mark - Public Command List Mutators -
 
@@ -182,7 +173,7 @@ static __attribute__ ((optnone)) void HandleIntelMetalCommandBufferError(MTL::Co
 	HandleMetalCommandBufferError(CompletedBuffer);
 }
 
-void FMetalCommandList::HandleMetalCommandBufferFailure(MTL::CommandBuffer* CompletedBuffer)
+void FMetalCommandQueue::HandleMetalCommandBufferFailure(MTL::CommandBuffer* CompletedBuffer)
 {
 	if (CompletedBuffer->error()->domain()->isEqualToString(NS::String::string("MTLCommandBufferErrorDomain", NS::UTF8StringEncoding)))
 	{
@@ -204,15 +195,3 @@ void FMetalCommandList::HandleMetalCommandBufferFailure(MTL::CommandBuffer* Comp
 		ReportMetalCommandBufferFailure(CompletedBuffer, TEXT("Unknown"));
 	}
 }
-
-void FMetalCommandList::FinalizeCommandBuffer(FMetalCommandBuffer* Buffer)
-{
-#if RHI_NEW_GPU_PROFILER == 0
-	check(Buffer);
-	
-	FMetalCommandBufferTimer& Timer = FMetalCommandBufferTimer::GetFrameBufferTimer();
-	Timer.Submit();
-	
-	Buffer->SetTimer(&Timer);
-#endif
-}
diff --git a/Engine/Source/Runtime/Apple/MetalRHI/Private/MetalCommandList.h b/Engine/Source/Runtime/Apple/MetalRHI/Private/MetalCommandList.h
deleted file mode 100644
diff --git a/Engine/Source/Runtime/Apple/MetalRHI/Private/MetalCommandQueue.cpp b/Engine/Source/Runtime/Apple/MetalRHI/Private/MetalCommandQueue.cpp
--- a/Engine/Source/Runtime/Apple/MetalRHI/Private/MetalCommandQueue.cpp
+++ b/Engine/Source/Runtime/Apple/MetalRHI/Private/MetalCommandQueue.cpp
@@ -6,7 +6,6 @@
 
 #include "MetalCommandQueue.h"
 #include "MetalCommandBuffer.h"
-#include "MetalCommandList.h"
 #include "MetalDevice.h"
 #include "MetalFence.h"
 #include "MetalProfiler.h"
@@ -25,13 +24,11 @@ bool GMetalCommandBufferDebuggingEnabled = 0;
 
 #pragma mark - Public C++ Boilerplate -
 
-#if RHI_NEW_GPU_PROFILER
 FMetalTiming::FMetalTiming(FMetalCommandQueue& Queue) 
 	: Queue(Queue)
 	, EventStream(Queue.GetProfilerQueue())
 {
 }
-#endif
 
 FMetalCommandQueue::FMetalCommandQueue(FMetalDevice& MetalDevice, uint32 const MaxNumCommandBuffers /* = 0 */)
 	: Device(MetalDevice)
@@ -203,7 +200,6 @@ void FMetalCommandQueue::InsertDebugCaptureBoundary(void)
 	CommandQueue->insertDebugCaptureBoundary();
 }
 
-#if RHI_NEW_GPU_PROFILER
 UE::RHI::GPUProfiler::FQueue FMetalCommandQueue::GetProfilerQueue() const
 {
 	UE::RHI::GPUProfiler::FQueue Queue;
@@ -215,4 +211,3 @@ UE::RHI::GPUProfiler::FQueue FMetalCommandQueue::GetProfilerQueue() const
 	
 	return Queue;
 }
-#endif
diff --git a/Engine/Source/Runtime/Apple/MetalRHI/Private/MetalCommandQueue.h b/Engine/Source/Runtime/Apple/MetalRHI/Private/MetalCommandQueue.h
--- a/Engine/Source/Runtime/Apple/MetalRHI/Private/MetalCommandQueue.h
+++ b/Engine/Source/Runtime/Apple/MetalRHI/Private/MetalCommandQueue.h
@@ -9,14 +9,10 @@
 #include "Containers/LockFreeList.h"
 
 class FMetalCommandBuffer;
-class FMetalCommandList;
 class FMetalFence;
 class FMetalDevice;
-struct FMetalCounterSample;
 struct FMetalBreadcrumbEvent;
-class FMetalEventNode;
 
-#if RHI_NEW_GPU_PROFILER
 class FMetalTiming
 {
 public:
@@ -30,7 +26,6 @@ public:
 
 	FMetalTiming(FMetalCommandQueue& Queue);
 };
-#endif
 
 /**
  * FMetalCommandQueue:
@@ -122,16 +117,16 @@ public:
 	TArray<FMetalRHIRenderQuery*> TimestampQueries;
 	TArray<FMetalRHIRenderQuery*> OcclusionQueries;
 	TArray<FMetalCounterSamplePtr> CounterSamples;
-	
-#if RHI_NEW_GPU_PROFILER == 0
-	TMap<FMetalEventNode*, TArray<FMetalCounterSamplePtr>> EventSampleCounters;
-#endif
-	
-#if RHI_NEW_GPU_PROFILER
+
 	UE::RHI::GPUProfiler::FQueue GetProfilerQueue() const;
 	FMetalTiming* Timing = nullptr;
-#endif
 	
+	/**
+	 * Command buffer failure reporting function.
+	 * @param CompletedBuffer The buffer to check for failure.
+	 */
+	static void HandleMetalCommandBufferFailure(MTL::CommandBuffer* CompletedBuffer);
+
 private:
 #pragma mark - Private Member Variables -
 	FMetalDevice& Device;
diff --git a/Engine/Source/Runtime/Apple/MetalRHI/Private/MetalCommands.cpp b/Engine/Source/Runtime/Apple/MetalRHI/Private/MetalCommands.cpp
--- a/Engine/Source/Runtime/Apple/MetalRHI/Private/MetalCommands.cpp
+++ b/Engine/Source/Runtime/Apple/MetalRHI/Private/MetalCommands.cpp
@@ -297,8 +297,6 @@ void FMetalRHICommandContext::RHIDispatchComputeShader(uint32 ThreadGroupCountX,
 	TRefCountPtr<FMetalComputeShader> ComputeShader = StateCache.GetComputeShader();
 	check(ComputeShader);
 	
-	METAL_GPUPROFILE(FMetalProfiler::GetProfiler()->EncodeDispatch(CurrentEncoder.GetCommandBufferStats(), __FUNCTION__));
-	
 	MTL::Size ThreadgroupCounts = MTL::Size(ComputeShader->NumThreadsX, ComputeShader->NumThreadsY, ComputeShader->NumThreadsZ);
 	check(ComputeShader->NumThreadsX > 0 && ComputeShader->NumThreadsY > 0 && ComputeShader->NumThreadsZ > 0);
 	MTL::Size Threadgroups = MTL::Size(ThreadGroupCountX, ThreadGroupCountY, ThreadGroupCountZ);
@@ -343,8 +341,7 @@ void FMetalRHICommandContext::RHIDispatchIndirectComputeShader(FRHIBuffer* Argum
 			
 			TRefCountPtr<FMetalComputeShader> ComputeShader = StateCache.GetComputeShader();
 			check(ComputeShader);
-			
-			METAL_GPUPROFILE(FMetalProfiler::GetProfiler()->EncodeDispatch(CurrentEncoder.GetCommandBufferStats(), __FUNCTION__));
+
 			MTL::Size ThreadgroupCounts = MTL::Size(ComputeShader->NumThreadsX, ComputeShader->NumThreadsY, ComputeShader->NumThreadsZ);
 			check(ComputeShader->NumThreadsX > 0 && ComputeShader->NumThreadsY > 0 && ComputeShader->NumThreadsZ > 0);
 
@@ -970,7 +967,6 @@ void FMetalRHICommandContext::RHIDrawPrimitive(uint32 BaseVertexIndex, uint32 Nu
 	else
 #endif
 	{
-		METAL_GPUPROFILE(FMetalProfiler::GetProfiler()->EncodeDraw(CurrentEncoder.GetCommandBufferStats(), __FUNCTION__, NumPrimitives, NumVertices, NumInstances));
 		CurrentEncoder.GetRenderCommandEncoder()->drawPrimitives(TranslatePrimitiveType(PrimitiveType), BaseVertexIndex, NumVertices, NumInstances);
 	}
 }
@@ -1007,7 +1003,6 @@ void FMetalRHICommandContext::RHIDrawPrimitiveIndirect(FRHIBuffer* ArgumentBuffe
 		else
 #endif
 		{
-			METAL_GPUPROFILE(FMetalProfiler::GetProfiler()->EncodeDraw(CurrentEncoder.GetCommandBufferStats(), __FUNCTION__, 1, 1, 1));
 			CurrentEncoder.GetRenderCommandEncoder()->drawPrimitives(TranslatePrimitiveType(PrimitiveType),
 																	 TheBackingBuffer->GetMTLBuffer(), TheBackingBuffer->GetOffset() + ArgumentOffset);
 		}
@@ -1152,8 +1147,7 @@ void FMetalRHICommandContext::RHIDrawIndexedPrimitive(FRHIBuffer* IndexBufferRHI
 #endif
 	{
 		FMetalBufferPtr IndexBufferPtr = IndexBuffer->GetCurrentBuffer();
-		
-		METAL_GPUPROFILE(FMetalProfiler::GetProfiler()->EncodeDraw(CurrentEncoder.GetCommandBufferStats(), __FUNCTION__, NumPrimitives, NumVertices, NumInstances));
+
 		if (GRHISupportsBaseVertexIndex && GRHISupportsFirstInstance)
 		{
 			CurrentEncoder.GetRenderCommandEncoder()->drawIndexedPrimitives(TranslatePrimitiveType(PrimitiveType), NumIndices,
@@ -1218,12 +1212,14 @@ void FMetalRHICommandContext::RHIDrawIndexedPrimitiveIndirect(FRHIBuffer* IndexB
 		else
 #endif
 		{
-			METAL_GPUPROFILE(FMetalProfiler::GetProfiler()->EncodeDraw(CurrentEncoder.GetCommandBufferStats(), __FUNCTION__, 1, 1, 1));
-			CurrentEncoder.GetRenderCommandEncoder()->drawIndexedPrimitives(TranslatePrimitiveType(PrimitiveType), 															IndexBuffer->GetIndexType(),
-																			TheBackingIndexBuffer->GetMTLBuffer(),
-																			TheBackingIndexBuffer->GetOffset(),
-																			TheBackingBuffer->GetMTLBuffer(), 
-																			TheBackingBuffer->GetOffset() + ArgumentOffset);
+			CurrentEncoder.GetRenderCommandEncoder()->drawIndexedPrimitives(
+				TranslatePrimitiveType(PrimitiveType),
+				IndexBuffer->GetIndexType(),
+				TheBackingIndexBuffer->GetMTLBuffer(),
+				TheBackingIndexBuffer->GetOffset(),
+				TheBackingBuffer->GetMTLBuffer(), 
+				TheBackingBuffer->GetOffset() + ArgumentOffset
+			);
 		}
 	}
 	else
diff --git a/Engine/Source/Runtime/Apple/MetalRHI/Private/MetalDynamicRHI.h b/Engine/Source/Runtime/Apple/MetalRHI/Private/MetalDynamicRHI.h
--- a/Engine/Source/Runtime/Apple/MetalRHI/Private/MetalDynamicRHI.h
+++ b/Engine/Source/Runtime/Apple/MetalRHI/Private/MetalDynamicRHI.h
@@ -275,7 +275,6 @@ public:
 	// Called when the GPU has crashed. This function will not return.
 	void ProcessInterruptQueueOnGPUCrash();
 
-#if RHI_NEW_GPU_PROFILER
 	struct FMetalTimingArray : public TArray<TUniquePtr<FMetalTiming>, TInlineAllocator<GMetalMaxNumQueues>>
 	{
 		FMetalTiming* CreateNew(FMetalCommandQueue& Queue)
@@ -285,7 +284,6 @@ public:
 	};
 
 	FMetalTimingArray CurrentTimingPerQueue;
-#endif
 	
 	void InitializeSubmissionPipe();
 	void ShutdownSubmissionPipe();
@@ -306,9 +304,6 @@ public:
 private:
 	FMetalDevice* Device;
 	FTextureMemoryStats MemoryStats;
-#if (RHI_NEW_GPU_PROFILER == 0)
-	FMetalRHICommandContext ImmediateContext;
-#endif
 	TMap<uint32, FVertexDeclarationRHIRef> VertexDeclarationCache;
 	TLockFreePointerListUnordered<FMetalRHICommandContext, PLATFORM_CACHE_LINE_SIZE> MetalCommandContextPool;
 	
diff --git a/Engine/Source/Runtime/Apple/MetalRHI/Private/MetalPipeline.cpp b/Engine/Source/Runtime/Apple/MetalRHI/Private/MetalPipeline.cpp
--- a/Engine/Source/Runtime/Apple/MetalRHI/Private/MetalPipeline.cpp
+++ b/Engine/Source/Runtime/Apple/MetalRHI/Private/MetalPipeline.cpp
@@ -1084,7 +1084,6 @@ static FMetalShaderPipelinePtr CreateMTLRenderPipeline(FMetalDevice& Device, boo
 		}
 
 		{
-			METAL_GPUPROFILE(FScopedMetalCPUStats CPUStat(FString::Printf(TEXT("NewRenderPipeline: %s"), TEXT("")/**FString([RenderPipelineDesc.GetPtr() description])*/)));
 			RelatePipelineStateToCache(RenderPipelineDesc.get(), &RenderOption);
             
 #if METAL_DEBUG_OPTIONS
@@ -1184,7 +1183,6 @@ static FMetalShaderPipelinePtr CreateMTLRenderPipeline(FMetalDevice& Device, boo
 
         {
             NS::Error* RenderError = nullptr;
-            METAL_GPUPROFILE(FScopedMetalCPUStats CPUStat(FString::Printf(TEXT("NewMeshRenderPipeline: %s"), TEXT("")/**FString([RenderPipelineDesc.GetPtr() description])*/)));
 #if 0
             // Binary Archive does not support Mesh shaders...
             RelatePipelineStateToCache(MeshPipelineDesc, &RenderOption);
@@ -1263,7 +1261,6 @@ static FMetalShaderPipelinePtr CreateMTLRenderPipeline(FMetalDevice& Device, boo
 
         {
             NS::Error* RenderError;
-            METAL_GPUPROFILE(FScopedMetalCPUStats CPUStat(FString::Printf(TEXT("NewMeshRenderPipeline: %s"), TEXT("")/**FString([RenderPipelineDesc.GetPtr() description])*/)));
 #if 0
             // Binary Archive does not support Mesh shaders...
             RelatePipelineStateToCache(MeshPipelineDesc, &RenderOption);
diff --git a/Engine/Source/Runtime/Apple/MetalRHI/Private/MetalProfiler.cpp b/Engine/Source/Runtime/Apple/MetalRHI/Private/MetalProfiler.cpp
--- a/Engine/Source/Runtime/Apple/MetalRHI/Private/MetalProfiler.cpp
+++ b/Engine/Source/Runtime/Apple/MetalRHI/Private/MetalProfiler.cpp
@@ -35,9 +35,6 @@ DEFINE_STAT(STAT_MetalMakeDrawableTime);
 DEFINE_STAT(STAT_MetalBufferPageOffTime);
 DEFINE_STAT(STAT_MetalTexturePageOnTime);
 DEFINE_STAT(STAT_MetalTexturePageOffTime);
-DEFINE_STAT(STAT_MetalGPUWorkTime);
-DEFINE_STAT(STAT_MetalGPUIdleTime);
-DEFINE_STAT(STAT_MetalPresentTime);
 DEFINE_STAT(STAT_MetalCustomPresentTime);
 DEFINE_STAT(STAT_MetalCommandBufferCreatedPerFrame);
 DEFINE_STAT(STAT_MetalCommandBufferCommittedPerFrame);
@@ -57,12 +54,7 @@ DEFINE_STAT(STAT_MetalUniformBytesPerFrame);
 
 DEFINE_STAT(STAT_MetalTempAllocatorAllocatedMemory);
 
-int64 volatile GMetalTexturePageOnTime = 0;
-int64 volatile GMetalGPUWorkTime = 0;
-int64 volatile GMetalGPUIdleTime = 0;
-int64 volatile GMetalPresentTime = 0;
-
-#if RHI_NEW_GPU_PROFILER && WITH_RHI_BREADCRUMBS
+#if WITH_RHI_BREADCRUMBS
 FMetalBreadcrumbProfiler* FMetalBreadcrumbProfiler::Instance = nullptr;
 #endif
 
@@ -71,846 +63,7 @@ void WriteString(FArchive* OutputFile, const char* String)
 	OutputFile->Serialize((void*)String, sizeof(ANSICHAR)*FCStringAnsi::Strlen(String));
 }
 
-#if RHI_NEW_GPU_PROFILER == 0
-
-void FMetalCommandBufferTimer::AddTiming(FMetalCommandBufferTiming Timing)
-{
-	FScopeLock Lock(&Mutex);
-	Counter--;
-	Timings.Add(Timing);
-	
-	if(Counter == 0 && bFrameEnded)
-	{
-		RecordFrame();
-	}
-}
-
-void FMetalCommandBufferTimer::FrameEnd()
-{
-	FScopeLock Lock(&Mutex);
-	bFrameEnded = true;
-	
-	if(Counter == 0)
-	{
-		RecordFrame();
-	}
-}
-
-void FMetalCommandBufferTimer::RecordFrame()
-{
-	Timings.Sort();
-	
-	FMetalCommandBufferTiming LastBufferTiming;
-	const double CyclesPerSecond = 1.0 / FPlatformTime::GetSecondsPerCycle();
-
-	double RunningFrameTimeSeconds = 0.0;
-	uint64 FrameStartGPUCycles = 0;
-	uint64 FrameEndGPUCycles = 0;
-
-	CFTimeInterval FirstStartTime = 0.0;
-
-	// Add the timings excluding any overlapping time
-	for (const FMetalCommandBufferTiming& Timing : Timings)
-	{
-		if (FirstStartTime == 0.0)
-		{
-			FirstStartTime = Timing.StartTime;
-			LastBufferTiming = Timing;
-		}
-		
-		// Only process if the previous buffer finished before the end of this one
-		if (LastBufferTiming.EndTime < Timing.EndTime)
-		{
-			// Check if the end of the previous buffer finished before the start of this one
-			if (LastBufferTiming.EndTime > Timing.StartTime)
-			{
-				// Segment from end of last buffer to end of current
-				RunningFrameTimeSeconds += Timing.EndTime - LastBufferTiming.EndTime;
-			}
-			else
-			{
-				// Full timing of this buffer
-				RunningFrameTimeSeconds += Timing.EndTime - Timing.StartTime;
-			}
-
-			LastBufferTiming = Timing;
-		}
-	}
-	
-	FrameStartGPUCycles = FirstStartTime * CyclesPerSecond;
-	FrameEndGPUCycles = LastBufferTiming.EndTime * CyclesPerSecond;
-	
-	uint64 FrameGPUTimeCycles = uint64(CyclesPerSecond * RunningFrameTimeSeconds);
-	GRHIGPUFrameTimeHistory.PushFrameCycles(1.0, RunningFrameTimeSeconds);
-	
-#if STATS
-	FPlatformAtomics::AtomicStore_Relaxed(&GMetalGPUWorkTime, FrameGPUTimeCycles);
-	int64 FrameIdleTimeCycles = int64(FrameEndGPUCycles - FrameStartGPUCycles - FrameGPUTimeCycles);
-	FPlatformAtomics::AtomicStore_Relaxed(&GMetalGPUIdleTime, FrameIdleTimeCycles);
-#endif //STATS
-	
-	FMetalDynamicRHI::Get().DeferredDelete([InTimer=this](){
-		delete InTimer;
-	});
-}
-
-void FMetalCommandBufferTimer::RecordPresent(MTL::CommandBuffer* CommandBuffer)
-{
-	const CFTimeInterval GpuStartTimeSeconds = CommandBuffer->GPUStartTime();
-	const CFTimeInterval GpuEndTimeSeconds = CommandBuffer->GPUStartTime();
-	const double CyclesPerSecond = 1.0 / FPlatformTime::GetSecondsPerCycle();
-	uint64 StartTimeCycles = uint64(GpuStartTimeSeconds * CyclesPerSecond);
-	uint64 EndTimeCycles = uint64(GpuEndTimeSeconds * CyclesPerSecond);
-	int64 Time = int64(EndTimeCycles - StartTimeCycles);
-	FPlatformAtomics::AtomicStore_Relaxed(&GMetalPresentTime, Time);
-}
-// END WARNING
-
-void FMetalCommandBufferTimer::ResetFrameBufferTimings()
-{
-	if(Timer)
-	{
-		Timer->FrameEnd();
-	}
-	
-	Timer = new FMetalCommandBufferTimer();
-}
-
-FMetalCommandBufferTimer& FMetalCommandBufferTimer::GetFrameBufferTimer()
-{
-	if(!Timer)
-	{
-		Timer = new FMetalCommandBufferTimer();
-	}
-	return *Timer;
-}
-
-FMetalCommandBufferTimer* FMetalCommandBufferTimer::Timer = nullptr;
-
-void FMetalGPUTiming::PlatformStaticInitialize(void* UserData)
-{
-	// Are the static variables initialized?
-	if ( !GAreGlobalsInitialized )
-	{
-		GIsSupported = true;
-		SetTimingFrequency(1000 * 1000 * 1000);
-		GAreGlobalsInitialized = true;
-		
-		FMetalDevice& Device = ((FMetalRHICommandContext*)UserData)->GetDevice();
-		
-		MTL::Timestamp CPUTimeStamp, GPUTimestamp;
-		Device.GetDevice()->sampleTimestamps(&CPUTimeStamp, &GPUTimestamp);
-
-		FGPUTiming::SetCalibrationTimestamp({ GPUTimestamp, CPUTimeStamp });
-	}
-}
-
-FMetalEventNode::~FMetalEventNode()
-{
-
-}
-
-float FMetalEventNode::GetTiming()
-{
-	SyncPoint->Wait();
-	return FPlatformTime::ToSeconds(EndTime - StartTime);
-}
-
-void FMetalEventNode::StartTiming()
-{
-	StartTime = 0;
-	EndTime = 0;
-
-	Context.StartTiming(this);
-}
-
-void FMetalEventNode::StopTiming()
-{
-	Context.EndTiming(this);
-}
-
-bool MetalGPUProfilerIsInSafeThread()
-{
-	return (GIsMetalInitialized && !GIsRHIInitialized) || (IsInRHIThread() || IsInActualRenderingThread());
-}
-	
-/** Start this frame of per tracking */
-void FMetalEventNodeFrame::StartFrame()
-{
-	RootNode->StartTiming();
-}
-
-/** End this frame of per tracking, but do not block yet */
-void FMetalEventNodeFrame::EndFrame()
-{
-	RootNode->StopTiming();
-}
-
-/** Calculates root timing base frequency (if needed by this RHI) */
-float FMetalEventNodeFrame::GetRootTimingResults()
-{
-	return RootNode->GetTiming();
-}
-
-void FMetalEventNodeFrame::LogDisjointQuery()
-{
-	
-}
-
-FGPUProfilerEventNode* FMetalGPUProfiler::CreateEventNode(const TCHAR* InName, FGPUProfilerEventNode* InParent)
-{
-#if ENABLE_METAL_GPUPROFILE
-	FMetalEventNode* EventNode = new FMetalEventNode(Context, InName, InParent, false, false);
-	return EventNode;
-#else
-	return nullptr;
-#endif
-}
-
-void FMetalGPUProfiler::Cleanup()
-{
-	
-}
-
-void FMetalGPUProfiler::PushEvent(const TCHAR* Name, FColor Color)
-{
-	if(MetalGPUProfilerIsInSafeThread())
-	{
-		FGPUProfiler::PushEvent(Name, Color);
-	}
-}
-
-void FMetalGPUProfiler::PopEvent()
-{
-	if(MetalGPUProfilerIsInSafeThread())
-	{
-		FGPUProfiler::PopEvent();
-	}
-}
-
-//TGlobalResource<FVector4VertexDeclaration> GMetalVector4VertexDeclaration;
-TGlobalResource<FTexture> GMetalLongTaskRT;
-
-void FMetalGPUProfiler::BeginFrame()
-{
-	if(GTriggerGPUProfile)
-	{
-		bTrackingEvents = true;
-		bLatchedGProfilingGPU = true;
-		GTriggerGPUProfile = false;
-	}
-	
-	if (bLatchedGProfilingGPU)
-	{
-		// Start tracking the frame
-		CurrentEventNodeFrame = new FMetalEventNodeFrame(Context, GTriggerGPUProfile);
-		CurrentEventNodeFrame->StartFrame();
-	}
-}
-
-void FMetalGPUProfiler::EndFrame()
-{
-	dispatch_async(dispatch_get_global_queue(QOS_CLASS_UTILITY, 0), ^{
-#if PLATFORM_MAC
-	FPlatformMisc::UpdateDriverMonitorStatistics(Context.GetDevice().GetDeviceIndex());
-#endif
-	});
-	
-#if STATS
-	SET_CYCLE_COUNTER(STAT_MetalTexturePageOnTime, GMetalTexturePageOnTime);
-	GMetalTexturePageOnTime = 0;
-	
-	SET_CYCLE_COUNTER(STAT_MetalGPUIdleTime, GMetalGPUIdleTime);
-	SET_CYCLE_COUNTER(STAT_MetalGPUWorkTime, GMetalGPUWorkTime);
-	SET_CYCLE_COUNTER(STAT_MetalPresentTime, GMetalPresentTime);
-#endif
-	
-	if(CurrentEventNodeFrame)
-	{
-		GDynamicRHI->RHIBlockUntilGPUIdle();
-		
-		CurrentEventNodeFrame->EndFrame();
-		
-		if(bLatchedGProfilingGPU)
-		{
-			bTrackingEvents = false;
-			bLatchedGProfilingGPU = false;
-		
-			UE_LOG(LogRHI, Warning, TEXT(""));
-			UE_LOG(LogRHI, Warning, TEXT(""));
-			CurrentEventNodeFrame->DumpEventTree();
-		}
-		
-		delete CurrentEventNodeFrame;
-		CurrentEventNodeFrame = NULL;
-	}
-}
-
-IMetalStatsScope::~IMetalStatsScope()
-{
-	for (IMetalStatsScope* Stat : Children)
-	{
-		delete Stat;
-	}
-}
-
-FString IMetalStatsScope::GetJSONRepresentation(uint32 Pid)
-{
-	FString JSONOutput;
-	
-	{
-		if (GPUStartTime && GPUEndTime)
-		{
-			uint64 ChildStartCallTime = GPUStartTime;
-			uint64 ChildDrawCallTime = GPUEndTime - GPUStartTime;
-			
-			JSONOutput += FString::Printf(TEXT("{\"pid\":%d, \"tid\":%d, \"ph\": \"X\", \"name\": \"%s\", \"ts\": %llu, \"dur\": %llu, \"args\":{\"num_child\":%u}},\n"),
-				  Pid,
-				  GPUThreadIndex,
-				  *Name,
-				  ChildStartCallTime,
-				  ChildDrawCallTime,
-				  Children.Num()
-			  );
-		}
-	}
-	
-	if (CPUStartTime && CPUEndTime)
-	{
-		uint64 ChildStartCallTime = CPUStartTime;
-		uint64 ChildDrawCallTime = FMath::Max(CPUEndTime - CPUStartTime, 1llu);
-		
-		JSONOutput += FString::Printf(TEXT("{\"pid\":%d, \"tid\":%d, \"ph\": \"X\", \"name\": \"%s\", \"ts\": %llu, \"dur\": %llu, \"args\":{\"num_child\":%u}},\n"),
-			 Pid,
-			 CPUThreadIndex,
-			 *Name,
-			 ChildStartCallTime,
-			 ChildDrawCallTime,
-			 Children.Num()
-		);
-	}
-	
-	return JSONOutput;
-}
-
-FMetalCommandBufferStats::FMetalCommandBufferStats(MTL::CommandBuffer* CommandBuffer, uint64 InGPUThreadIndex)
-{
-	CmdBuffer = CommandBuffer;
-	
-	Name = FString::Printf(TEXT("CommandBuffer: %p %s"), CmdBuffer, *NSStringToFString(CmdBuffer->label()));
-	
-	CPUThreadIndex = FPlatformTLS::GetCurrentThreadId();
-	GPUThreadIndex = InGPUThreadIndex;
-	
-	Start(CmdBuffer);
-}
-
-FMetalCommandBufferStats::~FMetalCommandBufferStats()
-{
-}
-
-void FMetalCommandBufferStats::Start(MTL::CommandBuffer*& CommandBuffer)
-{
-	CPUStartTime = FPlatformTime::ToMilliseconds64(mach_absolute_time()) * 1000.0;
-	CPUEndTime = 0;
-	
-	GPUStartTime = 0;
-	GPUEndTime = 0;
-}
-
-void FMetalCommandBufferStats::End(MTL::CommandBuffer*& CommandBuffer)
-{
-	check(CommandBuffer == CmdBuffer);
-	
-	bool const bTracing = FMetalProfiler::GetProfiler() && FMetalProfiler::GetProfiler()->TracingEnabled();
-    
-    MTL::HandlerFunction CompletionHandler = [this, bTracing](MTL::CommandBuffer* CommandBuffer) {
-        const CFTimeInterval GpuTimeSeconds = CommandBuffer->GPUStartTime();
-        GPUStartTime = GpuTimeSeconds * 1000000.0;
-        
-        const CFTimeInterval GpuEndTimeSeconds = CommandBuffer->GPUEndTime();
-        GPUEndTime = GpuEndTimeSeconds * 1000000.0;
-    
-        if (bTracing)
-        {
-            FMetalProfiler::GetProfiler()->AddCommandBuffer(this);
-        }
-        else
-        {
-            delete this;
-        }
-    };
-    
-    CmdBuffer->addCompletedHandler(CompletionHandler);
-	
-	CPUEndTime = FPlatformTime::ToMilliseconds64(mach_absolute_time()) * 1000.0;
-}
-
-
-#pragma mark -- FMetalProfiler
-// ----------------------------------------------------------------
-
-
-FMetalProfiler* FMetalProfiler::Self = nullptr;
-static FMetalViewportPresentHandler PresentHandler = ^(uint32 DisplayID, double OutputSeconds, double OutputDuration){
-	FMetalProfiler* Profiler = FMetalProfiler::GetProfiler();
-	Profiler->AddDisplayVBlank(DisplayID, OutputSeconds, OutputDuration);
-};
-
-FMetalDisplayStats::FMetalDisplayStats(uint32 DisplayID, double OutputSeconds, double Duration)
-{
-	Name = TEXT("V-Blank");
-	
-	CPUThreadIndex = FPlatformTLS::GetCurrentThreadId();
-	GPUThreadIndex = DisplayID;
-	
-	CPUStartTime = FPlatformTime::ToMilliseconds64(mach_absolute_time()) * 1000.0;
-	CPUEndTime = CPUStartTime+1;
-	
-	GPUStartTime = OutputSeconds * 1000000.0;
-	GPUEndTime = GPUStartTime + (Duration * 1000000.0);
-}
-FMetalDisplayStats::~FMetalDisplayStats()
-{
-}
-
-void FMetalDisplayStats::Start(MTL::CommandBuffer*& Buffer)
-{
-}
-void FMetalDisplayStats::End(MTL::CommandBuffer*& Buffer)
-{
-}
-
-FMetalCPUStats::FMetalCPUStats(FString const& InName)
-{
-	Name = InName;
-	
-	CPUThreadIndex = 0;
-	GPUThreadIndex = 0;
-	
-	CPUStartTime = 0;
-	CPUEndTime = 0;
-	
-	GPUStartTime = 0;
-	GPUEndTime = 0;
-}
-FMetalCPUStats::~FMetalCPUStats()
-{
-	
-}
-
-void FMetalCPUStats::Start(void)
-{
-	CPUThreadIndex = FPlatformTLS::GetCurrentThreadId();
-	CPUStartTime = FPlatformTime::ToMilliseconds64(mach_absolute_time()) * 1000.0;
-	
-}
-void FMetalCPUStats::End(void)
-{
-	CPUEndTime = FPlatformTime::ToMilliseconds64(mach_absolute_time()) * 1000.0;
-}
-
-void FMetalCPUStats::Start(MTL::CommandBuffer*& Buffer)
-{
-	
-}
-void FMetalCPUStats::End(MTL::CommandBuffer*& Buffer)
-{
-	
-}
-
-void FMetalProfiler::AddDisplayVBlank(uint32 DisplayID, double OutputSeconds, double OutputDuration)
-{
-	if (GIsRHIInitialized && bEnabled)
-	{
-		FScopeLock Lock(&Mutex);
-		DisplayStats.Add(new FMetalDisplayStats(DisplayID, OutputSeconds, OutputDuration));
-	}
-}
-
-FMetalProfiler::FMetalProfiler(FMetalRHICommandContext& Context)
-: FMetalGPUProfiler(Context)
-, bEnabled(false)
-{
-	NumFramesToCapture = -1;
-	CaptureFrameNumber = 0;
-	
-	bRequestStartCapture = false;
-	bRequestStopCapture = false;
-	
-	if (FPlatformRHIFramePacer::IsEnabled())
-	{
-		FPlatformRHIFramePacer::AddHandler(PresentHandler);
-	}
-}
-
-FMetalProfiler::~FMetalProfiler()
-{
-	check(bEnabled == false);
-	if (FPlatformRHIFramePacer::IsEnabled())
-	{
-		FPlatformRHIFramePacer::RemoveHandler(PresentHandler);
-	}
-}
-
-FMetalProfiler* FMetalProfiler::CreateProfiler(FMetalRHICommandContext& InContext)
-{
-	if (!Self)
-	{
-		Self = new FMetalProfiler(InContext);
-		
-		int32 CaptureFrames = 0;
-		if (FParse::Value(FCommandLine::Get(), TEXT("MetalProfileFrames="), CaptureFrames))
-		{
-			Self->BeginCapture(CaptureFrames);
-		}
-	}
-	return Self;
-}
-
-FMetalProfiler* FMetalProfiler::GetProfiler()
-{
-	return Self;
-}
-
-void FMetalProfiler::DestroyProfiler()
-{
-	delete Self;
-	Self = nullptr;
-}
-
-void FMetalProfiler::BeginCapture(int InNumFramesToCapture)
-{
-	check(IsInGameThread());
-	
-	NumFramesToCapture = InNumFramesToCapture;
-	CaptureFrameNumber = 0;
-	
-	bRequestStartCapture = true;
-}
-
-void FMetalProfiler::EndCapture()
-{
-	bRequestStopCapture = true;
-}
-
-bool FMetalProfiler::TracingEnabled() const
-{
-	return bEnabled;
-}
-
-void FMetalProfiler::BeginFrame()
-{
-	if (MetalGPUProfilerIsInSafeThread())
-	{
-		if (bRequestStartCapture && !bEnabled)
-		{
-			bEnabled = true;
-			bRequestStartCapture = false;
-		}
-	}
-	
-	FMetalGPUProfiler::BeginFrame();
-}
-
-void FMetalProfiler::EndFrame()
-{	
-	FMetalGPUProfiler::EndFrame();
-	
-	if (MetalGPUProfilerIsInSafeThread() && bEnabled)
-	{
-		CaptureFrameNumber++;
-		if (bRequestStopCapture || (NumFramesToCapture > 0 && CaptureFrameNumber >= NumFramesToCapture))
-		{
-			bRequestStopCapture = false;
-			NumFramesToCapture = -1;
-			bEnabled = false;
-			SaveTrace();
-		}
-	}
-}
-
-void FMetalProfiler::EncodeDraw(FMetalCommandBufferStats* CmdBufStats, char const* DrawCall, uint32 RHIPrimitives, uint32 RHIVertices, uint32 RHIInstances)
-{
-#if (RHI_NEW_GPU_PROFILER == 0)
-	if (MetalGPUProfilerIsInSafeThread())
-	{
-		FMetalGPUProfiler::RegisterGPUWork(RHIPrimitives, RHIVertices);
-	}
-#endif
-}
-
-void FMetalProfiler::EncodeBlit(FMetalCommandBufferStats* CmdBufStats, char const* DrawCall)
-{
-#if (RHI_NEW_GPU_PROFILER == 0)
-	if (MetalGPUProfilerIsInSafeThread())
-	{
-		FMetalGPUProfiler::RegisterGPUWork(1, 1);
-	}
-#endif
-}
-
-void FMetalProfiler::EncodeBlit(FMetalCommandBufferStats* CmdBufStats, FString DrawCall)
-{
-#if (RHI_NEW_GPU_PROFILER == 0)
-	if (MetalGPUProfilerIsInSafeThread())
-	{
-		FMetalGPUProfiler::RegisterGPUWork(1, 1);
-	}
-#endif
-}
-
-void FMetalProfiler::EncodeDispatch(FMetalCommandBufferStats* CmdBufStats, char const* DrawCall)
-{
-#if (RHI_NEW_GPU_PROFILER == 0)
-	if (MetalGPUProfilerIsInSafeThread())
-	{
-		FMetalGPUProfiler::RegisterGPUWork(1, 1);
-	}
-#endif
-}
-
-FMetalCPUStats* FMetalProfiler::AddCPUStat(FString const& Name)
-{
-	if (GIsRHIInitialized && bEnabled)
-	{
-		FScopeLock Lock(&Mutex);
-		FMetalCPUStats* Stat = new FMetalCPUStats(Name);
-		CPUStats.Add(Stat);
-		return Stat;
-	}
-	else
-	{
-		return nullptr;
-	}
-}
-
-FMetalCommandBufferStats* FMetalProfiler::AllocateCommandBuffer(MTL::CommandBuffer* Buffer, uint64 GPUThreadIndex)
-{
-	return new FMetalCommandBufferStats(Buffer, GPUThreadIndex);
-}
-
-void FMetalProfiler::AddCommandBuffer(FMetalCommandBufferStats *CommandBuffer)
-{
-	if (GIsRHIInitialized)
-	{
-		FScopeLock Lock(&Mutex);
-		TracedBuffers.Add(CommandBuffer);
-	}
-	else
-	{
-		delete CommandBuffer;
-	}
-}
-
-void FMetalProfiler::PushEvent(const TCHAR *Name, FColor Color)
-{
-	FMetalGPUProfiler::PushEvent(Name, Color);
-}
-
-void FMetalProfiler::PopEvent()
-{
-	FMetalGPUProfiler::PopEvent();
-}
-
-void FMetalProfiler::SaveTrace()
-{
-	GDynamicRHI->RHIBlockUntilGPUIdle();
-	
-	{
-		FScopeLock Lock(&Mutex);
-		
-		TSet<uint32> ThreadIDs;
-		
-		for (FMetalCommandBufferStats* CmdBufStats : TracedBuffers)
-		{
-			ThreadIDs.Add(CmdBufStats->CPUThreadIndex);
-			
-			for (IMetalStatsScope* ES : CmdBufStats->Children)
-			{
-				ThreadIDs.Add(ES->CPUThreadIndex);
-				
-				for (IMetalStatsScope* DS : ES->Children)
-				{
-					ThreadIDs.Add(DS->CPUThreadIndex);
-				}
-			}
-		}
-		
-		TSet<uint32> Displays;
-		for (FMetalDisplayStats* DisplayStat : DisplayStats)
-		{
-			ThreadIDs.Add(DisplayStat->CPUThreadIndex);
-			Displays.Add(DisplayStat->GPUThreadIndex);
-		}
-		
-		for (FMetalCPUStats* CPUStat : CPUStats)
-		{
-			ThreadIDs.Add(CPUStat->CPUThreadIndex);
-		}
-		
-		FString Filename = FString::Printf(TEXT("Profile(%s)"), *FDateTime::Now().ToString(TEXT("%Y%m%d_%H%M%S")));
-		FString TracingRootPath = FPaths::ProfilingDir() + TEXT("Traces/");
-		FString OutputFilename = TracingRootPath + Filename + TEXT(".json");
-		
-		FArchive* OutputFile = IFileManager::Get().CreateFileWriter(*OutputFilename);
-		
-		WriteString(OutputFile, R"({"traceEvents":[)" "\n");
-		
-		int32 SortIndex = 0; // Lower numbers result in higher position in the visualizer.
-		const uint32 Pid = FPlatformProcess::GetCurrentProcessId();
-		
-		for (int32 GPUIndex = 0; GPUIndex <= 0/*MaxGPUIndex*/; ++GPUIndex)
-		{
-			FString Output = FString::Printf(TEXT("{\"pid\":%d, \"tid\":%d, \"ph\": \"M\", \"name\": \"thread_name\", \"args\":{\"name\":\"GPU %d Command Buffers\"}},{\"pid\":%d, \"tid\":%d, \"ph\": \"M\", \"name\": \"thread_sort_index\", \"args\":{\"sort_index\": %d}},\n"),
-											 Pid, GPUIndex, GPUIndex, Pid, GPUIndex, SortIndex
-											 );
-			
-			WriteString(OutputFile, TCHAR_TO_UTF8(*Output));
-			SortIndex++;
-			
-			Output = FString::Printf(TEXT("{\"pid\":%d, \"tid\":%d, \"ph\": \"M\", \"name\": \"thread_name\", \"args\":{\"name\":\"GPU %d Operations\"}},{\"pid\":%d, \"tid\":%d, \"ph\": \"M\", \"name\": \"thread_sort_index\", \"args\":{\"sort_index\": %d}},\n"),
-											 Pid, GPUIndex+SortIndex, GPUIndex, Pid, GPUIndex+SortIndex, SortIndex
-											 );
-			
-			WriteString(OutputFile, TCHAR_TO_UTF8(*Output));
-			SortIndex++;
-			
-			Output = FString::Printf(TEXT("{\"pid\":%d, \"tid\":%d, \"ph\": \"M\", \"name\": \"thread_name\", \"args\":{\"name\":\"Render Events %d\"}},{\"pid\":%d, \"tid\":%d, \"ph\": \"M\", \"name\": \"thread_sort_index\", \"args\":{\"sort_index\": %d}},\n"),
-									 Pid, GPUIndex+SortIndex, GPUIndex, Pid, GPUIndex+SortIndex, SortIndex
-									 );
-			
-			WriteString(OutputFile, TCHAR_TO_UTF8(*Output));
-			SortIndex++;
-			
-			Output = FString::Printf(TEXT("{\"pid\":%d, \"tid\":%d, \"ph\": \"M\", \"name\": \"thread_name\", \"args\":{\"name\":\"Driver Stats %d\"}},{\"pid\":%d, \"tid\":%d, \"ph\": \"M\", \"name\": \"thread_sort_index\", \"args\":{\"sort_index\": %d}},\n"),
-									 Pid, GPUIndex+SortIndex, GPUIndex, Pid, GPUIndex+SortIndex, SortIndex
-									 );
-			
-			WriteString(OutputFile, TCHAR_TO_UTF8(*Output));
-			SortIndex++;
-			
-			for (uint32 Display : Displays)
-			{
-				Output = FString::Printf(TEXT("{\"pid\":%d, \"tid\":%d, \"ph\": \"M\", \"name\": \"thread_name\", \"args\":{\"name\":\"Display %d\"}},{\"pid\":%d, \"tid\":%d, \"ph\": \"M\", \"name\": \"thread_sort_index\", \"args\":{\"sort_index\": %d}},\n"),
-										 Pid, Display + SortIndex, SortIndex - 3, Pid, Display + SortIndex, SortIndex
-										 );
-				
-				WriteString(OutputFile, TCHAR_TO_UTF8(*Output));
-				SortIndex++;
-			}
-		}
-		
-		static const uint32 BufferSize = 128;
-		ANSICHAR Buffer[BufferSize];
-		for (uint32 CPUIndex : ThreadIDs)
-		{
-			bool bThreadName = false;
-			pthread_t PThread = pthread_from_mach_thread_np((mach_port_t)CPUIndex);
-			if (PThread)
-			{
-				if (!pthread_getname_np(PThread,Buffer,BufferSize))
-				{
-					bThreadName = true;
-				}
-			}
-			if (!bThreadName)
-			{
-				snprintf(Buffer, BufferSize, "Thread %d", CPUIndex);
-			}
-			
-			FString Output = FString::Printf(TEXT("{\"pid\":%d, \"tid\":%d, \"ph\": \"M\", \"name\": \"thread_name\", \"args\":{\"name\":\"%s\"}},{\"pid\":%d, \"tid\":%d, \"ph\": \"M\", \"name\": \"thread_sort_index\", \"args\":{\"sort_index\": %d}},\n"),
-											 Pid, CPUIndex, ANSI_TO_TCHAR(Buffer), Pid, CPUIndex, SortIndex
-											 );
-			
-			
-			WriteString(OutputFile, TCHAR_TO_UTF8(*Output));
-			SortIndex++;
-		}
-		
-		for (FMetalCommandBufferStats* CmdBufStats : TracedBuffers)
-		{
-			WriteString(OutputFile, TCHAR_TO_UTF8(*CmdBufStats->GetJSONRepresentation(Pid)));
-			
-			for (IMetalStatsScope* ES : CmdBufStats->Children)
-			{
-				WriteString(OutputFile, TCHAR_TO_UTF8(*ES->GetJSONRepresentation(Pid)));
-				
-				uint64 PrevTime = ES->GPUStartTime;
-				for (IMetalStatsScope* DS : ES->Children)
-				{
-					WriteString(OutputFile, TCHAR_TO_UTF8(*DS->GetJSONRepresentation(Pid)));
-					if (!DS->GPUStartTime)
-					{
-						DS->GPUStartTime = FMath::Max(PrevTime, DS->GPUStartTime);
-						DS->GPUEndTime = DS->GPUStartTime + 1llu;
-						WriteString(OutputFile, TCHAR_TO_UTF8(*DS->GetJSONRepresentation(Pid)));
-					}
-					PrevTime = DS->GPUEndTime;
-				}
-			}
-			
-			delete CmdBufStats;
-		}
-		TracedBuffers.Empty();
-		
-		for (FMetalDisplayStats* DisplayStat : DisplayStats)
-		{
-			DisplayStat->GPUThreadIndex += 3;
-			WriteString(OutputFile, TCHAR_TO_UTF8(*DisplayStat->GetJSONRepresentation(Pid)));
-			delete DisplayStat;
-		}
-		DisplayStats.Empty();
-		
-		for (FMetalCPUStats* CPUStat : CPUStats)
-		{
-			WriteString(OutputFile, TCHAR_TO_UTF8(*CPUStat->GetJSONRepresentation(Pid)));
-			delete CPUStat;
-		}
-		CPUStats.Empty();
-		
-		// All done
-		
-		WriteString(OutputFile, "{}]}");
-		
-		OutputFile->Close();
-	}
-}
-
-static void HandleMetalProfileCommand(const TArray<FString>& Args, UWorld*, FOutputDevice& Ar)
-{
-	if (Args.Num() < 1)
-	{
-		return;
-	}
-	FString Param = Args[0];
-	if (Param == TEXT("START"))
-	{
-		FMetalProfiler::GetProfiler()->BeginCapture();
-	}
-	else if (Param == TEXT("STOP"))
-	{
-		FMetalProfiler::GetProfiler()->EndCapture();
-	}
-	else
-	{
-		int32 CaptureFrames = 0;
-		if (FParse::Value(*Param, TEXT("FRAMES="), CaptureFrames))
-		{
-			FMetalProfiler::GetProfiler()->BeginCapture(CaptureFrames);
-		}
-	}
-}
-
-static FAutoConsoleCommand HandleMetalProfilerCmd(
-	TEXT("MetalProfiler"),
-	TEXT("Starts or stops Metal profiler"),
-	FConsoleCommandWithWorldArgsAndOutputDeviceDelegate::CreateStatic(&HandleMetalProfileCommand)
-);
-
-#endif
-
-#if RHI_NEW_GPU_PROFILER && WITH_RHI_BREADCRUMBS
+#if WITH_RHI_BREADCRUMBS
 
 void FMetalBreadcrumbProfiler::ResolveBreadcrumbTrackerStream(TArray<FMetalBreadcrumbTrackerObject>& BreadcrumbTrackerStream)
 {
diff --git a/Engine/Source/Runtime/Apple/MetalRHI/Private/MetalProfiler.h b/Engine/Source/Runtime/Apple/MetalRHI/Private/MetalProfiler.h
--- a/Engine/Source/Runtime/Apple/MetalRHI/Private/MetalProfiler.h
+++ b/Engine/Source/Runtime/Apple/MetalRHI/Private/MetalProfiler.h
@@ -51,16 +51,7 @@ DECLARE_DWORD_ACCUMULATOR_STAT_EXTERN(TEXT("Heap Count"), STAT_MetalHeapCount, S
 DECLARE_DWORD_ACCUMULATOR_STAT_EXTERN(TEXT("Fence Count"), STAT_MetalFenceCount, STATGROUP_MetalRHI, );
 
 DECLARE_CYCLE_STAT_EXTERN(TEXT("Texture Page-On time"), STAT_MetalTexturePageOnTime, STATGROUP_MetalRHI, );
-DECLARE_CYCLE_STAT_EXTERN(TEXT("GPU Work time"), STAT_MetalGPUWorkTime, STATGROUP_MetalRHI, );
-DECLARE_CYCLE_STAT_EXTERN(TEXT("GPU Idle time"), STAT_MetalGPUIdleTime, STATGROUP_MetalRHI, );
-DECLARE_CYCLE_STAT_EXTERN(TEXT("Present time"), STAT_MetalPresentTime, STATGROUP_MetalRHI, );
 DECLARE_CYCLE_STAT_EXTERN(TEXT("CustomPresent time"), STAT_MetalCustomPresentTime, STATGROUP_MetalRHI, );
-#if STATS
-extern int64 volatile GMetalTexturePageOnTime;
-extern int64 volatile GMetalGPUWorkTime;
-extern int64 volatile GMetalGPUIdleTime;
-extern int64 volatile GMetalPresentTime;
-#endif
 
 DECLARE_DWORD_COUNTER_STAT_EXTERN(TEXT("Number Command Buffers Created Per-Frame"), STAT_MetalCommandBufferCreatedPerFrame, STATGROUP_MetalRHI, );
 DECLARE_DWORD_COUNTER_STAT_EXTERN(TEXT("Number Command Buffers Committed Per-Frame"), STAT_MetalCommandBufferCommittedPerFrame, STATGROUP_MetalRHI, );
@@ -68,332 +59,7 @@ DECLARE_DWORD_COUNTER_STAT_EXTERN(TEXT("Number Command Buffers Committed Per-Fra
 class FMetalRHICommandContext;
 class FMetalSyncPoint;
 
-/**
- * Simple struct to hold sortable command buffer start and end timestamps.
- */
-struct FMetalCommandBufferTiming
-{
-	CFTimeInterval StartTime;
-	CFTimeInterval EndTime;
-
-	bool operator<(const FMetalCommandBufferTiming& RHS) const
-	{
-		// Sort by start time and then by length if the commandbuffer started at the same time
-		if (this->StartTime < RHS.StartTime)
-		{
-			return true;
-		}
-		else if ((this->StartTime == RHS.StartTime) && (this->EndTime > RHS.EndTime))
-		{
-			return true;
-		}
-		return false;
-	}
-};
-
-#if RHI_NEW_GPU_PROFILER == 0
-
-class FMetalCommandBufferTimer
-{
-public:
-	void Submit()
-	{
-		FScopeLock Lock(&Mutex);
-		Counter++;
-	}
-	
-	void AddTiming(FMetalCommandBufferTiming Timing);
-	void FrameEnd();
-	void RecordFrame();
-	
-	const TArray<FMetalCommandBufferTiming>& GetTimings()
-	{
-		return Timings;
-	}
-	
-	static void ResetFrameBufferTimings();
-	static FMetalCommandBufferTimer& GetFrameBufferTimer();
-	
-	static void RecordPresent(MTL::CommandBuffer* CommandBuffer);
-	// END WARNING
-	
-private:
-	bool bFrameEnded;
-	int32_t Counter = 0;
-	TArray<FMetalCommandBufferTiming> Timings;
-	FCriticalSection Mutex;
-	
-	static FMetalCommandBufferTimer* Timer;
-};
-
-/** A single perf event node, which tracks information about a appBeginDrawEvent/appEndDrawEvent range. */
-class FMetalEventNode : public FGPUProfilerEventNode
-{
-public:
-	
-	FMetalEventNode(FMetalRHICommandContext& InContext, const TCHAR* InName, FGPUProfilerEventNode* InParent, bool bIsRoot, bool bInFullProfiling)
-	: FGPUProfilerEventNode(InName, InParent)
-	, StartTime(0)
-	, EndTime(0)
-	, Context(InContext)
-	, bRoot(bIsRoot)
-    , bFullProfiling(bInFullProfiling)
-	{
-	}
-	
-	virtual ~FMetalEventNode();
-	
-	/**
-	 * Returns the time in ms that the GPU spent in this draw event.
-	 * This blocks the CPU if necessary, so can cause hitching.
-	 */
-	virtual float GetTiming() override;
-	virtual void StartTiming() override;
-	virtual void StopTiming() override;
-	
-	FMetalSyncPoint* SyncPoint = nullptr;
-
-	bool IsRoot() const { return bRoot; }	
-	uint64 GetCycles() { return EndTime - StartTime; }
-	
-	uint64 StartTime;
-	uint64 EndTime;
-private:
-	FMetalRHICommandContext& Context;
-	bool bRoot;
-    bool bFullProfiling;
-};
-
-/** An entire frame of perf event nodes, including ancillary timers. */
-class FMetalEventNodeFrame : public FGPUProfilerEventNodeFrame
-{
-public:
-	FMetalEventNodeFrame(FMetalRHICommandContext& InContext, bool bInFullProfiling)
-	: RootNode(new FMetalEventNode(InContext, TEXT("Frame"), nullptr, true, bInFullProfiling))
-    , bFullProfiling(bInFullProfiling)
-	{
-	}
-	
-	virtual ~FMetalEventNodeFrame()
-	{
-        delete RootNode;
-		RootNode = nullptr;
-	}
-	
-	/** Start this frame of per tracking */
-	virtual void StartFrame() override;
-	
-	/** End this frame of per tracking, but do not block yet */
-	virtual void EndFrame() override;
-	
-	/** Calculates root timing base frequency (if needed by this RHI) */
-	virtual float GetRootTimingResults() override;
-	
-	virtual void LogDisjointQuery() override;
-	
-	FMetalEventNode* RootNode;
-    bool bFullProfiling;
-};
-
-class FMetalContext;
-
-// This class has multiple inheritance but really FGPUTiming is a static class
-class FMetalGPUTiming : public FGPUTiming
-{
-public:
-	
-	/**
-	 * Constructor.
-	 */
-	FMetalGPUTiming(FMetalRHICommandContext& Context)
-	{
-		StaticInitialize((void*)&Context, PlatformStaticInitialize);
-	}
-	
-	void SetCalibrationTimestamp(uint64 GPU, uint64 CPU)
-	{
-		FGPUTiming::SetCalibrationTimestamp({ GPU, CPU });
-	}
-	
-private:
-	
-	/**
-	 * Initializes the static variables, if necessary.
-	 */
-	static void PlatformStaticInitialize(void* UserData);
-};
-
-struct IMetalStatsScope
-{
-	FString Name;
-	FString Parent;
-	TArray<IMetalStatsScope*> Children;
-	
-	uint64 CPUStartTime;
-	uint64 CPUEndTime;
-	
-	uint64 GPUStartTime;
-	uint64 GPUEndTime;
-	
-	uint64 CPUThreadIndex;
-	uint64 GPUThreadIndex;
-	
-	virtual ~IMetalStatsScope();
-	
-	virtual void Start(MTL::CommandBuffer*& CommandBuffer) = 0;
-	virtual void End(MTL::CommandBuffer*& CommandBuffer) = 0;
-	
-	FString GetJSONRepresentation(uint32 PID);
-};
-
-struct FMetalCPUStats : public IMetalStatsScope
-{
-	FMetalCPUStats(FString const& Name);
-	virtual ~FMetalCPUStats();
-	
-	void Start(void);
-	void End(void);
-	
-	virtual void Start(MTL::CommandBuffer*& CommandBuffer) final override;
-	virtual void End(MTL::CommandBuffer*& CommandBuffer) final override;
-};
-
-struct FMetalDisplayStats : public IMetalStatsScope
-{
-	FMetalDisplayStats(uint32 DisplayID, double OutputSeconds, double Duration);
-	virtual ~FMetalDisplayStats();
-	
-	virtual void Start(MTL::CommandBuffer*& CommandBuffer) final override;
-	virtual void End(MTL::CommandBuffer*& CommandBuffer) final override;
-};
-
-enum EMTLFenceType
-{
-	EMTLFenceTypeWait,
-	EMTLFenceTypeUpdate,
-};
-
-struct FMetalCommandBufferStats : public IMetalStatsScope
-{
-	FMetalCommandBufferStats(MTL::CommandBuffer* CommandBuffer, uint64 GPUThreadIndex);
-	virtual ~FMetalCommandBufferStats();
-	
-	virtual void Start(MTL::CommandBuffer*& CommandBuffer) final override;
-	virtual void End(MTL::CommandBuffer*& CommandBuffer) final override;
-
-	MTL::CommandBuffer* CmdBuffer;
-};
-
-class FMetalRHICommandContext;
-/**
- * Encapsulates GPU profiling logic and data.
- * There's only one global instance of this struct so it should only contain global data, nothing specific to a frame.
- */
-struct FMetalGPUProfiler : public FGPUProfiler
-{
-	/** GPU hitch profile histories */
-	TIndirectArray<FMetalEventNodeFrame> GPUHitchEventNodeFrames;
-	
-	FMetalGPUProfiler(FMetalRHICommandContext& InContext)
-	:	FGPUProfiler()
-	,	TimingSupport(InContext)
-	,	Context(InContext)
-	{}
-	
-	virtual ~FMetalGPUProfiler() {}
-	
-	virtual FGPUProfilerEventNode* CreateEventNode(const TCHAR* InName, FGPUProfilerEventNode* InParent) override;
-	
-	void Cleanup();
-	
-	virtual void PushEvent(const TCHAR* Name, FColor Color) override;
-	virtual void PopEvent() override;
-	
-	void BeginFrame();
-	void EndFrame();
-	
-	FMetalGPUTiming TimingSupport;
-	FMetalRHICommandContext& Context;
-};
-
-class FMetalProfiler : public FMetalGPUProfiler
-{
-	static FMetalProfiler* Self;
-public:
-	FMetalProfiler(FMetalRHICommandContext& InContext);
-	~FMetalProfiler();
-	
-	static FMetalProfiler* CreateProfiler(FMetalRHICommandContext& InContext);
-	static FMetalProfiler* GetProfiler();
-	static void DestroyProfiler();
-	
-	void BeginCapture(int InNumFramesToCapture = -1);
-	void EndCapture();
-	bool TracingEnabled() const;
-	
-	void BeginFrame();
-	void EndFrame();
-	
-	void AddDisplayVBlank(uint32 DisplayID, double OutputSeconds, double OutputDuration);
-	
-	void EncodeDraw(FMetalCommandBufferStats* CmdBufStats, char const* DrawCall, uint32 RHIPrimitives, uint32 RHIVertices, uint32 RHIInstances);
-	void EncodeBlit(FMetalCommandBufferStats* CmdBufStats, char const* DrawCall);
-	void EncodeBlit(FMetalCommandBufferStats* CmdBufStats, FString DrawCall);
-	void EncodeDispatch(FMetalCommandBufferStats* CmdBufStats, char const* DrawCall);
-	
-	FMetalCPUStats* AddCPUStat(FString const& Name);
-	FMetalCommandBufferStats* AllocateCommandBuffer(MTL::CommandBuffer* CommandBuffer, uint64 GPUThreadIndex);
-	void AddCommandBuffer(FMetalCommandBufferStats* CommandBuffer);
-	virtual void PushEvent(const TCHAR* Name, FColor Color) final override;
-	virtual void PopEvent() final override;
-	
-	void SaveTrace();
-	
-private:
-	FCriticalSection Mutex;
-	
-	TArray<FMetalCommandBufferStats*> TracedBuffers;
-	TArray<FMetalDisplayStats*> DisplayStats;
-	TArray<FMetalCPUStats*> CPUStats;
-	
-	int32 NumFramesToCapture;
-	int32 CaptureFrameNumber;
-	
-	bool bRequestStartCapture;
-	bool bRequestStopCapture;
-	bool bEnabled;
-};
-
-struct FScopedMetalCPUStats
-{
-	FScopedMetalCPUStats(FString const& Name)
-	: Stats(nullptr)
-	{
-		FMetalProfiler* Profiler = FMetalProfiler::GetProfiler();
-		if (Profiler)
-		{
-			Stats = Profiler->AddCPUStat(Name);
-			if (Stats)
-			{
-				Stats->Start();
-			}
-		}
-	}
-	
-	~FScopedMetalCPUStats()
-	{
-		if (Stats)
-		{
-			Stats->End();
-		}
-	}
-	
-	FMetalCPUStats* Stats;
-};
-
-#endif
-
-#if RHI_NEW_GPU_PROFILER && WITH_RHI_BREADCRUMBS
+#if WITH_RHI_BREADCRUMBS
 
 enum class EMetalBreadcrumbTrackerType : uint8
 {
diff --git a/Engine/Source/Runtime/Apple/MetalRHI/Private/MetalRHI.cpp b/Engine/Source/Runtime/Apple/MetalRHI/Private/MetalRHI.cpp
--- a/Engine/Source/Runtime/Apple/MetalRHI/Private/MetalRHI.cpp
+++ b/Engine/Source/Runtime/Apple/MetalRHI/Private/MetalRHI.cpp
@@ -194,9 +194,6 @@ static void VerifyMetalCompiler()
 
 FMetalDynamicRHI::FMetalDynamicRHI(ERHIFeatureLevel::Type RequestedFeatureLevel)
 	: Device(FMetalDevice::CreateDevice())
-#if (RHI_NEW_GPU_PROFILER == 0)
-	, ImmediateContext(*Device, nullptr)
-#endif
 {
 	check(Singleton == nullptr);
 	Singleton = this;
@@ -1197,16 +1194,6 @@ FMetalDynamicRHI::FMetalDynamicRHI(ERHIFeatureLevel::Type RequestedFeatureLevel)
 	
 	GIsMetalInitialized = true;
 
-#if ENABLE_METAL_GPUPROFILE && RHI_NEW_GPU_PROFILER == 0
-	FMetalProfiler* Profiler = FMetalProfiler::CreateProfiler(ImmediateContext);
-	ImmediateContext.SetProfiler(Profiler);
-	
-	if (Profiler)
-	{
-		Profiler->BeginFrame();
-	}
-#endif
-
 #if METAL_USE_METAL_SHADER_CONVERTER
 	CompilerInstance = IRCompilerCreate();
 #endif
@@ -1236,10 +1223,6 @@ FMetalDynamicRHI::~FMetalDynamicRHI()
 #if METAL_USE_METAL_SHADER_CONVERTER
     IRCompilerDestroy(CompilerInstance);
 #endif
-
-#if ENABLE_METAL_GPUPROFILE && RHI_NEW_GPU_PROFILER == 0
-	FMetalProfiler::DestroyProfiler();
-#endif
 }
 
 FDynamicRHI::FRHICalcTextureSizeResult FMetalDynamicRHI::RHICalcTexturePlatformSize(FRHITextureDesc const& Desc, uint32 FirstMipIndex)
@@ -1268,32 +1251,12 @@ void FMetalDynamicRHI::RHIEndFrame_RenderThread(FRHICommandListImmediate& RHICmd
 	{
 		MTL_SCOPED_AUTORELEASE_POOL;
 
-		
-#if RHI_NEW_GPU_PROFILER == 0
-		FMetalCommandBufferTimer::ResetFrameBufferTimings();
-#if ENABLE_METAL_GPUPROFILE
-		Contexts[ERHIPipeline::Graphics]->GetProfiler()->EndFrame();
-#endif
-#endif
-
 #if METAL_RHI_RAYTRACING
 		Device->UpdateRayTracing(*Contexts[ERHIPipeline::Graphics]);
 #endif // METAL_RHI_RAYTRACING
 	});
 
 	FDynamicRHI::RHIEndFrame_RenderThread(RHICmdList);
-
-	RHICmdList.EnqueueLambdaMultiPipe(ERHIPipeline::Graphics, FRHICommandListBase::EThreadFence::Enabled, TEXT("Metal BeginFrame"),
-		[this](FMetalContextArray const& Contexts)
-	{
-		MTL_SCOPED_AUTORELEASE_POOL;
-		
-#if ENABLE_METAL_GPUPROFILE
-#if RHI_NEW_GPU_PROFILER == 0
-		Contexts[ERHIPipeline::Graphics]->GetProfiler()->BeginFrame();
-#endif
-#endif
-	});
 }
 
 void FMetalDynamicRHI::RHIEndFrame(const FRHIEndFrameArgs& Args)
@@ -1302,7 +1265,6 @@ void FMetalDynamicRHI::RHIEndFrame(const FRHIEndFrameArgs& Args)
 	Device->IncrementFrameRHIThread();
     Device->GarbageCollect();
 
-#if RHI_NEW_GPU_PROFILER	
 	// Close the previous frame's timing and start a new one
 	auto Lambda = [this, OldTiming = MoveTemp(CurrentTimingPerQueue)]()
 	{
@@ -1333,7 +1295,6 @@ void FMetalDynamicRHI::RHIEndFrame(const FRHIEndFrameArgs& Args)
 		#endif
 		);
 	});
-#endif
 	
 	// Pump the interrupt queue to gather completed events
 	// (required if we're not using an interrupt thread).
@@ -1343,33 +1304,24 @@ void FMetalDynamicRHI::RHIEndFrame(const FRHIEndFrameArgs& Args)
 #if WITH_RHI_BREADCRUMBS
 void FMetalRHICommandContext::RHIBeginBreadcrumbGPU(FRHIBreadcrumbNode* Breadcrumb)
 {
-	const TCHAR* NameStr = nullptr;
-	FRHIBreadcrumb::FBuffer Buffer;
-	auto GetNameStr = [&]()
-	{
-		if (!NameStr)
-		{
-			NameStr = Breadcrumb->GetTCHAR(Buffer);
-		}
-		return NameStr;
-	};
-
+#if ENABLE_METAL_GPUEVENTS
 	if (ShouldEmitBreadcrumbs())
 	{
-#if ENABLE_METAL_GPUEVENTS
 		MTL_SCOPED_AUTORELEASE_POOL;
 		{
 			// @todo dev-pr avoid TCHAR -> ANSI conversion
-			CurrentEncoder.PushDebugGroup(NS::String::string(TCHAR_TO_UTF8(GetNameStr()), NS::UTF8StringEncoding));
+			FRHIBreadcrumb::FBuffer Buffer;
+			TCHAR const* Name = Breadcrumb->GetTCHAR(Buffer);
+
+			CurrentEncoder.PushDebugGroup(NS::String::string(TCHAR_TO_UTF8(Name), NS::UTF8StringEncoding));
 		}
-#endif
 	}
+#endif
 
 #if ENABLE_METAL_GPUPROFILE
-#if RHI_NEW_GPU_PROFILER
-	if(!CurrentEncoder.IsParallelEncoding())
+	if (!CurrentEncoder.IsParallelEncoding())
 	{
-		if(!CurrentEncoder.GetCommandBuffer())
+		if (!CurrentEncoder.GetCommandBuffer())
 		{
 			StartCommandBuffer();
 		}
@@ -1377,7 +1329,7 @@ void FMetalRHICommandContext::RHIBeginBreadcrumbGPU(FRHIBreadcrumbNode* Breadcru
 		FMetalCommandBuffer* CmdBuffer = CurrentEncoder.GetCommandBuffer();
 		
 		// Can't process breadcrumbs if we are within a render pass
-		if(Device.SupportsFeature(EMetalFeaturesStageCounterSampling))
+		if (Device.SupportsFeature(EMetalFeaturesStageCounterSampling))
 		{
 			FMetalBreadcrumbEvent& Event = FMetalBreadcrumbProfiler::GetInstance()->GetBreadcrumbEvent(Breadcrumb, bWithinRenderPass);
 			Event.TimestampTOP = &CmdBuffer->EmplaceProfilerEvent<UE::RHI::GPUProfiler::FEvent::FBeginBreadcrumb>(Breadcrumb).GPUTimestampTOP;
@@ -1385,20 +1337,13 @@ void FMetalRHICommandContext::RHIBeginBreadcrumbGPU(FRHIBreadcrumbNode* Breadcru
 			CmdBuffer->BeginBreadcrumb(Breadcrumb);
 		}
 	}
-#else
-	if (Profiler && Profiler->IsProfilingGPU())
-	{
-		Profiler->PushEvent(GetNameStr(), FColor::White);
-	}
-#endif
 #endif
 }
 
 void FMetalRHICommandContext::RHIEndBreadcrumbGPU(FRHIBreadcrumbNode* Breadcrumb)
 {
 #if ENABLE_METAL_GPUPROFILE
-#if RHI_NEW_GPU_PROFILER
-	if(!CurrentEncoder.IsParallelEncoding())
+	if (!CurrentEncoder.IsParallelEncoding())
 	{
 		if(!CurrentEncoder.GetCommandBuffer())
 		{
@@ -1407,7 +1352,7 @@ void FMetalRHICommandContext::RHIEndBreadcrumbGPU(FRHIBreadcrumbNode* Breadcrumb
 		
 		FMetalCommandBuffer* CmdBuffer = CurrentEncoder.GetCommandBuffer();
 		
-		if(Device.SupportsFeature(EMetalFeaturesStageCounterSampling))
+		if (Device.SupportsFeature(EMetalFeaturesStageCounterSampling))
 		{
 			FMetalBreadcrumbEvent& Event = FMetalBreadcrumbProfiler::GetInstance()->GetBreadcrumbEvent(Breadcrumb, bWithinRenderPass);
 			Event.TimestampBOP = &CmdBuffer->EmplaceProfilerEvent<UE::RHI::GPUProfiler::FEvent::FEndBreadcrumb>(Breadcrumb).GPUTimestampBOP;
@@ -1416,23 +1361,17 @@ void FMetalRHICommandContext::RHIEndBreadcrumbGPU(FRHIBreadcrumbNode* Breadcrumb
 			CmdBuffer->EndBreadcrumb(Breadcrumb);
 		}
 	}
-#else
-	if (Profiler && Profiler->IsProfilingGPU())
-	{
-		Profiler->PopEvent();
-	}
-#endif
 #endif
 
+#if ENABLE_METAL_GPUEVENTS
 	if (ShouldEmitBreadcrumbs())
 	{
-#if ENABLE_METAL_GPUEVENTS
 		MTL_SCOPED_AUTORELEASE_POOL;
 		{
 			CurrentEncoder.PopDebugGroup();
 		}
-#endif
 	}
+#endif
 }
 #endif // WITH_RHI_BREADCRUMBS
 
@@ -1553,20 +1492,12 @@ MTL::Device* FMetalDynamicRHI::RHIGetDevice()
 
 void* FMetalDynamicRHI::RHIGetNativeGraphicsQueue()
 {
-#if RHI_NEW_GPU_PROFILER
 	return nullptr;
-#else
-	return ImmediateContext.GetCommandQueue().GetQueue();
-#endif
 }
 
 void* FMetalDynamicRHI::RHIGetNativeComputeQueue()
 {
-#if RHI_NEW_GPU_PROFILER
 	return nullptr;
-#else
-	return ImmediateContext.GetCommandQueue().GetQueue();
-#endif
 }
 
 uint16 FMetalDynamicRHI::RHIGetPlatformTextureMaxSampleCount()
@@ -1620,13 +1551,9 @@ void FMetalDynamicRHI::RHIBlockUntilGPUIdle()
 
 IRHICommandContext* FMetalDynamicRHI::RHIGetDefaultContext()
 {
-#if RHI_NEW_GPU_PROFILER
 	// There is no "default"/"immediate" context in Metal now, just return a regular (i.e. parallel) RHI context.
 	// Forward this call through the GDynamicRHI pointer so the validation layer can still do its per-context setup.
 	return static_cast<IRHICommandContext*>(&GDynamicRHI->RHIGetCommandContext(ERHIPipeline::Graphics)->GetLowestLevelContext());
-#else
-	return &ImmediateContext;
-#endif
 }
 
 IRHIComputeContext* FMetalDynamicRHI::RHIGetCommandContext(ERHIPipeline Pipeline)
@@ -1636,7 +1563,7 @@ IRHIComputeContext* FMetalDynamicRHI::RHIGetCommandContext(ERHIPipeline Pipeline
 	FMetalRHICommandContext* Context = MetalCommandContextPool.Pop();
 	if (!Context)
 	{
-		Context = new FMetalRHICommandContext(*Device, nullptr);
+		Context = new FMetalRHICommandContext(*Device);
 	}
 	
 	Context->ResetContext();
@@ -1651,7 +1578,7 @@ IRHIComputeContext* FMetalDynamicRHI::RHIGetParallelCommandContext(FRHIParallelR
 	FMetalRHICommandContext* Context = MetalCommandContextPool.Pop();
 	if (!Context)
 	{
-		Context = new FMetalRHICommandContext(*Device, nullptr);
+		Context = new FMetalRHICommandContext(*Device);
 	}
 	
 	Context->ResetContext();
diff --git a/Engine/Source/Runtime/Apple/MetalRHI/Private/MetalRHIBuffer.cpp b/Engine/Source/Runtime/Apple/MetalRHI/Private/MetalRHIBuffer.cpp
--- a/Engine/Source/Runtime/Apple/MetalRHI/Private/MetalRHIBuffer.cpp
+++ b/Engine/Source/Runtime/Apple/MetalRHI/Private/MetalRHIBuffer.cpp
@@ -258,8 +258,7 @@ void FMetalRHIBuffer::AllocateBuffer()
 	Buffer = Device.CreatePooledBuffer(ArgsCPU);
 #else
 	NS::UInteger Options = (((NS::UInteger) Mode) << MTL::ResourceStorageModeShift);
-	
-	METAL_GPUPROFILE(FScopedMetalCPUStats CPUStat(FString::Printf(TEXT("AllocBuffer: %llu, %llu"), AllocSize, Options)));
+
 	// Allocate one.
 	MTL::Buffer* BufferPtr = Device.GetDevice()->newBuffer(AllocSize, (MTL::ResourceOptions) Options);
 	Buffer = FMetalBufferPtr(new FMetalBuffer(BufferPtr, FMetalBuffer::FreePolicy::Owner));
diff --git a/Engine/Source/Runtime/Apple/MetalRHI/Private/MetalRHIContext.cpp b/Engine/Source/Runtime/Apple/MetalRHI/Private/MetalRHIContext.cpp
--- a/Engine/Source/Runtime/Apple/MetalRHI/Private/MetalRHIContext.cpp
+++ b/Engine/Source/Runtime/Apple/MetalRHI/Private/MetalRHIContext.cpp
@@ -54,15 +54,13 @@ void METALRHI_API SafeReleaseMetalObject(NS::Object* Object)
 	Object->release();
 }
 
-FMetalRHICommandContext::FMetalRHICommandContext(FMetalDevice& MetalDevice, class FMetalProfiler* InProfiler)
+FMetalRHICommandContext::FMetalRHICommandContext(FMetalDevice& MetalDevice)
 	: Device(MetalDevice)
 	, CommandQueue(Device.GetCommandQueue(EMetalQueueType::Direct))
-	, CommandList(Device.GetCommandQueue(EMetalQueueType::Direct))
-	, CurrentEncoder(MetalDevice, CommandList)
+	, CurrentEncoder(MetalDevice, CommandQueue)
 	, StateCache(MetalDevice)
 	, QueryBuffer(new FMetalQueryBufferPool(MetalDevice))
 	, RenderPassDesc(nullptr)
-	, Profiler(InProfiler)
 	, bWithinRenderPass(false)
 {
 	GlobalUniformBuffers.AddZeroed(FUniformBufferStaticSlotRegistry::Get().GetSlotCount());
@@ -452,13 +450,6 @@ void FMetalRHICommandContext::ResolveTexture(UE::RHICore::FResolveTextureInfo In
 		Size.height = FMath::Max<uint32>(1, SourceDesc.Extent.Y >> Info.MipLevel);
 	}
 
-#if RHI_NEW_GPU_PROFILER == 0
-	if (Profiler)
-	{
-		Profiler->RegisterGPUWork();
-	}
-#endif
-
 	int32 ArraySliceBegin = Info.ArraySlice;
 	int32 ArraySliceEnd   = Info.ArraySlice + 1;
 
@@ -476,7 +467,6 @@ void FMetalRHICommandContext::ResolveTexture(UE::RHICore::FResolveTextureInfo In
 	
 	for (int32 ArraySlice = ArraySliceBegin; ArraySlice < ArraySliceEnd; ArraySlice++)
 	{
-		METAL_GPUPROFILE(FMetalProfiler::GetProfiler()->EncodeBlit(CurrentEncoder.GetCommandBufferStats(), __FUNCTION__));
 		Encoder->copyFromTexture(Source->MSAAResolveTexture.get(), ArraySlice, Info.MipLevel, Origin, Size, Destination->Texture.get(), ArraySlice, Info.MipLevel, Origin);
 	}
 	
@@ -498,29 +488,12 @@ void FMetalRHICommandContext::RHINextSubpass()
 #endif
 }
 
-#if (RHI_NEW_GPU_PROFILER == 0)
-void FMetalRHICommandContext::RHICalibrateTimers(FRHITimestampCalibrationQuery* CalibrationQuery)
-{
-    MTL::Device* MTLDevice = Device.GetDevice();
-    
-    MTL::Timestamp CPUTimeStamp, GPUTimestamp;
-    MTLDevice->sampleTimestamps(&CPUTimeStamp, &GPUTimestamp);
-
-    CalibrationQuery->CPUMicroseconds[0] = uint64(CPUTimeStamp / 1000.0);
-    CalibrationQuery->GPUMicroseconds[0] = uint64(GPUTimestamp / 1000.0);
-}
-#endif
-
 void FMetalRHICommandContext::FillBuffer(MTL::Buffer* Buffer, NS::Range Range, uint8 Value)
 {
 	check(Buffer);
-	
-	MTL::BlitCommandEncoder *TargetEncoder;
-	
 	BeginBlitEncoder();
-	TargetEncoder = CurrentEncoder.GetBlitCommandEncoder();
-	METAL_GPUPROFILE(FMetalProfiler::GetProfiler()->EncodeBlit(CurrentEncoder.GetCommandBufferStats(), FString::Printf(TEXT("FillBuffer: %p %llu %llu"), Buffer, Range.location, Range.length)));
-	
+
+	MTL::BlitCommandEncoder* TargetEncoder = CurrentEncoder.GetBlitCommandEncoder();
 	check(TargetEncoder);
 	
 	TargetEncoder->fillBuffer(Buffer, Range, Value);
@@ -533,15 +506,23 @@ void FMetalRHICommandContext::CopyFromTextureToBuffer(MTL::Texture* Texture, uin
 	BeginBlitEncoder();
 	MTL::BlitCommandEncoder* Encoder = CurrentEncoder.GetBlitCommandEncoder();
 	check(Encoder);
-	
-	METAL_GPUPROFILE(FMetalProfiler::GetProfiler()->EncodeBlit(CurrentEncoder.GetCommandBufferStats(), __FUNCTION__));
+
+	if (Texture)
 	{
-		if(Texture)
-		{
-			Encoder->copyFromTexture(Texture, sourceSlice, sourceLevel, sourceOrigin, sourceSize,
-									 toBuffer->GetMTLBuffer(), destinationOffset + toBuffer->GetOffset(), destinationBytesPerRow, destinationBytesPerImage, options);
-		}
+		Encoder->copyFromTexture(
+			Texture,
+			sourceSlice,
+			sourceLevel,
+			sourceOrigin,
+			sourceSize,
+			toBuffer->GetMTLBuffer(),
+			destinationOffset + toBuffer->GetOffset(),
+			destinationBytesPerRow,
+			destinationBytesPerImage,
+			options
+		);
 	}
+
 	EndBlitEncoder();
 }
 
@@ -550,8 +531,7 @@ void FMetalRHICommandContext::CopyFromBufferToTexture(FMetalBufferPtr Buffer, ui
 	BeginBlitEncoder();
 	MTL::BlitCommandEncoder* Encoder = CurrentEncoder.GetBlitCommandEncoder();
 	check(Encoder);
-	
-	METAL_GPUPROFILE(FMetalProfiler::GetProfiler()->EncodeBlit(CurrentEncoder.GetCommandBufferStats(), __FUNCTION__));
+
 	if (options == MTL::BlitOptionNone)
 	{
 		Encoder->copyFromBuffer(Buffer->GetMTLBuffer(), sourceOffset + Buffer->GetOffset(), sourceBytesPerRow, sourceBytesPerImage, sourceSize,
@@ -572,8 +552,7 @@ void FMetalRHICommandContext::CopyFromTextureToTexture(MTL::Texture* Texture, ui
 	
 	MTL::BlitCommandEncoder* Encoder = CurrentEncoder.GetBlitCommandEncoder();
 	check(Encoder);
-	
-	METAL_GPUPROFILE(FMetalProfiler::GetProfiler()->EncodeBlit(CurrentEncoder.GetCommandBufferStats(), __FUNCTION__));
+
 	Encoder->copyFromTexture(Texture, sourceSlice, sourceLevel, sourceOrigin, sourceSize, toTexture, destinationSlice, destinationLevel, destinationOrigin);
 	
 	EndBlitEncoder();
@@ -586,10 +565,13 @@ void FMetalRHICommandContext::CopyFromBufferToBuffer(FMetalBufferPtr SourceBuffe
 	MTL::BlitCommandEncoder* Encoder = CurrentEncoder.GetBlitCommandEncoder();
 	check(Encoder);
 	
-	METAL_GPUPROFILE(FMetalProfiler::GetProfiler()->EncodeBlit(CurrentEncoder.GetCommandBufferStats(), __FUNCTION__));
-	
-	Encoder->copyFromBuffer(SourceBuffer->GetMTLBuffer(), SourceOffset + SourceBuffer->GetOffset(),
-							DestinationBuffer->GetMTLBuffer(), DestinationOffset + DestinationBuffer->GetOffset(), Size);
+	Encoder->copyFromBuffer(
+		SourceBuffer->GetMTLBuffer(),
+		SourceOffset + SourceBuffer->GetOffset(),
+		DestinationBuffer->GetMTLBuffer(),
+		DestinationOffset + DestinationBuffer->GetOffset(),
+		Size
+	);
 	
 	EndBlitEncoder();
 }
@@ -703,11 +685,9 @@ void FMetalRHICommandContext::StartCommandBuffer()
 	
 	CurrentEncoder.StartCommandBuffer();
 	
-#if RHI_NEW_GPU_PROFILER
 	auto& Event = CurrentEncoder.GetCommandBuffer()->EmplaceProfilerEvent<UE::RHI::GPUProfiler::FEvent::FBeginWork>(0);
 	Event.GPUTimestampTOP = 0;
 	CurrentEncoder.GetCommandBuffer()->SetBeginWorkTimestamp(&Event.GPUTimestampTOP);
-#endif
 	
 	// Add new command buffer to payload
 	GetPayload(EPhase::Execute)->CommandBuffersToExecute.Add(CurrentEncoder.GetCommandBuffer());
@@ -718,11 +698,9 @@ void FMetalRHICommandContext::EndCommandBuffer()
 	check(CurrentEncoder.GetCommandBuffer());
 	check(!bWithinRenderPass);
 	
-#if RHI_NEW_GPU_PROFILER
 	auto& Event = CurrentEncoder.GetCommandBuffer()->EmplaceProfilerEvent<UE::RHI::GPUProfiler::FEvent::FEndWork>();
 	Event.GPUTimestampBOP = 0;
 	CurrentEncoder.GetCommandBuffer()->SetEndWorkTimestamp(&Event.GPUTimestampBOP);
-#endif
 	
 	if(CurrentEncoder.IsAnyCommandEncoderActive())
 	{
@@ -731,32 +709,6 @@ void FMetalRHICommandContext::EndCommandBuffer()
 	CurrentEncoder.EndCommandBuffer(this);
 }
 
-void FMetalRHICommandContext::StartTiming(class FMetalEventNode* EventNode)
-{
-#if RHI_NEW_GPU_PROFILER == 0
-	bool const bHasCurrentCommandBuffer = CurrentEncoder.GetCommandBuffer();
-	
-	if(!bHasCurrentCommandBuffer)
-	{
-		StartCommandBuffer();
-	}
-
-	CurrentEncoder.GetCommandBuffer()->ActiveEventNodes.Add(EventNode);
-
-	EventNode->SyncPoint = GetContextSyncPoint();
-#endif
-}
-
-void FMetalRHICommandContext::EndTiming(class FMetalEventNode* EventNode)
-{
-#if RHI_NEW_GPU_PROFILER == 0
-	if(CurrentEncoder.GetCommandBuffer())
-	{
-		CurrentEncoder.GetCommandBuffer()->ActiveEventNodes.Remove(EventNode);
-	}
-#endif
-}
-
 void FMetalRHICommandContext::SynchronizeResource(MTL::Resource* Resource)
 {
 	check(Resource);
@@ -764,8 +716,7 @@ void FMetalRHICommandContext::SynchronizeResource(MTL::Resource* Resource)
 	BeginBlitEncoder();
 	MTL::BlitCommandEncoder* Encoder = CurrentEncoder.GetBlitCommandEncoder();
 	check(Encoder);
-	
-	METAL_GPUPROFILE(FMetalProfiler::GetProfiler()->EncodeBlit(CurrentEncoder.GetCommandBufferStats(), __FUNCTION__));
+
 	Encoder->synchronizeResource(Resource);
 	EndBlitEncoder();
 #endif
@@ -778,8 +729,7 @@ void FMetalRHICommandContext::SynchronizeTexture(MTL::Texture* Texture, uint32 S
 	BeginBlitEncoder();
 	MTL::BlitCommandEncoder* Encoder = CurrentEncoder.GetBlitCommandEncoder();
 	check(Encoder);
-	
-	METAL_GPUPROFILE(FMetalProfiler::GetProfiler()->EncodeBlit(CurrentEncoder.GetCommandBufferStats(), __FUNCTION__));
+
 	Encoder->synchronizeTexture(Texture, Slice, Level);
 	EndBlitEncoder();
 #endif
@@ -795,8 +745,6 @@ FMetalCommandBuffer* FMetalRHICommandContext::GetCurrentCommandBuffer()
 	return CurrentEncoder.GetCommandBuffer();
 }
 
-
-#if RHI_NEW_GPU_PROFILER
 void FMetalRHICommandContext::FlushProfilerStats()
 {
 	// Flush accumulated draw stats
@@ -806,7 +754,6 @@ void FMetalRHICommandContext::FlushProfilerStats()
 		StatEvent = {};
 	}
 }
-#endif
 
 void FMetalRHICommandContext::FlushCommands(EMetalFlushFlags FlushFlags)
 {
@@ -845,7 +792,7 @@ void FMetalRHICommandContext::FlushCommands(EMetalFlushFlags FlushFlags)
 
 FMetalRHIUploadContext::FMetalRHIUploadContext(FMetalDevice& Device)
 {
-	UploadContext = new FMetalRHICommandContext(Device, nullptr);
+	UploadContext = new FMetalRHICommandContext(Device);
 	UploadContext->ResetContext();
 }
 
diff --git a/Engine/Source/Runtime/Apple/MetalRHI/Private/MetalRHIPrivate.h b/Engine/Source/Runtime/Apple/MetalRHI/Private/MetalRHIPrivate.h
--- a/Engine/Source/Runtime/Apple/MetalRHI/Private/MetalRHIPrivate.h
+++ b/Engine/Source/Runtime/Apple/MetalRHI/Private/MetalRHIPrivate.h
@@ -120,15 +120,9 @@ extern bool GMetalCommandBufferDebuggingEnabled;
 #ifndef ENABLE_METAL_GPUEVENTS_IN_TEST
     #define ENABLE_METAL_GPUEVENTS_IN_TEST 0
 #endif
-#define ENABLE_METAL_GPUEVENTS    (UE_BUILD_DEBUG || UE_BUILD_DEVELOPMENT || (UE_BUILD_TEST && ENABLE_METAL_GPUEVENTS_IN_TEST))
+#define ENABLE_METAL_GPUEVENTS     (UE_BUILD_DEBUG || UE_BUILD_DEVELOPMENT || (UE_BUILD_TEST && ENABLE_METAL_GPUEVENTS_IN_TEST))
 #define ENABLE_METAL_GPUPROFILE    (ENABLE_METAL_GPUEVENTS && 1)
 
-#if ENABLE_METAL_GPUPROFILE && RHI_NEW_GPU_PROFILER == 0
-    #define METAL_GPUPROFILE(Code) Code
-#else
-    #define METAL_GPUPROFILE(Code)
-#endif
-
 #define UNREAL_TO_METAL_BUFFER_INDEX(Index) ((MaxMetalStreams - 1) - Index)
 #define METAL_TO_UNREAL_BUFFER_INDEX(Index) ((MaxMetalStreams - 1) - Index)
 
diff --git a/Engine/Source/Runtime/Apple/MetalRHI/Private/MetalRHIRenderQuery.cpp b/Engine/Source/Runtime/Apple/MetalRHI/Private/MetalRHIRenderQuery.cpp
--- a/Engine/Source/Runtime/Apple/MetalRHI/Private/MetalRHIRenderQuery.cpp
+++ b/Engine/Source/Runtime/Apple/MetalRHI/Private/MetalRHIRenderQuery.cpp
@@ -150,9 +150,7 @@ FMetalQueryBufferRef FMetalQueryBufferPool::AcquireQueryBuffer(uint32_t NumOcclu
 		Buffer = Buffers.Pop();
 	}
 	else
-	{
-		METAL_GPUPROFILE(FScopedMetalCPUStats CPUStat(FString::Printf(TEXT("AllocBuffer: %llu, %llu"), BufferSize, MTL::ResourceStorageModeShared)));
-		
+	{		
 		MTL::ResourceOptions HazardTrackingMode = MTL::ResourceHazardTrackingModeUntracked;
 		static bool bSupportsHeaps = Device.SupportsFeature(EMetalFeaturesHeaps);
 		if(bSupportsHeaps)
diff --git a/Engine/Source/Runtime/Apple/MetalRHI/Private/MetalStateCache.cpp b/Engine/Source/Runtime/Apple/MetalRHI/Private/MetalStateCache.cpp
--- a/Engine/Source/Runtime/Apple/MetalRHI/Private/MetalStateCache.cpp
+++ b/Engine/Source/Runtime/Apple/MetalRHI/Private/MetalStateCache.cpp
@@ -2170,7 +2170,6 @@ void FMetalStateCache::FlushVisibilityResults(FMetalCommandEncoder& CommandEncod
 	{
 		MTL::BlitCommandEncoder* Encoder = CommandEncoder.GetBlitCommandEncoder();
 
-		METAL_GPUPROFILE(FMetalProfiler::GetProfiler()->EncodeBlit(CommandEncoder.GetCommandBufferStats(), __FUNCTION__));
         Encoder->synchronizeResource(VisibilityResults->Buffer->GetMTLBuffer());
 		VisibilityWritten = 0;
 	}
diff --git a/Engine/Source/Runtime/Apple/MetalRHI/Private/MetalSubmission.cpp b/Engine/Source/Runtime/Apple/MetalRHI/Private/MetalSubmission.cpp
--- a/Engine/Source/Runtime/Apple/MetalRHI/Private/MetalSubmission.cpp
+++ b/Engine/Source/Runtime/Apple/MetalRHI/Private/MetalSubmission.cpp
@@ -117,7 +117,6 @@ void FMetalDynamicRHI::InitializeSubmissionPipe()
 	}
 
 	// Initialize the timing structs in each queue, and the engine GPU profilers
-#if RHI_NEW_GPU_PROFILER
 	{
 		TArray<FMetalPayload*> Payloads;
 		TArray<UE::RHI::GPUProfiler::FQueue> ProfilerQueues;
@@ -134,7 +133,6 @@ void FMetalDynamicRHI::InitializeSubmissionPipe()
 	
 		SubmitPayloads(MoveTemp(Payloads));
 	}
-#endif
 }
 
 void FMetalDynamicRHI::ShutdownSubmissionPipe()
@@ -224,7 +222,6 @@ void FMetalDynamicRHI::SubmitCommands(TConstArrayView<FMetalFinalizedCommands*>
 {	
 	SCOPED_NAMED_EVENT_TEXT("CommandList_Submit", FColor::Magenta);
 
-#if RHI_NEW_GPU_PROFILER
 	TArray<FMetalPayload*> AllPayloads;
 	for (FMetalFinalizedCommands* Payloads : Commands)
 	{
@@ -251,37 +248,6 @@ void FMetalDynamicRHI::SubmitCommands(TConstArrayView<FMetalFinalizedCommands*>
 	}
 
 	SubmitPayloads(MoveTemp(AllPayloads));
-
-#else
-	
-	TArray<FMetalPayload*> AllPayloads;
-	#if WITH_RHI_BREADCRUMBS
-	TArray<TSharedPtr<FRHIBreadcrumbAllocator>> BreadcrumbAllocators;
-	#endif
-
-	for (FMetalFinalizedCommands* Payloads : Commands)
-	{
-	#if WITH_RHI_BREADCRUMBS
-		for (FMetalPayload* Payload : *Payloads)
-		{
-			Payload->BreadcrumbRange = Payloads->BreadcrumbRange;
-		}
-	#endif
-
-		AllPayloads.Append(MoveTemp(static_cast<TArray<FMetalPayload*>&>(*Payloads)));
-	#if WITH_RHI_BREADCRUMBS
-		BreadcrumbAllocators.Append(MoveTemp(Payloads->BreadcrumbAllocators));
-	#endif
-		delete Payloads;
-	}
-
-	SubmitPayloads(MoveTemp(AllPayloads));
-
-	#if WITH_RHI_BREADCRUMBS
-	// Enqueue the breadcrumb allocator references for cleanup once all prior payloads have completed on the GPU.
-	DeferredDelete([Array = MoveTemp(BreadcrumbAllocators)]() {});
-	#endif
-#endif
 }
 
 void FMetalDynamicRHI::SubmitPayloads(TArray<FMetalPayload*>&& Payloads)
@@ -440,9 +406,6 @@ FMetalDynamicRHI::FProcessResult FMetalDynamicRHI::ProcessSubmissionQueue()
 						
 						CurrentQueue.OcclusionQueries.Append(MoveTemp(CurrentCommandBuffer->OcclusionQueries));
 						CurrentQueue.TimestampQueries.Append(MoveTemp(CurrentCommandBuffer->TimestampQueries));
-#if RHI_NEW_GPU_PROFILER == 0
-						CurrentQueue.EventSampleCounters.Append(MoveTemp(CurrentCommandBuffer->EventSampleCounters));
-#endif
 					}
 				}
 
@@ -499,10 +462,6 @@ uint64 FMetalCommandQueue::FinalizePayload(bool bRequiresSignal, FPayloadArray&
 	PayloadToSubmit->BatchedObjects.OcclusionQueries = MoveTemp(OcclusionQueries);
 	PayloadToSubmit->BatchedObjects.TimestampQueries = MoveTemp(TimestampQueries);
 	
-#if RHI_NEW_GPU_PROFILER == 0
-	PayloadToSubmit->BatchedObjects.EventSampleCounters = MoveTemp(EventSampleCounters);
-#endif
-	
 	PayloadsToHandDown.Add(PayloadToSubmit);
 	PayloadToSubmit = nullptr;
 
@@ -549,9 +508,7 @@ void FMetalDynamicRHI::FlushBatchedPayloads(FMetalCommandQueue::FPayloadArray& P
 
 			for (FMetalCommandBuffer* CommandBuffer : Payload->CommandBuffersToExecute)
 			{
-#if RHI_NEW_GPU_PROFILER
 				CommandBuffer->FlushProfilerEvents(Payload->EventStream, Time);
-#endif
 				Payload->Queue.CommitCommandBuffer(CommandBuffer);
 			}
 		}
@@ -585,14 +542,12 @@ void FMetalDynamicRHI::FlushBatchedPayloads(FMetalCommandQueue::FPayloadArray& P
 			Queue.CommitCommandBuffer(CommandBuffer);
 			Queue.GetSignalEvent().LastSignaledValue.store(Payload->CompletionFenceValue, std::memory_order_release);
 		}
-		
-#if RHI_NEW_GPU_PROFILER
+
 		if (Payload->EndFrameEvent.IsSet())
 		{
 			Payload->EndFrameEvent->CPUTimestamp = FPlatformTime::Cycles64();
 			Payload->EventStream.Emplace<UE::RHI::GPUProfiler::FEvent::FFrameBoundary>(*Payload->EndFrameEvent);
 		}
-#endif
 
 		// Submission of this payload is completed. Signal the submission event if one was provided.
 		if (Payload->SubmissionEvent)
@@ -728,11 +683,10 @@ void FMetalDynamicRHI::ProcessPendingCommandBuffers()
 		MTL::CommandBufferStatus Status = CompletedBuffer->status();
 		if (Status == MTL::CommandBufferStatusCompleted && bContinueProcessing)
 		{
-#if RHI_NEW_GPU_PROFILER
 #if WITH_RHI_BREADCRUMBS
 			FMetalBreadcrumbProfiler::GetInstance()->ResolveBreadcrumbTrackerStream(CommandBuffer->BreadcrumbTrackerStream);
 #endif
-			
+
 			uint64_t* Start = CommandBuffer->BeginWorkTimestamp;
 			uint64_t* End = CommandBuffer->EndWorkTimestamp;
 			
@@ -752,11 +706,7 @@ void FMetalDynamicRHI::ProcessPendingCommandBuffers()
 				*Start = CompletedBuffer->GPUStartTime() / SecondsPerCycle;
 				*End = CompletedBuffer->GPUEndTime() / SecondsPerCycle;
 			}
-#else
-			FMetalCommandBufferTimer* Timer = CommandBuffer->GetTimer();
-			Timer->AddTiming({CompletedBuffer->GPUStartTime(), CompletedBuffer->GPUEndTime()});
-#endif
-			
+
 			delete CommandBuffer;
 			
 			return true;
@@ -817,7 +767,7 @@ FMetalDynamicRHI::FProcessResult FMetalDynamicRHI::ProcessInterruptQueue()
 				MTL::CommandBuffer* CompletedBuffer = CommandBuffer->GetMTLCmdBuffer();
 				if (CompletedBuffer->status() == MTL::CommandBufferStatusError)
 				{
-					FMetalCommandList::HandleMetalCommandBufferFailure(CompletedBuffer);
+					FMetalCommandQueue::HandleMetalCommandBufferFailure(CompletedBuffer);
 				}
 			}
 			
@@ -829,7 +779,7 @@ FMetalDynamicRHI::FProcessResult FMetalDynamicRHI::ProcessInterruptQueue()
 				Status = SignalBuffer->status();
 				if(Status == MTL::CommandBufferStatusError)
 				{
-					FMetalCommandList::HandleMetalCommandBufferFailure(SignalBuffer);
+					FMetalCommandQueue::HandleMetalCommandBufferFailure(SignalBuffer);
 				}
 			}
 			
@@ -846,8 +796,7 @@ FMetalDynamicRHI::FProcessResult FMetalDynamicRHI::ProcessInterruptQueue()
 			{
 				delete Payload->SignalCommandBuffer;
 			}
-			
-#if RHI_NEW_GPU_PROFILER
+
 			if (!Payload->EventStream.IsEmpty())
 			{			
 				check(CurrentQueue.Timing);
@@ -859,7 +808,7 @@ FMetalDynamicRHI::FProcessResult FMetalDynamicRHI::ProcessInterruptQueue()
 				// Switch the new timing struct into the queue. This redirects timestamp results to separate each frame's work.
 				CurrentQueue.Timing = Payload->Timing.GetValue();
 			}
-#endif
+
 			// Resolve query results
 			{
 				for (FMetalRHIRenderQuery* OcclusionQuery : Payload->BatchedObjects.OcclusionQueries)
@@ -885,27 +834,6 @@ FMetalDynamicRHI::FProcessResult FMetalDynamicRHI::ProcessInterruptQueue()
 					
 					TimestampQuery->Release();
 				}
-				
-				
-#if RHI_NEW_GPU_PROFILER == 0
-				for(auto& Pair : Payload->BatchedObjects.EventSampleCounters)
-				{
-					uint64_t& Start = Pair.Key->StartTime;
-					uint64_t& End = Pair.Key->EndTime;
-					
-					for(FMetalCounterSamplePtr Sample: Pair.Value)
-					{
-						uint64_t StartTime, EndTime;
-						Sample->ResolveStageCounters(StartTime, EndTime);
-						
-						StartTime /= 1000.0;
-						EndTime /= 1000.0;
-						
-						Start = Start > 0 ? FMath::Min(StartTime, Start) : StartTime;
-						End = End > 0 ? FMath::Max(EndTime, End) : EndTime;
-					}
-				} 
-#endif
 			}
 			
 			// Signal the CPU events of all sync points associated with this batch.
@@ -940,9 +868,7 @@ FMetalDynamicRHI::FProcessResult FMetalDynamicRHI::ProcessInterruptQueue()
 
 FMetalPayload::FMetalPayload(FMetalCommandQueue& Queue)
 	: Queue(Queue)
-#if RHI_NEW_GPU_PROFILER
 	, EventStream(Queue.GetProfilerQueue())
-#endif
 {}
 
 FMetalPayload::~FMetalPayload()
diff --git a/Engine/Source/Runtime/Apple/MetalRHI/Private/MetalSubmission.h b/Engine/Source/Runtime/Apple/MetalRHI/Private/MetalSubmission.h
--- a/Engine/Source/Runtime/Apple/MetalRHI/Private/MetalSubmission.h
+++ b/Engine/Source/Runtime/Apple/MetalRHI/Private/MetalSubmission.h
@@ -121,24 +121,15 @@ public:
 };
 
 class FMetalRHIRenderQuery;
-struct FMetalCounterSample;
-class FMetalEventNode;
 
 struct FMetalBatchedPayloadObjects
 {
 	TArray<FMetalRHIRenderQuery*> OcclusionQueries;
 	TArray<FMetalRHIRenderQuery*> TimestampQueries;
-#if RHI_NEW_GPU_PROFILER == 0
-	TMap<FMetalEventNode*, TArray<FMetalCounterSamplePtr>> EventSampleCounters;
-#endif
 
 	bool IsEmpty() const
 	{
-		return TimestampQueries.Num() == 0 && OcclusionQueries.Num() == 0
-#if RHI_NEW_GPU_PROFILER == 0
-				&& EventSampleCounters.Num() == 0
-#endif
-		;
+		return TimestampQueries.Num() == 0 && OcclusionQueries.Num() == 0;
 	}
 };
 
@@ -197,19 +188,15 @@ struct FMetalPayload : public FMetalPayloadBaseFixLayout
 	TOptional<uint64> SubmissionTime;
 
 	FMetalBatchedPayloadObjects BatchedObjects;
-	
-	TMap<FMetalEventNode*, TArray<FMetalCounterSamplePtr>> EventSampleCounters;
 
 #if WITH_RHI_BREADCRUMBS
 	FRHIBreadcrumbRange BreadcrumbRange {};
 	TSharedPtr<FRHIBreadcrumbAllocatorArray> BreadcrumbAllocators {};
 #endif
 
-#if RHI_NEW_GPU_PROFILER
 	TOptional<FMetalTiming*> Timing;
 	UE::RHI::GPUProfiler::FEventStream EventStream;
 	TOptional<UE::RHI::GPUProfiler::FEvent::FFrameBoundary> EndFrameEvent;
-#endif
 
 	virtual ~FMetalPayload();
 
diff --git a/Engine/Source/Runtime/Apple/MetalRHI/Private/MetalTexture.cpp b/Engine/Source/Runtime/Apple/MetalRHI/Private/MetalTexture.cpp
--- a/Engine/Source/Runtime/Apple/MetalRHI/Private/MetalTexture.cpp
+++ b/Engine/Source/Runtime/Apple/MetalRHI/Private/MetalTexture.cpp
@@ -1031,10 +1031,6 @@ void FMetalSurface::UpdateSurfaceAndDestroySourceBuffer(FMetalRHICommandContext*
 	}
 	
 	FMetalDynamicRHI::Get().DeferredDelete(Source);
-	
-#if STATS
-	GMetalTexturePageOnTime += (FPlatformTime::Cycles64() - Start);
-#endif
 }
 
 FRHILockTextureResult FMetalSurface::Lock(const FRHILockTextureArgs& Arguments, bool bSingleLayer)
diff --git a/Engine/Source/Runtime/Apple/MetalRHI/Private/MetalViewport.cpp b/Engine/Source/Runtime/Apple/MetalRHI/Private/MetalViewport.cpp
--- a/Engine/Source/Runtime/Apple/MetalRHI/Private/MetalViewport.cpp
+++ b/Engine/Source/Runtime/Apple/MetalRHI/Private/MetalViewport.cpp
@@ -565,11 +565,6 @@ void FMetalViewport::Present(FMetalCommandQueue& CommandQueue, bool bLockToVsync
 					{
 						check(CurrentCommandBuffer);
 						
-#if ENABLE_METAL_GPUPROFILE && RHI_NEW_GPU_PROFILER == 0
-						FMetalProfiler* Profiler = FMetalProfiler::GetProfiler();
-						FMetalCommandBufferStats* Stats = Profiler->AllocateCommandBuffer(CurrentCommandBuffer->GetMTLCmdBuffer(), 0);
-#endif
-						
 						if (GMetalSupportsIntermediateBackBuffer)
 						{
 							TRefCountPtr<FMetalSurface> Texture = LastCompleteFrame;
@@ -583,7 +578,6 @@ void FMetalViewport::Present(FMetalCommandQueue& CommandQueue, bool bLockToVsync
 							
 							MTLBlitCommandEncoderPtr Encoder = NS::RetainPtr(CurrentCommandBuffer->GetMTLCmdBuffer()->blitCommandEncoder());
 							check(Encoder);
-							METAL_GPUPROFILE(Profiler->EncodeBlit(Stats, __FUNCTION__));
 
 							Encoder->copyFromTexture(Src.get(), 0, 0, MTL::Origin(0, 0, 0), MTL::Size(Width, Height, 1), Dst.get(), 0, 0, MTL::Origin(0, 0, 0));
 							Encoder->endEncoding();
@@ -591,15 +585,6 @@ void FMetalViewport::Present(FMetalCommandQueue& CommandQueue, bool bLockToVsync
 							Drawable->release();
 							Drawable = nullptr;
 						}
-
-						METAL_GPUPROFILE(Stats->End(CurrentCommandBuffer->GetMTLCmdBuffer()));
-						
-						// This is a bit different than the usual pattern.
-						// This command buffer here is committed directly, instead of going through
-						// FMetalCommandList::Commit. So long as Present() is called within
-						// high level RHI BeginFrame/EndFrame this will not fine.
-						// Otherwise the recording of the Present time will be offset by one in the
-						// FMetalGPUProfiler frame indices.
 						
 						dispatch_semaphore_t& FrameSemaphore = Device.GetFrameSemaphore();
 						dispatch_retain(FrameSemaphore);
@@ -613,10 +598,7 @@ void FMetalViewport::Present(FMetalCommandQueue& CommandQueue, bool bLockToVsync
 						{
 							dispatch_semaphore_signal(FrameSemaphore);
 							dispatch_release(FrameSemaphore);
-							
-#if RHI_NEW_GPU_PROFILER == 0
-							FMetalCommandBufferTimer::RecordPresent(cmd_buf);
-#endif
+
 							LocalDrawable->release();
 #if PLATFORM_MAC
 							MainThreadCall(^{
@@ -793,11 +775,6 @@ void FMetalViewport::PresentImmersive(const MetalRHIVisionOS::PresentImmersivePa
 		}
 				
 		FMetalCommandBuffer* FinalCommandBuffer = Context.GetCurrentCommandBuffer();;
-
-#if ENABLE_METAL_GPUPROFILE && RHI_NEW_GPU_PROFILER == 0
-		FMetalProfiler* Profiler = FMetalProfiler::GetProfiler();
-		FMetalCommandBufferStats* Stats = Profiler->AllocateCommandBuffer(FinalCommandBuffer->GetMTLCmdBuffer(), 0);
-#endif
 		
 		{
 			dispatch_semaphore_t& FrameSemaphore = Device.GetFrameSemaphore();
@@ -806,16 +783,10 @@ void FMetalViewport::PresentImmersive(const MetalRHIVisionOS::PresentImmersivePa
 			{
 				dispatch_semaphore_signal(FrameSemaphore);
 				dispatch_release(FrameSemaphore);
-				
-#if RHI_NEW_GPU_PROFILER == 0
-				FMetalCommandBufferTimer::RecordPresent(cmd_buf);
-#endif
 			};
 			FinalCommandBuffer->GetMTLCmdBuffer()->addCompletedHandler(CommandBufferHandler);
 		}
-		
-		METAL_GPUPROFILE(Stats->End(FinalCommandBuffer->GetMTLCmdBuffer()));			
-		
+
 		TArray<FMetalPayload*> Payloads;
 		Context.Finalize(Payloads);
 
@@ -832,10 +803,6 @@ void FMetalViewport::PresentImmersive(const MetalRHIVisionOS::PresentImmersivePa
 		// Wait for the frame semaphore
 		dispatch_semaphore_wait(Device.GetFrameSemaphore(), DISPATCH_TIME_FOREVER);
 	}
-	
-#if RHI_NEW_GPU_PROFILER == 0
-	FMetalCommandBufferTimer::ResetFrameBufferTimings();
-#endif
 }
 #endif //PLATFORM_VISIONOS
 
diff --git a/Engine/Source/Runtime/Apple/MetalRHI/Private/Shaders/Types/MetalComputeShader.cpp b/Engine/Source/Runtime/Apple/MetalRHI/Private/Shaders/Types/MetalComputeShader.cpp
--- a/Engine/Source/Runtime/Apple/MetalRHI/Private/Shaders/Types/MetalComputeShader.cpp
+++ b/Engine/Source/Runtime/Apple/MetalRHI/Private/Shaders/Types/MetalComputeShader.cpp
@@ -117,7 +117,6 @@ FMetalShaderPipelinePtr FMetalComputeShader::GetPipeline()
         MTLComputePipelineStatePtr Kernel;
 		MTL::ComputePipelineReflection* Reflection = nullptr;
 
-		METAL_GPUPROFILE(FScopedMetalCPUStats CPUStat(FString::Printf(TEXT("NewComputePipeline: %d_%d"), SourceLen, SourceCRC)));
 #if METAL_DEBUG_OPTIONS
 		if (Device.GetRuntimeDebuggingLevel() >= EMetalDebugLevelFastValidation)
 		{
diff --git a/Engine/Source/Runtime/Apple/MetalRHI/Private/Shaders/Types/Templates/MetalBaseShader.h b/Engine/Source/Runtime/Apple/MetalRHI/Private/Shaders/Types/Templates/MetalBaseShader.h
--- a/Engine/Source/Runtime/Apple/MetalRHI/Private/Shaders/Types/Templates/MetalBaseShader.h
+++ b/Engine/Source/Runtime/Apple/MetalRHI/Private/Shaders/Types/Templates/MetalBaseShader.h
@@ -252,8 +252,6 @@ void TMetalBaseShader<BaseResourceType, ShaderType>::Init(TArrayView<const uint8
 			}
 			else
 			{
-				METAL_GPUPROFILE(FScopedMetalCPUStats CPUStat(FString::Printf(TEXT("NewLibraryBinary: %d_%d"), SourceLen, SourceCRC)));
-
 				// Archived shaders should never get in here.
 				check(!(Header.CompileFlags & (1 << CFLAG_Archive)) || BufferSize > 0);
 
@@ -277,8 +275,7 @@ void TMetalBaseShader<BaseResourceType, ShaderType>::Init(TArrayView<const uint8
 		else
 		{
 			MTL_SCOPED_AUTORELEASE_POOL;
-			
-			METAL_GPUPROFILE(FScopedMetalCPUStats CPUStat(FString::Printf(TEXT("NewLibrarySource: %d_%d"), SourceLen, SourceCRC)));
+
 			NS::String* ShaderString = ((OfflineCompiledFlag == 0) ? NS::String::string(SourceCode, NS::UTF8StringEncoding) : GlslCodeNSString);
 
             FString FinalShaderString;
@@ -483,7 +480,6 @@ MTLFunctionPtr TMetalBaseShader<BaseResourceType, ShaderType>::GetCompiledFuncti
 
 			if (!bHasFunctionConstants || !bAsync)
 			{
-				METAL_GPUPROFILE(FScopedMetalCPUStats CPUStat(FString::Printf(TEXT("NewFunction: %s"), *NSStringToFString(Name))));
 				if (!bHasFunctionConstants)
 				{
 					Function = NS::TransferPtr(Library->newFunction(Name));
@@ -503,25 +499,12 @@ MTLFunctionPtr TMetalBaseShader<BaseResourceType, ShaderType>::GetCompiledFuncti
 			}
 			else
 			{
-				METAL_GPUPROFILE(FScopedMetalCPUStats CPUStat(FString::Printf(TEXT("NewFunctionAsync: %s"), *NSStringToFString(Name))));
-				METAL_GPUPROFILE(uint64 CPUStart = CPUStat.Stats ? CPUStat.Stats->CPUStartTime : 0);
-#if ENABLE_METAL_GPUPROFILE && RHI_NEW_GPU_PROFILER == 0
-                NS::String* nsName = Name;
-                const std::function<void(MTL::Function* pFunction, NS::Error* pError)> Handler = [Key, this, CPUStart, nsName](MTL::Function* NewFunction, NS::Error* Error){
-#else
-                const std::function<void(MTL::Function* pFunction, NS::Error* pError)> Handler = [Key, this](MTL::Function* NewFunction, NS::Error* Error){
-#endif
-					METAL_GPUPROFILE(FScopedMetalCPUStats CompletionStat(FString::Printf(TEXT("NewFunctionCompletion: %s"), *NSStringToFString(nsName))));
+				const std::function<void(MTL::Function* pFunction, NS::Error* pError)> Handler = [Key, this](MTL::Function* NewFunction, NS::Error* Error)
+				{
 					UE_CLOG(NewFunction == nullptr, LogMetal, Error, TEXT("Failed to create function: %s"), *NSStringToFString(Error->description()));
 					UE_CLOG(NewFunction == nullptr, LogMetal, Fatal, TEXT("*********** Error\n%s"), *NSStringToFString(GetSourceCode()));
 
 					GetMetalCompiledShaderCache().Add(Key, Library, NS::RetainPtr(NewFunction));
-#if ENABLE_METAL_GPUPROFILE && RHI_NEW_GPU_PROFILER == 0
-					if (CompletionStat.Stats)
-					{
-						CompletionStat.Stats->CPUStartTime = CPUStart;
-					}
-#endif
 				};
                     
                 Library->newFunction(Name, ConstantValues, Handler);
diff --git a/Engine/Source/Runtime/Apple/MetalRHI/Public/MetalRHIContext.h b/Engine/Source/Runtime/Apple/MetalRHI/Public/MetalRHIContext.h
--- a/Engine/Source/Runtime/Apple/MetalRHI/Public/MetalRHIContext.h
+++ b/Engine/Source/Runtime/Apple/MetalRHI/Public/MetalRHIContext.h
@@ -9,14 +9,11 @@
 #include "MetalResources.h"
 #include "MetalViewport.h"
 #include "MetalDevice.h"
-#include "MetalCommandList.h"
 #include "MetalCommandEncoder.h"
 #include "MetalRHIRenderQuery.h"
 #include "MetalBindlessDescriptors.h"
 #include "RHICore.h"
 
-class FMetalEventNode;
-
 #if PLATFORM_VISIONOS
 namespace MetalRHIVisionOS
 {
@@ -47,7 +44,7 @@ ENUM_CLASS_FLAGS(EMetalFlushFlags)
 class FMetalRHICommandContext : public IRHICommandContext
 {
 public:
-	FMetalRHICommandContext(FMetalDevice& Device, class FMetalProfiler* InProfiler);
+	FMetalRHICommandContext(FMetalDevice& Device);
 	virtual ~FMetalRHICommandContext();
 	
 	static inline FMetalRHICommandContext& Get(FRHICommandListBase& CmdList)
@@ -61,9 +58,6 @@ public:
 	void EndComputeEncoder();
 	void BeginBlitEncoder();
 	void EndBlitEncoder();
-	
-	/** Get the profiler pointer */
-	inline class FMetalProfiler* GetProfiler() const { return Profiler; }
 
 	virtual void RHISetComputePipelineState(FRHIComputePipelineState* ComputePipelineState) override;
 	
@@ -82,10 +76,6 @@ public:
 	
 	virtual void RHICopyTexture(FRHITexture* SourceTextureRHI, FRHITexture* DestTextureRHI, const FRHICopyTextureInfo& CopyInfo) final override;
 	virtual void RHICopyBufferRegion(FRHIBuffer* DstBufferRHI, uint64 DstOffset, FRHIBuffer* SrcBufferRHI, uint64 SrcOffset, uint64 NumBytes) final override;
-
-#if (RHI_NEW_GPU_PROFILER == 0)
-    virtual void RHICalibrateTimers(FRHITimestampCalibrationQuery* CalibrationQuery) final override;
-#endif
     
 	virtual void RHIBeginRenderQuery(FRHIRenderQuery* RenderQuery) final override;
 	virtual void RHIEndRenderQuery(FRHIRenderQuery* RenderQuery) final override;
@@ -268,9 +258,6 @@ public:
 	// Inserts a command that blocks the GPU queue until the specified sync point is signaled.
 	void WaitSyncPoint(FMetalSyncPoint* SyncPoint);
 	
-	void StartTiming(class FMetalEventNode* EventNode);
-	void EndTiming(class FMetalEventNode* EventNode);
-	
 	void SynchronizeResource(MTL::Resource* Resource);
 	void SynchronizeTexture(MTL::Texture* Texture, uint32 Slice, uint32 Level);
 	
@@ -305,16 +292,6 @@ public:
         return GlobalUniformBuffers;
     }
 	
-	inline void SetProfiler(FMetalProfiler* InProfiler)
-	{
-		Profiler = InProfiler;
-	}
-	
-	inline FMetalProfiler* GetProfiler()
-	{
-		return Profiler;
-	}
-	
 	inline TSharedRef<FMetalQueryBufferPool, ESPMode::ThreadSafe> GetQueryBufferPool()
 	{
 		return QueryBuffer.ToSharedRef();
@@ -343,9 +320,7 @@ public:
 	void StartCommandBuffer();
 	void EndCommandBuffer();
 	
-#if RHI_NEW_GPU_PROFILER
 	void FlushProfilerStats();
-#endif
 	
 	void FlushCommands(EMetalFlushFlags Flags);
 	
@@ -354,10 +329,7 @@ protected:
 	
 	/** The wrapper around the device command-queue for creating & committing command buffers to */
 	FMetalCommandQueue& CommandQueue;
-	
-	/** The wrapper around command buffers for ensuring correct parallel execution order */
-	FMetalCommandList CommandList;
-	
+
 	FMetalCommandEncoder CurrentEncoder;
 	
 	/** The cache of all tracked & accessible state. */
@@ -367,9 +339,6 @@ protected:
 	TSharedPtr<FMetalQueryBufferPool, ESPMode::ThreadSafe> QueryBuffer;
 	
 	MTL::RenderPassDescriptor* RenderPassDesc = nullptr;
-	
-	/** Profiling implementation details. */
-	class FMetalProfiler* Profiler = nullptr;
 
 	TRefCountPtr<FMetalFence> CurrentEncoderFence;
 	uint64_t UploadSyncCounter = 0;
