Metal RHI

Fix for check for parallel support when creating default context.

diff --git a/Engine/Source/Runtime/Apple/MetalRHI/Private/MetalRHI.cpp b/Engine/Source/Runtime/Apple/MetalRHI/Private/MetalRHI.cpp
--- a/Engine/Source/Runtime/Apple/MetalRHI/Private/MetalRHI.cpp
+++ b/Engine/Source/Runtime/Apple/MetalRHI/Private/MetalRHI.cpp
@@ -1558,8 +1558,6 @@ IRHICommandContext* FMetalDynamicRHI::RHIGetDefaultContext()
 
 IRHIComputeContext* FMetalDynamicRHI::RHIGetCommandContext(ERHIPipeline Pipeline)
 {
-	check(GRHISupportsParallelRHIExecute);
-	
 	FMetalRHICommandContext* Context = MetalCommandContextPool.Pop();
 	if (!Context)
 	{
