Metal RHI

- Replace cases when we are using FMetalDynamicRHI::Get().RHIGetDefaultContext()
- Remove GMetalSeparatePresentThread

diff --git a/Engine/Source/Runtime/Apple/MetalRHI/Private/MetalDevice.cpp b/Engine/Source/Runtime/Apple/MetalRHI/Private/MetalDevice.cpp
--- a/Engine/Source/Runtime/Apple/MetalRHI/Private/MetalDevice.cpp
+++ b/Engine/Source/Runtime/Apple/MetalRHI/Private/MetalDevice.cpp
@@ -27,16 +27,7 @@ int32 GMetalSupportsIntermediateBackBuffer = 0;
 static FAutoConsoleVariableRef CVarMetalSupportsIntermediateBackBuffer(
 	TEXT("rhi.Metal.SupportsIntermediateBackBuffer"),
 	GMetalSupportsIntermediateBackBuffer,
-	TEXT("When enabled (> 0) allocate an intermediate texture to use as the back-buffer & blit from there into the actual device back-buffer, this is required if we use the experimental separate presentation thread. (Off by default (0))"), ECVF_ReadOnly);
-
-int32 GMetalSeparatePresentThread = 0;
-static FAutoConsoleVariableRef CVarMetalSeparatePresentThread(
-	TEXT("rhi.Metal.SeparatePresentThread"),
-	GMetalSeparatePresentThread,
-	TEXT("When enabled (> 0) requires rhi.Metal.SupportsIntermediateBackBuffer be enabled and will cause two intermediate back-buffers be allocated so that the presentation of frames to the screen can be run on a separate thread.\n")
-	TEXT("This option uncouples the Render/RHI thread from calls to -[CAMetalLayer nextDrawable] and will run arbitrarily fast by rendering but not waiting to present all frames. This is equivalent to running without V-Sync, but without the screen tearing.\n")
-	TEXT("On iOS/tvOS this is the only way to run without locking the CPU to V-Sync somewhere - this shouldn't be used in a shipping title without understanding the power/heat implications.\n")
-	TEXT("(Off by default (0))"), ECVF_ReadOnly);
+	TEXT("When enabled (> 0) allocate an intermediate texture to use as the back-buffer & blit from there into the actual device back-buffer (Off by default (0))"), ECVF_ReadOnly);
 
 #if PLATFORM_MAC
 static int32 GMetalCommandQueueSize = 5120; // This number is large due to texture streaming - currently each texture is its own command-buffer.
@@ -383,13 +374,10 @@ FMetalDevice::FMetalDevice(MTL::Device* MetalDevice, uint32 InDeviceIndex)
 	RuntimeDebuggingLevel = GMetalRuntimeDebugLevel;
 	
 	CaptureManager = new FMetalCaptureManager(MetalDevice, *CommandQueues[(uint32_t)EMetalQueueType::Direct]);
-	
-	// If the separate present thread is enabled then an intermediate backbuffer is required
-	check(!GMetalSeparatePresentThread || GMetalSupportsIntermediateBackBuffer);
-	
+
 	// Hook into the ios framepacer, if it's enabled for this platform.
 	FrameReadyEvent = NULL;
-	if( FPlatformRHIFramePacer::IsEnabled() || GMetalSeparatePresentThread )
+	if( FPlatformRHIFramePacer::IsEnabled() )
 	{
 		FrameReadyEvent = FPlatformProcess::GetSynchEventFromPool();
 		FPlatformRHIFramePacer::InitWithEvent( FrameReadyEvent );
@@ -569,7 +557,7 @@ void FMetalDevice::EnumerateFeatureSupport()
 void FMetalDevice::EndDrawingViewport(const FRHIPresentArgs& InPresentArgs)
 {
 	// We may be limiting our framerate to the display link
-	if( FrameReadyEvent != nullptr && !GMetalSeparatePresentThread )
+	if( FrameReadyEvent != nullptr  )
 	{
 		bool bIgnoreThreadIdleStats = true; // Idle time is already counted by the caller
 		FrameReadyEvent->Wait(MAX_uint32, bIgnoreThreadIdleStats);
diff --git a/Engine/Source/Runtime/Apple/MetalRHI/Private/MetalViewport.cpp b/Engine/Source/Runtime/Apple/MetalRHI/Private/MetalViewport.cpp
--- a/Engine/Source/Runtime/Apple/MetalRHI/Private/MetalViewport.cpp
+++ b/Engine/Source/Runtime/Apple/MetalRHI/Private/MetalViewport.cpp
@@ -27,7 +27,6 @@
 #include "Engine/RendererSettings.h"
 
 extern int32 GMetalSupportsIntermediateBackBuffer;
-extern int32 GMetalSeparatePresentThread;
 extern float GMetalPresentFramePacing;
 
 #if PLATFORM_IOS
@@ -163,20 +162,8 @@ FMetalViewport::FMetalViewport(FMetalDevice& InDevice, void* WindowHandle, uint3
 
 FMetalViewport::~FMetalViewport()
 {
-	if (Block)
-	{
-		FScopeLock BlockLock(&Mutex);
-		if (GMetalSeparatePresentThread)
-		{
-			FPlatformRHIFramePacer::RemoveHandler(Block);
-		}
-		Block_release(Block);
-		Block = nullptr;
-	}
-	{
-		FScopeLock Lock(&ViewportsMutex);
-		Viewports.Remove(this);
-	}
+	FScopeLock Lock(&ViewportsMutex);
+	Viewports.Remove(this);
 	
 #if PLATFORM_IOS && !PLATFORM_TVOS && !PLATFORM_VISIONOS
 	FCoreDelegates::ApplicationReceivedScreenOrientationChangedNotificationDelegate.Remove(OrientationChangedHandle);
@@ -303,7 +290,6 @@ void FMetalViewport::Resize(uint32 InSizeX, uint32 InSizeY, bool bInIsFullscreen
         FScopeLock Lock(&Mutex);
 
 		TRefCountPtr<FMetalSurface> NewBackBuffer;
-		TRefCountPtr<FMetalSurface> DoubleBuffer;
 
 		FRHITextureCreateDesc CreateDesc =
 			FRHITextureCreateDesc::Create2D(TEXT("BackBuffer"), InSizeX, InSizeY, InFormat)
@@ -320,21 +306,8 @@ void FMetalViewport::Resize(uint32 InSizeX, uint32 InSizeY, bool bInIsFullscreen
 		NewBackBuffer = new FMetalSurface(Device, FMetalTextureCreateDesc(Device, CreateDesc));
 		NewBackBuffer->Viewport = this;
 
-        if (GMetalSupportsIntermediateBackBuffer && GMetalSeparatePresentThread)
-        {
-            DoubleBuffer = new FMetalSurface(Device, FMetalTextureCreateDesc(Device, CreateDesc));
-            DoubleBuffer->Viewport = this;
-        }
-
         BackBuffer[Index] = NewBackBuffer;
-        if (GMetalSeparatePresentThread)
-        {
-            BackBuffer[EMetalViewportAccessRHI] = DoubleBuffer;
-        }
-        else
-        {
-            BackBuffer[EMetalViewportAccessRHI] = BackBuffer[Index];
-        }
+		BackBuffer[EMetalViewportAccessRHI] = BackBuffer[Index];
     }
 }
 
@@ -463,18 +436,15 @@ MTL::Texture* FMetalViewport::GetCurrentTexture(EMetalViewportAccessFlag Accesso
 
 void FMetalViewport::ReleaseDrawable()
 {
-	if (!GMetalSeparatePresentThread)
+	if (Drawable != nullptr)
 	{
-		if (Drawable != nullptr)
-		{
-			Drawable->release();
-			Drawable = nullptr;
-		}
+		Drawable->release();
+		Drawable = nullptr;
+	}
 
-		if (!GMetalSupportsIntermediateBackBuffer && IsValidRef(BackBuffer[GetViewportIndex(EMetalViewportAccessRHI)]))
-		{
-			BackBuffer[GetViewportIndex(EMetalViewportAccessRHI)]->ReleaseDrawableTexture();
-		}
+	if (!GMetalSupportsIntermediateBackBuffer && IsValidRef(BackBuffer[GetViewportIndex(EMetalViewportAccessRHI)]))
+	{
+		BackBuffer[GetViewportIndex(EMetalViewportAccessRHI)]->ReleaseDrawableTexture();
 	}
 }
 
@@ -485,7 +455,7 @@ NSWindow* FMetalViewport::GetWindow() const
 }
 #endif
 
-void FMetalViewport::Present(FMetalCommandQueue& CommandQueue, bool bLockToVsync)
+void FMetalViewport::Present(FMetalCommandBuffer* CommandBuffer, bool bLockToVsync)
 {
 	FScopeLock Lock(&Mutex);
 	
@@ -523,164 +493,95 @@ void FMetalViewport::Present(FMetalCommandQueue& CommandQueue, bool bLockToVsync
 	LastCompleteFrame = GetBackBuffer(EMetalViewportAccessRHI);
 	FPlatformAtomics::InterlockedExchange(&FrameAvailable, 1);
 	
-	if (!Block)
-	{
-		Block = Block_copy(^(uint32 InDisplayID, double OutputSeconds, double OutputDuration)
-		{
 #if !PLATFORM_MAC
-			uint32 FramePace = FPlatformRHIFramePacer::GetFramePace();
-			float MinPresentDuration = FramePace ? (1.0f / (float)FramePace) : 0.0f;
-#endif
-			bool bIsInLiveResize = false;
-#if PLATFORM_MAC
-			// DisplayID is always equal to 0 when GMetalSeparatePresentThread == false so we don't need to access inLiveResize outside of main thread here
-			// @TODO: UE-227397
-			if (GMetalSeparatePresentThread)
-			{
-				bIsInLiveResize = View.inLiveResize;
-			}
+	uint32 FramePace = FPlatformRHIFramePacer::GetFramePace();
+	float MinPresentDuration = FramePace ? (1.0f / (float)FramePace) : 0.0f;
 #endif
-			if (FrameAvailable > 0 && (InDisplayID == 0 || (DisplayID == InDisplayID && !bIsInLiveResize)))
+	bool bIsInLiveResize = false;
+
+	if (FrameAvailable > 0)
+	{
+		FPlatformAtomics::InterlockedDecrement(&FrameAvailable);
+		CA::MetalDrawable* LocalDrawable = GetDrawable(EMetalViewportAccessDisplayLink);
+		LocalDrawable->retain();
+		MTL::Texture* DrawableTexture = GetDrawableTexture(EMetalViewportAccessDisplayLink);
+		
+		if (DrawableTexture)
+		{
+			check(CommandBuffer);
+			
+			if (GMetalSupportsIntermediateBackBuffer)
 			{
-				FPlatformAtomics::InterlockedDecrement(&FrameAvailable);
-				CA::MetalDrawable* LocalDrawable = GetDrawable(EMetalViewportAccessDisplayLink);
-                LocalDrawable->retain();
-				MTL::Texture* DrawableTexture = GetDrawableTexture(EMetalViewportAccessDisplayLink);
+				TRefCountPtr<FMetalSurface> Texture = LastCompleteFrame;
+				check(IsValidRef(Texture));
 				
-				{
-					FScopeLock BlockLock(&Mutex);
-#if PLATFORM_MAC
-					// DisplayID is always equal to 0 when GMetalSeparatePresentThread == false so we don't need to access inLiveResize outside of main thread here
-					// @TODO: UE-227397
-					if (GMetalSeparatePresentThread)
-					{
-						bIsInLiveResize = View.inLiveResize;
-					}
-#endif
-					
-					FMetalRHICommandContext& Context = *static_cast<FMetalRHICommandContext*>(FMetalDynamicRHI::Get().RHIGetDefaultContext());
-					FMetalCommandBuffer* CurrentCommandBuffer = Context.GetCurrentCommandBuffer();
-					
-					if (DrawableTexture && (InDisplayID == 0 || !bIsInLiveResize))
-					{
-						check(CurrentCommandBuffer);
-						
-						if (GMetalSupportsIntermediateBackBuffer)
-						{
-							TRefCountPtr<FMetalSurface> Texture = LastCompleteFrame;
-							check(IsValidRef(Texture));
-							
-							MTLTexturePtr Src = Texture->Texture;
-                            MTLTexturePtr Dst = NS::RetainPtr(DrawableTexture);
-							
-							NS::UInteger Width = FMath::Min(Src->width(), Dst->width());
-							NS::UInteger Height = FMath::Min(Src->height(), Dst->height());
-							
-							MTLBlitCommandEncoderPtr Encoder = NS::RetainPtr(CurrentCommandBuffer->GetMTLCmdBuffer()->blitCommandEncoder());
-							check(Encoder);
-
-							Encoder->copyFromTexture(Src.get(), 0, 0, MTL::Origin(0, 0, 0), MTL::Size(Width, Height, 1), Dst.get(), 0, 0, MTL::Origin(0, 0, 0));
-							Encoder->endEncoding();
-                            
-							Drawable->release();
-							Drawable = nullptr;
-						}
-						
-						dispatch_semaphore_t& FrameSemaphore = Device.GetFrameSemaphore();
-						dispatch_retain(FrameSemaphore);
-						
+				MTLTexturePtr Src = Texture->Texture;
+				MTLTexturePtr Dst = NS::RetainPtr(DrawableTexture);
+				
+				NS::UInteger Width = FMath::Min(Src->width(), Dst->width());
+				NS::UInteger Height = FMath::Min(Src->height(), Dst->height());
+				
+				MTLBlitCommandEncoderPtr Encoder = NS::RetainPtr(CommandBuffer->GetMTLCmdBuffer()->blitCommandEncoder());
+				check(Encoder);
+
+				Encoder->copyFromTexture(Src.get(), 0, 0, MTL::Origin(0, 0, 0), MTL::Size(Width, Height, 1), Dst.get(), 0, 0, MTL::Origin(0, 0, 0));
+				Encoder->endEncoding();
+				
+				Drawable->release();
+				Drawable = nullptr;
+			}
+			
+			dispatch_semaphore_t& FrameSemaphore = Device.GetFrameSemaphore();
+			dispatch_retain(FrameSemaphore);
+			
 #if PLATFORM_MAC
-						FMetalView* theView = View;
-						MTL::HandlerFunction CommandBufferHandler = [LocalDrawable, theView, FrameSemaphore](MTL::CommandBuffer* cmd_buf)
+			FMetalView* theView = View;
+			MTL::HandlerFunction CommandBufferHandler = [LocalDrawable, theView, FrameSemaphore](MTL::CommandBuffer* cmd_buf)
 #else
-                        MTL::HandlerFunction CommandBufferHandler = [LocalDrawable, FrameSemaphore](MTL::CommandBuffer* cmd_buf)
+			MTL::HandlerFunction CommandBufferHandler = [LocalDrawable, FrameSemaphore](MTL::CommandBuffer* cmd_buf)
 #endif
-						{
-							dispatch_semaphore_signal(FrameSemaphore);
-							dispatch_release(FrameSemaphore);
+			{
+				dispatch_semaphore_signal(FrameSemaphore);
+				dispatch_release(FrameSemaphore);
 
-							LocalDrawable->release();
+				LocalDrawable->release();
 #if PLATFORM_MAC
-							MainThreadCall(^{
-								FCocoaWindow* Window = (FCocoaWindow*)[theView window];
-								[Window startRendering];
-							}, false);
+				MainThreadCall(^{
+					FCocoaWindow* Window = (FCocoaWindow*)[theView window];
+					[Window startRendering];
+				}, false);
 #endif
-						};
-						
+			};
+			
 #if PLATFORM_MAC		// Mac needs the older way to present otherwise we end up with bad behaviour of the completion handlers that causes GPU timeouts.
-                        MTL::HandlerFunction ScheduledHandler = [LocalDrawable](MTL::CommandBuffer*)
-						{
-							LocalDrawable->present();
-						};
-								
-						CurrentCommandBuffer->GetMTLCmdBuffer()->addCompletedHandler(CommandBufferHandler);
-						CurrentCommandBuffer->GetMTLCmdBuffer()->addScheduledHandler(ScheduledHandler);
+			MTL::HandlerFunction ScheduledHandler = [LocalDrawable](MTL::CommandBuffer*)
+			{
+				LocalDrawable->present();
+			};
+					
+			CommandBuffer->GetMTLCmdBuffer()->addCompletedHandler(CommandBufferHandler);
+			CommandBuffer->GetMTLCmdBuffer()->addScheduledHandler(ScheduledHandler);
 
 #else // PLATFORM_MAC
-						CurrentCommandBuffer->GetMTLCmdBuffer()->addCompletedHandler(CommandBufferHandler);
-
-                        {
-                            // Queue this on the current command buffer to ensure that all work is committed prior to the present, present only knows about dependencies on committed work.
-                            if (MinPresentDuration && GEnablePresentPacing)
-                            {
-                                CurrentCommandBuffer->GetMTLCmdBuffer()->presentDrawableAfterMinimumDuration(LocalDrawable, 1.0f/(float)FramePace);
-                            }
-                            else
-                            {
-                                CurrentCommandBuffer->GetMTLCmdBuffer()->presentDrawable(LocalDrawable);
-                            }
-                        }
-#endif // PLATFORM_MAC
-						
-						TArray<FMetalPayload*> Payloads;
-						Context.Finalize(Payloads);
+			CommandBuffer->GetMTLCmdBuffer()->addCompletedHandler(CommandBufferHandler);
 
-						FMetalDynamicRHI::Get().SubmitPayloads(MoveTemp(Payloads));
-						
-						// Wait for the frame semaphore
-						dispatch_semaphore_wait(Device.GetFrameSemaphore(), DISPATCH_TIME_FOREVER);
-					}
+			{
+				// Queue this on the current command buffer to ensure that all work is committed prior to the present, present only knows about dependencies on committed work.
+				if (MinPresentDuration && GEnablePresentPacing)
+				{
+					CommandBuffer->GetMTLCmdBuffer()->presentDrawableAfterMinimumDuration(LocalDrawable, 1.0f/(float)FramePace);
+				}
+				else
+				{
+					CommandBuffer->GetMTLCmdBuffer()->presentDrawable(LocalDrawable);
 				}
 			}
-		});
-		
-		if (GMetalSeparatePresentThread)
-		{
-			FPlatformRHIFramePacer::AddHandler(Block);
+#endif // PLATFORM_MAC
+			
+			// Wait for the frame semaphore
+			dispatch_semaphore_wait(Device.GetFrameSemaphore(), DISPATCH_TIME_FOREVER);
 		}
 	}
-	
-	if (!GMetalSeparatePresentThread
-#if PLATFORM_MAC
-		|| View.inLiveResize
-#endif
-		)
-	{
-		Block(0, 0.0, 0.0);
-	}
-	
-	if (!(GRHISupportsRHIThread && IsRunningRHIInSeparateThread()))
-	{
-		Swap();
-	}
-}
-
-void FMetalViewport::Swap()
-{
-	if (GMetalSeparatePresentThread)
-	{
-		FScopeLock Lock(&Mutex);
-		
-		check(IsValidRef(BackBuffer[0]));
-		check(IsValidRef(BackBuffer[1]));
-		
-		TRefCountPtr<FMetalSurface> BB0 = BackBuffer[0];
-		TRefCountPtr<FMetalSurface> BB1 = BackBuffer[1];
-		
-		BackBuffer[0] = BB1;
-		BackBuffer[1] = BB0;
-	}
 }
 
 #if PLATFORM_VISIONOS
@@ -786,14 +687,6 @@ void FMetalViewport::PresentImmersive(const MetalRHIVisionOS::PresentImmersivePa
 			};
 			FinalCommandBuffer->GetMTLCmdBuffer()->addCompletedHandler(CommandBufferHandler);
 		}
-
-		TArray<FMetalPayload*> Payloads;
-		Context.Finalize(Payloads);
-
-
-		FMetalDynamicRHI::Get().SubmitPayloads(MoveTemp(Payloads));
-		
-		Context.ResetContext();
 		
 		if (VisionOSParams.SwiftDrawable)
 		{
@@ -896,7 +789,13 @@ void FMetalRHICommandContext::RHIEndDrawingViewport(FRHIViewport* ViewportRHI, c
 		
 		if (bNeedNativePresent)
 		{
-			Viewport->Present(CommandQueue, InPresentArgs.bLockToVsync);
+			FMetalCommandBuffer* CommandBuffer = CurrentEncoder.GetCommandBuffer();
+			if(!CommandBuffer)
+			{
+				StartCommandBuffer();
+				CommandBuffer = CurrentEncoder.GetCommandBuffer();
+			}
+			Viewport->Present(CommandBuffer, InPresentArgs.bLockToVsync);
 		}
 	}
 	
@@ -915,12 +814,4 @@ FTextureRHIRef FMetalDynamicRHI::RHIGetViewportBackBuffer(FRHIViewport* Viewport
 
 void FMetalDynamicRHI::RHIAdvanceFrameForGetViewportBackBuffer(FRHIViewport* ViewportRHI, bool bPresent)
 {
-	if (GMetalSeparatePresentThread && (GRHISupportsRHIThread && IsRunningRHIInSeparateThread()))
-	{
-		FScopeLock Lock(&ViewportsMutex);
-		for (FMetalViewport* Viewport : Viewports)
-		{
-			Viewport->Swap();
-		}
-	}
 }
diff --git a/Engine/Source/Runtime/Apple/MetalRHI/Public/MetalViewport.h b/Engine/Source/Runtime/Apple/MetalRHI/Public/MetalViewport.h
--- a/Engine/Source/Runtime/Apple/MetalRHI/Public/MetalViewport.h
+++ b/Engine/Source/Runtime/Apple/MetalRHI/Public/MetalViewport.h
@@ -76,8 +76,7 @@ public:
 	virtual FRHICustomPresent* GetCustomPresent() const override { return CustomPresent; }
 #endif
 	
-	void Present(FMetalCommandQueue& CommandQueue, bool bLockToVsync);
-	void Swap();
+	void Present(FMetalCommandBuffer* CommandBuffer, bool bLockToVsync);
 	
 #if PLATFORM_VISIONOS
 	void GetDrawableImmersiveTextures(EMetalViewportAccessFlag Accessor, cp_drawable_t SwiftDrawable, MTL::Texture*& OutColorTexture, MTL::Texture*& OutDepthTexture );
