Fix Metal RHI leaking contexts due to use of RHIGetDefaultContext function
- RHIGetDefaultContext has been reimplemented to return a newly allocated RHI context, since there is no longer a default context in Metal RHI, meaning the comparison in the if statement leaks a context.
- The only caller of this function should be RHI command list management code in RHICommandList.cpp. The global scope RHIGetDefaultContext will be deprecated in a later CL. The comparison here is no longer necessary.

diff --git a/Engine/Source/Runtime/Apple/MetalRHI/Private/MetalSubmission.cpp b/Engine/Source/Runtime/Apple/MetalRHI/Private/MetalSubmission.cpp
--- a/Engine/Source/Runtime/Apple/MetalRHI/Private/MetalSubmission.cpp
+++ b/Engine/Source/Runtime/Apple/MetalRHI/Private/MetalSubmission.cpp
@@ -171,7 +171,7 @@ void FMetalDynamicRHI::RHIFinalizeContext(FRHIFinalizeContextArgs&& Args, TRHIPi
 
 	FMetalFinalizedCommands Commands;
 	
-	if(UploadContext)
+	if (UploadContext)
 	{
 		UploadContext->Finalize(Commands);
 		MetalUploadContextPool.Push(UploadContext);
@@ -181,18 +181,12 @@ void FMetalDynamicRHI::RHIFinalizeContext(FRHIFinalizeContextArgs&& Args, TRHIPi
 	{
 		FMetalRHICommandContext* CmdContext = static_cast<FMetalRHICommandContext*>(Context);
 		
-		if(!CmdContext->IsInsideRenderPass())
+		if (!CmdContext->IsInsideRenderPass())
 		{
 			CmdContext->Finalize(Commands);
-
 			CmdContext->ResetContext();
-			if(GRHISupportsParallelRHIExecute)
-			{
-				if(CmdContext != RHIGetDefaultContext())
-				{
-					MetalCommandContextPool.Push(CmdContext);
-				}
-			}
+			
+			MetalCommandContextPool.Push(CmdContext);
 		}
 		
 		Output[Context->GetPipeline()] = new FMetalFinalizedCommands(MoveTemp(Commands));
