Metal - Temporary disable validation while a proper fix is being worked on

#tests local run

diff --git a/Engine/Source/Runtime/Apple/MetalRHI/Private/MetalVertexDeclaration.cpp b/Engine/Source/Runtime/Apple/MetalRHI/Private/MetalVertexDeclaration.cpp
--- a/Engine/Source/Runtime/Apple/MetalRHI/Private/MetalVertexDeclaration.cpp
+++ b/Engine/Source/Runtime/Apple/MetalRHI/Private/MetalVertexDeclaration.cpp
@@ -209,8 +209,8 @@ void FMetalVertexDeclaration::GenerateLayout(const FVertexDeclarationElementList
 		{
 			const FVertexElement& Element = InElements[ElementIndex];
 			
-			checkf(Element.Stride == 0 || Element.Offset + TranslateElementTypeToSize(Element.Type) <= Element.Stride,
-				   TEXT("Stream component is bigger than stride: Offset: %d, Size: %d [Type %d], Stride: %d"), Element.Offset, TranslateElementTypeToSize(Element.Type), (uint32)Element.Type, Element.Stride);
+			//checkf(Element.Stride == 0 || Element.Offset + TranslateElementTypeToSize(Element.Type) <= Element.Stride,
+			//	   TEXT("Stream component is bigger than stride: Offset: %d, Size: %d [Type %d], Stride: %d"), Element.Offset, TranslateElementTypeToSize(Element.Type), (uint32)Element.Type, Element.Stride);
 			
 			BaseHash = FCrc::MemCrc32(&Element.StreamIndex, sizeof(Element.StreamIndex), BaseHash);
 			BaseHash = FCrc::MemCrc32(&Element.Offset, sizeof(Element.Offset), BaseHash);
