Fix bug in ProcessSubmissionQueue.
- The function should return EQueueStatus::Processed if any payload was processed, but the flag was getting reset at the top of the do/while loop, meaning it never gets returned.

#Synced-CL 48124752

diff --git a/Engine/Source/Runtime/Apple/MetalRHI/Private/MetalSubmission.cpp b/Engine/Source/Runtime/Apple/MetalRHI/Private/MetalSubmission.cpp
--- a/Engine/Source/Runtime/Apple/MetalRHI/Private/MetalSubmission.cpp
+++ b/Engine/Source/Runtime/Apple/MetalRHI/Private/MetalSubmission.cpp
@@ -303,7 +303,7 @@ FMetalDynamicRHI::FProcessResult FMetalDynamicRHI::ProcessSubmissionQueue()
 	do
 	{
 		bProgress = false;
-		Result.Status = EQueueStatus::None;
+		EnumRemoveFlags(Result.Status, EQueueStatus::Pending);
 
 		// Push all pending payloads into the ordered per-device, per-pipe pending queues
 		{
diff --git a/Engine/Source/Runtime/D3D12RHI/Private/D3D12Submission.cpp b/Engine/Source/Runtime/D3D12RHI/Private/D3D12Submission.cpp
--- a/Engine/Source/Runtime/D3D12RHI/Private/D3D12Submission.cpp
+++ b/Engine/Source/Runtime/D3D12RHI/Private/D3D12Submission.cpp
@@ -352,7 +352,7 @@ FD3D12DynamicRHI::FProcessResult FD3D12DynamicRHI::ProcessSubmissionQueue()
 	do
 	{
 		bProgress = false;
-		Result.Status = EQueueStatus::None;
+		EnumRemoveFlags(Result.Status, EQueueStatus::Pending);
 
 		// Push all pending payloads into the ordered per-device, per-pipe pending queues
 		{
