Updating Vulkan and Metal binding to support bindless Uniform Buffers.

#Synced-CL 48400404

diff --git a/Engine/Source/Runtime/Apple/MetalRHI/Private/MetalCommands.cpp b/Engine/Source/Runtime/Apple/MetalRHI/Private/MetalCommands.cpp
--- a/Engine/Source/Runtime/Apple/MetalRHI/Private/MetalCommands.cpp
+++ b/Engine/Source/Runtime/Apple/MetalRHI/Private/MetalCommands.cpp
@@ -26,6 +26,7 @@
 #include "MetalBindlessDescriptors.h"
 #include "MetalResourceCollection.h"
 #include "DataDrivenShaderPlatformInfo.h"
+#include "RHIUniformBufferUtilities.h"
 
 static const bool GUsesInvertedZ = true;
 
@@ -476,73 +477,114 @@ void FMetalRHICommandContext::RHISetStaticUniformBuffer(FUniformBufferStaticSlot
 
 struct FMetalShaderBinder
 {
-    FMetalStateCache& StateCache;
-    const EMetalShaderStages Stage;
-    FMetalShaderParameterCache& ShaderParameters;
+	FMetalStateCache& StateCache;
+	const EMetalShaderStages Stage;
+	FMetalShaderParameterCache& ShaderParameters;
 #if PLATFORM_SUPPORTS_BINDLESS_RENDERING
-    const bool bBindlessResources;
-    const bool bBindlessSamplers;
+	const bool bBindlessResources;
+	const bool bBindlessSamplers;
 #endif
-    
-    FMetalShaderBinder(FMetalStateCache& InStateCache, EShaderFrequency ShaderFrequency)
-    : StateCache(InStateCache)
-    , Stage(GetMetalShaderFrequency(ShaderFrequency))
-    , ShaderParameters(StateCache.GetShaderParameters(Stage))
+
+	FMetalShaderBinder(FMetalStateCache& InStateCache, EShaderFrequency ShaderFrequency)
+		: StateCache(InStateCache)
+		, Stage(GetMetalShaderFrequency(ShaderFrequency))
+		, ShaderParameters(StateCache.GetShaderParameters(Stage))
 #if PLATFORM_SUPPORTS_BINDLESS_RENDERING
-    , bBindlessResources(IsMetalBindlessEnabled())
-    , bBindlessSamplers(IsMetalBindlessEnabled())
+		, bBindlessResources(IsMetalBindlessEnabled())
+		, bBindlessSamplers(IsMetalBindlessEnabled())
 #endif
-    {
-    }
-    
+	{
+	}
+
+	void SetParameter(uint32 BaseIndex, uint32 NumBytes, const void* NewValue)
+	{
+		ShaderParameters.Set(0, BaseIndex, NumBytes, NewValue);
+	}
+
 #if PLATFORM_SUPPORTS_BINDLESS_RENDERING
-    void SetBindlessHandle(const FRHIDescriptorHandle& Handle, uint32 Offset)
-    {
-        if (Handle.IsValid())
-        {
-            const uint32 BindlessIndex = Handle.GetIndex();
-            StateCache.GetShaderParameters(Stage).Set(0, Offset, 4, &BindlessIndex);
-        }
-    }
+	void SetBindlessHandle(const FRHIDescriptorHandle& Handle, uint32 Offset)
+	{
+		if (Handle.IsValid())
+		{
+			const uint32 BindlessIndex = Handle.GetIndex();
+			StateCache.GetShaderParameters(Stage).Set(0, Offset, 4, &BindlessIndex);
+		}
+	}
 #endif
 
-    void SetUAV(FRHIUnorderedAccessView* InUnorderedAccessView, uint32 Index, bool bClearResources = false)
-    {
-        FMetalUnorderedAccessView* UAV = ResourceCast(InUnorderedAccessView);
-        StateCache.SetShaderUnorderedAccessView(Stage, Index, UAV);
-    }
+	void SetUAV(FRHIUnorderedAccessView* InUnorderedAccessView, uint32 Index, bool bClearResources = false)
+	{
+		FMetalUnorderedAccessView* UAV = ResourceCast(InUnorderedAccessView);
+#if PLATFORM_SUPPORTS_BINDLESS_RENDERING
+		if (bBindlessResources)
+		{
+			StateCache.IRMakeUAVResident(Stage, UAV);
+		}
+		else
+#endif
+		{
+			StateCache.SetShaderUnorderedAccessView(Stage, Index, UAV);
+		}
+	}
 
-    void SetSRV(FRHIShaderResourceView* InShaderResourceView, uint32 Index)
-    {
-        FMetalShaderResourceView* SRV = ResourceCast(InShaderResourceView);
-        StateCache.SetShaderResourceView(Stage, Index, SRV);
-    }
+	void SetSRV(FRHIShaderResourceView* InShaderResourceView, uint32 Index)
+	{
+		FMetalShaderResourceView* SRV = ResourceCast(InShaderResourceView);
+#if PLATFORM_SUPPORTS_BINDLESS_RENDERING
+		if (bBindlessResources)
+		{
+			StateCache.IRMakeSRVResident(Stage, SRV);
+		}
+		else
+#endif
+		{
+			StateCache.SetShaderResourceView(Stage, Index, SRV);
+		}
+	}
 
-    void SetTexture(FRHITexture* InTexture, uint32 Index)
-    {
-        if (FMetalSurface* Surface = GetMetalSurfaceFromRHITexture(InTexture))
-        {
-            if (Surface->Texture || !EnumHasAnyFlags(Surface->GetDesc().Flags, ETextureCreateFlags::Presentable))
-            {
-                StateCache.SetShaderTexture(Stage, Surface->Texture.get(), Index, (MTL::ResourceUsage)(MTL::ResourceUsageRead|MTL::ResourceUsageSample));
-            }
-            else
-            {
-                MTLTexturePtr Tex = Surface->GetCurrentTexture();
-                StateCache.SetShaderTexture(Stage, Tex.get(), Index, (MTL::ResourceUsage)(MTL::ResourceUsageRead|MTL::ResourceUsageSample));
-            }
-        }
-        else
-        {
-            StateCache.SetShaderTexture(Stage, nullptr, Index, MTL::ResourceUsage(0));
-        }
-    }
+	void SetTexture(FRHITexture* InTexture, uint32 Index)
+	{
+		FMetalSurface* Surface = GetMetalSurfaceFromRHITexture(InTexture);
+#if PLATFORM_SUPPORTS_BINDLESS_RENDERING
+		if (bBindlessResources)
+		{
+			StateCache.IRMakeTextureResident(Stage, Surface->Texture.get());
+		}
+		else
+#endif
+		{
+			if (Surface)
+			{
+				if (Surface->Texture || !EnumHasAnyFlags(Surface->GetDesc().Flags, ETextureCreateFlags::Presentable))
+				{
+					StateCache.SetShaderTexture(Stage, Surface->Texture.get(), Index, (MTL::ResourceUsage)(MTL::ResourceUsageRead | MTL::ResourceUsageSample));
+				}
+				else
+				{
+					MTLTexturePtr Tex = Surface->GetCurrentTexture();
+					StateCache.SetShaderTexture(Stage, Tex.get(), Index, (MTL::ResourceUsage)(MTL::ResourceUsageRead | MTL::ResourceUsageSample));
+				}
+			}
+			else
+			{
+				StateCache.SetShaderTexture(Stage, nullptr, Index, MTL::ResourceUsage(0));
+			}
+		}
+	}
 
-    void SetSampler(FRHISamplerState* InSampler, uint32 Index)
-    {
-        FMetalSamplerState* Sampler = ResourceCast(InSampler);
-        StateCache.SetShaderSamplerState(Stage, Sampler, Index);
-    }
+	void SetSampler(FRHISamplerState* InSampler, uint32 Index)
+	{
+#if PLATFORM_SUPPORTS_BINDLESS_RENDERING
+		if (bBindlessSamplers)
+		{
+		}
+		else
+#endif
+		{
+			FMetalSamplerState* Sampler = ResourceCast(InSampler);
+			StateCache.SetShaderSamplerState(Stage, Sampler, Index);
+		}
+	}
 
 #if PLATFORM_SUPPORTS_BINDLESS_RENDERING
 	void SetResourceCollection(FRHIResourceCollection* ResourceCollection, uint32 Index)
@@ -553,112 +595,133 @@ struct FMetalShaderBinder
 #endif
 };
 
+FRHIDescriptorHandle GetDescriptorHandleFromTexture(FRHITexture* InTexture)
+{
+	if (FMetalSurface* Surface = GetMetalSurfaceFromRHITexture(InTexture))
+	{
+		FRHIDescriptorHandle Handle = Surface->GetDefaultBindlessHandle();
+		if (Handle.IsValid())
+		{
+			return Handle;
+		}
+
+		if (Surface->Viewport)
+		{
+			// We need to update the drawable texture on the surface and create a bindless handle, we can ignore the return value.
+			Surface->GetDrawableTexture();
+			return Surface->GetDefaultBindlessHandle();
+		}
+	}
+
+	return FRHIDescriptorHandle();
+}
+
 static void SetShaderParameters(
-    FMetalStateCache& StateCache
-    , FRHIShader* Shader
-    , EShaderFrequency ShaderFrequency
-    , TArrayView<const uint8> InParametersData
-    , TArrayView<const FRHIShaderParameter> InParameters
-    , TArrayView<const FRHIShaderParameterResource> InResourceParameters
-    , TArrayView<const FRHIShaderParameterResource> InBindlessParameters)
+	FMetalStateCache& StateCache
+	, FRHIShader* Shader
+	, EShaderFrequency ShaderFrequency
+	, TArrayView<const uint8> InParametersData
+	, TArrayView<const FRHIShaderParameter> InParameters
+	, TArrayView<const FRHIShaderParameterResource> InResourceParameters
+	, TArrayView<const FRHIShaderParameterResource> InBindlessParameters)
 {
-    MTL_SCOPED_AUTORELEASE_POOL;
+	MTL_SCOPED_AUTORELEASE_POOL;
 
-    FMetalShaderBinder Binder(StateCache, ShaderFrequency);
-    
-    for (const FRHIShaderParameter& Parameter : InParameters)
-    {
-        Binder.ShaderParameters.Set(Parameter.BufferIndex, Parameter.BaseIndex, Parameter.ByteSize, &InParametersData[Parameter.ByteOffset]);
-    }
+	FMetalShaderBinder Binder(StateCache, ShaderFrequency);
+
+	for (const FRHIShaderParameter& Parameter : InParameters)
+	{
+		check(Parameter.BufferIndex == 0);
+		Binder.SetParameter(Parameter.BaseIndex, Parameter.ByteSize, &InParametersData[Parameter.ByteOffset]);
+	}
 
 #if PLATFORM_SUPPORTS_BINDLESS_RENDERING
-    for (const FRHIShaderParameterResource& Parameter : InBindlessParameters)
-    {
-        FRHIDescriptorHandle Handle = UE::RHICore::GetBindlessParameterHandle(Parameter);
-        if (!Handle.IsValid())
-        {
-			if (Parameter.Type == FRHIShaderParameterResource::EType::Texture)
+	for (const FRHIShaderParameterResource& Parameter : InBindlessParameters)
+	{
+		if (FRHIResource* Resource = Parameter.Resource)
+		{
+			if (Parameter.Type == FRHIShaderParameterResource::EType::UniformBuffer)
 			{
-				if (FRHIResource* Resource = Parameter.Resource)
+				UE::RHI::Private::SetAllUniformBufferResourcesForBindless(Binder, static_cast<FRHIUniformBuffer*>(Resource));
+			}
+			else
+			{
+				FRHIDescriptorHandle Handle;
+
+				switch (Parameter.Type)
 				{
-					FMetalSurface* Surface = static_cast<FMetalSurface*>(Resource);
-					if(Surface->Viewport)
-					{
-						// We need to update the drawable texture on the surface and create a bindless handle, we can ignore the return value.
-						Surface->GetDrawableTexture();
-						Handle = UE::RHICore::GetBindlessParameterHandle(Parameter);
-					}
+				case FRHIShaderParameterResource::EType::Texture:
+					Handle = GetDescriptorHandleFromTexture(static_cast<FRHITexture*>(Resource));
+					Binder.SetTexture(static_cast<FRHITexture*>(Resource), Parameter.Index);
+					break;
+				case FRHIShaderParameterResource::EType::ResourceView:
+					Handle = static_cast<FRHIShaderResourceView*>(Resource)->GetBindlessHandle();
+					Binder.SetSRV(static_cast<FRHIShaderResourceView*>(Resource), Parameter.Index);
+					break;
+				case FRHIShaderParameterResource::EType::UnorderedAccessView:
+					Handle = static_cast<FRHIUnorderedAccessView*>(Resource)->GetBindlessHandle();
+					Binder.SetUAV(static_cast<FRHIUnorderedAccessView*>(Resource), Parameter.Index, true);
+					break;
+				case FRHIShaderParameterResource::EType::Sampler:
+					Handle = static_cast<FRHISamplerState*>(Resource)->GetBindlessHandle();
+					Binder.SetSampler(static_cast<FRHISamplerState*>(Resource), Parameter.Index);
+					break;
+				case FRHIShaderParameterResource::EType::ResourceCollection:
+					Handle = static_cast<FRHIResourceCollection*>(Resource)->GetBindlessHandle();
+					Binder.SetResourceCollection(static_cast<FRHIResourceCollection*>(Resource), Parameter.Index);
+					break;
 				}
+
+				checkf(Handle.IsValid(), TEXT("Metal resource did not provide a valid descriptor handle. Please validate that all D3D12 types can provide this or that the resource is still valid."));
+				Binder.SetBindlessHandle(Handle, Parameter.Index);
 			}
-        }
-				
-		Binder.SetBindlessHandle(Handle, Parameter.Index);
-			
+		}
+	}
+#endif 
+
+	for (const FRHIShaderParameterResource& Parameter : InResourceParameters)
+	{
 		if (Parameter.Type == FRHIShaderParameterResource::EType::UnorderedAccessView)
 		{
 			if (ShaderFrequency == SF_Pixel || ShaderFrequency == SF_Compute)
 			{
-				Binder.StateCache.IRMakeUAVResident(GetShaderStage(ShaderFrequency), static_cast<FMetalUnorderedAccessView*>(Parameter.Resource));
+				Binder.SetUAV(static_cast<FRHIUnorderedAccessView*>(Parameter.Resource), Parameter.Index, true);
 			}
 			else
 			{
 				checkf(false, TEXT("TShaderRHI Can't have compute shader to be set. UAVs are not supported on vertex, tessellation and geometry shaders."));
 			}
 		}
-		else if(Parameter.Type == FRHIShaderParameterResource::EType::ResourceView)
-		{
-			Binder.StateCache.IRMakeSRVResident(GetShaderStage(ShaderFrequency), static_cast<FMetalShaderResourceView*>(Parameter.Resource));
-		}
-		else if(Parameter.Type == FRHIShaderParameterResource::EType::Texture)
-		{
-			Binder.StateCache.IRMakeTextureResident(GetShaderStage(ShaderFrequency), static_cast<FMetalSurface*>(Parameter.Resource)->Texture.get());
-		}
-    }
-#endif 
-	
-    for (const FRHIShaderParameterResource& Parameter : InResourceParameters)
-    {
-        if (Parameter.Type == FRHIShaderParameterResource::EType::UnorderedAccessView)
-        {
-            if (ShaderFrequency == SF_Pixel || ShaderFrequency == SF_Compute)
-            {
-                Binder.SetUAV(static_cast<FRHIUnorderedAccessView*>(Parameter.Resource), Parameter.Index, true);
-            }
-            else
-            {
-                checkf(false, TEXT("TShaderRHI Can't have compute shader to be set. UAVs are not supported on vertex, tessellation and geometry shaders."));
-            }
-        }
-    }
+	}
 
-    for (const FRHIShaderParameterResource& Parameter : InResourceParameters)
-    {
-        switch (Parameter.Type)
-        {
-        case FRHIShaderParameterResource::EType::Texture:
-            Binder.SetTexture(static_cast<FRHITexture*>(Parameter.Resource), Parameter.Index);
-            break;
-        case FRHIShaderParameterResource::EType::ResourceView:
-            Binder.SetSRV(static_cast<FRHIShaderResourceView*>(Parameter.Resource), Parameter.Index);
-            break;
-        case FRHIShaderParameterResource::EType::UnorderedAccessView:
-            break;
-        case FRHIShaderParameterResource::EType::Sampler:
-            Binder.SetSampler(static_cast<FRHISamplerState*>(Parameter.Resource), Parameter.Index);
-            break;
-        case FRHIShaderParameterResource::EType::UniformBuffer:
-            BindUniformBuffer(Binder.StateCache, Shader, Binder.Stage, Parameter.Index, static_cast<FRHIUniformBuffer*>(Parameter.Resource));
-            break;
+	for (const FRHIShaderParameterResource& Parameter : InResourceParameters)
+	{
+		switch (Parameter.Type)
+		{
+		case FRHIShaderParameterResource::EType::Texture:
+			Binder.SetTexture(static_cast<FRHITexture*>(Parameter.Resource), Parameter.Index);
+			break;
+		case FRHIShaderParameterResource::EType::ResourceView:
+			Binder.SetSRV(static_cast<FRHIShaderResourceView*>(Parameter.Resource), Parameter.Index);
+			break;
+		case FRHIShaderParameterResource::EType::UnorderedAccessView:
+			break;
+		case FRHIShaderParameterResource::EType::Sampler:
+			Binder.SetSampler(static_cast<FRHISamplerState*>(Parameter.Resource), Parameter.Index);
+			break;
+		case FRHIShaderParameterResource::EType::UniformBuffer:
+			BindUniformBuffer(Binder.StateCache, Shader, Binder.Stage, Parameter.Index, static_cast<FRHIUniformBuffer*>(Parameter.Resource));
+			break;
 #if PLATFORM_SUPPORTS_BINDLESS_RENDERING
-        case FRHIShaderParameterResource::EType::ResourceCollection:
-            Binder.SetResourceCollection(static_cast<FRHIResourceCollection*>(Parameter.Resource), Parameter.Index);
-            break;
+		case FRHIShaderParameterResource::EType::ResourceCollection:
+			Binder.SetResourceCollection(static_cast<FRHIResourceCollection*>(Parameter.Resource), Parameter.Index);
+			break;
 #endif
-        default:
-            checkf(false, TEXT("Unhandled resource type?"));
-            break;
-        }
-    }
+		default:
+			checkf(false, TEXT("Unhandled resource type?"));
+			break;
+		}
+	}
 }
 
 void FMetalRHICommandContext::RHISetShaderParameters(FRHIGraphicsShader* Shader, TConstArrayView<uint8> InParametersData, TConstArrayView<FRHIShaderParameter> InParameters, TConstArrayView<FRHIShaderParameterResource> InResourceParameters, TConstArrayView<FRHIShaderParameterResource> InBindlessParameters)
diff --git a/Engine/Source/Runtime/VulkanRHI/Private/VulkanCommands.cpp b/Engine/Source/Runtime/VulkanRHI/Private/VulkanCommands.cpp
--- a/Engine/Source/Runtime/VulkanRHI/Private/VulkanCommands.cpp
+++ b/Engine/Source/Runtime/VulkanRHI/Private/VulkanCommands.cpp
@@ -344,9 +344,7 @@ static void SetShaderParametersOnBinder(
 		{
 			if (Parameter.Type == FRHIShaderParameterResource::EType::UniformBuffer)
 			{
-				// ue-todo: christopher.waters - Enable SetAllUniformBufferResourcesForBindless once it's checked in.
-				//UE::RHI::Private::SetAllUniformBufferResourcesForBindless(Binder, static_cast<FRHIUniformBuffer*>(Resource));
-				checkNoEntry();
+				UE::RHI::Private::SetAllUniformBufferResourcesForBindless(Binder, static_cast<FRHIUniformBuffer*>(Resource));
 			}
 			else
 			{
