Fixing parameter binding to not require BufferIndex of 0

#Synced-CL Unknown

diff --git a/Engine/Source/Runtime/Apple/MetalRHI/Private/MetalCommands.cpp b/Engine/Source/Runtime/Apple/MetalRHI/Private/MetalCommands.cpp
--- a/Engine/Source/Runtime/Apple/MetalRHI/Private/MetalCommands.cpp
+++ b/Engine/Source/Runtime/Apple/MetalRHI/Private/MetalCommands.cpp
@@ -496,9 +496,9 @@ struct FMetalShaderBinder
 	{
 	}
 
-	void SetParameter(uint32 BaseIndex, uint32 NumBytes, const void* NewValue)
+	void SetParameter(uint32 BufferIndex, uint32 BaseIndex, uint32 NumBytes, const void* NewValue)
 	{
-		ShaderParameters.Set(0, BaseIndex, NumBytes, NewValue);
+		ShaderParameters.Set(BufferIndex, BaseIndex, NumBytes, NewValue);
 	}
 
 #if PLATFORM_SUPPORTS_BINDLESS_RENDERING
@@ -507,7 +507,7 @@ struct FMetalShaderBinder
 		if (Handle.IsValid())
 		{
 			const uint32 BindlessIndex = Handle.GetIndex();
-			StateCache.GetShaderParameters(Stage).Set(0, Offset, 4, &BindlessIndex);
+			ShaderParameters.Set(0, Offset, 4, &BindlessIndex);
 		}
 	}
 #endif
@@ -631,8 +631,7 @@ static void SetShaderParameters(
 
 	for (const FRHIShaderParameter& Parameter : InParameters)
 	{
-		check(Parameter.BufferIndex == 0);
-		Binder.SetParameter(Parameter.BaseIndex, Parameter.ByteSize, &InParametersData[Parameter.ByteOffset]);
+		Binder.SetParameter(Parameter.BufferIndex, Parameter.BaseIndex, Parameter.ByteSize, &InParametersData[Parameter.ByteOffset]);
 	}
 
 #if PLATFORM_SUPPORTS_BINDLESS_RENDERING
