Metal RHI

- FMetalQueryBuffer no longer inherits from FRHIResource

diff --git a/Engine/Source/Runtime/Apple/MetalRHI/Private/MetalRHIRenderQuery.cpp b/Engine/Source/Runtime/Apple/MetalRHI/Private/MetalRHIRenderQuery.cpp
--- a/Engine/Source/Runtime/Apple/MetalRHI/Private/MetalRHIRenderQuery.cpp
+++ b/Engine/Source/Runtime/Apple/MetalRHI/Private/MetalRHIRenderQuery.cpp
@@ -35,11 +35,7 @@ FAutoConsoleVariableRef CVarQueryUseControlPrecisionBit(
 );
 
 FMetalQueryBuffer::FMetalQueryBuffer(FMetalQueryBufferPool* InPool, FMetalBufferPtr InBuffer)
-PRAGMA_DISABLE_DEPRECATION_WARNINGS
-	// @todo Metal RHI - Remove the use of RRT_TimestampCalibrationQuery. FMetalQueryBuffer probably shouldn't be an FRHIResource type.
-	: FRHIResource(RRT_TimestampCalibrationQuery)
-PRAGMA_ENABLE_DEPRECATION_WARNINGS
-	, Pool(InPool)
+	: Pool(InPool)
 	, Buffer(InBuffer)
 	, WriteOffset(0)
 {
@@ -48,15 +44,9 @@ PRAGMA_ENABLE_DEPRECATION_WARNINGS
 
 FMetalQueryBuffer::~FMetalQueryBuffer()
 {
-	if (GIsMetalInitialized)
+	if (GIsMetalInitialized && Buffer && Pool)
 	{
-		if (Buffer)
-		{
-			if (Pool)
-			{
-				Pool->ReleaseQueryBuffer(Buffer);
-			}
-		}
+		Pool->ReleaseQueryBuffer(Buffer);
 	}
 }
 
diff --git a/Engine/Source/Runtime/Apple/MetalRHI/Private/MetalRHIRenderQuery.h b/Engine/Source/Runtime/Apple/MetalRHI/Private/MetalRHIRenderQuery.h
--- a/Engine/Source/Runtime/Apple/MetalRHI/Private/MetalRHIRenderQuery.h
+++ b/Engine/Source/Runtime/Apple/MetalRHI/Private/MetalRHIRenderQuery.h
@@ -27,7 +27,7 @@ typedef TSharedPtr<FMetalBuffer, ESPMode::ThreadSafe> FMetalBufferPtr;
 #pragma mark - Metal RHI Private Query Buffer Resource Class -
 
 
-class FMetalQueryBuffer : public FRHIResource
+class FMetalQueryBuffer : public FRefCountedObject
 {
 public:
 	FMetalQueryBuffer(FMetalQueryBufferPool* Pool, FMetalBufferPtr InBuffer);
@@ -35,7 +35,7 @@ public:
 
 	uint64 GetResult(uint32 Offset);
 
-	FMetalQueryBufferPool* Pool;
+	FMetalQueryBufferPool* Pool = nullptr;
 	FMetalBufferPtr Buffer;
 	uint32 WriteOffset;
 };
