Metal RHI

Fix memory size calculation not taking into account NumFaces

diff --git a/Engine/Source/Runtime/Apple/MetalRHI/Private/MetalTexture.cpp b/Engine/Source/Runtime/Apple/MetalRHI/Private/MetalTexture.cpp
--- a/Engine/Source/Runtime/Apple/MetalRHI/Private/MetalTexture.cpp
+++ b/Engine/Source/Runtime/Apple/MetalRHI/Private/MetalTexture.cpp
@@ -1251,7 +1251,7 @@ uint32 FMetalSurface::GetMemorySize()
 	uint32 TotalSize = 0;
 	for (uint32 MipIndex = 0; MipIndex < Texture->mipmapLevelCount(); MipIndex++)
 	{
-		TotalSize += GetMipSize(MipIndex, NULL, false);
+		TotalSize += GetMipSize(MipIndex, NULL, false) * GetNumFaces();
 	}
 	
 	return TotalSize;
