Metal - Add Vertex Descriptor information when failing to compile a PSO

#tests Local run & preflight
#rnx
#Synced-CL 49035639

diff --git a/Engine/Source/Runtime/Apple/MetalRHI/Private/MetalPipeline.cpp b/Engine/Source/Runtime/Apple/MetalRHI/Private/MetalPipeline.cpp
--- a/Engine/Source/Runtime/Apple/MetalRHI/Private/MetalPipeline.cpp
+++ b/Engine/Source/Runtime/Apple/MetalRHI/Private/MetalPipeline.cpp
@@ -1101,6 +1101,8 @@ static FMetalShaderPipelinePtr CreateMTLRenderPipeline(FMetalDevice& Device, boo
 		UE_CLOG((Pipeline->RenderPipelineState.get() == nullptr), LogMetal, Error, TEXT("Vertex shader: %s"), *NSStringToFString(VertexShader->GetSourceCode()));
 		UE_CLOG((Pipeline->RenderPipelineState.get() == nullptr), LogMetal, Error, TEXT("Pixel shader: %s"), PixelShader ? *NSStringToFString(PixelShader->GetSourceCode()) : TEXT("NULL"));
 		UE_CLOG((Pipeline->RenderPipelineState.get() == nullptr), LogMetal, Error, TEXT("Descriptor: %s"), *NSStringToFString(RenderPipelineDesc->description()));
+		UE_CLOG((Pipeline->RenderPipelineState.get() == nullptr), LogMetal, Error, TEXT("Vertex Descriptor: %s"), *NSStringToFString(VertexDecl->Layout.VertexDesc->description()));
+		UE_CLOG((Pipeline->RenderPipelineState.get() == nullptr), LogMetal, Error, TEXT("Masked Vertex Descriptor: %s"), *NSStringToFString(MaskedVertexDesc->description()));
 		UE_CLOG((Pipeline->RenderPipelineState.get() == nullptr), LogMetal, Error, TEXT("Failed to generate a render pipeline state object:\n\n %s\n\n"), *NSStringToFString(Error->localizedDescription()));
 		
 		// We need to pass a failure up the chain, so we'll clean up here.
