[IOS/Mac] Initial pass to remove iOS/macOS sdk headers from Engine platform header files where possible.

Currently most platform headers are including huge apple sdk headers, which means they are included in every module. This change makes it so these sdk headers are only included where they are needed.

We are still using monolithic headers (like iossystemincludes.h) for now, but at least they are included where they are really needed

#tests Tested that iOS and other platforms still compile as before. Also did a replay run to confirm the de-inlining of some methods do not have a negative performance impact

#Synced-CL 49747958

diff --git a/Engine/Plugins/Developer/GitSourceControl/Source/GitSourceControl/Private/GitSourceControlUtils.cpp b/Engine/Plugins/Developer/GitSourceControl/Source/GitSourceControl/Private/GitSourceControlUtils.cpp
--- a/Engine/Plugins/Developer/GitSourceControl/Source/GitSourceControl/Private/GitSourceControlUtils.cpp
+++ b/Engine/Plugins/Developer/GitSourceControl/Source/GitSourceControl/Private/GitSourceControlUtils.cpp
@@ -11,6 +11,10 @@
 #include "ISourceControlModule.h"
 #include "GitSourceControlModule.h"
 
+#if PLATFORM_MAC
+#include "Apple/ScopeAutoreleasePool.h"
+#endif
+
 #if PLATFORM_LINUX
 #endif
 
diff --git a/Engine/Plugins/Developer/XcodeGPUDebuggerPlugin/Source/XcodeGPUDebuggerPlugin/Private/XcodeGPUDebuggerPluginModule.cpp b/Engine/Plugins/Developer/XcodeGPUDebuggerPlugin/Source/XcodeGPUDebuggerPlugin/Private/XcodeGPUDebuggerPluginModule.cpp
--- a/Engine/Plugins/Developer/XcodeGPUDebuggerPlugin/Source/XcodeGPUDebuggerPlugin/Private/XcodeGPUDebuggerPluginModule.cpp
+++ b/Engine/Plugins/Developer/XcodeGPUDebuggerPlugin/Source/XcodeGPUDebuggerPlugin/Private/XcodeGPUDebuggerPluginModule.cpp
@@ -24,6 +24,7 @@
 extern UNREALED_API UEditorEngine* GEditor;
 #endif // WITH_EDITOR
 
+#include "Apple/AppleStringUtils.h"
 #include "MetalInclude.h"
 
 DEFINE_LOG_CATEGORY(XcodeGPUDebuggerPlugin);
diff --git a/Engine/Plugins/Experimental/NNERuntimeCoreML/Source/NNERuntimeCoreML/Private/NNERuntimeCoreMLModel.cpp b/Engine/Plugins/Experimental/NNERuntimeCoreML/Source/NNERuntimeCoreML/Private/NNERuntimeCoreMLModel.cpp
--- a/Engine/Plugins/Experimental/NNERuntimeCoreML/Source/NNERuntimeCoreML/Private/NNERuntimeCoreMLModel.cpp
+++ b/Engine/Plugins/Experimental/NNERuntimeCoreML/Source/NNERuntimeCoreML/Private/NNERuntimeCoreMLModel.cpp
@@ -16,6 +16,9 @@
 #import <Foundation/Foundation.h>
 #import <CoreML/CoreML.h>
 
+#include "Apple/ScopeAutoreleasePool.h"
+#include "Apple/AppleStringUtils.h"
+
 @interface CoreMLInputWrapper : NSObject<MLFeatureProvider>
 @property (readwrite, nonatomic, retain) NSMutableDictionary<NSString*, MLFeatureValue*>* featureByNames;
 @end
diff --git a/Engine/Plugins/Experimental/NNERuntimeCoreML/Source/NNERuntimeCoreML/Private/NNERuntimeCoreMLModule.cpp b/Engine/Plugins/Experimental/NNERuntimeCoreML/Source/NNERuntimeCoreML/Private/NNERuntimeCoreMLModule.cpp
--- a/Engine/Plugins/Experimental/NNERuntimeCoreML/Source/NNERuntimeCoreML/Private/NNERuntimeCoreMLModule.cpp
+++ b/Engine/Plugins/Experimental/NNERuntimeCoreML/Source/NNERuntimeCoreML/Private/NNERuntimeCoreMLModule.cpp
@@ -8,6 +8,9 @@
 #if defined(__APPLE__)
 #import <Foundation/Foundation.h>
 #import <CoreML/CoreML.h>
+
+#include "Apple/ScopeAutoreleasePool.h"
+
 #endif // defined(__APPLE__)
 
 DEFINE_LOG_CATEGORY(LogNNERuntimeCoreML);
diff --git a/Engine/Plugins/Experimental/NNERuntimeCoreML/Source/NNERuntimeCoreML/Private/NNERuntimeCoreMLNPUHelper.cpp b/Engine/Plugins/Experimental/NNERuntimeCoreML/Source/NNERuntimeCoreML/Private/NNERuntimeCoreMLNPUHelper.cpp
--- a/Engine/Plugins/Experimental/NNERuntimeCoreML/Source/NNERuntimeCoreML/Private/NNERuntimeCoreMLNPUHelper.cpp
+++ b/Engine/Plugins/Experimental/NNERuntimeCoreML/Source/NNERuntimeCoreML/Private/NNERuntimeCoreMLNPUHelper.cpp
@@ -3,7 +3,9 @@
 #include "NNERuntimeCoreMLNPUHelper.h"
 
 #if defined(__APPLE__)
+#include "Apple/ScopeAutoreleasePool.h"
 #include "CoreMinimal.h"
+
 #import <Foundation/Foundation.h>
 #import <CoreML/CoreML.h>
 #endif // defined(__APPLE__)
diff --git a/Engine/Plugins/Experimental/TextToSpeech/Source/TextToSpeech/Private/Apple/AppleTextToSpeech.cpp b/Engine/Plugins/Experimental/TextToSpeech/Source/TextToSpeech/Private/Apple/AppleTextToSpeech.cpp
--- a/Engine/Plugins/Experimental/TextToSpeech/Source/TextToSpeech/Private/Apple/AppleTextToSpeech.cpp
+++ b/Engine/Plugins/Experimental/TextToSpeech/Source/TextToSpeech/Private/Apple/AppleTextToSpeech.cpp
@@ -2,6 +2,8 @@
 #include "Apple/AppleTextToSpeech.h"
 #include "GenericPlatform/TextToSpeechBase.h"
 #include "TextToSpeechLog.h"
+#include "Apple/AppleStringUtils.h"
+#include "Apple/ScopeAutoreleasePool.h"
 
 #if PLATFORM_MAC
 #include "Mac/CocoaThread.h"
@@ -14,6 +16,8 @@
 #import <UIKit/UIKit.h>
 #endif
 
+#include "Apple/AppleMemoryObject.h"
+
 // #import <Foundation/Foundation.h>
 #import <AVFoundation/AVFoundation.h>
 
diff --git a/Engine/Plugins/FastBuildController/Source/Private/FastBuildControllerModule.cpp b/Engine/Plugins/FastBuildController/Source/Private/FastBuildControllerModule.cpp
--- a/Engine/Plugins/FastBuildController/Source/Private/FastBuildControllerModule.cpp
+++ b/Engine/Plugins/FastBuildController/Source/Private/FastBuildControllerModule.cpp
@@ -25,6 +25,10 @@ namespace FASTBuildControllerVariables
 }
 
 #if PLATFORM_MAC
+
+#include "Apple/AppleStringUtils.h"
+#include "Apple/ScopeAutoreleasePool.h"
+
 static FString GetMetalCompilerFolder()
 {
 	FString Result;
diff --git a/Engine/Plugins/Media/AvfMedia/Source/AvfMedia/Private/Player/AvfMediaPlayer.cpp b/Engine/Plugins/Media/AvfMedia/Source/AvfMedia/Private/Player/AvfMediaPlayer.cpp
--- a/Engine/Plugins/Media/AvfMedia/Source/AvfMedia/Private/Player/AvfMediaPlayer.cpp
+++ b/Engine/Plugins/Media/AvfMedia/Source/AvfMedia/Private/Player/AvfMediaPlayer.cpp
@@ -6,6 +6,7 @@
 #include "HAL/PlatformProcess.h"
 #include "IMediaEventSink.h"
 #include "MediaSamples.h"
+#include "Apple/AppleStringUtils.h"
 #include "Misc/CommandLine.h"
 #include "Misc/Paths.h"
 #include "Misc/CoreDelegates.h"
diff --git a/Engine/Plugins/Media/AvfMedia/Source/AvfMedia/Private/Player/AvfMediaTracks.cpp b/Engine/Plugins/Media/AvfMedia/Source/AvfMedia/Private/Player/AvfMediaTracks.cpp
--- a/Engine/Plugins/Media/AvfMedia/Source/AvfMedia/Private/Player/AvfMediaTracks.cpp
+++ b/Engine/Plugins/Media/AvfMedia/Source/AvfMedia/Private/Player/AvfMediaTracks.cpp
@@ -1,6 +1,16 @@
 // Copyright Epic Games, Inc. All Rights Reserved.
 
 #include "AvfMediaTracks.h"
+
+#if PLATFORM_MAC
+#include "Apple/ScopeAutoreleasePool.h"
+#include "Mac/MacSystemIncludes.h"
+#elif PLATFORM_IOS
+#include "Apple/ScopeAutoreleasePool.h"
+#include "IOS/IOSSystemIncludes.h"
+#endif
+#include "Misc/AssertionMacros.h"
+
 #include "AvfMediaPrivate.h"
 
 #include "Containers/ResourceArray.h"
diff --git a/Engine/Plugins/Media/AvfMedia/Source/AvfMediaCapture/Private/AvfMediaCaptureModule.cpp b/Engine/Plugins/Media/AvfMedia/Source/AvfMediaCapture/Private/AvfMediaCaptureModule.cpp
--- a/Engine/Plugins/Media/AvfMedia/Source/AvfMediaCapture/Private/AvfMediaCaptureModule.cpp
+++ b/Engine/Plugins/Media/AvfMedia/Source/AvfMediaCapture/Private/AvfMediaCaptureModule.cpp
@@ -21,6 +21,8 @@
 	#include "UObject/WeakObjectPtr.h"
 #endif
 
+#include "Apple/ScopeAutoreleasePool.h"
+
 #import <AVFoundation/AVFoundation.h>
 
 
diff --git a/Engine/Plugins/Media/AvfMedia/Source/AvfMediaCapture/Private/Player/AvfMediaCapturePlayer.cpp b/Engine/Plugins/Media/AvfMedia/Source/AvfMediaCapture/Private/Player/AvfMediaCapturePlayer.cpp
--- a/Engine/Plugins/Media/AvfMedia/Source/AvfMediaCapture/Private/Player/AvfMediaCapturePlayer.cpp
+++ b/Engine/Plugins/Media/AvfMedia/Source/AvfMediaCapture/Private/Player/AvfMediaCapturePlayer.cpp
@@ -5,6 +5,8 @@
 #include "Misc/MessageDialog.h"
 
 #include "IMetalDynamicRHI.h"
+#include "Apple/AppleStringUtils.h"
+#include "Apple/ScopeAutoreleasePool.h"
 
 inline FString FStringFromAppleFourCharWordCode(uint32 Word)
 {
diff --git a/Engine/Plugins/Online/IOS/OnlineSubsystemIOS/Source/Private/OnlineAchievementsInterfaceIOS.cpp b/Engine/Plugins/Online/IOS/OnlineSubsystemIOS/Source/Private/OnlineAchievementsInterfaceIOS.cpp
--- a/Engine/Plugins/Online/IOS/OnlineSubsystemIOS/Source/Private/OnlineAchievementsInterfaceIOS.cpp
+++ b/Engine/Plugins/Online/IOS/OnlineSubsystemIOS/Source/Private/OnlineAchievementsInterfaceIOS.cpp
@@ -2,12 +2,17 @@
 
 // Module includes
 #include "OnlineAchievementsInterfaceIOS.h"
+
+#include "Apple/AppleStringUtils.h"
+
 #include "OnlineSubsystem.h"
 
 // Game Center includes
 #include <GameKit/GKAchievement.h>
 #include <GameKit/GKAchievementDescription.h>
 
+#include "IOS/IOSAsyncTask.h"
+
 FOnlineAchievementsIOS::FOnlineAchievementsIOS(class FOnlineSubsystemIOS* InSubsystem)
 {
 	UE_LOG_ONLINE_ACHIEVEMENTS(Display, TEXT("FOnlineSubsystemIOS::FOnlineAchievementsIOS()"));
diff --git a/Engine/Plugins/Online/IOS/OnlineSubsystemIOS/Source/Private/OnlineAppStoreUtils.cpp b/Engine/Plugins/Online/IOS/OnlineSubsystemIOS/Source/Private/OnlineAppStoreUtils.cpp
--- a/Engine/Plugins/Online/IOS/OnlineSubsystemIOS/Source/Private/OnlineAppStoreUtils.cpp
+++ b/Engine/Plugins/Online/IOS/OnlineSubsystemIOS/Source/Private/OnlineAppStoreUtils.cpp
@@ -4,6 +4,9 @@
 #include "OnlineSubsystem.h"
 
 #include "HAL/PlatformProcess.h"
+#include "IOS/IOSAsyncTask.h"
+
+#import <Foundation/Foundation.h>
 
 @implementation FAppStoreUtils
 
diff --git a/Engine/Plugins/Online/IOS/OnlineSubsystemIOS/Source/Private/OnlineAppStoreUtils.h b/Engine/Plugins/Online/IOS/OnlineSubsystemIOS/Source/Private/OnlineAppStoreUtils.h
--- a/Engine/Plugins/Online/IOS/OnlineSubsystemIOS/Source/Private/OnlineAppStoreUtils.h
+++ b/Engine/Plugins/Online/IOS/OnlineSubsystemIOS/Source/Private/OnlineAppStoreUtils.h
@@ -4,6 +4,8 @@
 
 #include "OnlineSubsystemTypes.h"
 
+#import <Foundation/Foundation.h> 
+
 /**
  * Delegate fires when a query for the app bundle id metadata completes
  *
diff --git a/Engine/Plugins/Online/IOS/OnlineSubsystemIOS/Source/Private/OnlineIdentityInterfaceIOS.cpp b/Engine/Plugins/Online/IOS/OnlineSubsystemIOS/Source/Private/OnlineIdentityInterfaceIOS.cpp
--- a/Engine/Plugins/Online/IOS/OnlineSubsystemIOS/Source/Private/OnlineIdentityInterfaceIOS.cpp
+++ b/Engine/Plugins/Online/IOS/OnlineSubsystemIOS/Source/Private/OnlineIdentityInterfaceIOS.cpp
@@ -5,6 +5,7 @@
 #include "OnlineSubsystemIOS.h"
 #include "OnlineError.h"
 #include "IOS/IOSAppDelegate.h"
+#include "IOS/IOSAsyncTask.h"
 
 #import <GameKit/GameKit.h>
 #import "OnlineAppStoreUtils.h"
diff --git a/Engine/Plugins/Online/IOS/OnlineSubsystemIOS/Source/Private/OnlineLeaderboardsInterfaceIOS.cpp b/Engine/Plugins/Online/IOS/OnlineSubsystemIOS/Source/Private/OnlineLeaderboardsInterfaceIOS.cpp
--- a/Engine/Plugins/Online/IOS/OnlineSubsystemIOS/Source/Private/OnlineLeaderboardsInterfaceIOS.cpp
+++ b/Engine/Plugins/Online/IOS/OnlineSubsystemIOS/Source/Private/OnlineLeaderboardsInterfaceIOS.cpp
@@ -2,6 +2,9 @@
 
 // Module includes
 #include "OnlineLeaderboardsInterfaceIOS.h"
+
+#include "Apple/AppleStringUtils.h"
+
 #include "OnlineSubsystemIOS.h"
 
 // GameCenter includes
@@ -9,6 +12,8 @@
 #include <GameKit/GKScore.h>
 #include <GameKit/GKLocalPlayer.h>
 
+#include "IOS/IOSAsyncTask.h"
+
 FOnlineLeaderboardsIOS::FOnlineLeaderboardsIOS(FOnlineSubsystemIOS* InSubsystem)
 {
 	UE_LOG_ONLINE_LEADERBOARD(Display, TEXT("FOnlineLeaderboardsIOS::FOnlineLeaderboardsIOS()"));
diff --git a/Engine/Plugins/Online/IOS/OnlineSubsystemIOS/Source/Private/OnlinePurchaseIOS.cpp b/Engine/Plugins/Online/IOS/OnlineSubsystemIOS/Source/Private/OnlinePurchaseIOS.cpp
--- a/Engine/Plugins/Online/IOS/OnlineSubsystemIOS/Source/Private/OnlinePurchaseIOS.cpp
+++ b/Engine/Plugins/Online/IOS/OnlineSubsystemIOS/Source/Private/OnlinePurchaseIOS.cpp
@@ -15,6 +15,8 @@
 #import <StoreKit/SKPayment.h>
 #import <StoreKit/SKPaymentQueue.h>
 
+#include "IOS/IOSAsyncTask.h"
+
 #define LOCTEXT_NAMESPACE "OnlineSubsystemIOS"
 
 namespace OnlinePurchaseIOSPrivate
diff --git a/Engine/Plugins/Online/IOS/OnlineSubsystemIOS/Source/Private/OnlineSessionInterfaceIOS.cpp b/Engine/Plugins/Online/IOS/OnlineSubsystemIOS/Source/Private/OnlineSessionInterfaceIOS.cpp
--- a/Engine/Plugins/Online/IOS/OnlineSubsystemIOS/Source/Private/OnlineSessionInterfaceIOS.cpp
+++ b/Engine/Plugins/Online/IOS/OnlineSubsystemIOS/Source/Private/OnlineSessionInterfaceIOS.cpp
@@ -1,6 +1,10 @@
 // Copyright Epic Games, Inc. All Rights Reserved.
 
 #include "OnlineSessionInterfaceIOS.h"
+
+#include "Apple/AppleStringUtils.h"
+#include "IOS/IOSAsyncTask.h"
+
 #include "OnlineSessionSettings.h"
 #include "OnlineSubsystemIOS.h"
 
diff --git a/Engine/Plugins/Online/IOS/OnlineSubsystemIOS/Source/Private/OnlineStoreIOS.cpp b/Engine/Plugins/Online/IOS/OnlineSubsystemIOS/Source/Private/OnlineStoreIOS.cpp
--- a/Engine/Plugins/Online/IOS/OnlineSubsystemIOS/Source/Private/OnlineStoreIOS.cpp
+++ b/Engine/Plugins/Online/IOS/OnlineSubsystemIOS/Source/Private/OnlineStoreIOS.cpp
@@ -3,6 +3,7 @@
 #if UE_WITH_STORE_KIT
 
 #include "OnlineStoreIOS.h"
+#include "Apple/AppleStringUtils.h"
 #include "Internationalization/Internationalization.h"
 #include "Internationalization/Culture.h"
 #include "Internationalization/FastDecimalFormat.h"
@@ -14,6 +15,7 @@
 #include "OnlineSubsystemIOS-Swift.h"
 #include "StoreKitSwiftInterop.h"
 #include "IOSAppDelegate.h"
+#include "IOS/IOSAsyncTask.h"
 
 @interface FSKProductsRequestHelper : SKProductsRequest
 {
diff --git a/Engine/Plugins/Online/IOS/OnlineSubsystemIOS/Source/Private/OnlineTurnBasedInterfaceIOS.cpp b/Engine/Plugins/Online/IOS/OnlineSubsystemIOS/Source/Private/OnlineTurnBasedInterfaceIOS.cpp
--- a/Engine/Plugins/Online/IOS/OnlineSubsystemIOS/Source/Private/OnlineTurnBasedInterfaceIOS.cpp
+++ b/Engine/Plugins/Online/IOS/OnlineSubsystemIOS/Source/Private/OnlineTurnBasedInterfaceIOS.cpp
@@ -10,6 +10,8 @@
 #include "Interfaces/TurnBasedMatchInterface.h"
 #include <GameKit/GKLocalPlayer.h>
 
+#include "IOS/IOSAsyncTask.h"
+
 DEFINE_LOG_CATEGORY_STATIC(LogTurnBasedInterfaceIOS, Verbose, All);
 
 #define INCLUDE_LOCAL_PLAYER (true)
diff --git a/Engine/Plugins/Online/IOS/OnlineSubsystemIOS/Source/Public/OnlinePurchaseIOS.h b/Engine/Plugins/Online/IOS/OnlineSubsystemIOS/Source/Public/OnlinePurchaseIOS.h
--- a/Engine/Plugins/Online/IOS/OnlineSubsystemIOS/Source/Public/OnlinePurchaseIOS.h
+++ b/Engine/Plugins/Online/IOS/OnlineSubsystemIOS/Source/Public/OnlinePurchaseIOS.h
@@ -5,6 +5,7 @@
 #if UE_WITH_STORE_KIT
 
 #include "Interfaces/OnlinePurchaseInterface.h"
+#import <Foundation/Foundation.h>
 #include "RetainedObjCInstance.h"
 
 enum class EPurchaseTransactionState : uint8;
diff --git a/Engine/Plugins/Online/IOS/OnlineSubsystemIOS/Source/Public/OnlineStoreIOS.h b/Engine/Plugins/Online/IOS/OnlineSubsystemIOS/Source/Public/OnlineStoreIOS.h
--- a/Engine/Plugins/Online/IOS/OnlineSubsystemIOS/Source/Public/OnlineStoreIOS.h
+++ b/Engine/Plugins/Online/IOS/OnlineSubsystemIOS/Source/Public/OnlineStoreIOS.h
@@ -5,6 +5,7 @@
 #if UE_WITH_STORE_KIT
 
 #include "Interfaces/OnlineStoreInterfaceV2.h"
+#import <Foundation/Foundation.h>
 #include "RetainedObjCInstance.h"
 
 enum class EPurchaseTransactionState : uint8;
diff --git a/Engine/Plugins/Online/OnlineSubsystemGoogle/Source/Private/IOS/OnlineIdentityGoogle.cpp b/Engine/Plugins/Online/OnlineSubsystemGoogle/Source/Private/IOS/OnlineIdentityGoogle.cpp
--- a/Engine/Plugins/Online/OnlineSubsystemGoogle/Source/Private/IOS/OnlineIdentityGoogle.cpp
+++ b/Engine/Plugins/Online/OnlineSubsystemGoogle/Source/Private/IOS/OnlineIdentityGoogle.cpp
@@ -1,6 +1,11 @@
 // Copyright Epic Games, Inc. All Rights Reserved.
 
 #include "OnlineIdentityGoogle.h"
+
+#if PLATFORM_APPLE
+#include "Apple/AppleStringUtils.h"
+#endif
+
 #include "OnlineSubsystemGooglePrivate.h"
 #include "Interfaces/OnlineExternalUIInterface.h"
 #include "GoogleHelper.h"
diff --git a/Engine/Plugins/Runtime/AR/AppleAR/AppleARKit/Source/AppleARKit/Private/AppleARKitTextures.cpp b/Engine/Plugins/Runtime/AR/AppleAR/AppleARKit/Source/AppleARKit/Private/AppleARKitTextures.cpp
--- a/Engine/Plugins/Runtime/AR/AppleAR/AppleARKit/Source/AppleARKit/Private/AppleARKitTextures.cpp
+++ b/Engine/Plugins/Runtime/AR/AppleAR/AppleARKit/Source/AppleARKit/Private/AppleARKitTextures.cpp
@@ -23,6 +23,8 @@
 	#include "IMetalDynamicRHI.h"
 	#import <Metal/Metal.h>
 	#import <MetalPerformanceShaders/MetalPerformanceShaders.h>
+	#include "Apple/ScopeAutoreleasePool.h"
+
 #endif
 
 DECLARE_CYCLE_STAT(TEXT("Scale Image"), STAT_ScaleImage, STATGROUP_ARKIT);
diff --git a/Engine/Plugins/Runtime/AppleImageUtils/Source/AppleImageUtils/Private/AppleImageUtilsPlugin.cpp b/Engine/Plugins/Runtime/AppleImageUtils/Source/AppleImageUtils/Private/AppleImageUtilsPlugin.cpp
--- a/Engine/Plugins/Runtime/AppleImageUtils/Source/AppleImageUtils/Private/AppleImageUtilsPlugin.cpp
+++ b/Engine/Plugins/Runtime/AppleImageUtils/Source/AppleImageUtils/Private/AppleImageUtilsPlugin.cpp
@@ -5,7 +5,9 @@
 
 #if SUPPORTS_IMAGE_UTILS_1_0
 	#include "Apple/ApplePlatformMisc.h"
-	#include "Async/Async.h"
+	#include "Apple/ScopeAutoreleasePool.h"
+
+#include "Async/Async.h"
 	#include "Engine/Texture2D.h"
 
 	#import <CoreImage/CIContext.h>
diff --git a/Engine/Plugins/Runtime/AppleMoviePlayer/Source/AppleMoviePlayer/Private/AppleMovieStreamer.cpp b/Engine/Plugins/Runtime/AppleMoviePlayer/Source/AppleMoviePlayer/Private/AppleMovieStreamer.cpp
--- a/Engine/Plugins/Runtime/AppleMoviePlayer/Source/AppleMoviePlayer/Private/AppleMovieStreamer.cpp
+++ b/Engine/Plugins/Runtime/AppleMoviePlayer/Source/AppleMoviePlayer/Private/AppleMovieStreamer.cpp
@@ -13,6 +13,7 @@
 #include "Containers/ResourceArray.h"
 
 #include "IMetalDynamicRHI.h"
+#include "Apple/AppleStringUtils.h"
 
 DEFINE_LOG_CATEGORY_STATIC(LogMoviePlayer, Log, All);
 
diff --git a/Engine/Source/Developer/CrashDebugHelper/Private/Mac/CrashDebugHelperMac.cpp b/Engine/Source/Developer/CrashDebugHelper/Private/Mac/CrashDebugHelperMac.cpp
--- a/Engine/Source/Developer/CrashDebugHelper/Private/Mac/CrashDebugHelperMac.cpp
+++ b/Engine/Source/Developer/CrashDebugHelper/Private/Mac/CrashDebugHelperMac.cpp
@@ -2,6 +2,8 @@
 
 #include "CrashDebugHelperMac.h"
 #include "Misc/EngineVersion.h"
+#include "Apple/AppleStringUtils.h"
+#include "Apple/ScopeAutoreleasePool.h"
 #include "Apple/ApplePlatformSymbolication.h"
 #include "CrashReporter.h"
 #include "CrashDebugHelperPrivate.h"
diff --git a/Engine/Source/Developer/DesktopPlatform/Private/Mac/DesktopPlatformMac.cpp b/Engine/Source/Developer/DesktopPlatform/Private/Mac/DesktopPlatformMac.cpp
--- a/Engine/Source/Developer/DesktopPlatform/Private/Mac/DesktopPlatformMac.cpp
+++ b/Engine/Source/Developer/DesktopPlatform/Private/Mac/DesktopPlatformMac.cpp
@@ -1,6 +1,8 @@
 // Copyright Epic Games, Inc. All Rights Reserved.
 
 #include "Mac/DesktopPlatformMac.h"
+#include "Apple/AppleStringUtils.h"
+#include "Apple/ScopeAutoreleasePool.h"
 #include "Mac/MacApplication.h"
 #include "Misc/FeedbackContextMarkup.h"
 #include "MacNativeFeedbackContext.h"
diff --git a/Engine/Source/Developer/DesktopPlatform/Private/Mac/MacNativeFeedbackContext.cpp b/Engine/Source/Developer/DesktopPlatform/Private/Mac/MacNativeFeedbackContext.cpp
--- a/Engine/Source/Developer/DesktopPlatform/Private/Mac/MacNativeFeedbackContext.cpp
+++ b/Engine/Source/Developer/DesktopPlatform/Private/Mac/MacNativeFeedbackContext.cpp
@@ -1,6 +1,9 @@
 // Copyright Epic Games, Inc. All Rights Reserved.
 
 #include "MacNativeFeedbackContext.h"
+
+#include "Apple/ScopeAutoreleasePool.h"
+#include "Apple/AppleStringUtils.h"
 #include "HAL/PlatformTime.h"
 #include "Logging/StructuredLog.h"
 #include "Mac/MacApplication.h"
diff --git a/Engine/Source/Programs/CrashReportClient/Private/Mac/CrashReportClientMainMac.cpp b/Engine/Source/Programs/CrashReportClient/Private/Mac/CrashReportClientMainMac.cpp
--- a/Engine/Source/Programs/CrashReportClient/Private/Mac/CrashReportClientMainMac.cpp
+++ b/Engine/Source/Programs/CrashReportClient/Private/Mac/CrashReportClientMainMac.cpp
@@ -1,6 +1,8 @@
 // Copyright Epic Games, Inc. All Rights Reserved.
 
 #include "CrashReportClientApp.h"
+
+#include "Apple/ScopeAutoreleasePool.h"
 #include "HAL/ExceptionHandling.h"
 #include "Mac/MacPlatformCrashContext.h"
 #include "Mac/MacPlatformProcess.h"
diff --git a/Engine/Source/Programs/Mac/UnrealEditorServices/Private/Mac/UnrealEditorServicesMain.cpp b/Engine/Source/Programs/Mac/UnrealEditorServices/Private/Mac/UnrealEditorServicesMain.cpp
--- a/Engine/Source/Programs/Mac/UnrealEditorServices/Private/Mac/UnrealEditorServicesMain.cpp
+++ b/Engine/Source/Programs/Mac/UnrealEditorServices/Private/Mac/UnrealEditorServicesMain.cpp
@@ -1,5 +1,6 @@
 // Copyright Epic Games, Inc. All Rights Reserved.
 
+#include "Apple/ScopeAutoreleasePool.h"
 #include "CoreMinimal.h"
 #include "RequiredProgramMainCPPInclude.h"
 #include "Modules/ModuleManager.h"
diff --git a/Engine/Source/Programs/UnrealMultiUserServer/Private/UnrealMultiUserServerMain.cpp b/Engine/Source/Programs/UnrealMultiUserServer/Private/UnrealMultiUserServerMain.cpp
--- a/Engine/Source/Programs/UnrealMultiUserServer/Private/UnrealMultiUserServerMain.cpp
+++ b/Engine/Source/Programs/UnrealMultiUserServer/Private/UnrealMultiUserServerMain.cpp
@@ -5,12 +5,15 @@
 
 #include "UnrealMultiUserServerRun.h"
 #include "RequiredProgramMainCPPInclude.h"
+
 DEFINE_LOG_CATEGORY_STATIC(LogMultiUserServer, Log, All)
 
 IMPLEMENT_APPLICATION(UnrealMultiUserServer, "UnrealMultiUserServer");
 
 #if PLATFORM_MAC // On Mac, to get a properly logging console that play nice, we need to build a mac application (.app) rather than a console application.
 
+#include "Apple/ScopeAutoreleasePool.h"
+
 class CommandLineArguments
 {
 public:
diff --git a/Engine/Source/Runtime/AVIWriter/Private/AVIWriter.cpp b/Engine/Source/Runtime/AVIWriter/Private/AVIWriter.cpp
--- a/Engine/Source/Runtime/AVIWriter/Private/AVIWriter.cpp
+++ b/Engine/Source/Runtime/AVIWriter/Private/AVIWriter.cpp
@@ -426,6 +426,8 @@ private:
 
 #elif PLATFORM_MAC && WITH_UNREAL_DEVELOPER_TOOLS
 
+#include "Apple/ScopeAutoreleasePool.h"
+
 #import <AVFoundation/AVFoundation.h>
 
 /**
diff --git a/Engine/Source/Runtime/Apple/BackgroundHttpIOS/Private/BackgroundURLSessionHandler.cpp b/Engine/Source/Runtime/Apple/BackgroundHttpIOS/Private/BackgroundURLSessionHandler.cpp
--- a/Engine/Source/Runtime/Apple/BackgroundHttpIOS/Private/BackgroundURLSessionHandler.cpp
+++ b/Engine/Source/Runtime/Apple/BackgroundHttpIOS/Private/BackgroundURLSessionHandler.cpp
@@ -13,6 +13,12 @@
 
 #include "BuildSettings.h"
 
+#if PLATFORM_MAC
+#include "Mac/MacSystemIncludes.h"
+#elif PLATFORM_IOS
+#include "IOS/IOSSystemIncludes.h"
+#endif
+
 static bool GBackgroundHttpUseExistingNonResumableTasks = true;
 static FAutoConsoleVariableRef CVarBackgroundHttpUseExistingNonResumableTasks(
 	TEXT("ios.BackgroundHttp.bUseExistingNonResumableTasks"),
diff --git a/Engine/Source/Runtime/Apple/BackgroundHttpIOS/Public/BackgroundURLSessionHandler.h b/Engine/Source/Runtime/Apple/BackgroundHttpIOS/Public/BackgroundURLSessionHandler.h
--- a/Engine/Source/Runtime/Apple/BackgroundHttpIOS/Public/BackgroundURLSessionHandler.h
+++ b/Engine/Source/Runtime/Apple/BackgroundHttpIOS/Public/BackgroundURLSessionHandler.h
@@ -10,22 +10,8 @@
 //Included for FBackgroundHttpRequestMetricsExtended
 #include "BackgroundHttpMetrics.h"
 
-//Call backs called by the bellow FBackgroundURLSessionHandler so higher-level systems can respond to task updates.
-class UE_DEPRECATED(5.5, "Use new download methods in FBackgroundURLSessionHandler.") BACKGROUNDHTTPIOS_API FIOSBackgroundDownloadCoreDelegates
-{
-public:
-	DECLARE_MULTICAST_DELEGATE_ThreeParams(FIOSBackgroundDownload_DidFinishDownloadingToURL, NSURLSessionDownloadTask*, NSError*, const FString&);
-	DECLARE_MULTICAST_DELEGATE_FourParams(FIOSBackgroundDownload_DidWriteData, NSURLSessionDownloadTask*, int64_t /*Bytes Written Since Last Call */, int64_t /*Total Bytes Written */, int64_t /*Total Bytes Expedted To Write */);
-	DECLARE_MULTICAST_DELEGATE_TwoParams(FIOSBackgroundDownload_DidCompleteWithError, NSURLSessionTask*, NSError*);
-	DECLARE_DELEGATE(FIOSBackgroundDownload_DelayedBackgroundURLSessionCompleteHandler);
-	DECLARE_MULTICAST_DELEGATE_TwoParams(FIOSBackgroundDownload_SessionDidFinishAllEvents, NSURLSession*, FIOSBackgroundDownload_DelayedBackgroundURLSessionCompleteHandler);
-
-	UE_DEPRECATED(5.5, "Use new API in FBackgroundURLSessionHandler.") static FIOSBackgroundDownload_DidFinishDownloadingToURL OnIOSBackgroundDownload_DidFinishDownloadingToURL;
-	UE_DEPRECATED(5.5, "Use new API in FBackgroundURLSessionHandler.") static FIOSBackgroundDownload_DidWriteData OnIOSBackgroundDownload_DidWriteData;
-	UE_DEPRECATED(5.5, "Use new API in FBackgroundURLSessionHandler.") static FIOSBackgroundDownload_DidCompleteWithError OnIOSBackgroundDownload_DidCompleteWithError;
-	UE_DEPRECATED(5.5, "Use new API in FBackgroundURLSessionHandler.") static FIOSBackgroundDownload_SessionDidFinishAllEvents OnIOSBackgroundDownload_SessionDidFinishAllEvents;
-	UE_DEPRECATED(5.5, "Use new API in FBackgroundURLSessionHandler.") static FIOSBackgroundDownload_DelayedBackgroundURLSessionCompleteHandler OnDelayedBackgroundURLSessionCompleteHandler;
-};
+#import <Foundation/NSDictionary.h>
+@class NSString;
 
 //Interface for wrapping a NSURLSession configured to support background downloading of NSURLSessionDownloadTasks.
 //This exists here as we can have to re-associate with our background session after app launch and need to re-associate with downloads
@@ -33,26 +19,6 @@ public:
 class BACKGROUNDHTTPIOS_API FBackgroundURLSessionHandler
 {
 public:
-	// Initializes a BackgroundSession with the given identifier. If the current background session already exists, returns true if the identifier matches. False if identifier doesn't match or if the session fails to create.
-	UE_DEPRECATED(5.5, "Use new API in FBackgroundURLSessionHandler.") static bool InitBackgroundSession(const FString& SessionIdentifier) {return false;}
-
-	//bShouldInvalidateExistingTasks determines if the session cancells all outstanding tasks immediately and cancels the session immediately or waits for them to finish and then invalidates the session
-	UE_DEPRECATED(5.5, "Use new API in FBackgroundURLSessionHandler.") static void ShutdownBackgroundSession(bool bShouldFinishTasksFirst = true) {}
-
-	//Gets a pointer to the current background session
-	UE_DEPRECATED(5.5, "Use new API in FBackgroundURLSessionHandler.") static NSURLSession* GetBackgroundSession() {return nullptr;}
-
-	UE_DEPRECATED(5.5, "Use new API in FBackgroundURLSessionHandler.") static void CreateBackgroundSessionWorkingDirectory() {}
-	
-	//Function to mark if you would like for the NSURLSession to wait to call the completion handler when
-	//OnIOSBackgroundDownload_SessionDidFinishAllEvents is called for you to call the passed completion handler
-	//NOTE: Call DURING OnIOSBackgroundDownload_SessionDidFinishAllEvents
-	UE_DEPRECATED(5.5, "Use new API in FBackgroundURLSessionHandler.") static void AddDelayedBackgroundURLSessionComplete() {}
-	
-	//Function to handle calls to OnDelayedBackgroundURLSessionCompleteHandler
-	//The intention is to call this for every call to AddDelayedBackgroundURLSessionComplete.
-	//NOTE: Once calling this your task should be completely finished with work and ready to be backgrounded!
-	UE_DEPRECATED(5.5, "Use new API in FBackgroundURLSessionHandler.") static void OnDelayedBackgroundURLSessionCompleteHandlerCalled();
 
 	// New API for background downloading
 
diff --git a/Engine/Source/Runtime/Apple/MetalRHI/Private/MetalBuffer.cpp b/Engine/Source/Runtime/Apple/MetalRHI/Private/MetalBuffer.cpp
--- a/Engine/Source/Runtime/Apple/MetalRHI/Private/MetalBuffer.cpp
+++ b/Engine/Source/Runtime/Apple/MetalRHI/Private/MetalBuffer.cpp
@@ -1,6 +1,7 @@
 // Copyright Epic Games, Inc. All Rights Reserved.
 
 #include "MetalBuffer.h"
+#include "Apple/ScopeAutoreleasePool.h"
 #include "MetalCommandBuffer.h"
 #include "MetalCommandQueue.h"
 #include "MetalDevice.h"
diff --git a/Engine/Source/Runtime/Apple/MetalRHI/Private/MetalCommandEncoder.cpp b/Engine/Source/Runtime/Apple/MetalRHI/Private/MetalCommandEncoder.cpp
--- a/Engine/Source/Runtime/Apple/MetalRHI/Private/MetalCommandEncoder.cpp
+++ b/Engine/Source/Runtime/Apple/MetalRHI/Private/MetalCommandEncoder.cpp
@@ -5,6 +5,7 @@
 =============================================================================*/
 
 #include "MetalCommandEncoder.h"
+#include "Apple/ScopeAutoreleasePool.h"
 #include "MetalCommandBuffer.h"
 #include "MetalRHIPrivate.h"
 #include "MetalShaderTypes.h"
diff --git a/Engine/Source/Runtime/Apple/MetalRHI/Private/MetalCommandQueue.cpp b/Engine/Source/Runtime/Apple/MetalRHI/Private/MetalCommandQueue.cpp
--- a/Engine/Source/Runtime/Apple/MetalRHI/Private/MetalCommandQueue.cpp
+++ b/Engine/Source/Runtime/Apple/MetalRHI/Private/MetalCommandQueue.cpp
@@ -5,6 +5,7 @@
 =============================================================================*/
 
 #include "MetalCommandQueue.h"
+#include "Apple/ScopeAutoreleasePool.h"
 #include "MetalCommandBuffer.h"
 #include "MetalDevice.h"
 #include "MetalFence.h"
diff --git a/Engine/Source/Runtime/Apple/MetalRHI/Private/MetalCommands.cpp b/Engine/Source/Runtime/Apple/MetalRHI/Private/MetalCommands.cpp
--- a/Engine/Source/Runtime/Apple/MetalRHI/Private/MetalCommands.cpp
+++ b/Engine/Source/Runtime/Apple/MetalRHI/Private/MetalCommands.cpp
@@ -5,6 +5,7 @@
 =============================================================================*/
 
 #include "MetalRHIPrivate.h"
+#include "Apple/ScopeAutoreleasePool.h"
 #include "MetalShaderTypes.h"
 #include "MetalGraphicsPipelineState.h"
 #include "MetalVertexDeclaration.h"
diff --git a/Engine/Source/Runtime/Apple/MetalRHI/Private/MetalCounterSampler.cpp b/Engine/Source/Runtime/Apple/MetalRHI/Private/MetalCounterSampler.cpp
--- a/Engine/Source/Runtime/Apple/MetalRHI/Private/MetalCounterSampler.cpp
+++ b/Engine/Source/Runtime/Apple/MetalRHI/Private/MetalCounterSampler.cpp
@@ -1,6 +1,7 @@
 // Copyright Epic Games, Inc. All Rights Reserved.
 
 #include "MetalCounterSampler.h"
+#include "Apple/ScopeAutoreleasePool.h"
 #include "MetalDynamicRHI.h"
 
 static TAutoConsoleVariable<bool> CVarMetalRHIInsertCounterSampleBarrier(
diff --git a/Engine/Source/Runtime/Apple/MetalRHI/Private/MetalDevice.cpp b/Engine/Source/Runtime/Apple/MetalRHI/Private/MetalDevice.cpp
--- a/Engine/Source/Runtime/Apple/MetalRHI/Private/MetalDevice.cpp
+++ b/Engine/Source/Runtime/Apple/MetalRHI/Private/MetalDevice.cpp
@@ -1,6 +1,7 @@
 // Copyright Epic Games, Inc. All Rights Reserved.
 
 #include "MetalDevice.h"
+#include "Apple/ScopeAutoreleasePool.h"
 #include "MetalRHI.h"
 #include "MetalRHIPrivate.h"
 #include "MetalRHIRenderQuery.h"
diff --git a/Engine/Source/Runtime/Apple/MetalRHI/Private/MetalDynamicRHI.cpp b/Engine/Source/Runtime/Apple/MetalRHI/Private/MetalDynamicRHI.cpp
--- a/Engine/Source/Runtime/Apple/MetalRHI/Private/MetalDynamicRHI.cpp
+++ b/Engine/Source/Runtime/Apple/MetalRHI/Private/MetalDynamicRHI.cpp
@@ -5,6 +5,7 @@
 =============================================================================*/
 
 #include "MetalDynamicRHI.h"
+#include "Apple/ScopeAutoreleasePool.h"
 #include "MetalRHIPrivate.h"
 #include "MetalRHIRenderQuery.h"
 #include "MetalRHIStagingBuffer.h"
diff --git a/Engine/Source/Runtime/Apple/MetalRHI/Private/MetalDynamicRHI_Shaders.cpp b/Engine/Source/Runtime/Apple/MetalRHI/Private/MetalDynamicRHI_Shaders.cpp
--- a/Engine/Source/Runtime/Apple/MetalRHI/Private/MetalDynamicRHI_Shaders.cpp
+++ b/Engine/Source/Runtime/Apple/MetalRHI/Private/MetalDynamicRHI_Shaders.cpp
@@ -5,6 +5,7 @@
 =============================================================================*/
 
 #include "MetalDynamicRHI.h"
+#include "Apple/ScopeAutoreleasePool.h"
 #include "MetalRHIPrivate.h"
 #include "MetalShaderTypes.h"
 #include "Shaders/MetalShaderLibrary.h"
diff --git a/Engine/Source/Runtime/Apple/MetalRHI/Private/MetalLLM.h b/Engine/Source/Runtime/Apple/MetalRHI/Private/MetalLLM.h
--- a/Engine/Source/Runtime/Apple/MetalRHI/Private/MetalLLM.h
+++ b/Engine/Source/Runtime/Apple/MetalRHI/Private/MetalLLM.h
@@ -5,6 +5,7 @@
 #include "MetalRHIPrivate.h"
 #include "HAL/LowLevelMemTracker.h"
 #include "Apple/AppleLLM.h"
+#include "Apple/AppleMemoryObject.h"
 #include "RHICoreStats.h"
 
 class FMetalBuffer;
diff --git a/Engine/Source/Runtime/Apple/MetalRHI/Private/MetalPipeline.cpp b/Engine/Source/Runtime/Apple/MetalRHI/Private/MetalPipeline.cpp
--- a/Engine/Source/Runtime/Apple/MetalRHI/Private/MetalPipeline.cpp
+++ b/Engine/Source/Runtime/Apple/MetalRHI/Private/MetalPipeline.cpp
@@ -5,6 +5,13 @@
 =============================================================================*/
 
 #include "MetalPipeline.h"
+
+#if PLATFORM_MAC
+#include "Mac/MacSystemIncludes.h"
+#elif PLATFORM_IOS
+#include "IOS/IOSSystemIncludes.h"
+#endif
+
 #include "MetalRHIPrivate.h"
 #include "MetalVertexDeclaration.h"
 #include "MetalShaderTypes.h"
diff --git a/Engine/Source/Runtime/Apple/MetalRHI/Private/MetalRHI.cpp b/Engine/Source/Runtime/Apple/MetalRHI/Private/MetalRHI.cpp
--- a/Engine/Source/Runtime/Apple/MetalRHI/Private/MetalRHI.cpp
+++ b/Engine/Source/Runtime/Apple/MetalRHI/Private/MetalRHI.cpp
@@ -5,6 +5,7 @@
 =============================================================================*/
 
 #include "MetalRHI.h"
+#include "Apple/ScopeAutoreleasePool.h"
 
 #include "MetalThirdParty.h"
 
diff --git a/Engine/Source/Runtime/Apple/MetalRHI/Private/MetalRHIBuffer.cpp b/Engine/Source/Runtime/Apple/MetalRHI/Private/MetalRHIBuffer.cpp
--- a/Engine/Source/Runtime/Apple/MetalRHI/Private/MetalRHIBuffer.cpp
+++ b/Engine/Source/Runtime/Apple/MetalRHI/Private/MetalRHIBuffer.cpp
@@ -5,6 +5,7 @@
 =============================================================================*/
 
 #include "MetalRHIPrivate.h"
+#include "Apple/ScopeAutoreleasePool.h"
 #include "MetalProfiler.h"
 #include "MetalCommandBuffer.h"
 #include "MetalCommandQueue.h"
diff --git a/Engine/Source/Runtime/Apple/MetalRHI/Private/MetalRHIContext.cpp b/Engine/Source/Runtime/Apple/MetalRHI/Private/MetalRHIContext.cpp
--- a/Engine/Source/Runtime/Apple/MetalRHI/Private/MetalRHIContext.cpp
+++ b/Engine/Source/Runtime/Apple/MetalRHI/Private/MetalRHIContext.cpp
@@ -1,6 +1,7 @@
 // Copyright Epic Games, Inc. All Rights Reserved.
 
 #include "MetalRHIContext.h"
+#include "Apple/ScopeAutoreleasePool.h"
 #include "MetalRHIPrivate.h"
 #include "MetalRHIRenderQuery.h"
 #include "MetalRHIVisionOSBridge.h"
diff --git a/Engine/Source/Runtime/Apple/MetalRHI/Private/MetalRHIRenderQuery.cpp b/Engine/Source/Runtime/Apple/MetalRHI/Private/MetalRHIRenderQuery.cpp
--- a/Engine/Source/Runtime/Apple/MetalRHI/Private/MetalRHIRenderQuery.cpp
+++ b/Engine/Source/Runtime/Apple/MetalRHI/Private/MetalRHIRenderQuery.cpp
@@ -5,6 +5,7 @@
 =============================================================================*/
 
 #include "MetalRHIRenderQuery.h"
+#include "Apple/ScopeAutoreleasePool.h"
 #include "MetalDevice.h"
 #include "MetalRHIPrivate.h"
 #include "MetalDynamicRHI.h"
diff --git a/Engine/Source/Runtime/Apple/MetalRHI/Private/MetalRayTracing.cpp b/Engine/Source/Runtime/Apple/MetalRHI/Private/MetalRayTracing.cpp
--- a/Engine/Source/Runtime/Apple/MetalRHI/Private/MetalRayTracing.cpp
+++ b/Engine/Source/Runtime/Apple/MetalRHI/Private/MetalRayTracing.cpp
@@ -15,6 +15,7 @@
 #include "BuiltInRayTracingShaders.h"
 #include "RayTracingValidationShaders.h"
 #include "Async/ParallelFor.h"
+#include "Apple/ScopeAutoreleasePool.h"
 
 static int32 GMetalRayTracingAllowCompaction = 1;
 static FAutoConsoleVariableRef CVarMetalRayTracingAllowCompaction(
diff --git a/Engine/Source/Runtime/Apple/MetalRHI/Private/MetalRenderTarget.cpp b/Engine/Source/Runtime/Apple/MetalRHI/Private/MetalRenderTarget.cpp
--- a/Engine/Source/Runtime/Apple/MetalRHI/Private/MetalRenderTarget.cpp
+++ b/Engine/Source/Runtime/Apple/MetalRHI/Private/MetalRenderTarget.cpp
@@ -5,6 +5,7 @@
 =============================================================================*/
 
 #include "MetalRHIPrivate.h"
+#include "Apple/ScopeAutoreleasePool.h"
 #include "ScreenRendering.h"
 #include "MetalCommandBuffer.h"
 #include "MetalCommandQueue.h"
diff --git a/Engine/Source/Runtime/Apple/MetalRHI/Private/MetalState.cpp b/Engine/Source/Runtime/Apple/MetalRHI/Private/MetalState.cpp
--- a/Engine/Source/Runtime/Apple/MetalRHI/Private/MetalState.cpp
+++ b/Engine/Source/Runtime/Apple/MetalRHI/Private/MetalState.cpp
@@ -5,6 +5,7 @@
 =============================================================================*/
 
 #include "MetalState.h"
+#include "Apple/ScopeAutoreleasePool.h"
 #include "MetalDynamicRHI.h"
 #include "MetalRHIPrivate.h"
 #include "MetalProfiler.h"
diff --git a/Engine/Source/Runtime/Apple/MetalRHI/Private/MetalSubmission.cpp b/Engine/Source/Runtime/Apple/MetalRHI/Private/MetalSubmission.cpp
--- a/Engine/Source/Runtime/Apple/MetalRHI/Private/MetalSubmission.cpp
+++ b/Engine/Source/Runtime/Apple/MetalRHI/Private/MetalSubmission.cpp
@@ -1,6 +1,7 @@
 // Copyright Epic Games, Inc. All Rights Reserved.
 
 #include "MetalSubmission.h"
+#include "Apple/ScopeAutoreleasePool.h"
 #include "MetalRHIPrivate.h"
 #include "MetalDynamicRHI.h"
 #include "MetalCommandBuffer.h"
diff --git a/Engine/Source/Runtime/Apple/MetalRHI/Private/MetalTexture.cpp b/Engine/Source/Runtime/Apple/MetalRHI/Private/MetalTexture.cpp
--- a/Engine/Source/Runtime/Apple/MetalRHI/Private/MetalTexture.cpp
+++ b/Engine/Source/Runtime/Apple/MetalRHI/Private/MetalTexture.cpp
@@ -5,6 +5,7 @@
  =============================================================================*/
 
 #include "MetalRHIPrivate.h"
+#include "Apple/ScopeAutoreleasePool.h"
 #include "MetalDynamicRHI.h"
 #include "MetalProfiler.h" // for STAT_MetalTexturePageOffTime
 #include "MetalCommandBuffer.h"
diff --git a/Engine/Source/Runtime/Apple/MetalRHI/Private/MetalUAV.cpp b/Engine/Source/Runtime/Apple/MetalRHI/Private/MetalUAV.cpp
--- a/Engine/Source/Runtime/Apple/MetalRHI/Private/MetalUAV.cpp
+++ b/Engine/Source/Runtime/Apple/MetalRHI/Private/MetalUAV.cpp
@@ -2,6 +2,7 @@
 
 
 #include "MetalRHIPrivate.h"
+#include "Apple/ScopeAutoreleasePool.h"
 #include "MetalDynamicRHI.h"
 #include "MetalRHIContext.h"
 #include "MetalRHIStagingBuffer.h"
diff --git a/Engine/Source/Runtime/Apple/MetalRHI/Private/MetalViewport.cpp b/Engine/Source/Runtime/Apple/MetalRHI/Private/MetalViewport.cpp
--- a/Engine/Source/Runtime/Apple/MetalRHI/Private/MetalViewport.cpp
+++ b/Engine/Source/Runtime/Apple/MetalRHI/Private/MetalViewport.cpp
@@ -5,6 +5,7 @@
 =============================================================================*/
 
 #include "MetalViewport.h"
+#include "Apple/ScopeAutoreleasePool.h"
 #include "MetalDynamicRHI.h"
 #include "MetalRHIPrivate.h"
 #include "MetalCommandBuffer.h"
@@ -15,9 +16,11 @@
 #import <QuartzCore/CAMetalLayer.h>
 
 #if PLATFORM_MAC
+#include "Mac/MacSystemIncludes.h"
 #include "Mac/CocoaWindow.h"
 #include "Mac/CocoaThread.h"
 #else
+#include "IOS/IOSSystemIncludes.h"
 #include "IOS/IOSApplication.h"
 #include "IOS/IOSAppDelegate.h"
 #endif
diff --git a/Engine/Source/Runtime/Apple/MetalRHI/Private/Shaders/Types/MetalComputeShader.cpp b/Engine/Source/Runtime/Apple/MetalRHI/Private/Shaders/Types/MetalComputeShader.cpp
--- a/Engine/Source/Runtime/Apple/MetalRHI/Private/Shaders/Types/MetalComputeShader.cpp
+++ b/Engine/Source/Runtime/Apple/MetalRHI/Private/Shaders/Types/MetalComputeShader.cpp
@@ -5,6 +5,7 @@
 =============================================================================*/
 
 #include "MetalComputeShader.h"
+#include "Apple/ScopeAutoreleasePool.h"
 #include "MetalCommandBuffer.h"
 #include "MetalDevice.h"
 #include "MetalProfiler.h"
diff --git a/Engine/Source/Runtime/Apple/MetalRHI/Private/Shaders/Types/Templates/MetalBaseShader.h b/Engine/Source/Runtime/Apple/MetalRHI/Private/Shaders/Types/Templates/MetalBaseShader.h
--- a/Engine/Source/Runtime/Apple/MetalRHI/Private/Shaders/Types/Templates/MetalBaseShader.h
+++ b/Engine/Source/Runtime/Apple/MetalRHI/Private/Shaders/Types/Templates/MetalBaseShader.h
@@ -6,6 +6,7 @@
 
 #pragma once
 
+#include "Apple/ScopeAutoreleasePool.h"
 #include "Misc/Paths.h"
 #include "HAL/FileManager.h"
 #include "HAL/PlatformFileManager.h"
diff --git a/Engine/Source/Runtime/Apple/MetalRHI/Public/IMetalDynamicRHI.h b/Engine/Source/Runtime/Apple/MetalRHI/Public/IMetalDynamicRHI.h
--- a/Engine/Source/Runtime/Apple/MetalRHI/Public/IMetalDynamicRHI.h
+++ b/Engine/Source/Runtime/Apple/MetalRHI/Public/IMetalDynamicRHI.h
@@ -2,6 +2,12 @@
 
 #pragma once 
 
+#if PLATFORM_MAC
+#include "Mac/MacSystemIncludes.h"
+#elif PLATFORM_IOS
+#include "IOS/IOSSystemIncludes.h"
+#endif
+
 #include "RHI.h"
 #include "MetalThirdParty.h"
 
diff --git a/Engine/Source/Runtime/Apple/MetalRHI/Public/MetalViewport.h b/Engine/Source/Runtime/Apple/MetalRHI/Public/MetalViewport.h
--- a/Engine/Source/Runtime/Apple/MetalRHI/Public/MetalViewport.h
+++ b/Engine/Source/Runtime/Apple/MetalRHI/Public/MetalViewport.h
@@ -7,10 +7,17 @@
 #pragma once
 
 #include "MetalRHIPrivate.h"
+#include "Apple/AppleStringUtils.h"
 #include "PixelFormat.h"
 #include "RHIResources.h"
 #include "MetalResources.h"
 
+#if PLATFORM_MAC
+#include "Mac/MacSystemIncludes.h"
+#elif PLATFORM_IOS
+#include "IOS/IOSSystemIncludes.h"
+#endif
+
 #if PLATFORM_MAC
 #include "Mac/CocoaTextView.h"
 @interface FMetalView : FCocoaTextView
diff --git a/Engine/Source/Runtime/ApplicationCore/Private/Apple/AppleControllerInterface.cpp b/Engine/Source/Runtime/ApplicationCore/Private/Apple/AppleControllerInterface.cpp
--- a/Engine/Source/Runtime/ApplicationCore/Private/Apple/AppleControllerInterface.cpp
+++ b/Engine/Source/Runtime/ApplicationCore/Private/Apple/AppleControllerInterface.cpp
@@ -1,6 +1,8 @@
 // Copyright Epic Games, Inc. All Rights Reserved.
 
 #include "AppleControllerInterface.h"
+
+#include "Apple/AppleStringUtils.h"
 #include "HAL/PlatformTime.h"
 #include "Misc/ScopeLock.h"
 #include "GenericPlatform/GenericPlatformInputDeviceMapper.h"
diff --git a/Engine/Source/Runtime/ApplicationCore/Private/GenericPlatform/GenericPlatformSoftwareCursor.cpp b/Engine/Source/Runtime/ApplicationCore/Private/GenericPlatform/GenericPlatformSoftwareCursor.cpp
--- a/Engine/Source/Runtime/ApplicationCore/Private/GenericPlatform/GenericPlatformSoftwareCursor.cpp
+++ b/Engine/Source/Runtime/ApplicationCore/Private/GenericPlatform/GenericPlatformSoftwareCursor.cpp
@@ -2,6 +2,11 @@
 
 #include "GenericPlatform/GenericPlatformSoftwareCursor.h"
 #include "GenericPlatform/GenericApplication.h"
+#if PLATFORM_MAC
+#include "Mac/MacSystemIncludes.h"
+#elif PLATFORM_IOS
+#include "IOS/IOSSystemIncludes.h"
+#endif
 
 // Windows has special needs with RECT, and we don't need this class on Windows anyway, so just skip it
 #if !PLATFORM_MICROSOFT
diff --git a/Engine/Source/Runtime/ApplicationCore/Private/IOS/Accessibility/IOSAccessibilityCache.cpp b/Engine/Source/Runtime/ApplicationCore/Private/IOS/Accessibility/IOSAccessibilityCache.cpp
--- a/Engine/Source/Runtime/ApplicationCore/Private/IOS/Accessibility/IOSAccessibilityCache.cpp
+++ b/Engine/Source/Runtime/ApplicationCore/Private/IOS/Accessibility/IOSAccessibilityCache.cpp
@@ -9,6 +9,7 @@
 #include "IOS/IOSAsyncTask.h"
 #include "IOS/IOSView.h"
 #include "HAL/IConsoleManager.h"
+#include "Apple/AppleStringUtils.h"
 #include "Async/TaskGraphInterfaces.h"
 
 @implementation FIOSAccessibilityCache
diff --git a/Engine/Source/Runtime/ApplicationCore/Private/IOS/Accessibility/IOSAccessibilityElement.cpp b/Engine/Source/Runtime/ApplicationCore/Private/IOS/Accessibility/IOSAccessibilityElement.cpp
--- a/Engine/Source/Runtime/ApplicationCore/Private/IOS/Accessibility/IOSAccessibilityElement.cpp
+++ b/Engine/Source/Runtime/ApplicationCore/Private/IOS/Accessibility/IOSAccessibilityElement.cpp
@@ -4,6 +4,7 @@
 
 #include "IOS/Accessibility/IOSAccessibilityElement.h"
 
+#include "Apple/AppleStringUtils.h"
 #include "Async/TaskGraphInterfaces.h"
 #include "GenericPlatform/Accessibility/GenericAccessibleInterfaces.h"
 #include "IOS/Accessibility/IOSAccessibilityCache.h"
diff --git a/Engine/Source/Runtime/ApplicationCore/Private/IOS/IOSAppDelegate.cpp b/Engine/Source/Runtime/ApplicationCore/Private/IOS/IOSAppDelegate.cpp
--- a/Engine/Source/Runtime/ApplicationCore/Private/IOS/IOSAppDelegate.cpp
+++ b/Engine/Source/Runtime/ApplicationCore/Private/IOS/IOSAppDelegate.cpp
@@ -1,6 +1,7 @@
 // Copyright Epic Games, Inc. All Rights Reserved.
 
 #include "IOS/IOSAppDelegate.h"
+#include "Apple/AppleStringUtils.h"
 #include "IOS/IOSCommandLineHelper.h"
 #include "HAL/ExceptionHandling.h"
 #include "Modules/Boilerplate/ModuleBoilerplate.h"
diff --git a/Engine/Source/Runtime/ApplicationCore/Private/IOS/IOSApplication.cpp b/Engine/Source/Runtime/ApplicationCore/Private/IOS/IOSApplication.cpp
--- a/Engine/Source/Runtime/ApplicationCore/Private/IOS/IOSApplication.cpp
+++ b/Engine/Source/Runtime/ApplicationCore/Private/IOS/IOSApplication.cpp
@@ -1,6 +1,8 @@
 // Copyright Epic Games, Inc. All Rights Reserved.
 
 #include "IOS/IOSApplication.h"
+
+#include "Apple/AppleStringUtils.h"
 #include "IOS/IOSInputInterface.h"
 #include "IOS/IOSCursor.h"
 #include "IOS/IOSView.h"
diff --git a/Engine/Source/Runtime/ApplicationCore/Private/IOS/IOSConfigRules.cpp b/Engine/Source/Runtime/ApplicationCore/Private/IOS/IOSConfigRules.cpp
--- a/Engine/Source/Runtime/ApplicationCore/Private/IOS/IOSConfigRules.cpp
+++ b/Engine/Source/Runtime/ApplicationCore/Private/IOS/IOSConfigRules.cpp
@@ -5,6 +5,8 @@
 #include "zlib.h"
 #include <CommonCrypto/CommonCryptor.h>
 #include <CommonCrypto/CommonKeyDerivation.h>
+
+#include "Apple/AppleStringUtils.h"
 #include "String/ParseLines.h"
 #include "String/ParseTokens.h"
 #include "Containers/StringConv.h"
diff --git a/Engine/Source/Runtime/ApplicationCore/Private/IOS/IOSPlatformFramePacer.cpp b/Engine/Source/Runtime/ApplicationCore/Private/IOS/IOSPlatformFramePacer.cpp
--- a/Engine/Source/Runtime/ApplicationCore/Private/IOS/IOSPlatformFramePacer.cpp
+++ b/Engine/Source/Runtime/ApplicationCore/Private/IOS/IOSPlatformFramePacer.cpp
@@ -1,6 +1,11 @@
 // Copyright Epic Games, Inc. All Rights Reserved.
 
 #include "IOS/IOSPlatformFramePacer.h"
+#if PLATFORM_MAC
+#include "Mac/MacSystemIncludes.h"
+#elif PLATFORM_IOS
+#include "IOS/IOSSystemIncludes.h"
+#endif
 #include "Containers/Array.h"
 #include "HAL/ThreadingBase.h"
 #include "HAL/IConsoleManager.h"
diff --git a/Engine/Source/Runtime/ApplicationCore/Private/IOS/IOSPlatformSurvey.cpp b/Engine/Source/Runtime/ApplicationCore/Private/IOS/IOSPlatformSurvey.cpp
--- a/Engine/Source/Runtime/ApplicationCore/Private/IOS/IOSPlatformSurvey.cpp
+++ b/Engine/Source/Runtime/ApplicationCore/Private/IOS/IOSPlatformSurvey.cpp
@@ -5,6 +5,7 @@
 =============================================================================*/
 
 #include "IOS/IOSPlatformSurvey.h"
+#include "Apple/AppleStringUtils.h"
 #include "IOS/IOSAppDelegate.h"
 #include "IOSWindow.h"
 
diff --git a/Engine/Source/Runtime/ApplicationCore/Private/Mac/Accessibility/MacAccessibilityElement.cpp b/Engine/Source/Runtime/ApplicationCore/Private/Mac/Accessibility/MacAccessibilityElement.cpp
--- a/Engine/Source/Runtime/ApplicationCore/Private/Mac/Accessibility/MacAccessibilityElement.cpp
+++ b/Engine/Source/Runtime/ApplicationCore/Private/Mac/Accessibility/MacAccessibilityElement.cpp
@@ -3,6 +3,7 @@
 #if WITH_ACCESSIBILITY
 
 #include "Mac/Accessibility/MacAccessibilityElement.h"
+#include "Apple/AppleStringUtils.h"
 #include "Mac/CocoaThread.h"
 #include "GenericPlatform/Accessibility/GenericAccessibleInterfaces.h"
 #include "Mac/Accessibility/MacAccessibilityManager.h"
diff --git a/Engine/Source/Runtime/ApplicationCore/Private/Mac/Accessibility/MacAccessibilityManager.cpp b/Engine/Source/Runtime/ApplicationCore/Private/Mac/Accessibility/MacAccessibilityManager.cpp
--- a/Engine/Source/Runtime/ApplicationCore/Private/Mac/Accessibility/MacAccessibilityManager.cpp
+++ b/Engine/Source/Runtime/ApplicationCore/Private/Mac/Accessibility/MacAccessibilityManager.cpp
@@ -7,6 +7,7 @@
 #include "Mac/MacApplication.h"
 #include "HAL/IConsoleManager.h"
 #include "Mac/CocoaThread.h"
+#include "Apple/AppleStringUtils.h"
 
 #import <AppKit/AppKit.h>
 
diff --git a/Engine/Source/Runtime/ApplicationCore/Private/Mac/CocoaMenu.cpp b/Engine/Source/Runtime/ApplicationCore/Private/Mac/CocoaMenu.cpp
--- a/Engine/Source/Runtime/ApplicationCore/Private/Mac/CocoaMenu.cpp
+++ b/Engine/Source/Runtime/ApplicationCore/Private/Mac/CocoaMenu.cpp
@@ -2,6 +2,7 @@
 
 #include "Mac/CocoaMenu.h"
 #include "Mac/MacApplication.h"
+#include "Apple/ScopeAutoreleasePool.h"
 
 @implementation FCocoaMenu
 
diff --git a/Engine/Source/Runtime/ApplicationCore/Private/Mac/CocoaTextView.cpp b/Engine/Source/Runtime/ApplicationCore/Private/Mac/CocoaTextView.cpp
--- a/Engine/Source/Runtime/ApplicationCore/Private/Mac/CocoaTextView.cpp
+++ b/Engine/Source/Runtime/ApplicationCore/Private/Mac/CocoaTextView.cpp
@@ -3,6 +3,7 @@
 #include "Mac/CocoaTextView.h"
 #include "Mac/CocoaWindow.h"
 #include "Mac/CocoaThread.h"
+#include "Apple/ScopeAutoreleasePool.h"
 
 @implementation FCocoaTextView
 
diff --git a/Engine/Source/Runtime/ApplicationCore/Private/Mac/CocoaWindow.cpp b/Engine/Source/Runtime/ApplicationCore/Private/Mac/CocoaWindow.cpp
--- a/Engine/Source/Runtime/ApplicationCore/Private/Mac/CocoaWindow.cpp
+++ b/Engine/Source/Runtime/ApplicationCore/Private/Mac/CocoaWindow.cpp
@@ -1,6 +1,8 @@
 // Copyright Epic Games, Inc. All Rights Reserved.
 
 #include "Mac/CocoaWindow.h"
+
+#include "Apple/ScopeAutoreleasePool.h"
 #include "Mac/MacApplication.h"
 #include "Mac/CocoaTextView.h"
 #include "Mac/CocoaThread.h"
diff --git a/Engine/Source/Runtime/ApplicationCore/Private/Mac/MacApplication.cpp b/Engine/Source/Runtime/ApplicationCore/Private/Mac/MacApplication.cpp
--- a/Engine/Source/Runtime/ApplicationCore/Private/Mac/MacApplication.cpp
+++ b/Engine/Source/Runtime/ApplicationCore/Private/Mac/MacApplication.cpp
@@ -28,6 +28,8 @@
 #include <IOKit/graphics/IOGraphicsLib.h>
 #include "Misc/CoreDelegates.h"
 
+#include "Apple/ScopeAutoreleasePool.h"
+#include "Apple/AppleStringUtils.h"
 #include "Apple/AppleControllerInterface.h"
 #include "HIDInputInterface.h"
 
diff --git a/Engine/Source/Runtime/ApplicationCore/Private/Mac/MacConsoleOutputDevice.cpp b/Engine/Source/Runtime/ApplicationCore/Private/Mac/MacConsoleOutputDevice.cpp
--- a/Engine/Source/Runtime/ApplicationCore/Private/Mac/MacConsoleOutputDevice.cpp
+++ b/Engine/Source/Runtime/ApplicationCore/Private/Mac/MacConsoleOutputDevice.cpp
@@ -1,6 +1,7 @@
 // Copyright Epic Games, Inc. All Rights Reserved.
 
 #include "Mac/MacConsoleOutputDevice.h"
+#include "Apple/ScopeAutoreleasePool.h"
 #include "Misc/App.h"
 #include "Misc/CoreDelegates.h"
 #include "Misc/ConfigCacheIni.h"
@@ -10,6 +11,7 @@
 #include "Misc/OutputDeviceRedirector.h"
 #include "Misc/ScopeLock.h"
 #include "Mac/CocoaThread.h"
+
 #include "HAL/PlatformApplicationMisc.h"
 #include "CoreGlobals.h"
 
diff --git a/Engine/Source/Runtime/ApplicationCore/Private/Mac/MacCursor.cpp b/Engine/Source/Runtime/ApplicationCore/Private/Mac/MacCursor.cpp
--- a/Engine/Source/Runtime/ApplicationCore/Private/Mac/MacCursor.cpp
+++ b/Engine/Source/Runtime/ApplicationCore/Private/Mac/MacCursor.cpp
@@ -1,6 +1,7 @@
 // Copyright Epic Games, Inc. All Rights Reserved.
 
 #include "Mac/MacCursor.h"
+#include "Apple/ScopeAutoreleasePool.h"
 #include "Mac/MacWindow.h"
 #include "Mac/MacApplication.h"
 #include "Math/IntRect.h"
diff --git a/Engine/Source/Runtime/ApplicationCore/Private/Mac/MacPlatformApplicationMisc.cpp b/Engine/Source/Runtime/ApplicationCore/Private/Mac/MacPlatformApplicationMisc.cpp
--- a/Engine/Source/Runtime/ApplicationCore/Private/Mac/MacPlatformApplicationMisc.cpp
+++ b/Engine/Source/Runtime/ApplicationCore/Private/Mac/MacPlatformApplicationMisc.cpp
@@ -33,6 +33,9 @@
 #include <notify.h>
 #include <uuid/uuid.h>
 #import <ScreenCaptureKit/ScreenCaptureKit.h>
+
+#include "Apple/ScopeAutoreleasePool.h"
+#include "Apple/AppleStringUtils.h"
 #include "Apple/PostAppleSystemHeaders.h"
 
 
diff --git a/Engine/Source/Runtime/ApplicationCore/Private/Mac/MacPlatformSplash.cpp b/Engine/Source/Runtime/ApplicationCore/Private/Mac/MacPlatformSplash.cpp
--- a/Engine/Source/Runtime/ApplicationCore/Private/Mac/MacPlatformSplash.cpp
+++ b/Engine/Source/Runtime/ApplicationCore/Private/Mac/MacPlatformSplash.cpp
@@ -1,11 +1,13 @@
 // Copyright Epic Games, Inc. All Rights Reserved.
 
 #include "Mac/MacPlatformSplash.h"
+#include "Apple/ScopeAutoreleasePool.h"
 #include "Misc/App.h"
 #include "Misc/ScopeLock.h"
 #include "Misc/EngineVersion.h"
 #include "Misc/EngineBuildSettings.h"
 #include "Mac/CocoaThread.h"
+
 #include "HAL/PlatformApplicationMisc.h"
 #include "Misc/ScopeLock.h"
 
diff --git a/Engine/Source/Runtime/ApplicationCore/Private/Mac/MacProgramDelegate.cpp b/Engine/Source/Runtime/ApplicationCore/Private/Mac/MacProgramDelegate.cpp
--- a/Engine/Source/Runtime/ApplicationCore/Private/Mac/MacProgramDelegate.cpp
+++ b/Engine/Source/Runtime/ApplicationCore/Private/Mac/MacProgramDelegate.cpp
@@ -1,6 +1,7 @@
 // Copyright Epic Games, Inc. All Rights Reserved.
 
 #include "Mac/MacProgramDelegate.h"
+#include "Apple/ScopeAutoreleasePool.h"
 #include "CoreMinimal.h"
 #include "HAL/ExceptionHandling.h"
 #include "Misc/CommandLine.h"
diff --git a/Engine/Source/Runtime/ApplicationCore/Private/Mac/MacTextInputMethodSystem.cpp b/Engine/Source/Runtime/ApplicationCore/Private/Mac/MacTextInputMethodSystem.cpp
--- a/Engine/Source/Runtime/ApplicationCore/Private/Mac/MacTextInputMethodSystem.cpp
+++ b/Engine/Source/Runtime/ApplicationCore/Private/Mac/MacTextInputMethodSystem.cpp
@@ -1,6 +1,7 @@
 // Copyright Epic Games, Inc. All Rights Reserved.
 
 #include "Mac/MacTextInputMethodSystem.h"
+#include "Apple/ScopeAutoreleasePool.h"
 #include "Mac/MacApplication.h"
 #include "Mac/MacWindow.h"
 #include "Mac/CocoaTextView.h"
diff --git a/Engine/Source/Runtime/ApplicationCore/Private/Mac/MacWindow.cpp b/Engine/Source/Runtime/ApplicationCore/Private/Mac/MacWindow.cpp
--- a/Engine/Source/Runtime/ApplicationCore/Private/Mac/MacWindow.cpp
+++ b/Engine/Source/Runtime/ApplicationCore/Private/Mac/MacWindow.cpp
@@ -1,11 +1,15 @@
 // Copyright Epic Games, Inc. All Rights Reserved.
 
 #include "Mac/MacWindow.h"
+
+#include "Apple/ScopeAutoreleasePool.h"
+#include "Apple/AppleStringUtils.h"
 #include "Mac/MacApplication.h"
 #include "Mac/MacCursor.h"
 #include "Mac/CocoaTextView.h"
 #include "Mac/CocoaThread.h"
 #include "Mac/MacPlatformApplicationMisc.h"
+
 #include "HAL/PlatformProcess.h"
 
 FMacWindow::FMacWindow()
diff --git a/Engine/Source/Runtime/ApplicationCore/Public/IOS/Accessibility/IOSAccessibilityCache.h b/Engine/Source/Runtime/ApplicationCore/Public/IOS/Accessibility/IOSAccessibilityCache.h
--- a/Engine/Source/Runtime/ApplicationCore/Public/IOS/Accessibility/IOSAccessibilityCache.h
+++ b/Engine/Source/Runtime/ApplicationCore/Public/IOS/Accessibility/IOSAccessibilityCache.h
@@ -6,6 +6,8 @@
 
 #include "GenericPlatform/Accessibility/GenericAccessibleInterfaces.h"
 
+#import <Foundation/Foundation.h> 
+
 @class FIOSAccessibilityLeaf;
 
 /**
diff --git a/Engine/Source/Runtime/ApplicationCore/Public/IOS/IOSApplication.h b/Engine/Source/Runtime/ApplicationCore/Public/IOS/IOSApplication.h
--- a/Engine/Source/Runtime/ApplicationCore/Public/IOS/IOSApplication.h
+++ b/Engine/Source/Runtime/ApplicationCore/Public/IOS/IOSApplication.h
@@ -5,6 +5,12 @@
 #include "GenericPlatform/Accessibility/GenericAccessibleInterfaces.h"
 #include "GenericPlatform/GenericApplication.h"
 
+#if PLATFORM_MAC
+#include "Mac/MacSystemIncludes.h"
+#elif PLATFORM_IOS
+#include "IOS/IOSSystemIncludes.h"
+#endif
+
 class FIOSInputInterface;
 class FIOSWindow;
 class IModularFeature;
diff --git a/Engine/Source/Runtime/ApplicationCore/Public/IOS/IOSConfigRules.h b/Engine/Source/Runtime/ApplicationCore/Public/IOS/IOSConfigRules.h
--- a/Engine/Source/Runtime/ApplicationCore/Public/IOS/IOSConfigRules.h
+++ b/Engine/Source/Runtime/ApplicationCore/Public/IOS/IOSConfigRules.h
@@ -5,6 +5,8 @@
 #include "CoreMinimal.h"
 #include "GenericPlatform/GenericConfigRules.h"
 
+#import <Foundation/NSObject.h>
+
 // iOSConfigRuleProvider represents the interface used to return the info required to load multiple config files
 // instances are registered with iOSConfigRuleProviders during +load time and queried during FIOSConfigRules::InitRules
 @protocol iOSConfigRuleProvider <NSObject>
diff --git a/Engine/Source/Runtime/ApplicationCore/Public/Mac/MacApplication.h b/Engine/Source/Runtime/ApplicationCore/Public/Mac/MacApplication.h
--- a/Engine/Source/Runtime/ApplicationCore/Public/Mac/MacApplication.h
+++ b/Engine/Source/Runtime/ApplicationCore/Public/Mac/MacApplication.h
@@ -2,6 +2,7 @@
 
 #pragma once
 
+#include "Apple/ScopeAutoreleasePool.h"
 #include "CoreMinimal.h"
 #include "GenericPlatform/GenericApplication.h"
 #include "Mac/MacWindow.h"
diff --git a/Engine/Source/Runtime/Core/Private/Apple/AppleLLM.cpp b/Engine/Source/Runtime/Core/Private/Apple/AppleLLM.cpp
--- a/Engine/Source/Runtime/Core/Private/Apple/AppleLLM.cpp
+++ b/Engine/Source/Runtime/Core/Private/Apple/AppleLLM.cpp
@@ -6,6 +6,7 @@
 #if ENABLE_LOW_LEVEL_MEM_TRACKER
 
 #include <objc/runtime.h>
+#import <Foundation/Foundation.h>
 
 struct FLLMTagInfoApple
 {
diff --git a/Engine/Source/Runtime/Core/Private/Apple/AppleMemoryObject.cpp b/Engine/Source/Runtime/Core/Private/Apple/AppleMemoryObject.cpp
new file mode 100644
--- /dev/null
+++ b/Engine/Source/Runtime/Core/Private/Apple/AppleMemoryObject.cpp
@@ -0,0 +1,148 @@
+// Copyright Epic Games, Inc. All Rights Reserved.
+
+#include "Apple/AppleMemoryObject.h"
+
+#include <stdlib.h>
+#include <sys/param.h>
+#include <sys/mman.h>
+#include <sys/mount.h>
+#include <objc/runtime.h>
+
+#include "CoreTypes.h"
+#if PLATFORM_MAC
+#include <libproc.h>
+#endif
+#if PLATFORM_IOS
+#include "IOS/IOSPlatformMisc.h"
+#include <os/proc.h>
+#endif
+#include <CoreFoundation/CFBase.h>
+#include <mach/vm_map.h>
+#include <mach/vm_region.h>
+#include <mach/vm_statistics.h>
+#include <mach/vm_page_size.h>
+#include "HAL/LowLevelMemTracker.h"
+#include "Apple/AppleLLM.h"
+
+#if PLATFORM_IOS
+#include <sys/resource.h> // On iOS, iPadOS and tvOS.
+extern "C" int proc_pid_rusage(int pid, int flavor, rusage_info_t *buffer)  
+	__OSX_AVAILABLE_STARTING(__MAC_10_9, __IPHONE_7_0);
+#endif
+
+NS_ASSUME_NONNULL_BEGIN
+
+/** 
+ * Zombie object implementation so that we can implement NSZombie behaviour for our custom allocated objects.
+ * Will leak memory - just like Cocoa's NSZombie - but allows for debugging of invalid usage of the pooled types.
+ */
+@interface FApplePlatformObjectZombie : NSObject 
+{
+	@public
+	Class OriginalClass;
+}
+@end
+
+@implementation FApplePlatformObjectZombie
+-(id)init
+{
+	self = (FApplePlatformObjectZombie*)[super init];
+	if (self)
+	{
+		OriginalClass = nil;
+	}
+	return self;
+}
+
+-(void)dealloc
+{
+	// Denied!
+	return;
+	
+	[super dealloc];
+}
+
+- (nullable NSMethodSignature *)methodSignatureForSelector:(SEL)sel
+{
+	NSLog(@"Selector %@ sent to deallocated instance %p of class %@", NSStringFromSelector(sel), self, OriginalClass);
+	abort();
+}
+@end
+
+@implementation FApplePlatformObject
+
++ (nullable OSQueueHead*)classAllocator
+{
+	return nullptr;
+}
+
++ (id)allocClass: (Class)NewClass
+{
+	static bool NSZombieEnabled = (getenv("NSZombieEnabled") != nullptr);
+	
+	// Allocate the correct size & zero it
+	// All allocations must be 16 byte aligned
+	SIZE_T Size = Align(FPlatformMath::Max(class_getInstanceSize(NewClass), class_getInstanceSize([FApplePlatformObjectZombie class])), 16);
+	void* Mem = nullptr;
+	
+	OSQueueHead* Alloc = [NewClass classAllocator];
+	if (Alloc && !NSZombieEnabled)
+	{
+		Mem = OSAtomicDequeue(Alloc, 0);
+		if (!Mem)
+		{
+			static uint8 BlocksPerChunk = 32;
+			char* Chunk = (char*)FMemory::Malloc(Size * BlocksPerChunk);
+			Mem = Chunk;
+			Chunk += Size;
+			for (uint8 i = 0; i < (BlocksPerChunk - 1); i++, Chunk += Size)
+			{
+				OSAtomicEnqueue(Alloc, Chunk, 0);
+			}
+		}
+	}
+	else
+	{
+		Mem = FMemory::Malloc(Size);
+	}
+	FMemory::Memzero(Mem, Size);
+	
+	// Construction assumes & requires zero-initialised memory
+	FApplePlatformObject* Obj = (FApplePlatformObject*)objc_constructInstance(NewClass, Mem);
+	object_setClass(Obj, NewClass);
+	Obj->AllocatorPtr = !NSZombieEnabled ? Alloc : nullptr;
+	return Obj;
+}
+
+- (void)dealloc
+{
+	static bool NSZombieEnabled = (getenv("NSZombieEnabled") != nullptr);
+	
+	// First call the destructor and then release the memory - like C++ placement new/delete
+	objc_destructInstance(self);
+	if (AllocatorPtr)
+	{
+		check(!NSZombieEnabled);
+		OSAtomicEnqueue(AllocatorPtr, self, 0);
+	}
+	else if (NSZombieEnabled)
+	{
+		Class CurrentClass = self.class;
+		object_setClass(self, [FApplePlatformObjectZombie class]);
+		FApplePlatformObjectZombie* ZombieSelf = (FApplePlatformObjectZombie*)self;
+		ZombieSelf->OriginalClass = CurrentClass;
+	}
+	else
+	{
+		FMemory::Free(self);
+	}
+	return;
+	
+	// Deliberately unreachable code to silence clang's error about not calling super - which in all other
+	// cases will be correct.
+	[super dealloc];
+}
+
+@end
+
+NS_ASSUME_NONNULL_END
diff --git a/Engine/Source/Runtime/Core/Private/Apple/ApplePlatformCrashContext.cpp b/Engine/Source/Runtime/Core/Private/Apple/ApplePlatformCrashContext.cpp
--- a/Engine/Source/Runtime/Core/Private/Apple/ApplePlatformCrashContext.cpp
+++ b/Engine/Source/Runtime/Core/Private/Apple/ApplePlatformCrashContext.cpp
@@ -10,6 +10,7 @@
 #include "HAL/FileManager.h"
 #include "Misc/Paths.h"
 #include "CoreGlobals.h"
+#include "Apple/AppleStringUtils.h"
 
 THIRD_PARTY_INCLUDES_START
 #if !PLATFORM_TVOS && !PLATFORM_VISIONOS
diff --git a/Engine/Source/Runtime/Core/Private/Apple/ApplePlatformDebugEvents.cpp b/Engine/Source/Runtime/Core/Private/Apple/ApplePlatformDebugEvents.cpp
--- a/Engine/Source/Runtime/Core/Private/Apple/ApplePlatformDebugEvents.cpp
+++ b/Engine/Source/Runtime/Core/Private/Apple/ApplePlatformDebugEvents.cpp
@@ -5,6 +5,8 @@
 =============================================================================*/
 
 #include "Apple/ApplePlatformDebugEvents.h"
+#include "Apple/AppleStringUtils.h"
+#include "Apple/AppleDebugEvent.h"
 #include "Apple/ApplePlatformTLS.h"
 #include "Containers/UnrealString.h"
 #include "HAL/IConsoleManager.h"
@@ -178,12 +180,12 @@ void FApplePlatformDebugEvents::DebugSignPostEnd(uint16 Code, uintptr_t Arg1, ui
 
 PRAGMA_ENABLE_DEPRECATION_WARNINGS
 
-TArray<FApplePlatformDebugEvents::FEvent>* FApplePlatformDebugEvents::GetEventStack()
+TArray<FAppleDebugEvent>* FApplePlatformDebugEvents::GetEventStack()
 {
-	TArray<FEvent>* Current = (TArray<FEvent>*)FPlatformTLS::GetTlsValue(TLSSlot);
+	TArray<FAppleDebugEvent>* Current = (TArray<FAppleDebugEvent>*)FPlatformTLS::GetTlsValue(TLSSlot);
 	if (!Current)
 	{
-		Current = new TArray<FEvent>;
+		Current = new TArray<FAppleDebugEvent>;
 		FPlatformTLS::SetTlsValue(TLSSlot, Current);
 	}
 	check(Current);
@@ -247,7 +249,7 @@ static void AppleBeginNamedEvent(const struct FColor& Color, const ANSICHAR* Tex
 	{
 		NSString* String = [FString(Text).GetNSString() retain];
 
-		FApplePlatformDebugEvents::FEvent Event;
+		FAppleDebugEvent Event;
 		Event.Tag = Text;
 		Event.Code = os_signpost_id_generate(GetLog());
 		Event.Destructor = Block_copy(^{ os_signpost_interval_end(GetLog(), Event.Code, "NamedEvent", "%s", [String UTF8String]); [String release]; });
@@ -272,7 +274,7 @@ void FApplePlatformDebugEvents::EndNamedEvent()
 {
 	if (GAppleInstrumentsEvents)
 	{
-		TArray<FEvent>* Stack = GetEventStack();
+		TArray<FAppleDebugEvent>* Stack = GetEventStack();
 		check(Stack);
 		if (Stack->Num() == 0)
 		{
@@ -281,7 +283,7 @@ void FApplePlatformDebugEvents::EndNamedEvent()
 			// anything to the event stack.
 			return;
 		}
-		FEvent Last = Stack->Last();
+		FAppleDebugEvent Last = Stack->Last();
 		Stack->RemoveAt(Stack->Num() - 1);
 		Last.Destructor();
 		Block_release(Last.Destructor);
diff --git a/Engine/Source/Runtime/Core/Private/Apple/ApplePlatformFile.cpp b/Engine/Source/Runtime/Core/Private/Apple/ApplePlatformFile.cpp
--- a/Engine/Source/Runtime/Core/Private/Apple/ApplePlatformFile.cpp
+++ b/Engine/Source/Runtime/Core/Private/Apple/ApplePlatformFile.cpp
@@ -6,6 +6,14 @@
 
 #include "Apple/ApplePlatformFile.h"
 
+#if PLATFORM_MAC
+#include "Apple/ScopeAutoreleasePool.h"
+#include "Mac/MacSystemIncludes.h"
+#elif PLATFORM_IOS
+#include "Apple/ScopeAutoreleasePool.h"
+#include "IOS/IOSSystemIncludes.h"
+#endif
+
 #include "Containers/UnrealString.h"
 #include "Containers/StringConv.h"
 #include "CoreGlobals.h"
diff --git a/Engine/Source/Runtime/Core/Private/Apple/ApplePlatformMemory.cpp b/Engine/Source/Runtime/Core/Private/Apple/ApplePlatformMemory.cpp
--- a/Engine/Source/Runtime/Core/Private/Apple/ApplePlatformMemory.cpp
+++ b/Engine/Source/Runtime/Core/Private/Apple/ApplePlatformMemory.cpp
@@ -5,6 +5,15 @@
 =============================================================================*/
 
 #include "Apple/ApplePlatformMemory.h"
+
+#include "Apple/AppleMemoryObject.h"
+
+#if PLATFORM_MAC
+#include "Mac/MacSystemIncludes.h"
+#elif PLATFORM_IOS
+#include "IOS/IOSSystemIncludes.h"
+#endif
+
 #include "HAL/PlatformMemory.h"
 #include "HAL/PlatformMath.h"
 #include "HAL/UnrealMemory.h"
@@ -66,119 +75,6 @@ static void* MmapWithTag(void* __nullable Addr, const size_t Len, const int Prot
 	return mmap(Addr, Len, Prot, Flags, (int32)Tag, Offset);
 }
 
-/** 
- * Zombie object implementation so that we can implement NSZombie behaviour for our custom allocated objects.
- * Will leak memory - just like Cocoa's NSZombie - but allows for debugging of invalid usage of the pooled types.
- */
-@interface FApplePlatformObjectZombie : NSObject 
-{
-	@public
-	Class OriginalClass;
-}
-@end
-
-@implementation FApplePlatformObjectZombie
--(id)init
-{
-	self = (FApplePlatformObjectZombie*)[super init];
-	if (self)
-	{
-		OriginalClass = nil;
-	}
-	return self;
-}
-
--(void)dealloc
-{
-	// Denied!
-	return;
-	
-	[super dealloc];
-}
-
-- (nullable NSMethodSignature *)methodSignatureForSelector:(SEL)sel
-{
-	NSLog(@"Selector %@ sent to deallocated instance %p of class %@", NSStringFromSelector(sel), self, OriginalClass);
-	abort();
-}
-@end
-
-@implementation FApplePlatformObject
-
-+ (nullable OSQueueHead*)classAllocator
-{
-	return nullptr;
-}
-
-+ (id)allocClass: (Class)NewClass
-{
-	static bool NSZombieEnabled = (getenv("NSZombieEnabled") != nullptr);
-	
-	// Allocate the correct size & zero it
-	// All allocations must be 16 byte aligned
-	SIZE_T Size = Align(FPlatformMath::Max(class_getInstanceSize(NewClass), class_getInstanceSize([FApplePlatformObjectZombie class])), 16);
-	void* Mem = nullptr;
-	
-	OSQueueHead* Alloc = [NewClass classAllocator];
-	if (Alloc && !NSZombieEnabled)
-	{
-		Mem = OSAtomicDequeue(Alloc, 0);
-		if (!Mem)
-		{
-			static uint8 BlocksPerChunk = 32;
-			char* Chunk = (char*)FMemory::Malloc(Size * BlocksPerChunk);
-			Mem = Chunk;
-			Chunk += Size;
-			for (uint8 i = 0; i < (BlocksPerChunk - 1); i++, Chunk += Size)
-			{
-				OSAtomicEnqueue(Alloc, Chunk, 0);
-			}
-		}
-	}
-	else
-	{
-		Mem = FMemory::Malloc(Size);
-	}
-	FMemory::Memzero(Mem, Size);
-	
-	// Construction assumes & requires zero-initialised memory
-	FApplePlatformObject* Obj = (FApplePlatformObject*)objc_constructInstance(NewClass, Mem);
-	object_setClass(Obj, NewClass);
-	Obj->AllocatorPtr = !NSZombieEnabled ? Alloc : nullptr;
-	return Obj;
-}
-
-- (void)dealloc
-{
-	static bool NSZombieEnabled = (getenv("NSZombieEnabled") != nullptr);
-	
-	// First call the destructor and then release the memory - like C++ placement new/delete
-	objc_destructInstance(self);
-	if (AllocatorPtr)
-	{
-		check(!NSZombieEnabled);
-		OSAtomicEnqueue(AllocatorPtr, self, 0);
-	}
-	else if (NSZombieEnabled)
-	{
-		Class CurrentClass = self.class;
-		object_setClass(self, [FApplePlatformObjectZombie class]);
-		FApplePlatformObjectZombie* ZombieSelf = (FApplePlatformObjectZombie*)self;
-		ZombieSelf->OriginalClass = CurrentClass;
-	}
-	else
-	{
-		FMemory::Free(self);
-	}
-	return;
-	
-	// Deliberately unreachable code to silence clang's error about not calling super - which in all other
-	// cases will be correct.
-	[super dealloc];
-}
-
-@end
-
 static void* FApplePlatformAllocatorAllocate(CFIndex AllocSize, CFOptionFlags Hint, void* Info)
 {
 	void* Mem = FMemory::Malloc(AllocSize, 16);
diff --git a/Engine/Source/Runtime/Core/Private/Apple/ApplePlatformMisc.cpp b/Engine/Source/Runtime/Core/Private/Apple/ApplePlatformMisc.cpp
--- a/Engine/Source/Runtime/Core/Private/Apple/ApplePlatformMisc.cpp
+++ b/Engine/Source/Runtime/Core/Private/Apple/ApplePlatformMisc.cpp
@@ -5,6 +5,11 @@
 =============================================================================*/
 
 #include "Apple/ApplePlatformMisc.h"
+#if PLATFORM_MAC
+#include "Mac/MacSystemIncludes.h"
+#elif PLATFORM_IOS
+#include "IOS/IOSSystemIncludes.h"
+#endif
 #include "HAL/ExceptionHandling.h"
 #include "Misc/SecureHash.h"
 #include "Misc/CommandLine.h"
@@ -41,6 +46,11 @@ bool FApplePlatformMisc::IsDebuggerPresent()
 
 #endif // !UE_BUILD_SHIPPING
 
+void FApplePlatformMisc::MemoryBarrier()
+{
+	__sync_synchronize();
+}
+
 void FApplePlatformMisc::GetEnvironmentVariable(const TCHAR* VariableName, TCHAR* Result, int32 ResultLength)
 {
 	// Replace hyphens with underscores. Some legacy UE environment variables (eg. UE-SharedDataCachePath) are in widespread
diff --git a/Engine/Source/Runtime/Core/Private/Apple/ApplePlatformProcess.cpp b/Engine/Source/Runtime/Core/Private/Apple/ApplePlatformProcess.cpp
--- a/Engine/Source/Runtime/Core/Private/Apple/ApplePlatformProcess.cpp
+++ b/Engine/Source/Runtime/Core/Private/Apple/ApplePlatformProcess.cpp
@@ -2,6 +2,9 @@
 
 #include "Apple/ApplePlatformProcess.h"
 
+#include "Apple/AppleStringUtils.h"
+#include "Apple/ScopeAutoreleasePool.h"
+
 #include "HAL/PlatformMisc.h"
 
 const TCHAR* FApplePlatformProcess::UserDir()
diff --git a/Engine/Source/Runtime/Core/Private/Apple/ApplePlatformString.cpp b/Engine/Source/Runtime/Core/Private/Apple/ApplePlatformString.cpp
--- a/Engine/Source/Runtime/Core/Private/Apple/ApplePlatformString.cpp
+++ b/Engine/Source/Runtime/Core/Private/Apple/ApplePlatformString.cpp
@@ -8,17 +8,26 @@
 #include "Containers/UnrealString.h"
 #include "Containers/StringConv.h"
 
-@implementation NSString (FString_Extensions)
+#if PLATFORM_MAC
+#include "Mac/MacSystemIncludes.h"
+#elif PLATFORM_IOS
+#include "IOS/IOSSystemIncludes.h"
+#endif
 
-+ (NSString*) stringWithTCHARString:(const TCHAR*)MyTCHARString
+void FApplePlatformString::CFStringToTCHAR( CFStringRef CFStr, TCHAR *TChar )
 {
-	return [NSString stringWithCString:TCHAR_TO_UTF8(MyTCHARString) encoding:NSUTF8StringEncoding];
+	const SIZE_T Length = CFStringGetLength( CFStr );
+	CFRange Range = CFRangeMake( 0, Length );
+	CFStringGetBytes( CFStr, Range, sizeof( TCHAR ) == 4 ? kCFStringEncodingUTF32LE : kCFStringEncodingUnicode, '?', false, ( uint8 *)TChar, Length * sizeof( TCHAR ) + 1, NULL );
+	TChar[Length] = 0;
 }
 
-+ (NSString*) stringWithFString:(const FString&)InFString
+CFStringRef FApplePlatformString::TCHARToCFString( const TCHAR *TChar )
 {
-	return [NSString stringWithTCHARString:*InFString];
+	const SIZE_T Length = Strlen( TChar );
+	CFStringRef String = CFStringCreateWithBytes( kCFAllocatorDefault, ( const uint8 *)TChar, Length * sizeof( TCHAR ), sizeof( TCHAR ) == 4 ? kCFStringEncodingUTF32LE : kCFStringEncodingUnicode, false );
+	// we are getting some crashes on strings that aren't converting - so, instead of crashing, print them out instead
+	ensureMsgf(String, TEXT("Failed to allocated CFString for '%s' -- Length id %d"), TChar ? TChar : TEXT("nullptr"), Length);
+	return String;
 }
 
-
-@end
diff --git a/Engine/Source/Runtime/Core/Private/Apple/ApplePlatformSymbolication.cpp b/Engine/Source/Runtime/Core/Private/Apple/ApplePlatformSymbolication.cpp
--- a/Engine/Source/Runtime/Core/Private/Apple/ApplePlatformSymbolication.cpp
+++ b/Engine/Source/Runtime/Core/Private/Apple/ApplePlatformSymbolication.cpp
@@ -14,6 +14,8 @@
 #include <CoreFoundation/CoreFoundation.h>
 #include <mach/mach.h>
 
+#include "Apple/AppleStringUtils.h"
+
 extern "C"
 {
 	struct CSRange
diff --git a/Engine/Source/Runtime/Core/Private/Apple/ApplePlatformTLS.cpp b/Engine/Source/Runtime/Core/Private/Apple/ApplePlatformTLS.cpp
new file mode 100644
--- /dev/null
+++ b/Engine/Source/Runtime/Core/Private/Apple/ApplePlatformTLS.cpp
@@ -0,0 +1,40 @@
+// Copyright Epic Games, Inc. All Rights Reserved.
+
+#include "Apple/ApplePlatformTLS.h"
+
+#if PLATFORM_MAC
+#include "Mac/MacSystemIncludes.h"
+#elif PLATFORM_IOS
+#include "IOS/IOSSystemIncludes.h"
+#endif
+
+uint32 FApplePlatformTLS::GetCurrentThreadId()
+{
+	return (uint32)pthread_mach_thread_np(pthread_self());
+}
+
+uint32 FApplePlatformTLS::AllocTlsSlot()
+{
+	// allocate a per-thread mem slot
+	pthread_key_t SlotKey = 0;
+	if (pthread_key_create(&SlotKey, NULL) != 0)
+	{
+		SlotKey = InvalidTlsSlot;  // matches the Windows TlsAlloc() retval.
+	}
+	return SlotKey;
+}
+
+void FApplePlatformTLS::SetTlsValue(uint32 SlotIndex, void* Value)
+{
+	pthread_setspecific((pthread_key_t)SlotIndex, Value);
+}
+
+void* FApplePlatformTLS::GetTlsValue(uint32 SlotIndex)
+{
+	return pthread_getspecific((pthread_key_t)SlotIndex);
+}
+
+void FApplePlatformTLS::FreeTlsSlot(uint32 SlotIndex)
+{
+	pthread_key_delete((pthread_key_t)SlotIndex);
+}
diff --git a/Engine/Source/Runtime/Core/Private/Apple/ApplePlatformTime.cpp b/Engine/Source/Runtime/Core/Private/Apple/ApplePlatformTime.cpp
--- a/Engine/Source/Runtime/Core/Private/Apple/ApplePlatformTime.cpp
+++ b/Engine/Source/Runtime/Core/Private/Apple/ApplePlatformTime.cpp
@@ -6,9 +6,18 @@
 
 #include "Apple/ApplePlatformTime.h"
 #include "HAL/PlatformTime.h"
+#if PLATFORM_MAC
+#include "Mac/MacSystemIncludes.h"
+#elif PLATFORM_IOS
+#include "IOS/IOSSystemIncludes.h"
+#endif
 #include "Misc/AssertionMacros.h"
 #include "CoreGlobals.h"
 
+#include "CoreTypes.h"
+
+
+
 namespace
 {
 	FORCEINLINE double TimeValToSecond(timeval & tv)
@@ -27,6 +36,25 @@ double FApplePlatformTime::InitTiming(void)
 	return FPlatformTime::Seconds();
 }
 
+double FApplePlatformTime::Seconds()
+{
+	uint64 Cycles = mach_absolute_time();
+	// Add big number to make bugs apparent where return value is being passed to float
+	return Cycles * GetSecondsPerCycle() + GetSecondsTimeOffset();
+}
+
+uint32 FApplePlatformTime::Cycles()
+{
+	uint64 Cycles = mach_absolute_time();
+	return Cycles;
+}
+
+uint64 FApplePlatformTime::Cycles64()
+{
+	uint64 Cycles = mach_absolute_time();
+	return Cycles;
+}
+
 FCPUTime FApplePlatformTime::GetCPUTime()
 {
 	// minimum delay between checks to minimize overhead (and also match Windows version)
diff --git a/Engine/Source/Runtime/Core/Private/Apple/AppleStringUtils.cpp b/Engine/Source/Runtime/Core/Private/Apple/AppleStringUtils.cpp
new file mode 100644
--- /dev/null
+++ b/Engine/Source/Runtime/Core/Private/Apple/AppleStringUtils.cpp
@@ -0,0 +1,38 @@
+// Copyright Epic Games, Inc. All Rights Reserved.
+
+#include "Apple/AppleStringUtils.h"
+#if PLATFORM_APPLE
+
+#include "Containers/UnrealString.h"
+
+#include "Containers/StringConv.h"
+
+@implementation NSString (FString_Extensions)
+
++ (NSString*) stringWithTCHARString:(const TCHAR*)MyTCHARString
+{
+	return [NSString stringWithCString:TCHAR_TO_UTF8(MyTCHARString) encoding:NSUTF8StringEncoding];
+}
+
++ (NSString*) stringWithFString:(const FString&)InFString
+{
+	return [NSString stringWithTCHARString:*InFString];
+}
+
+
+@end
+
+NSString* FAppleStringUtils::ConvertToNSString(const FString& InString)
+{
+	NSString* OutString = (NSString*)(InString.GetCFString());
+	[OutString autorelease];
+        
+	return OutString;
+}
+
+FString FAppleStringUtils::ConvertToFString(const  NSString* InNSString)
+{
+	return FString((__bridge CFStringRef)InNSString);
+}
+
+#endif
\ No newline at end of file
diff --git a/Engine/Source/Runtime/Core/Private/Containers/String.cpp.inl b/Engine/Source/Runtime/Core/Private/Containers/String.cpp.inl
--- a/Engine/Source/Runtime/Core/Private/Containers/String.cpp.inl
+++ b/Engine/Source/Runtime/Core/Private/Containers/String.cpp.inl
@@ -11,6 +11,10 @@
  *********************************************************************************************************************/
 #define UE_INCLUDETOOL_IGNORE_INCONSISTENT_STATE
 
+#if PLATFORM_MAC
+#include <MacTypes.h>
+#endif
+
 #ifndef UE_STRING_CLASS
 	#error "String.cpp.inl should only be included after defining UE_STRING_CLASS"
 #endif
@@ -21,6 +25,10 @@
 	#error "String.cpp.inl should only be included after defining UE_STRING_CHARTYPE_IS_TCHAR"
 #endif
 
+#if PLATFORM_APPLE
+#import <CoreFoundation/CoreFoundation.h>
+#endif
+
  /* String implementation
  *****************************************************************************/
 
@@ -352,6 +360,34 @@ NSString* UE_STRING_CLASS::GetNSString() const
 }
 #endif
 
+#if PLATFORM_APPLE && UE_STRING_CHARTYPE_IS_TCHAR
+
+#include "Apple/AppleStringUtils.h"
+#include "Misc/AssertionMacros.h"
+UE_STRING_CLASS::UE_STRING_CLASS(const CFStringRef In)
+{
+	uint32_t StringLength = In ? CFStringGetLength(In) : 0;
+
+	if (StringLength > 0)
+	{
+		// Convert the NSString data into the native TCHAR format for UE4
+		// This returns a buffer of bytes, but they can be safely cast to a buffer of TCHARs
+		const CFStringEncoding Encoding = kCFStringEncodingUTF16LE;
+		CFRange Range = CFRangeMake(0, StringLength);
+		CFIndex BytesNeeded;
+		if (CFStringGetBytes(In, Range, Encoding, '?', false, NULL, 0, &BytesNeeded) > 0)
+		{
+			const size_t Length = BytesNeeded / sizeof(TCHAR);
+			Data.Reserve(Length + 1);
+			Data.AddUninitialized(Length + 1);
+			CFStringGetBytes(In, Range, Encoding, '?', false, (uint8*)Data.GetData(), Length * sizeof(TCHAR) + 1, NULL);
+			Data[Length] = 0;
+		}
+	}
+}
+
+#endif
+
 #if PLATFORM_APPLE
 /** Convert string to bridgeable CFString */
 CFStringRef UE_STRING_CLASS::GetCFString() const
diff --git a/Engine/Source/Runtime/Core/Private/HAL/MallocBinned.cpp b/Engine/Source/Runtime/Core/Private/HAL/MallocBinned.cpp
--- a/Engine/Source/Runtime/Core/Private/HAL/MallocBinned.cpp
+++ b/Engine/Source/Runtime/Core/Private/HAL/MallocBinned.cpp
@@ -1,6 +1,12 @@
 // Copyright Epic Games, Inc. All Rights Reserved.
 
 #include "HAL/MallocBinned.h"
+#if PLATFORM_MAC
+#include "Mac/MacSystemIncludes.h"
+#elif PLATFORM_IOS
+#include "IOS/IOSSystemIncludes.h"
+#endif
+
 #include "Async/UniqueLock.h"
 #include "Misc/BufferedOutputDevice.h"
 
diff --git a/Engine/Source/Runtime/Core/Private/IOS/IOSChunkInstaller.cpp b/Engine/Source/Runtime/Core/Private/IOS/IOSChunkInstaller.cpp
--- a/Engine/Source/Runtime/Core/Private/IOS/IOSChunkInstaller.cpp
+++ b/Engine/Source/Runtime/Core/Private/IOS/IOSChunkInstaller.cpp
@@ -1,9 +1,9 @@
 // Copyright Epic Games, Inc. All Rights Reserved.
 
 #include "IOSChunkInstaller.h"
-#include "Misc/CallbackDevice.h"
-
 
+#include "Apple/AppleStringUtils.h"
+#include "Misc/CallbackDevice.h"
 
 FIOSChunkInstall::FIOSChunkInstall()
 {
diff --git a/Engine/Source/Runtime/Core/Private/IOS/IOSErrorOutputDevice.cpp b/Engine/Source/Runtime/Core/Private/IOS/IOSErrorOutputDevice.cpp
--- a/Engine/Source/Runtime/Core/Private/IOS/IOSErrorOutputDevice.cpp
+++ b/Engine/Source/Runtime/Core/Private/IOS/IOSErrorOutputDevice.cpp
@@ -11,6 +11,7 @@
 #include "Misc/OutputDeviceRedirector.h"
 #include "CoreGlobals.h"
 #include <cxxabi.h>
+#import <Foundation/Foundation.h>
 
 FIOSErrorOutputDevice::FIOSErrorOutputDevice()
 :	ErrorPos(0)
diff --git a/Engine/Source/Runtime/Core/Private/IOS/IOSMemoryTrace.cpp b/Engine/Source/Runtime/Core/Private/IOS/IOSMemoryTrace.cpp
--- a/Engine/Source/Runtime/Core/Private/IOS/IOSMemoryTrace.cpp
+++ b/Engine/Source/Runtime/Core/Private/IOS/IOSMemoryTrace.cpp
@@ -15,6 +15,8 @@
 #include "Misc/CString.h"
 #include "Trace/Trace.inl"
 
+#import <Foundation/Foundation.h>
+
 
 ////////////////////////////////////////////////////////////////////////////////
 FMalloc* MemoryTrace_CreateInternal(FMalloc*, int, const ANSICHAR* const*);
diff --git a/Engine/Source/Runtime/Core/Private/IOS/IOSModuleDiagnostics.cpp b/Engine/Source/Runtime/Core/Private/IOS/IOSModuleDiagnostics.cpp
--- a/Engine/Source/Runtime/Core/Private/IOS/IOSModuleDiagnostics.cpp
+++ b/Engine/Source/Runtime/Core/Private/IOS/IOSModuleDiagnostics.cpp
@@ -1,6 +1,7 @@
 ﻿// Copyright Epic Games, Inc. All Rights Reserved.
 
 #include "CoreTypes.h"
+#include "Apple/AppleStringUtils.h"
 
 #if UE_BUILD_SHIPPING
 
diff --git a/Engine/Source/Runtime/Core/Private/IOS/IOSPlatformCrashContext.cpp b/Engine/Source/Runtime/Core/Private/IOS/IOSPlatformCrashContext.cpp
--- a/Engine/Source/Runtime/Core/Private/IOS/IOSPlatformCrashContext.cpp
+++ b/Engine/Source/Runtime/Core/Private/IOS/IOSPlatformCrashContext.cpp
@@ -4,6 +4,14 @@
 	IOSPlatformCrashContext.cpp: implementations of IOS platform crash context.
 =============================================================================*/
 #include "IOS/IOSPlatformCrashContext.h"
+#include "Apple/ScopeAutoreleasePool.h"
+#include "Apple/AppleStringUtils.h"
+
+#if PLATFORM_MAC
+#include "Mac/MacSystemIncludes.h"
+#elif PLATFORM_IOS
+#include "IOS/IOSSystemIncludes.h"
+#endif
 
 #include "IOS/IOSPlatformPLCrashReporterIncludes.h"
 
@@ -307,8 +315,8 @@ void FIOSApplicationInfo::Init()
 	
 	RunUUID = RunGUID();
 	
-	OSXVersion = [[NSProcessInfo processInfo] operatingSystemVersion];
-	OSVersion = FString::Printf(TEXT("%ld.%ld.%ld"), OSXVersion.majorVersion, OSXVersion.minorVersion, OSXVersion.patchVersion);
+	NSOperatingSystemVersion CurrentOSXVersion = [[NSProcessInfo processInfo] operatingSystemVersion];
+	OSVersion = FString::Printf(TEXT("%ld.%ld.%ld"), CurrentOSXVersion.majorVersion, CurrentOSXVersion.minorVersion, CurrentOSXVersion.patchVersion);
 	FCStringAnsi::Strncpy(OSVersionUTF8, TCHAR_TO_UTF8(*OSVersion), PATH_MAX + 1);
 	
 	// macOS build number is only accessible on non-sandboxed applications as it resides outside the accessible sandbox
diff --git a/Engine/Source/Runtime/Core/Private/IOS/IOSPlatformFile.cpp b/Engine/Source/Runtime/Core/Private/IOS/IOSPlatformFile.cpp
--- a/Engine/Source/Runtime/Core/Private/IOS/IOSPlatformFile.cpp
+++ b/Engine/Source/Runtime/Core/Private/IOS/IOSPlatformFile.cpp
@@ -1,6 +1,15 @@
 // Copyright Epic Games, Inc. All Rights Reserved.
 
 #include "IOS/IOSPlatformFile.h"
+
+#include "Apple/AppleStringUtils.h"
+#include "Apple/ScopeAutoreleasePool.h"
+#if PLATFORM_MAC
+#include "Mac/MacSystemIncludes.h"
+#elif PLATFORM_IOS
+#include "IOS/IOSSystemIncludes.h"
+#endif
+
 #include "Apple/ApplePlatformMisc.h"
 #include "HAL/PlatformTLS.h"
 #include "Containers/StringConv.h"
diff --git a/Engine/Source/Runtime/Core/Private/IOS/IOSPlatformMisc.cpp b/Engine/Source/Runtime/Core/Private/IOS/IOSPlatformMisc.cpp
--- a/Engine/Source/Runtime/Core/Private/IOS/IOSPlatformMisc.cpp
+++ b/Engine/Source/Runtime/Core/Private/IOS/IOSPlatformMisc.cpp
@@ -5,6 +5,7 @@
 =============================================================================*/
 
 #include "IOS/IOSPlatformMisc.h"
+#include "Apple/AppleStringUtils.h"
 #include "Apple/ApplePlatformMisc.h"
 #include "Async/TaskGraphInterfaces.h"
 #include "Apple/ApplePlatformCrashContext.h"
@@ -39,6 +40,7 @@
 #include "Containers/StringConv.h"
 #include "Containers/StringView.h"
 #include "Apple/PreAppleSystemHeaders.h"
+#include "IOS/IOSSystemIncludes.h"
 
 #if !PLATFORM_TVOS && !PLATFORM_VISIONOS
 #include <AdSupport/ASIdentifierManager.h>
@@ -60,6 +62,7 @@
 #include <sys/sysctl.h> // sysctlbyname
 
 #include "Apple/PostAppleSystemHeaders.h"
+#include "Apple/ScopeAutoreleasePool.h"
 
 #include "IOSConfigRules.h"
 #include <Metal/Metal.h>
@@ -68,6 +71,8 @@
 	#define ENABLE_ADVERTISING_IDENTIFIER 0
 #endif
 
+#include "Misc/AssertionMacros.h"
+
 /** Amount of free memory in MB reported by the system at startup */
 CORE_API int32 GStartupFreeMemoryMB;
 
diff --git a/Engine/Source/Runtime/Core/Private/IOS/IOSPlatformProcess.cpp b/Engine/Source/Runtime/Core/Private/IOS/IOSPlatformProcess.cpp
--- a/Engine/Source/Runtime/Core/Private/IOS/IOSPlatformProcess.cpp
+++ b/Engine/Source/Runtime/Core/Private/IOS/IOSPlatformProcess.cpp
@@ -5,7 +5,18 @@ IOSPlatformProcess.cpp: iOS implementations of Process functions
 =============================================================================*/
 
 #include "IOS/IOSPlatformProcess.h"
+
+#include "Apple/ScopeAutoreleasePool.h"
+#include "Apple/AppleStringUtils.h"
+
+#if PLATFORM_MAC
+#include "Mac/MacSystemIncludes.h"
+#elif PLATFORM_IOS
+#include "IOS/IOSSystemIncludes.h"
+#endif
+
 #include "Apple/ApplePlatformRunnableThread.h"
+#include "IOS/IOSAsyncTask.h"
 #include "Misc/CoreDelegates.h"
 #include "Apple/PreAppleSystemHeaders.h"
 #include <mach-o/dyld.h>
diff --git a/Engine/Source/Runtime/Core/Private/Mac/CocoaThread.cpp b/Engine/Source/Runtime/Core/Private/Mac/CocoaThread.cpp
--- a/Engine/Source/Runtime/Core/Private/Mac/CocoaThread.cpp
+++ b/Engine/Source/Runtime/Core/Private/Mac/CocoaThread.cpp
@@ -1,6 +1,8 @@
 // Copyright Epic Games, Inc. All Rights Reserved.
 
 #include "Mac/CocoaThread.h"
+#include "Apple/ScopeAutoreleasePool.h"
+
 #include "Containers/LockFreeList.h"
 #include "HAL/IConsoleManager.h"
 #include "Templates/RefCounting.h"
diff --git a/Engine/Source/Runtime/Core/Private/Mac/MacPlatformMisc.cpp b/Engine/Source/Runtime/Core/Private/Mac/MacPlatformMisc.cpp
--- a/Engine/Source/Runtime/Core/Private/Mac/MacPlatformMisc.cpp
+++ b/Engine/Source/Runtime/Core/Private/Mac/MacPlatformMisc.cpp
@@ -48,6 +48,9 @@
 #include <notify.h>
 #include <uuid/uuid.h>
 #include <spawn.h>
+
+#include "Apple/ScopeAutoreleasePool.h"
+#include "Apple/AppleStringUtils.h"
 #include "Apple/PostAppleSystemHeaders.h"
 
 extern CORE_API bool GIsGPUCrashed;
diff --git a/Engine/Source/Runtime/Core/Private/Mac/MacPlatformProcess.cpp b/Engine/Source/Runtime/Core/Private/Mac/MacPlatformProcess.cpp
--- a/Engine/Source/Runtime/Core/Private/Mac/MacPlatformProcess.cpp
+++ b/Engine/Source/Runtime/Core/Private/Mac/MacPlatformProcess.cpp
@@ -20,6 +20,9 @@
 #include <mach/thread_policy.h>
 #include <libproc.h>
 #include <spawn.h>
+
+#include "Apple/ScopeAutoreleasePool.h"
+#include "Apple/AppleStringUtils.h"
 #include "Apple/PostAppleSystemHeaders.h"
 #if PLATFORM_MAC_X86
     #include <cpuid.h>
diff --git a/Engine/Source/Runtime/Core/Private/ProfilingDebugging/Apple/AppleMemoryProbe.cpp b/Engine/Source/Runtime/Core/Private/ProfilingDebugging/Apple/AppleMemoryProbe.cpp
--- a/Engine/Source/Runtime/Core/Private/ProfilingDebugging/Apple/AppleMemoryProbe.cpp
+++ b/Engine/Source/Runtime/Core/Private/ProfilingDebugging/Apple/AppleMemoryProbe.cpp
@@ -8,6 +8,12 @@
 
 #if UE_MEMORY_TRACE_ENABLED
 
+#pragma push_macro("FVector")
+#define FVector FVectorWorkaround
+#import <Foundation/Foundation.h>
+#undef FVector
+#pragma pop_macro("FVector")
+
 #include <Containers/Set.h>
 #include <mach/mach.h>
 #include <sys/types.h>
diff --git a/Engine/Source/Runtime/Core/Private/ProfilingDebugging/Apple/InstrumentsExternalProfiler.cpp b/Engine/Source/Runtime/Core/Private/ProfilingDebugging/Apple/InstrumentsExternalProfiler.cpp
--- a/Engine/Source/Runtime/Core/Private/ProfilingDebugging/Apple/InstrumentsExternalProfiler.cpp
+++ b/Engine/Source/Runtime/Core/Private/ProfilingDebugging/Apple/InstrumentsExternalProfiler.cpp
@@ -1,6 +1,13 @@
 // Copyright Epic Games, Inc. All Rights Reserved.
 
 #include "CoreTypes.h"
+
+#pragma push_macro("FVector")
+#define FVector FVectorWorkaround
+#import <Foundation/Foundation.h>
+#undef FVector
+#pragma pop_macro("FVector")
+
 #include "HAL/PlatformMisc.h"
 #include "HAL/PlatformProcess.h"
 #include "Misc/AssertionMacros.h"
diff --git a/Engine/Source/Runtime/Core/Public/Apple/AppleDebugEvent.h b/Engine/Source/Runtime/Core/Public/Apple/AppleDebugEvent.h
new file mode 100644
--- /dev/null
+++ b/Engine/Source/Runtime/Core/Public/Apple/AppleDebugEvent.h
@@ -0,0 +1,17 @@
+// Copyright Epic Games, Inc. All Rights Reserved.
+
+#pragma once
+
+#if PLATFORM_MAC
+#include "Mac/MacSystemIncludes.h"
+#elif PLATFORM_IOS
+#include "IOS/IOSSystemIncludes.h"
+#endif
+
+struct FAppleDebugEvent
+{
+	void const* Tag;
+	uint32 Color;
+	uint64 Code;
+	dispatch_block_t Destructor;
+};
\ No newline at end of file
diff --git a/Engine/Source/Runtime/Core/Public/Apple/AppleMemoryObject.h b/Engine/Source/Runtime/Core/Public/Apple/AppleMemoryObject.h
new file mode 100644
--- /dev/null
+++ b/Engine/Source/Runtime/Core/Public/Apple/AppleMemoryObject.h
@@ -0,0 +1,58 @@
+// Copyright Epic Games, Inc. All Rights Reserved.
+
+#pragma once
+
+#include <libkern/OSAtomic.h>
+#include <objc/objc.h>
+#include <Foundation/NSObject.h>
+
+#include <mach/mach.h>
+
+NS_ASSUME_NONNULL_BEGIN
+
+
+// Support for passing LLM alloc data to XCode Instruments.
+// Use USE_APPLE_SUPPORT_INSTRUMENTED_ALLOCS to enable it.
+// As APPLE_SUPPORT_INSTRUMENTED_ALLOCS define only enables the LLM API, because changes to this header trigger wide recompilation.
+#ifndef APPLE_SUPPORT_INSTRUMENTED_ALLOCS
+	#define APPLE_SUPPORT_INSTRUMENTED_ALLOCS (!UE_BUILD_SHIPPING)
+#endif // APPLE_SUPPORT_INSTRUMENTED_ALLOCS
+
+/**
+ * NSObject subclass that can be used to override the allocation functions to go through UE4's memory allocator.
+ * This ensures that memory allocated by custom Objective-C types can be tracked by UE4's tools and 
+ * that we benefit from the memory allocator's efficiencies.
+ */
+OBJC_EXPORT @interface FApplePlatformObject : NSObject
+{
+	@private
+		OSQueueHead* AllocatorPtr;
+}
+
+/** Sub-classes should override to provide the OSQueueHead* necessary to allocate from - handled by the macro */
++ (nullable OSQueueHead*)classAllocator;
+
+/** Sub-classes should override allocWithZone & alloc to call allocClass */
++ (id)allocClass: (Class)NewClass;
+
+/** Override the core NSObject deallocation function to correctly destruct */
+- (void)dealloc;
+
+@end
+
+#define APPLE_PLATFORM_OBJECT_ALLOC_OVERRIDES(ClassName)		\
++ (nullable OSQueueHead*)classAllocator							\
+{																\
+static OSQueueHead Queue = OS_ATOMIC_QUEUE_INIT;			\
+return &Queue;												\
+}																\
++ (id)allocWithZone:(NSZone*) Zone								\
+{																\
+return (ClassName*)[FApplePlatformObject allocClass:self];	\
+}																\
++ (id)alloc														\
+{																\
+return (ClassName*)[FApplePlatformObject allocClass:self];	\
+}
+
+NS_ASSUME_NONNULL_END
\ No newline at end of file
diff --git a/Engine/Source/Runtime/Core/Public/Apple/ApplePlatformDebugEvents.h b/Engine/Source/Runtime/Core/Public/Apple/ApplePlatformDebugEvents.h
--- a/Engine/Source/Runtime/Core/Public/Apple/ApplePlatformDebugEvents.h
+++ b/Engine/Source/Runtime/Core/Public/Apple/ApplePlatformDebugEvents.h
@@ -9,12 +9,17 @@
 #include "Logging/LogMacros.h"
 #include "Containers/Map.h"
 #include "Containers/UnrealString.h"
+
+#if UE_ENABLE_INCLUDE_ORDER_DEPRECATED_IN_5_8
+
 #if PLATFORM_MAC
 #include "Mac/MacSystemIncludes.h"
 #elif PLATFORM_IOS
 #include "IOS/IOSSystemIncludes.h"
 #endif
 
+#endif // UE_ENABLE_INCLUDE_ORDER_DEPRECATED_IN_5_8
+
 #if !defined(APPLE_PROFILING_ENABLED)
 #define APPLE_PROFILING_ENABLED 0
 #endif
@@ -23,6 +28,8 @@
 
 DECLARE_LOG_CATEGORY_EXTERN(LogInstruments, Display, All);
 
+struct FAppleDebugEvent;
+
 struct FApplePlatformDebugEvents
 {
 	static void DebugSignPost(uint16 Code, uintptr_t Arg1, uintptr_t Arg2, uintptr_t Arg3, uintptr_t Arg4);
@@ -31,6 +38,7 @@ struct FApplePlatformDebugEvents
 
 	static void DebugSignPostEnd(uint16 Code, uintptr_t Arg1, uintptr_t Arg2, uintptr_t Arg3, uintptr_t Arg4);
 	
+#if UE_ENABLE_INCLUDE_ORDER_DEPRECATED_IN_5_8
 	struct FEvent
 	{
 		void const* Tag;
@@ -38,8 +46,9 @@ struct FApplePlatformDebugEvents
 		uint64 Code;
 		dispatch_block_t Destructor;
 	};
+#endif // UE_ENABLE_INCLUDE_ORDER_DEPRECATED_IN_5_8
 	
-	static TArray<FEvent>* GetEventStack();
+	static TArray<FAppleDebugEvent>* GetEventStack();
 	
 	static uint16 GetEventCode(FString String);
 	
diff --git a/Engine/Source/Runtime/Core/Public/Apple/ApplePlatformMemory.h b/Engine/Source/Runtime/Core/Public/Apple/ApplePlatformMemory.h
--- a/Engine/Source/Runtime/Core/Public/Apple/ApplePlatformMemory.h
+++ b/Engine/Source/Runtime/Core/Public/Apple/ApplePlatformMemory.h
@@ -7,13 +7,13 @@
 
 #pragma once
 #include "GenericPlatform/GenericPlatformMemory.h"
+
+#if UE_ENABLE_INCLUDE_ORDER_DEPRECATED_IN_5_8
 #include <libkern/OSAtomic.h>
 #include <Foundation/NSObject.h>
 
 #include <mach/mach.h>
-
-NS_ASSUME_NONNULL_BEGIN
-
+#endif // UE_ENABLE_INCLUDE_ORDER_DEPRECATED_IN_5_8
 
 // Support for passing LLM alloc data to XCode Instruments.
 // Use USE_APPLE_SUPPORT_INSTRUMENTED_ALLOCS to enable it.
@@ -22,42 +22,6 @@ NS_ASSUME_NONNULL_BEGIN
 	#define APPLE_SUPPORT_INSTRUMENTED_ALLOCS (!UE_BUILD_SHIPPING)
 #endif // APPLE_SUPPORT_INSTRUMENTED_ALLOCS
 
-/**
- * NSObject subclass that can be used to override the allocation functions to go through UE4's memory allocator.
- * This ensures that memory allocated by custom Objective-C types can be tracked by UE4's tools and 
- * that we benefit from the memory allocator's efficiencies.
- */
-OBJC_EXPORT @interface FApplePlatformObject : NSObject
-{
-@private
-	OSQueueHead* AllocatorPtr;
-}
-
-/** Sub-classes should override to provide the OSQueueHead* necessary to allocate from - handled by the macro */
-+ (nullable OSQueueHead*)classAllocator;
-
-/** Sub-classes should override allocWithZone & alloc to call allocClass */
-+ (id)allocClass: (Class)NewClass;
-
-/** Override the core NSObject deallocation function to correctly destruct */
-- (void)dealloc;
-
-@end
-
-#define APPLE_PLATFORM_OBJECT_ALLOC_OVERRIDES(ClassName)		\
-+ (nullable OSQueueHead*)classAllocator							\
-{																\
-	static OSQueueHead Queue = OS_ATOMIC_QUEUE_INIT;			\
-	return &Queue;												\
-}																\
-+ (id)allocWithZone:(NSZone*) Zone								\
-{																\
-	return (ClassName*)[FApplePlatformObject allocClass:self];	\
-}																\
-+ (id)alloc														\
-{																\
-	return (ClassName*)[FApplePlatformObject allocClass:self];	\
-}
 
 /**
  *	Max implementation of the FGenericPlatformMemoryStats.
@@ -166,8 +130,10 @@ struct CORE_API FApplePlatformMemory : public FGenericPlatformMemory
 
 	static bool CanOverallocateVirtualMemory();
 	
+#ifndef _MACH_VM_TYPES_H_
+	typedef unsigned long vm_address_t;
+#endif
+	
 	static vm_address_t NanoRegionStart;
 	static vm_address_t NanoRegionEnd;
 };
-
-NS_ASSUME_NONNULL_END
diff --git a/Engine/Source/Runtime/Core/Public/Apple/ApplePlatformMisc.h b/Engine/Source/Runtime/Core/Public/Apple/ApplePlatformMisc.h
--- a/Engine/Source/Runtime/Core/Public/Apple/ApplePlatformMisc.h
+++ b/Engine/Source/Runtime/Core/Public/Apple/ApplePlatformMisc.h
@@ -6,12 +6,21 @@
 
 #pragma once
 #include "GenericPlatform/GenericPlatformMisc.h"
+
+#if UE_ENABLE_INCLUDE_ORDER_DEPRECATED_IN_5_8
+
 #if PLATFORM_MAC
 #include "Mac/MacSystemIncludes.h"
 #else
 #include "IOS/IOSSystemIncludes.h"
 #endif
 
+#include "Apple/ScopeAutoreleasePool.h"
+
+#endif // UE_ENABLE_INCLUDE_ORDER_DEPRECATED_IN_5_8
+
+#include "CoreTypes.h"
+
 #ifndef APPLE_PROFILING_ENABLED
 #define APPLE_PROFILING_ENABLED (UE_BUILD_DEBUG | UE_BUILD_DEVELOPMENT)
 #endif
@@ -22,33 +31,6 @@
 
 #define UE_DEBUG_BREAK_IMPL() PLATFORM_BREAK()
 
-#ifdef __OBJC__
-#if !__has_feature(objc_arc)
-
-class FScopeAutoreleasePool
-{
-public:
-
-	FScopeAutoreleasePool()
-	{
-		Pool = [[NSAutoreleasePool alloc] init];
-	}
-
-	~FScopeAutoreleasePool()
-	{
-		[Pool release];
-	}
-
-private:
-
-	NSAutoreleasePool*	Pool;
-};
-
-#define SCOPED_AUTORELEASE_POOL const FScopeAutoreleasePool UE_JOIN(Pool,__LINE__);
-
-#endif // !__has_feature(objc_arc)
-#endif // __OBJC__
-
 /**
 * Apple implementation of the misc OS functions
 **/
@@ -65,10 +47,7 @@ struct CORE_API FApplePlatformMisc : public FGenericPlatformMisc
 	static bool IsDebuggerPresent();
 #endif
 
-	UE_FORCEINLINE_HINT static void MemoryBarrier()
-	{
-		__sync_synchronize();
-	}
+	CORE_API static void MemoryBarrier();
 
 	static void LocalPrint(const TCHAR* Message);
 	static bool IsLocalPrintThreadSafe() { return true; }
diff --git a/Engine/Source/Runtime/Core/Public/Apple/ApplePlatformString.h b/Engine/Source/Runtime/Core/Public/Apple/ApplePlatformString.h
--- a/Engine/Source/Runtime/Core/Public/Apple/ApplePlatformString.h
+++ b/Engine/Source/Runtime/Core/Public/Apple/ApplePlatformString.h
@@ -7,85 +7,36 @@
 #include "Misc/AssertionMacros.h"
 #include "Misc/OutputDevice.h"
 
+static_assert(PLATFORM_TCHAR_IS_CHAR16, "TCHAR must be CHAR16.");
+
+// Manually declaring the type def here to avoid including a full sdk header
+typedef const struct __CFString * CFStringRef;
+
+#if UE_ENABLE_INCLUDE_ORDER_DEPRECATED_IN_5_8
+
 #if PLATFORM_MAC
+#include "AppleStringUtils.h"
 #include "Mac/MacSystemIncludes.h"
 #elif PLATFORM_IOS
+#include "AppleStringUtils.h"
 #include "IOS/IOSSystemIncludes.h"
 #endif
 
-static_assert(PLATFORM_TCHAR_IS_CHAR16, "TCHAR must be CHAR16.");
-
 #ifdef __OBJC__
 #import <Foundation/NSString.h>
+#endif
 
 class FString;
 
-@interface NSString (FString_Extensions)
-
-/**
- * Converts an TCHAR string to an NSString
- */
-+ (NSString*) stringWithTCHARString:(const TCHAR*)MyTCHARString;
-
-/**
- * Converts an FString to an NSString
- */
-+ (NSString*) stringWithFString:(const FString&)MyFString;
-
-@end
-#endif
-
+#endif // UE_ENABLE_INCLUDE_ORDER_DEPRECATED_IN_5_8
 
 /**
  * 2-byte string implementation
  */
 struct FApplePlatformString : public FGenericWidePlatformString
 {
-	// These should be replaced with equivalent Convert and ConvertedLength functions when we properly support variable-length encodings.
-/*	static void WideCharToMultiByte(const TCHAR *Source, uint32 LengthWM1, ANSICHAR *Dest, uint32 LengthA)
-	{
-		CFStringRef CFSource = CFStringCreateWithBytes( kCFAllocatorDefault, ( const uint8 *)Source, LengthWM1 * sizeof( TCHAR ), kCFStringEncodingUTF32LE, false );
-		check(CFSource);
-		const SIZE_T CFSourceLength = CFStringGetLength( CFSource );
-		LengthWM1 = LengthWM1 < CFSourceLength ? LengthWM1 : CFSourceLength;
-		CFRange Range = CFRangeMake( 0, LengthWM1 );
-		CFIndex ActualLength = 0;
-		CFIndex Converted = CFStringGetBytes( CFSource, Range, kCFStringEncodingUTF8, '?', false, NULL, LengthA, &ActualLength );
-		check(Converted == LengthWM1 && ActualLength <= LengthA);
-		CFStringGetBytes( CFSource, Range, kCFStringEncodingUTF8, '?', false, ( uint8 *)Dest, LengthA, &ActualLength );
-		Dest[ActualLength] = 0;
-		CFRelease( CFSource );
-	}
-	static void MultiByteToWideChar(const ANSICHAR *Source, TCHAR *Dest, uint32 LengthM1)
-	{
-		CFStringRef CFSource = CFStringCreateWithBytes( kCFAllocatorDefault, ( const uint8 *)Source, LengthM1, kCFStringEncodingUTF8, false );
-		check(CFSource);
-		const SIZE_T CFSourceLength = CFStringGetLength( CFSource );
-		LengthM1 = LengthM1 < CFSourceLength ? LengthM1 : CFSourceLength;
-		CFRange Range = CFRangeMake( 0, LengthM1 );
-		CFIndex ActualLength = 0;
-		CFStringGetBytes( CFSource, Range, kCFStringEncodingUTF32LE, '?', false, NULL, LengthM1 * sizeof( TCHAR ), &ActualLength );
-		check(ActualLength <= (LengthM1 * sizeof( TCHAR )));
-		CFStringGetBytes( CFSource, Range, kCFStringEncodingUTF32LE, '?', false, ( uint8 *)Dest, LengthM1 * sizeof( TCHAR ), &ActualLength );
-		uint8* Ptr = (uint8*)Dest;
-		Ptr[ActualLength] = Ptr[ActualLength + 1] = Ptr[ActualLength + 2] = Ptr[ActualLength + 3] = 0;
-		CFRelease( CFSource );
-	}*/
-	static void CFStringToTCHAR( CFStringRef CFStr, TCHAR *TChar )
-	{
-		const SIZE_T Length = CFStringGetLength( CFStr );
-		CFRange Range = CFRangeMake( 0, Length );
-		CFStringGetBytes( CFStr, Range, sizeof( TCHAR ) == 4 ? kCFStringEncodingUTF32LE : kCFStringEncodingUnicode, '?', false, ( uint8 *)TChar, Length * sizeof( TCHAR ) + 1, NULL );
-		TChar[Length] = 0;
-	}
-	static CFStringRef TCHARToCFString( const TCHAR *TChar )
-	{
-		const SIZE_T Length = Strlen( TChar );
-		CFStringRef String = CFStringCreateWithBytes( kCFAllocatorDefault, ( const uint8 *)TChar, Length * sizeof( TCHAR ), sizeof( TCHAR ) == 4 ? kCFStringEncodingUTF32LE : kCFStringEncodingUnicode, false );
-		// we are getting some crashes on strings that aren't converting - so, instead of crashing, print them out instead
-		ensureMsgf(String, TEXT("Failed to allocated CFString for '%s' -- Length id %d"), TChar ? TChar : TEXT("nullptr"), Length);
-		return String;
-	}
+	CORE_API static void CFStringToTCHAR( CFStringRef CFStr, TCHAR *TChar );
+	CORE_API static CFStringRef TCHARToCFString( const TCHAR *TChar );
 };
 
 typedef FApplePlatformString FPlatformString;
diff --git a/Engine/Source/Runtime/Core/Public/Apple/ApplePlatformTLS.h b/Engine/Source/Runtime/Core/Public/Apple/ApplePlatformTLS.h
--- a/Engine/Source/Runtime/Core/Public/Apple/ApplePlatformTLS.h
+++ b/Engine/Source/Runtime/Core/Public/Apple/ApplePlatformTLS.h
@@ -7,12 +7,17 @@
 #pragma once
 #include "GenericPlatform/GenericPlatformTLS.h"
 #include "CoreTypes.h"
+
+#if UE_ENABLE_INCLUDE_ORDER_DEPRECATED_IN_5_8
+
 #if PLATFORM_MAC
 #include "Mac/MacSystemIncludes.h"
 #elif PLATFORM_IOS
 #include "IOS/IOSSystemIncludes.h"
 #endif
 
+#endif // UE_ENABLE_INCLUDE_ORDER_DEPRECATED_IN_5_8
+
 /**
 * Apple implementation of the TLS OS functions
 **/
@@ -21,24 +26,12 @@ struct CORE_API FApplePlatformTLS : public FGenericPlatformTLS
 	/**
 	 * Returns the currently executing thread's id
 	 */
-	static UE_FORCEINLINE_HINT uint32 GetCurrentThreadId(void)
-	{
-		return (uint32)pthread_mach_thread_np(pthread_self());
-	}
+	static uint32 GetCurrentThreadId(void);
 
 	/**
 	 * Allocates a thread local store slot
 	 */
-	static inline uint32 AllocTlsSlot(void)
-	{
-		// allocate a per-thread mem slot
-		pthread_key_t SlotKey = 0;
-		if (pthread_key_create(&SlotKey, NULL) != 0)
-		{
-			SlotKey = InvalidTlsSlot;  // matches the Windows TlsAlloc() retval.
-		}
-		return SlotKey;
-	}
+	static uint32 AllocTlsSlot(void);
 
 	/**
 	 * Sets a value in the specified TLS slot
@@ -46,30 +39,21 @@ struct CORE_API FApplePlatformTLS : public FGenericPlatformTLS
 	 * @param SlotIndex the TLS index to store it in
 	 * @param Value the value to store in the slot
 	 */
-	static UE_FORCEINLINE_HINT void SetTlsValue(uint32 SlotIndex,void* Value)
-	{
-		pthread_setspecific((pthread_key_t)SlotIndex, Value);
-	}
+	static void SetTlsValue(uint32 SlotIndex,void* Value);
 
 	/**
 	 * Reads the value stored at the specified TLS slot
 	 *
 	 * @return the value stored in the slot
 	 */
-	static UE_FORCEINLINE_HINT void* GetTlsValue(uint32 SlotIndex)
-	{
-		return pthread_getspecific((pthread_key_t)SlotIndex);
-	}
+	static void* GetTlsValue(uint32 SlotIndex);
 
 	/**
 	 * Frees a previously allocated TLS slot
 	 *
 	 * @param SlotIndex the TLS index to store it in
 	 */
-	static UE_FORCEINLINE_HINT void FreeTlsSlot(uint32 SlotIndex)
-	{
-		pthread_key_delete((pthread_key_t)SlotIndex);
-	}
+	static void FreeTlsSlot(uint32 SlotIndex);
 };
 
 typedef FApplePlatformTLS FPlatformTLS;
diff --git a/Engine/Source/Runtime/Core/Public/Apple/ApplePlatformTime.h b/Engine/Source/Runtime/Core/Public/Apple/ApplePlatformTime.h
--- a/Engine/Source/Runtime/Core/Public/Apple/ApplePlatformTime.h
+++ b/Engine/Source/Runtime/Core/Public/Apple/ApplePlatformTime.h
@@ -6,12 +6,15 @@
 
 #pragma once
 #include "GenericPlatform/GenericPlatformTime.h"
+
+#if UE_ENABLE_INCLUDE_ORDER_DEPRECATED_IN_5_8
 #include "CoreTypes.h"
 #if PLATFORM_MAC
 #include "Mac/MacSystemIncludes.h"
 #elif PLATFORM_IOS
 #include "IOS/IOSSystemIncludes.h"
 #endif
+#endif // UE_ENABLE_INCLUDE_ORDER_DEPRECATED_IN_5_8
 
 /**
  * Please see following UDN post about using rdtsc on processors that support
@@ -29,24 +32,11 @@ struct CORE_API FApplePlatformTime : public FGenericPlatformTime
 {
 	static double InitTiming();
 
-	static inline double Seconds()
-	{
-		uint64 Cycles = mach_absolute_time();
-		// Add big number to make bugs apparent where return value is being passed to float
-		return Cycles * GetSecondsPerCycle() + GetSecondsTimeOffset();
-	}
+	static double Seconds();
 
-	static inline uint32 Cycles()
-	{
-		uint64 Cycles = mach_absolute_time();
-		return Cycles;
-	}
+	static uint32 Cycles();
 
-	static inline uint64 Cycles64()
-	{
-		uint64 Cycles = mach_absolute_time();
-		return Cycles;
-	}
+	static uint64 Cycles64();
 
 	static FCPUTime GetCPUTime();
 
diff --git a/Engine/Source/Runtime/Core/Public/Apple/AppleStringUtils.h b/Engine/Source/Runtime/Core/Public/Apple/AppleStringUtils.h
new file mode 100644
--- /dev/null
+++ b/Engine/Source/Runtime/Core/Public/Apple/AppleStringUtils.h
@@ -0,0 +1,42 @@
+// Copyright Epic Games, Inc. All Rights Reserved.
+
+#pragma once
+
+#include "CoreTypes.h"
+
+#if PLATFORM_APPLE
+
+#if PLATFORM_MAC
+#include "Mac/MacSystemIncludes.h"
+#elif PLATFORM_IOS
+#include "IOS/IOSSystemIncludes.h"
+#endif
+
+#include "Misc/AssertionMacros.h"
+
+#import <Foundation/NSString.h>
+
+class FString;
+
+@interface NSString (FString_Extensions)
+
+/**
+ * Converts an TCHAR string to an NSString
+ */
++ (NSString*) stringWithTCHARString:(const TCHAR*)MyTCHARString;
+
+/**
+ * Converts an FString to an NSString
+ */
++ (NSString*) stringWithFString:(const FString&)MyFString;
+
+@end
+
+class FAppleStringUtils
+{
+public:
+	CORE_API static NSString* ConvertToNSString(const FString& InString);
+	CORE_API static FString ConvertToFString(const  NSString* InNSString);
+};
+
+#endif
\ No newline at end of file
diff --git a/Engine/Source/Runtime/Core/Public/Apple/ScopeAutoreleasePool.h b/Engine/Source/Runtime/Core/Public/Apple/ScopeAutoreleasePool.h
new file mode 100644
--- /dev/null
+++ b/Engine/Source/Runtime/Core/Public/Apple/ScopeAutoreleasePool.h
@@ -0,0 +1,32 @@
+// Copyright Epic Games, Inc. All Rights Reserved.
+
+#pragma once
+
+#import <Foundation/NSAutoreleasePool.h>
+
+#ifdef __OBJC__
+#if !__has_feature(objc_arc)
+
+class FScopeAutoreleasePool
+{
+public:
+
+	FScopeAutoreleasePool()
+	{
+		Pool = [[NSAutoreleasePool alloc] init];
+	}
+
+	~FScopeAutoreleasePool()
+	{
+		[Pool release];
+	}
+
+private:
+
+	NSAutoreleasePool*	Pool;
+};
+
+#define SCOPED_AUTORELEASE_POOL const FScopeAutoreleasePool UE_JOIN(Pool,__LINE__);
+
+#endif // !__has_feature(objc_arc)
+#endif // __OBJC__
diff --git a/Engine/Source/Runtime/Core/Public/Containers/UnrealString.h.inl b/Engine/Source/Runtime/Core/Public/Containers/UnrealString.h.inl
--- a/Engine/Source/Runtime/Core/Public/Containers/UnrealString.h.inl
+++ b/Engine/Source/Runtime/Core/Public/Containers/UnrealString.h.inl
@@ -40,6 +40,10 @@ using FStringFormatOrderedArguments = TArray<FStringFormatArg>;
 [[nodiscard]] const UE_STRING_CHARTYPE* GetData(const UE_STRING_CLASS&);
 [[nodiscard]] int32                     GetNum(const UE_STRING_CLASS& String);
 
+#ifdef __OBJC__
+@class NSString;
+#endif
+
 /**
  * A dynamically sizeable string.
  * @see https://docs.unrealengine.com/latest/INT/Programming/UnrealArchitecture/StringHandling/FString/
@@ -182,27 +186,7 @@ public:
 public:
 
 #if PLATFORM_APPLE && UE_STRING_CHARTYPE_IS_TCHAR
-	[[nodiscard]] inline UE_STRING_CLASS(const CFStringRef In)
-	{
-		uint32_t StringLength = In ? CFStringGetLength(In) : 0;
-
-		if (StringLength > 0)
-		{
-			// Convert the NSString data into the native TCHAR format for UE4
-			// This returns a buffer of bytes, but they can be safely cast to a buffer of TCHARs
-			const CFStringEncoding Encoding = kCFStringEncodingUTF16LE;
-			CFRange Range = CFRangeMake(0, StringLength);
-			CFIndex BytesNeeded;
-			if (CFStringGetBytes(In, Range, Encoding, '?', false, NULL, 0, &BytesNeeded) > 0)
-			{
-				const size_t Length = BytesNeeded / sizeof(TCHAR);
-				Data.Reserve(Length + 1);
-				Data.AddUninitialized(Length + 1);
-				CFStringGetBytes(In, Range, Encoding, '?', false, (uint8*)Data.GetData(), Length * sizeof(TCHAR) + 1, NULL);
-				Data[Length] = 0;
-			}
-		}
-	}
+	[[nodiscard]] CORE_API UE_STRING_CLASS(const CFStringRef In);
 #endif
 
 	////////////////////////////////////////////////////////
diff --git a/Engine/Source/Runtime/Core/Public/IOS/IOSCommandLineHelper.h b/Engine/Source/Runtime/Core/Public/IOS/IOSCommandLineHelper.h
--- a/Engine/Source/Runtime/Core/Public/IOS/IOSCommandLineHelper.h
+++ b/Engine/Source/Runtime/Core/Public/IOS/IOSCommandLineHelper.h
@@ -2,6 +2,7 @@
 
 #pragma once
 
+#include "Apple/AppleStringUtils.h"
 #include "CoreMinimal.h"
 #include "Containers/StringView.h"
 #include "Misc/CommandLine.h"
diff --git a/Engine/Source/Runtime/Core/Public/IOS/IOSPlatformCrashContext.h b/Engine/Source/Runtime/Core/Public/IOS/IOSPlatformCrashContext.h
--- a/Engine/Source/Runtime/Core/Public/IOS/IOSPlatformCrashContext.h
+++ b/Engine/Source/Runtime/Core/Public/IOS/IOSPlatformCrashContext.h
@@ -5,6 +5,10 @@
 #include "Apple/ApplePlatformCrashContext.h"
 #include "Misc/Guid.h"
 
+#if UE_ENABLE_INCLUDE_ORDER_DEPRECATED_IN_5_8
+#import <Foundation/NSProcessInfo.h>
+#endif
+
 //Forward Declarations
 class FIOSMallocCrashHandler;
 #if !PLATFORM_TVOS
@@ -56,7 +60,11 @@ struct FIOSApplicationInfo
     FString BranchBaseDir;
     FString PrimaryGPU;
     FString ExecutableName;
+#if UE_ENABLE_INCLUDE_ORDER_DEPRECATED_IN_5_8
+	UE_DEPRECATED(5.8, "The OSXVersion property is no longer available. Use the already parsed version in string format")
     NSOperatingSystemVersion OSXVersion;
+#endif// UE_ENABLE_INCLUDE_ORDER_DEPRECATED_IN_5_8
+	
     FGuid RunUUID;
     FString XcodePath;
 #if !PLATFORM_TVOS
diff --git a/Engine/Source/Runtime/Core/Public/IOS/IOSPlatformMisc.h b/Engine/Source/Runtime/Core/Public/IOS/IOSPlatformMisc.h
--- a/Engine/Source/Runtime/Core/Public/IOS/IOSPlatformMisc.h
+++ b/Engine/Source/Runtime/Core/Public/IOS/IOSPlatformMisc.h
@@ -6,7 +6,14 @@
 
 #pragma once
 #include "GenericPlatform/GenericPlatformMisc.h"
+#if UE_ENABLE_INCLUDE_ORDER_DEPRECATED_IN_5_8
 #include "IOS/IOSSystemIncludes.h"
+#endif // UE_ENABLE_INCLUDE_ORDER_DEPRECATED_IN_5_8
+
+#ifdef __OBJC__
+@class Protocol;
+#endif
+
 #include "Apple/ApplePlatformMisc.h"
 
 template <typename FuncType>
diff --git a/Engine/Source/Runtime/Core/Public/IOS/IOSPlatformProcess.h b/Engine/Source/Runtime/Core/Public/IOS/IOSPlatformProcess.h
--- a/Engine/Source/Runtime/Core/Public/IOS/IOSPlatformProcess.h
+++ b/Engine/Source/Runtime/Core/Public/IOS/IOSPlatformProcess.h
@@ -6,7 +6,9 @@
 
 #pragma once
 #include "Apple/ApplePlatformProcess.h"
+#if UE_ENABLE_INCLUDE_ORDER_DEPRECATED_IN_5_8
 #import "IOSAsyncTask.h"
+#endif // UE_ENABLE_INCLUDE_ORDER_DEPRECATED_IN_5_8
 
 /** Dummy process handle for platforms that use generic implementation. */
 struct FProcHandle : public TProcHandle<void*, nullptr>
diff --git a/Engine/Source/Runtime/Core/Public/IOS/IOSSystemIncludes.h b/Engine/Source/Runtime/Core/Public/IOS/IOSSystemIncludes.h
--- a/Engine/Source/Runtime/Core/Public/IOS/IOSSystemIncludes.h
+++ b/Engine/Source/Runtime/Core/Public/IOS/IOSSystemIncludes.h
@@ -5,21 +5,24 @@
 #include "CoreTypes.h"
 #include "IOS/IOSPlatformCompilerSetup.h"
 
+#pragma push_macro("FVector")
 #define FVector FVectorWorkaround
 #ifdef __OBJC__
 #import <UIKit/UIKit.h>
 #import <CoreData/CoreData.h>
+#include <CoreFoundation/CoreFoundation.h>
 #else //@todo - zombie - iOS should always have Obj-C to function. This definition guard may be unnecessary. -astarkey
 #error "Objective C is undefined."
 #endif
+
+#pragma push_macro("check")
+#pragma push_macro("verify")
 #undef check
 #undef verify
 #undef FVector
 
 #undef TYPE_BOOL
 
-#include <CoreFoundation/CoreFoundation.h>
-
 #include <string.h>
 #include <alloca.h>
 #include <stdlib.h>
@@ -76,3 +79,8 @@ typedef struct tagRECT RECT;
 #define _aligned_realloc(Ptr,Size,Align) realloc(Ptr,Size)
 #define _aligned_free(Ptr) free(Ptr)
 #define _aligned_msize(Ptr,Align,Offset) malloc_size(Ptr)
+
+#pragma pop_macro("check")
+#pragma pop_macro("verify")
+#pragma pop_macro("FVector")
+
diff --git a/Engine/Source/Runtime/Core/Public/Mac/MacSystemIncludes.h b/Engine/Source/Runtime/Core/Public/Mac/MacSystemIncludes.h
--- a/Engine/Source/Runtime/Core/Public/Mac/MacSystemIncludes.h
+++ b/Engine/Source/Runtime/Core/Public/Mac/MacSystemIncludes.h
@@ -2,8 +2,10 @@
 #pragma once
 
 // Set up compiler pragmas, etc
+#include "CoreTypes.h"
 #include "Mac/MacPlatformCompilerSetup.h"
 
+#pragma push_macro("FVector")
 #define FVector FVectorWorkaround
 #ifdef __OBJC__
 #import <Cocoa/Cocoa.h>
@@ -11,6 +13,10 @@
 #endif
 #include <Carbon/Carbon.h>
 #include <CoreServices/CoreServices.h>
+
+#pragma push_macro("check")
+#pragma push_macro("verify")
+
 #undef check
 #undef verify
 #undef FVector
@@ -67,3 +73,7 @@ Memory. On Mac OS X malloc allocates memory aligned to 16 bytes.
 #define _aligned_realloc(Ptr,Size,Align) realloc(Ptr,Size)
 #define _aligned_free(Ptr) free(Ptr)
 #define _aligned_msize(Ptr,Align,Offset) malloc_size(Ptr)
+
+#pragma pop_macro("check")
+#pragma pop_macro("verify")
+#pragma pop_macro("FVector")
diff --git a/Engine/Source/Runtime/CrashReportCore/Private/IOS/IOSErrorReport.cpp b/Engine/Source/Runtime/CrashReportCore/Private/IOS/IOSErrorReport.cpp
--- a/Engine/Source/Runtime/CrashReportCore/Private/IOS/IOSErrorReport.cpp
+++ b/Engine/Source/Runtime/CrashReportCore/Private/IOS/IOSErrorReport.cpp
@@ -9,6 +9,8 @@
 #include "Misc/FileHelper.h"
 #include "HAL/PlatformFileManager.h"
 
+#import <CoreFoundation/CoreFoundation.h>
+
 #define LOCTEXT_NAMESPACE "CrashReport"
 
 namespace
diff --git a/Engine/Source/Runtime/IOS/IOSLocalNotification/Private/IOSLocalNotification.cpp b/Engine/Source/Runtime/IOS/IOSLocalNotification/Private/IOSLocalNotification.cpp
--- a/Engine/Source/Runtime/IOS/IOSLocalNotification/Private/IOSLocalNotification.cpp
+++ b/Engine/Source/Runtime/IOS/IOSLocalNotification/Private/IOSLocalNotification.cpp
@@ -10,6 +10,13 @@
 
 #include "IOSLocalNotification.h"
 
+#include "Apple/AppleStringUtils.h"
+#if PLATFORM_MAC
+#include "Mac/MacSystemIncludes.h"
+#elif PLATFORM_IOS
+#include "IOS/IOSSystemIncludes.h"
+#endif
+
 #include "IOS/IOSApplication.h"
 #include "IOS/IOSAppDelegate.h"
 
diff --git a/Engine/Source/Runtime/ImageWrapper/Private/Formats/IcnsImageWrapper.cpp b/Engine/Source/Runtime/ImageWrapper/Private/Formats/IcnsImageWrapper.cpp
--- a/Engine/Source/Runtime/ImageWrapper/Private/Formats/IcnsImageWrapper.cpp
+++ b/Engine/Source/Runtime/ImageWrapper/Private/Formats/IcnsImageWrapper.cpp
@@ -4,6 +4,9 @@
 #include "ImageWrapperPrivate.h"
 #include "ImageCoreUtils.h"
 
+#if PLATFORM_MAC
+#include "Apple/ScopeAutoreleasePool.h"
+#endif
 
 
 /* FIcnsImageWrapper structors
diff --git a/Engine/Source/Runtime/Launch/Private/Mac/LaunchMac.cpp b/Engine/Source/Runtime/Launch/Private/Mac/LaunchMac.cpp
--- a/Engine/Source/Runtime/Launch/Private/Mac/LaunchMac.cpp
+++ b/Engine/Source/Runtime/Launch/Private/Mac/LaunchMac.cpp
@@ -19,6 +19,7 @@
 	#include "ISettingsModule.h"
 #endif
 
+#include "Apple/ScopeAutoreleasePool.h"
 #include <signal.h>
 #include "Mac/CocoaThread.h"
 
diff --git a/Engine/Source/Runtime/Online/HTTP/Private/Apple/AppleHttp.cpp b/Engine/Source/Runtime/Online/HTTP/Private/Apple/AppleHttp.cpp
--- a/Engine/Source/Runtime/Online/HTTP/Private/Apple/AppleHttp.cpp
+++ b/Engine/Source/Runtime/Online/HTTP/Private/Apple/AppleHttp.cpp
@@ -11,6 +11,9 @@
 #include "Http.h"
 #include "HttpManager.h"
 #include "HttpModule.h"
+#include "Apple/AppleStringUtils.h"
+#include "Apple/ScopeAutoreleasePool.h"
+
 #include "Interfaces/IHttpRequest.h"
 #include "Misc/EngineVersion.h"
 #include "Misc/App.h"
diff --git a/Engine/Source/Runtime/Online/HTTP/Private/Apple/AppleHttp.h b/Engine/Source/Runtime/Online/HTTP/Private/Apple/AppleHttp.h
--- a/Engine/Source/Runtime/Online/HTTP/Private/Apple/AppleHttp.h
+++ b/Engine/Source/Runtime/Online/HTTP/Private/Apple/AppleHttp.h
@@ -15,6 +15,8 @@
  */
 DECLARE_DELEGATE(FNewAppleHttpEventDelegate);
 
+@class NSMutableURLRequest;
+
 /**
  * Apple implementation of an Http request
  */
diff --git a/Engine/Source/Runtime/Online/HTTP/Public/Apple/ApplePlatformHttp.h b/Engine/Source/Runtime/Online/HTTP/Public/Apple/ApplePlatformHttp.h
--- a/Engine/Source/Runtime/Online/HTTP/Public/Apple/ApplePlatformHttp.h
+++ b/Engine/Source/Runtime/Online/HTTP/Public/Apple/ApplePlatformHttp.h
@@ -4,6 +4,9 @@
 
 #include "GenericPlatform/GenericPlatformHttp.h"
 
+@class NSURLSessionTask;
+@class NSURLSession;
+
 /**
  * Platform specific Http implementations
  */
@@ -42,8 +45,8 @@ public:
 
 private:
     /** Session used to create Apple based requests */
-    static inline NSURLSession* Session = nil;
-	static inline NSURLSession* BackgroundSession = nil;
+    static inline NSURLSession* Session = nullptr;
+	static inline NSURLSession* BackgroundSession = nullptr;
 
 	static FRWLock TaskRequestMapLock;
 	static TMap<NSURLSessionTask*, TSharedPtr<class FAppleHttpRequest>> TaskRequestMap;
diff --git a/Engine/Source/Runtime/PreLoadScreen/Private/PreLoadScreenManager.cpp b/Engine/Source/Runtime/PreLoadScreen/Private/PreLoadScreenManager.cpp
--- a/Engine/Source/Runtime/PreLoadScreen/Private/PreLoadScreenManager.cpp
+++ b/Engine/Source/Runtime/PreLoadScreen/Private/PreLoadScreenManager.cpp
@@ -25,6 +25,10 @@
 #endif
 #endif
 
+#if PLATFORM_IOS
+#include "IOS/IOSAsyncTask.h"
+#endif
+
 IMPLEMENT_MODULE(FDefaultModuleImpl, PreLoadScreen);
 
 DEFINE_LOG_CATEGORY_STATIC(LogPreLoadScreenManager, Log, All);
diff --git a/Engine/Source/Runtime/RHI/Private/PipelineFileCache.cpp b/Engine/Source/Runtime/RHI/Private/PipelineFileCache.cpp
--- a/Engine/Source/Runtime/RHI/Private/PipelineFileCache.cpp
+++ b/Engine/Source/Runtime/RHI/Private/PipelineFileCache.cpp
@@ -5,6 +5,16 @@ PipelineFileCache.cpp: Pipeline state cache implementation.
 =============================================================================*/
 
 #include "PipelineFileCache.h"
+
+#if PLATFORM_APPLE
+	#include "Apple/ScopeAutoreleasePool.h"
+#if PLATFORM_MAC
+	#include "Mac/MacSystemIncludes.h"
+	#elif PLATFORM_IOS
+	#include "IOS/IOSSystemIncludes.h"
+	#endif 
+#endif // PLATFORM_APPLE
+
 #include "Containers/List.h"
 #include "Containers/Ticker.h"
 #include "PipelineStateCache.h"
diff --git a/Engine/Source/Runtime/Slate/Private/Framework/Application/SlateApplication.cpp b/Engine/Source/Runtime/Slate/Private/Framework/Application/SlateApplication.cpp
--- a/Engine/Source/Runtime/Slate/Private/Framework/Application/SlateApplication.cpp
+++ b/Engine/Source/Runtime/Slate/Private/Framework/Application/SlateApplication.cpp
@@ -27,6 +27,11 @@
 #include "Framework/Application/SWindowTitleBar.h"
 #include "Input/Events.h"
 #include "Input/HittestGrid.h"
+#if PLATFORM_MAC
+#include "Mac/MacSystemIncludes.h"
+#elif PLATFORM_IOS
+#include "IOS/IOSSystemIncludes.h"
+#endif
 #include "HAL/PlatformApplicationMisc.h"
 #include "HAL/PlatformStackWalk.h"
 #include "Null/NullPlatformApplicationMisc.h"
diff --git a/Engine/Source/Runtime/Slate/Private/Framework/Application/SlateUser.cpp b/Engine/Source/Runtime/Slate/Private/Framework/Application/SlateUser.cpp
--- a/Engine/Source/Runtime/Slate/Private/Framework/Application/SlateUser.cpp
+++ b/Engine/Source/Runtime/Slate/Private/Framework/Application/SlateUser.cpp
@@ -1,6 +1,13 @@
 // Copyright Epic Games, Inc. All Rights Reserved.
 
 #include "Framework/Application/SlateUser.h"
+
+#if PLATFORM_MAC
+#include "Mac/MacSystemIncludes.h"
+#elif PLATFORM_IOS
+#include "IOS/IOSSystemIncludes.h"
+#endif
+
 #include "Framework/Application/SlateApplication.h"
 #include "Framework/Application/NavigationConfig.h"
 #include "Misc/App.h"
diff --git a/Engine/Source/Runtime/Slate/Private/Framework/MultiBox/Mac/MacMenu.cpp b/Engine/Source/Runtime/Slate/Private/Framework/MultiBox/Mac/MacMenu.cpp
--- a/Engine/Source/Runtime/Slate/Private/Framework/MultiBox/Mac/MacMenu.cpp
+++ b/Engine/Source/Runtime/Slate/Private/Framework/MultiBox/Mac/MacMenu.cpp
@@ -1,6 +1,9 @@
 // Copyright Epic Games, Inc. All Rights Reserved.
 
 #include "MacMenu.h"
+
+#include "Apple/ScopeAutoreleasePool.h"
+#include "Apple/AppleStringUtils.h"
 #include "Framework/Application/SlateApplication.h"
 #include "Mac/CocoaThread.h"
 #include "Mac/MacApplication.h"
diff --git a/Engine/Source/Runtime/Slate/Private/Framework/Text/IOS/IOSPlatformTextField.cpp b/Engine/Source/Runtime/Slate/Private/Framework/Text/IOS/IOSPlatformTextField.cpp
--- a/Engine/Source/Runtime/Slate/Private/Framework/Text/IOS/IOSPlatformTextField.cpp
+++ b/Engine/Source/Runtime/Slate/Private/Framework/Text/IOS/IOSPlatformTextField.cpp
@@ -1,6 +1,7 @@
 // Copyright Epic Games, Inc. All Rights Reserved.
 
 #include "Framework/Text/IOS/IOSPlatformTextField.h"
+#include "Apple/AppleStringUtils.h"
 #include "IOS/IOSAppDelegate.h"
 #include "IOS/IOSAsyncTask.h"
 #include "IOS/IOSView.h"
diff --git a/Engine/Source/Runtime/WebBrowser/Private/Apple/AppleCookieManager.h b/Engine/Source/Runtime/WebBrowser/Private/Apple/AppleCookieManager.h
--- a/Engine/Source/Runtime/WebBrowser/Private/Apple/AppleCookieManager.h
+++ b/Engine/Source/Runtime/WebBrowser/Private/Apple/AppleCookieManager.h
@@ -7,6 +7,7 @@
 #if PLATFORM_IOS || PLATFORM_MAC
 #include "IWebBrowserCookieManager.h"
 
+@class NSUUID;
 /**
  * Implementation of interface for dealing with a Web Browser cookies for iOS.
  */
diff --git a/Engine/Source/Runtime/WebBrowser/Private/Apple/ApplePlatformWebBrowser.cpp b/Engine/Source/Runtime/WebBrowser/Private/Apple/ApplePlatformWebBrowser.cpp
--- a/Engine/Source/Runtime/WebBrowser/Private/Apple/ApplePlatformWebBrowser.cpp
+++ b/Engine/Source/Runtime/WebBrowser/Private/Apple/ApplePlatformWebBrowser.cpp
@@ -35,6 +35,8 @@
 #include "AppleCookieManager.h"
 #include "WebView3DInputHandler.h"
 
+#include "Apple/AppleStringUtils.h"
+
 DEFINE_LOG_CATEGORY(LogAppleWebBrowser);
 
 #if !UE_BUILD_SHIPPING
diff --git a/Engine/Source/Runtime/WebBrowser/Private/Apple/WebView3DInputHandler.cpp b/Engine/Source/Runtime/WebBrowser/Private/Apple/WebView3DInputHandler.cpp
--- a/Engine/Source/Runtime/WebBrowser/Private/Apple/WebView3DInputHandler.cpp
+++ b/Engine/Source/Runtime/WebBrowser/Private/Apple/WebView3DInputHandler.cpp
@@ -8,6 +8,8 @@
 #include "HAL/PlatformInput.h"
 #endif
 
+#import <Foundation/Foundation.h>
+
 DEFINE_LOG_CATEGORY(LogWebView3DInputHandler);
 
 FWebView3DInputHandler::FWebView3DInputHandler()
diff --git a/Engine/Source/Runtime/WebBrowser/Private/WebBrowserSingleton.cpp b/Engine/Source/Runtime/WebBrowser/Private/WebBrowserSingleton.cpp
--- a/Engine/Source/Runtime/WebBrowser/Private/WebBrowserSingleton.cpp
+++ b/Engine/Source/Runtime/WebBrowser/Private/WebBrowserSingleton.cpp
@@ -13,6 +13,10 @@
 #include "IWebBrowserCookieManager.h"
 #include "WebBrowserLog.h"
 
+#if PLATFORM_APPLE
+#include "Apple/ScopeAutoreleasePool.h"
+#endif
+
 #if WITH_CEF3
 #include "CEF3Utils.h"
 #include "Misc/ScopeLock.h"
