Adding MetalRHI support for multi view slate window rendering on macOS

#Synced-CL 49810091

diff --git a/Engine/Source/Runtime/Apple/MetalRHI/Private/MetalDynamicRHI.h b/Engine/Source/Runtime/Apple/MetalRHI/Private/MetalDynamicRHI.h
--- a/Engine/Source/Runtime/Apple/MetalRHI/Private/MetalDynamicRHI.h
+++ b/Engine/Source/Runtime/Apple/MetalRHI/Private/MetalDynamicRHI.h
@@ -175,6 +175,7 @@ public:
 	virtual FRenderQueryRHIRef RHICreateRenderQuery(ERenderQueryType QueryType) final override;
 	virtual bool RHIGetRenderQueryResult(FRHIRenderQuery* RenderQuery, uint64& OutResult, bool bWait, uint32 GPUIndex = INDEX_NONE) final override;
 	virtual FTextureRHIRef RHIGetViewportBackBuffer(FRHIViewport* Viewport) final override;
+	virtual FTextureRHIRef RHIGetViewportBackBufferAtLayerIndex(FRHIViewport* Viewport, uint32 DrawLayerIndex) final override;
 	virtual void RHIAdvanceFrameForGetViewportBackBuffer(FRHIViewport* Viewport, bool bPresent) final override;
 	virtual void RHIFlushResources() final override;
 	virtual FViewportRHIRef RHICreateViewport(void* WindowHandle, uint32 SizeX, uint32 SizeY, bool bIsFullscreen, EPixelFormat PreferredPixelFormat) final override;
diff --git a/Engine/Source/Runtime/Apple/MetalRHI/Private/MetalTexture.cpp b/Engine/Source/Runtime/Apple/MetalRHI/Private/MetalTexture.cpp
--- a/Engine/Source/Runtime/Apple/MetalRHI/Private/MetalTexture.cpp
+++ b/Engine/Source/Runtime/Apple/MetalRHI/Private/MetalTexture.cpp
@@ -1271,7 +1271,7 @@ MTLTexturePtr FMetalSurface::GetDrawableTexture()
 	if (!Texture && EnumHasAnyFlags(GetDesc().Flags, TexCreate_Presentable))
 	{
 		check(Viewport);
-		Texture = NS::RetainPtr(Viewport->GetDrawableTexture(EMetalViewportAccessRHI));
+		Texture = NS::RetainPtr(Viewport->GetDrawableTexture(EMetalViewportAccessRHI, this));
 
 #if PLATFORM_SUPPORTS_BINDLESS_RENDERING
 		if(IsMetalBindlessEnabled())
diff --git a/Engine/Source/Runtime/Apple/MetalRHI/Private/MetalViewport.cpp b/Engine/Source/Runtime/Apple/MetalRHI/Private/MetalViewport.cpp
--- a/Engine/Source/Runtime/Apple/MetalRHI/Private/MetalViewport.cpp
+++ b/Engine/Source/Runtime/Apple/MetalRHI/Private/MetalViewport.cpp
@@ -68,9 +68,19 @@ static FAutoConsoleVariableRef CVarMetalNonBlockingPresent(
 	return self;
 }
 
+- (void)setAllowMouseActions: (BOOL)allowMouseActions
+{
+	bAllowMouseActions = allowMouseActions;
+}
+
+- (NSView *) hitTest:(NSPoint) point
+{
+	return bAllowMouseActions ? [super hitTest: point] : nil;
+}
+
 - (BOOL)isOpaque
 {
-	return YES;
+	return !bAllowMouseActions;
 }
 
 - (BOOL)mouseDownCanMoveWindow
@@ -87,22 +97,28 @@ static TSet<FMetalViewport*> Viewports;
 
 FMetalViewport::FMetalViewport(FMetalDevice& InDevice, void* WindowHandle, uint32 InSizeX, uint32 InSizeY, bool bInIsFullscreen, EPixelFormat InFormat)
 	: Device(InDevice)
-	, Drawable{nullptr}
-	, BackBuffer{nullptr, nullptr}
 	, Mutex{}
-	, DrawableTextures{}
 	, DisplayID{0}
 	, Block{nullptr}
 	, FrameAvailable{0}
-	, LastCompleteFrame{nullptr}
 	, bIsFullScreen{bInIsFullscreen}
-#if PLATFORM_IOS && !PLATFORM_TVOS && !PLATFORM_VISIONOS
+#if !PLATFORM_TVOS && !PLATFORM_VISIONOS
 	, SizeX(InSizeX)
 	, SizeY(InSizeY)
+	, Format(InFormat)
+#endif
+	, DefaultNativeLayer{std::numeric_limits<int32>::min()}
+#if PLATFORM_MAC
+	, WindowHandle{(FCocoaWindow*)WindowHandle}
 #endif
+#if !WITH_MULTI_VIEW_SLATE_WINDOW_SUPPORT
+	, Drawable{nullptr}
+	, BackBuffer{nullptr, nullptr}
 #if PLATFORM_MAC
 	, View{nullptr}
 #endif
+#endif
+	, DrawableTextures{}
 #if PLATFORM_MAC || PLATFORM_VISIONOS
 	, CustomPresent{nullptr}
 #endif
@@ -112,39 +128,19 @@ FMetalViewport::FMetalViewport(FMetalDevice& InDevice, void* WindowHandle, uint3
 	SwiftLayer = [IOSAppDelegate GetDelegate].SwiftLayer;
 #endif
 
+#if !WITH_MULTI_VIEW_SLATE_WINDOW_SUPPORT
+	// if we are using the multi view slate window rendering path, the equivalent code
+	// gets called from FSlateRHIRenderer::CreateViewport() upon viewport
+	// creation, except initializing the base draw layer of std::numeric_limits<int32>::min() 
+	// instead of a view variable
 #if PLATFORM_MAC
 	MainThreadCall(^{
-		FCocoaWindow* Window = (FCocoaWindow*)WindowHandle;
-		const NSRect ContentRect = NSMakeRect(0, 0, InSizeX, InSizeY);
-		View = [[FMetalView alloc] initWithFrame:ContentRect];
-		[View setAutoresizingMask:NSViewWidthSizable | NSViewHeightSizable];
-		[View setWantsLayer:YES];
-
-		CAMetalLayer* Layer = [CAMetalLayer new];
-
-		CGFloat bgColor[] = { 0.0, 0.0, 0.0, 0.0 };
-		Layer.edgeAntialiasingMask = 0;
-		Layer.masksToBounds = YES;
-		Layer.backgroundColor = CGColorCreate(CGColorSpaceCreateDeviceRGB(), bgColor);
-		Layer.presentsWithTransaction = NO;
-		Layer.anchorPoint = CGPointMake(0.5, 0.5);
-		Layer.frame = ContentRect;
-		Layer.magnificationFilter = kCAFilterNearest;
-		Layer.minificationFilter = kCAFilterNearest;
-
-		[Layer setDevice:(__bridge id<MTLDevice>)Device.GetDevice()];
-		
-		[Layer setFramebufferOnly:NO];
-		[Layer removeAllAnimations];
-
-		[View setLayer:Layer];
-
-		[Window setContentView:View];
-		[[Window standardWindowButton:NSWindowCloseButton] setAction:@selector(performClose:)];
+		View = CreateNewMetalView(true, NSWindowAbove, nil);
 	});
 #endif
-
+	
 	Resize(InSizeX, InSizeY, bInIsFullscreen, InFormat);
+#endif
 
 #if PLATFORM_IOS && !PLATFORM_TVOS && !PLATFORM_VISIONOS
 	// An orientation change from LandscapeLeft to LandscapeRight won't trigger a SetRes, so we need to react to an orientation change here
@@ -172,12 +168,185 @@ FMetalViewport::~FMetalViewport()
 	FCoreDelegates::ApplicationReceivedScreenOrientationChangedNotificationDelegate.Remove(OrientationChangedHandle);
 #endif
 
+#if WITH_MULTI_VIEW_SLATE_WINDOW_SUPPORT
+	for(TPair<int32, DrawLayer>& DrawLayer : DrawLayers)
+	{
+		DrawLayer.Value.BackBuffer[0].SafeRelease();	// when the rest of the engine releases it, its framebuffers will be released too (those the engine knows about)
+		DrawLayer.Value.BackBuffer[1].SafeRelease();
+		check(!IsValidRef(DrawLayer.Value.BackBuffer[0]));
+		check(!IsValidRef(DrawLayer.Value.BackBuffer[1]));
+	}
+#else
 	BackBuffer[0].SafeRelease();	// when the rest of the engine releases it, its framebuffers will be released too (those the engine knows about)
 	BackBuffer[1].SafeRelease();
 	check(!IsValidRef(BackBuffer[0]));
 	check(!IsValidRef(BackBuffer[1]));
+#endif
+}
+
+#if PLATFORM_MAC
+FMetalView* FMetalViewport::CreateNewMetalView(bool bIsContentView, NSWindowOrderingMode ViewOrder, NSView* ViewRelativeTo)
+{	
+	const NSRect ContentRect = NSMakeRect(0, 0, SizeX, SizeY);
+	
+	FMetalView* MetalView = [[FMetalView alloc] initWithFrame:ContentRect];
+	[MetalView setAutoresizingMask:NSViewWidthSizable | NSViewHeightSizable];
+	[MetalView setWantsLayer:YES];
+	[MetalView setAllowMouseActions: bIsContentView];
+
+	CAMetalLayer* Layer = [CAMetalLayer new];
+
+	CGFloat bgColor[] = { 0.0, 0.0, 0.0, 0.0 };
+	Layer.edgeAntialiasingMask = 0;
+	Layer.masksToBounds = YES;
+	Layer.backgroundColor = CGColorCreate(CGColorSpaceCreateDeviceRGB(), bgColor);
+	Layer.presentsWithTransaction = NO;
+	Layer.anchorPoint = CGPointMake(0.5, 0.5);
+	Layer.frame = ContentRect;
+	Layer.magnificationFilter = kCAFilterNearest;
+	Layer.minificationFilter = kCAFilterNearest;
+	Layer.opaque = bIsContentView;
+
+	[Layer setDevice:(__bridge id<MTLDevice>)Device.GetDevice()];
+	
+	[Layer setFramebufferOnly:NO];
+	[Layer removeAllAnimations];
+	
+	[MetalView setLayer:Layer];
+	
+	if(bIsContentView)
+	{		
+		[WindowHandle setContentView: MetalView];
+		[[WindowHandle standardWindowButton:NSWindowCloseButton] setAction:@selector(performClose:)];
+	}
+	else
+	{
+		// we are adding a new metal view to expose slate drawing layers to a native widget
+		// (e.g. ApplePlatformWebBrowser) so that slate drawing layers above the native widget
+		// draw over the native widget
+		// For now, relativeTo is nil but we will need it to not be nil with multiple native layers
+		NSView* ContentView = [WindowHandle contentView];
+		
+		[ContentView addSubview: MetalView positioned: ViewOrder relativeTo: ViewRelativeTo];
+		
+		MetalView.frame = ContentView.frame;
+	}
+	
+	return MetalView;
+}
+#endif
+
+#if WITH_MULTI_VIEW_SLATE_WINDOW_SUPPORT
+void FMetalViewport::DeleteNativeLayer(int32 OldNativeLayer)
+{
+	UE_LOG(LogMetal, Display, TEXT("Deleting NativeLayer %d from Metal Viewport %p"), OldNativeLayer, this);
+	
+	ENQUEUE_RENDER_COMMAND(FlushPendingRHICommands)(
+		[Viewport = this, OldNativeLayer](FRHICommandListImmediate& RHICmdList)
+		{
+			GRHICommandList.GetImmediateCommandList().BlockUntilGPUIdle();
+			dispatch_sync(dispatch_get_main_queue(), ^
+		    {
+				FScopeLock Lock(&Viewport->Mutex);
+				
+				[Viewport->DrawLayers[OldNativeLayer].View removeFromSuperview];
+				
+				Viewport->DrawLayers[OldNativeLayer].BackBuffer[0].SafeRelease();
+				Viewport->DrawLayers[OldNativeLayer].BackBuffer[1].SafeRelease();
+				
+				Viewport->DrawLayers.Remove(OldNativeLayer);
+			});
+		});
 }
 
+void FMetalViewport::CreateNativeLayer(int32 NewNativeLayer, void* NativeViewHandle)
+{
+	if(DrawLayers.Contains(NewNativeLayer))
+	{
+		// the native layer already exists!
+		return;
+	}
+	
+	SizeX = WindowHandle.openGLFrame.size.width;
+	SizeY = WindowHandle.openGLFrame.size.height;
+	
+#if WITH_MULTI_VIEW_SLATE_WINDOW_SUPPORT
+	MainThreadCall(^{
+		UE_LOG(LogMetal, Display, TEXT("Creating NativeLayer %d in Metal Viewport %p"), NewNativeLayer, this);
+		
+		TOptional<int32> LayerRightBeforeNewLayer;
+		FScopeLock Lock(&Mutex);
+		
+		if(NewNativeLayer != std::numeric_limits<int32>::min())
+		{
+			for(const TPair<int32, DrawLayer>& DrawLayer : DrawLayers)
+			{
+				// this is reverse sorted, so once we find a DrawLayer lower
+				// than the NewNativeLayer, we know that it should be directly
+				// above that layer
+				if(DrawLayer.Key < NewNativeLayer)
+				{
+					LayerRightBeforeNewLayer = DrawLayer.Key;
+					break;
+				}
+			}
+		}
+		
+		FMetalView* View;
+		
+		if(NativeViewHandle != nil)
+		{
+			// we can't have a view we want to insert relative to if we are the
+			// base layer, because the base layer is the only native layer
+			// which exists when created.
+			check(NewNativeLayer != std::numeric_limits<int32>::min());
+			
+			NSView* ContentView = [WindowHandle contentView];
+			NSView* NativeView = (NSView*)NativeViewHandle;
+			
+			if(LayerRightBeforeNewLayer.IsSet()) 
+			{
+				UE_LOG(LogMetal, Display, TEXT("Adding NativeLayer %d right below %d"), NewNativeLayer, LayerRightBeforeNewLayer.GetValue());
+			}
+			
+			bool bPlaceBelowNativeView = 
+				LayerRightBeforeNewLayer.IsSet() && LayerRightBeforeNewLayer.GetValue() == std::numeric_limits<int32>::min();
+			bool bPlaceRelToSpecificView = LayerRightBeforeNewLayer.IsSet() && LayerRightBeforeNewLayer.GetValue() != std::numeric_limits<int32>::min();
+			
+			[ContentView addSubview: NativeView positioned: bPlaceBelowNativeView ? NSWindowBelow : NSWindowAbove 
+						 relativeTo: bPlaceRelToSpecificView ? DrawLayers[LayerRightBeforeNewLayer.GetValue()].View : nil];
+			View = CreateNewMetalView(false, NSWindowAbove, NativeView);
+		}
+		else
+		{
+			View = CreateNewMetalView(NewNativeLayer == std::numeric_limits<int32>::min(), NSWindowAbove, LayerRightBeforeNewLayer.IsSet() ? 					  DrawLayers[LayerRightBeforeNewLayer.GetValue()].View : nil);
+		}
+		
+		DrawLayers.Add(NewNativeLayer, DrawLayer
+	    {
+			.Drawable = nullptr,
+			.BackBuffer = {nullptr, nullptr},
+			.View = View
+		});
+		
+		DrawLayers.KeySort([](int32 A, int32 B)
+		{
+			return A > B;
+		});
+	});
+#endif
+	
+	Resize(SizeX, SizeY, bIsFullScreen, Format);
+}
+#endif
+
+#if WITH_MULTI_VIEW_SLATE_WINDOW_SUPPORT
+void FMetalViewport::SetDefaultNativeLayer(int32 NativeLayer)
+{
+	DefaultNativeLayer = NativeLayer;
+}
+#endif
+
 uint32 FMetalViewport::GetViewportIndex(EMetalViewportAccessFlag Accessor) const
 {
 	switch(Accessor)
@@ -199,6 +368,53 @@ uint32 FMetalViewport::GetViewportIndex(EMetalViewportAccessFlag Accessor) const
 	}
 }
 
+void FMetalViewport::ResizeBackBuffer(int32 DrawLayerKey, uint32 Index, uint32 InSizeX, uint32 InSizeY,
+									  EPixelFormat InFormat)
+{
+	TRefCountPtr<FMetalSurface> NewBackBuffer;
+
+	FString BackBufferName = FString::Printf(TEXT("BackBuffer-%d"), DrawLayerKey);
+	
+	FRHITextureCreateDesc CreateDesc =
+		FRHITextureCreateDesc::Create2D(*BackBufferName, InSizeX, InSizeY, InFormat)
+		.SetClearValue(FClearValueBinding::Transparent)
+		.SetFlags(ETextureCreateFlags::RenderTargetable);
+	
+	if (!GMetalSupportsIntermediateBackBuffer)
+	{
+		CreateDesc.AddFlags(ETextureCreateFlags::Presentable);
+	}
+	
+	CreateDesc.SetInitialState(RHIGetDefaultResourceState(CreateDesc.Flags, false));
+
+	NewBackBuffer = new FMetalSurface(Device, FMetalTextureCreateDesc(Device, CreateDesc));
+	NewBackBuffer->Viewport = this;
+
+#if WITH_MULTI_VIEW_SLATE_WINDOW_SUPPORT
+	DrawLayers[DrawLayerKey].BackBuffer[Index] = NewBackBuffer;
+	DrawLayers[DrawLayerKey].BackBuffer[EMetalViewportAccessRHI] = DrawLayers[DrawLayerKey].BackBuffer[Index];
+#else
+	BackBuffer[Index] = NewBackBuffer;
+	BackBuffer[EMetalViewportAccessRHI] = BackBuffer[Index];
+#endif
+}
+
+void ResizeMetalLayer(CAMetalLayer* MetalLayer, MTL::PixelFormat MetalFormat, bool bUseHDR,
+					  uint32 InSizeX, uint32 InSizeY)
+{	
+	MetalLayer.drawableSize = CGSizeMake(InSizeX, InSizeY);
+	
+	if (MetalFormat != (MTL::PixelFormat)MetalLayer.pixelFormat)
+	{
+		MetalLayer.pixelFormat = (MTLPixelFormat)MetalFormat;
+	}
+	
+	if (bUseHDR != MetalLayer.wantsExtendedDynamicRangeContent)
+	{
+		MetalLayer.wantsExtendedDynamicRangeContent = bUseHDR;
+	}
+}
+
 void FMetalViewport::Resize(uint32 InSizeX, uint32 InSizeY, bool bInIsFullscreen, EPixelFormat InFormat)
 {
 	bIsFullScreen = bInIsFullscreen;
@@ -213,16 +429,30 @@ void FMetalViewport::Resize(uint32 InSizeX, uint32 InSizeY, bool bInIsFullscreen
         {
             GRHICommandList.GetImmediateCommandList().BlockUntilGPUIdle();
         });
+	
+#if WITH_MULTI_VIEW_SLATE_WINDOW_SUPPORT
+	bool bNeedToReleaseBackBuffer = true;
+#else
+	bool bNeedToReleaseBackBuffer = IsValidRef(BackBuffer[Index]) && InFormat != BackBuffer[Index]->GetFormat();
+#endif
     
-	if (IsValidRef(BackBuffer[Index]) && InFormat != BackBuffer[Index]->GetFormat())
+	if (bNeedToReleaseBackBuffer)
 	{
 		// Really need to flush the RHI thread & GPU here...
 		AddRef();
 		ENQUEUE_RENDER_COMMAND(FlushPendingRHICommands)(
 			[Viewport = this](FRHICommandListImmediate& RHICmdList)
 			{
-				Viewport->ReleaseDrawable();
-				Viewport->Release();
+#if WITH_MULTI_VIEW_SLATE_WINDOW_SUPPORT
+				for(auto& DrawLayer : Viewport->DrawLayers)
+				{
+					Viewport->ReleaseDrawable(DrawLayer.Key);
+				}
+#else
+				Viewport->ReleaseDrawable(std::numeric_limits<int32>::min());
+#endif
+
+				Viewport->Release();			
 			});
 	}
     
@@ -234,20 +464,15 @@ void FMetalViewport::Resize(uint32 InSizeX, uint32 InSizeY, bool bInIsFullscreen
 #if PLATFORM_MAC
 	MainThreadCall(^
 	{
-		CAMetalLayer* MetalLayer = (CAMetalLayer*)View.layer;
-		
-		MetalLayer.drawableSize = CGSizeMake(InSizeX, InSizeY);
-		
-		if (MetalFormat != (MTL::PixelFormat)MetalLayer.pixelFormat)
-		{
-			MetalLayer.pixelFormat = (MTLPixelFormat)MetalFormat;
-		}
-		
-		if (bUseHDR != MetalLayer.wantsExtendedDynamicRangeContent)
+#if WITH_MULTI_VIEW_SLATE_WINDOW_SUPPORT
+		for(TPair<int32, DrawLayer>& DrawLayer : DrawLayers)
 		{
-			MetalLayer.wantsExtendedDynamicRangeContent = bUseHDR;
+			ResizeMetalLayer((CAMetalLayer*)DrawLayer.Value.View.layer, MetalFormat, 
+							 bUseHDR, InSizeX, InSizeY);
 		}
-		
+#else
+		ResizeMetalLayer((CAMetalLayer*)[View layer], MetalFormat, bUseHDR, InSizeX, InSizeY);
+#endif		
 	});
 #else
 	// A note on HDR in iOS
@@ -278,10 +503,10 @@ void FMetalViewport::Resize(uint32 InSizeX, uint32 InSizeY, bool bInIsFullscreen
 #endif
 	});
 
-#if PLATFORM_IOS && !PLATFORM_TVOS && !PLATFORM_VISIONOS
 		// Cache new size and orientation
 		SizeX = InSizeX;
 		SizeY = InSizeY;
+#if PLATFORM_IOS && !PLATFORM_TVOS && !PLATFORM_VISIONOS
 		Format = InFormat;
 		OrientationMask = CachedOrientationMask;
 		FPlatformMisc::LowLevelOutputDebugStringf(TEXT("[rotation] FMetalViewport::Resize: OrientationMask %d"), OrientationMask);
@@ -292,35 +517,33 @@ void FMetalViewport::Resize(uint32 InSizeX, uint32 InSizeY, bool bInIsFullscreen
     {
         FScopeLock Lock(&Mutex);
 
-		TRefCountPtr<FMetalSurface> NewBackBuffer;
-
-		FRHITextureCreateDesc CreateDesc =
-			FRHITextureCreateDesc::Create2D(TEXT("BackBuffer"), InSizeX, InSizeY, InFormat)
-			.SetClearValue(FClearValueBinding::None)
-			.SetFlags(ETextureCreateFlags::RenderTargetable);
-		
-		if (!GMetalSupportsIntermediateBackBuffer)
+#if WITH_MULTI_VIEW_SLATE_WINDOW_SUPPORT
+		for(TPair<int32, DrawLayer>& DrawLayer : DrawLayers)
 		{
-			CreateDesc.AddFlags(ETextureCreateFlags::Presentable);
+			ResizeBackBuffer(DrawLayer.Key, Index, InSizeX, InSizeY, InFormat);
 		}
-		
-		CreateDesc.SetInitialState(RHIGetDefaultResourceState(CreateDesc.Flags, false));
-
-		NewBackBuffer = new FMetalSurface(Device, FMetalTextureCreateDesc(Device, CreateDesc));
-		NewBackBuffer->Viewport = this;
-
-        BackBuffer[Index] = NewBackBuffer;
-		BackBuffer[EMetalViewportAccessRHI] = BackBuffer[Index];
-    }
+#else
+		ResizeBackBuffer(std::numeric_limits<int32>::min(), Index, InSizeX, InSizeY, InFormat);
+#endif
+	}
 }
 
-TRefCountPtr<FMetalSurface> FMetalViewport::GetBackBuffer(EMetalViewportAccessFlag Accessor) const
+TRefCountPtr<FMetalSurface> FMetalViewport::GetBackBuffer(EMetalViewportAccessFlag Accessor,
+														  TOptional<int32> DrawLayerIndex) const
 {
 	FScopeLock Lock(&Mutex);
-	
 	uint32 Index = GetViewportIndex(Accessor);
+	
+#if WITH_MULTI_VIEW_SLATE_WINDOW_SUPPORT
+	uint32 DrawLayerIndexValue = DrawLayerIndex.IsSet() ? DrawLayerIndex.GetValue() : DefaultNativeLayer;
+	
+	check(DrawLayers.Contains(DrawLayerIndexValue));
+	check(IsValidRef(DrawLayers[DrawLayerIndexValue].BackBuffer[Index]));
+	return DrawLayers[DrawLayerIndexValue].BackBuffer[Index];
+#else
 	check(IsValidRef(BackBuffer[Index]));
 	return BackBuffer[Index];
+#endif
 }
 
 #if PLATFORM_MAC
@@ -329,9 +552,17 @@ TRefCountPtr<FMetalSurface> FMetalViewport::GetBackBuffer(EMetalViewportAccessFl
 @end
 #endif
 
-CA::MetalDrawable* FMetalViewport::GetDrawable(EMetalViewportAccessFlag Accessor)
+CA::MetalDrawable* FMetalViewport::GetDrawableInternal(EMetalViewportAccessFlag Accessor, TOptional<int32> DrawLayerIndex,
+													   bool bLockAlreadyAcquired)
 {
-	FScopeLock Lock(&Mutex);
+	int32 DrawLayerIndexValue = DrawLayerIndex.IsSet() ? DrawLayerIndex.GetValue() : DefaultNativeLayer;
+	
+	TUniquePtr<FScopeLock> Lock;
+	
+	if(!bLockAlreadyAcquired)
+	{
+		Lock = MakeUnique<FScopeLock>(&Mutex);
+	}
 	
 #if PLATFORM_VISIONOS
 	// no CAMetalDrawable in Swift mode
@@ -341,9 +572,29 @@ CA::MetalDrawable* FMetalViewport::GetDrawable(EMetalViewportAccessFlag Accessor
 	}
 #endif
 	
+#if WITH_MULTI_VIEW_SLATE_WINDOW_SUPPORT
+	if(!DrawLayers.Contains(DrawLayerIndexValue))
+	{
+		UE_LOG(LogMetal, Warning, TEXT("Couldn't find DrawLayer with index %d in Metal Viewport %p"), DrawLayerIndexValue, this);
+		return nullptr;
+	}
+	
+	TRefCountPtr<FMetalSurface> BackBufferPtr = DrawLayers[DrawLayerIndexValue].BackBuffer[GetViewportIndex(Accessor)];
+	CA::MetalDrawable* Drawable = DrawLayers[DrawLayerIndexValue].Drawable;
+#else
+	TRefCountPtr<FMetalSurface> BackBufferPtr = BackBuffer[GetViewportIndex(Accessor)];
+#endif
+	
+	if(BackBufferPtr == nullptr)
+	{
+		UE_LOG(LogMetal, Warning, TEXT("BackBufferPtr is null with layer %d in Metal Viewport %p"), DrawLayerIndexValue, this);
+		return nullptr;
+	}
+	
+	
 	SCOPE_CYCLE_COUNTER(STAT_MetalMakeDrawableTime);
-    if (!Drawable || (Drawable->texture()->width() != BackBuffer[GetViewportIndex(Accessor)]->GetSizeX() ||
-                      Drawable->texture()->height() != BackBuffer[GetViewportIndex(Accessor)]->GetSizeY()))
+	if (!Drawable || (Drawable->texture()->width() != BackBufferPtr->GetSizeX() ||
+					  Drawable->texture()->height() != BackBufferPtr->GetSizeY()))
 	{
 		// Drawable changed, release the previously retained object.
 		if (Drawable != nullptr)
@@ -352,83 +603,158 @@ CA::MetalDrawable* FMetalViewport::GetDrawable(EMetalViewportAccessFlag Accessor
 			Drawable = nullptr;
 		}
 
-        MTL_SCOPED_AUTORELEASE_POOL;
-        {
-            FRenderThreadIdleScope IdleScope(ERenderThreadIdleTypes::WaitingForGPUPresent);
+		MTL_SCOPED_AUTORELEASE_POOL;
+		{
+			FRenderThreadIdleScope IdleScope(ERenderThreadIdleTypes::WaitingForGPUPresent);
 
 #if PLATFORM_MAC
-            CA::MetalLayer* CurrentLayer = (__bridge CA::MetalLayer*)[View layer];
-            if (GMetalNonBlockingPresent == 0 || [((id<CAMetalLayerSPI>)CurrentLayer) isDrawableAvailable])
-            {
-                Drawable = CurrentLayer ? CurrentLayer->nextDrawable() : nullptr;
-            }
+#if WITH_MULTI_VIEW_SLATE_WINDOW_SUPPORT
+			CA::MetalLayer* CurrentLayer = (__bridge CA::MetalLayer*)[DrawLayers[DrawLayerIndexValue].View layer];
+#else
+			CA::MetalLayer* CurrentLayer = (__bridge CA::MetalLayer*)[View layer];
+#endif
+			
+			if (GMetalNonBlockingPresent == 0 || [((id<CAMetalLayerSPI>)CurrentLayer) isDrawableAvailable])
+			{
+				Drawable = CurrentLayer ? CurrentLayer->nextDrawable() : nullptr;
+			}
 
 #if METAL_DEBUG_OPTIONS
-            if (Drawable)
-            {
-                CGSize Size = Drawable->layer()->drawableSize();
-                if ((Size.width != BackBuffer[GetViewportIndex(Accessor)]->GetSizeX() || Size.height != BackBuffer[GetViewportIndex(Accessor)]->GetSizeY()))
-                {
-                    UE_LOG(LogMetal, Display, TEXT("Viewport Size Mismatch: Drawable W:%f H:%f, Viewport W:%u H:%u"), Size.width, Size.height, BackBuffer[GetViewportIndex(Accessor)]->GetSizeX(), BackBuffer[GetViewportIndex(Accessor)]->GetSizeY());
-                }
-            }
+			if (Drawable)
+			{
+#if WITH_MULTI_VIEW_SLATE_WINDOW_SUPPORT
+				TRefCountPtr<FMetalSurface> CurrentBackBuffer = DrawLayers[DrawLayerIndexValue].BackBuffer[GetViewportIndex(Accessor)];
+#else
+				TRefCountPtr<FMetalSurface> CurrentBackBuffer = BackBuffer[GetViewportIndex(Accessor)];
+#endif
+				
+				CGSize Size = Drawable->layer()->drawableSize();
+				if ((Size.width != CurrentBackBuffer->GetSizeX() || Size.height != CurrentBackBuffer->GetSizeY()))
+				{
+					UE_LOG(LogMetal, Display, TEXT("Viewport Size Mismatch: Drawable W:%f H:%f, Viewport W:%u H:%u"), Size.width, Size.height, 
+						   CurrentBackBuffer->GetSizeX(), CurrentBackBuffer->GetSizeY());
+				}
+			}
 #endif // METAL_DEBUG_OPTIONS
 
 #else // PLATFORM_MAC
-            CGSize Size;
-            IOSAppDelegate* AppDelegate = [IOSAppDelegate GetDelegate];
-            do
-            {
-                Drawable = (__bridge CA::MetalDrawable*)[AppDelegate.IOSView MakeDrawable];
-                if (Drawable != nullptr)
-                {
-                    Size.width = Drawable->texture()->width();
-                    Size.height = Drawable->texture()->height();
-                }
-                else
-                {
-                    FPlatformProcess::SleepNoStats(0.001f);
-                }
-            }
-            while (Drawable == nullptr || Size.width != BackBuffer[GetViewportIndex(Accessor)]->GetSizeX() || Size.height != BackBuffer[GetViewportIndex(Accessor)]->GetSizeY());
+			CGSize Size;
+			IOSAppDelegate* AppDelegate = [IOSAppDelegate GetDelegate];
+			do
+			{
+				Drawable = (__bridge CA::MetalDrawable*)[AppDelegate.IOSView MakeDrawable];
+				if (Drawable != nullptr)
+				{
+					Size.width = Drawable->texture()->width();
+					Size.height = Drawable->texture()->height();
+				}
+				else
+				{
+					FPlatformProcess::SleepNoStats(0.001f);
+				}
+			}
+			while (Drawable == nullptr || Size.width != BackBuffer[GetViewportIndex(Accessor)]->GetSizeX() || Size.height != BackBuffer[GetViewportIndex(Accessor)]->GetSizeY());
 
 #endif // PLATFORM_MAC
-        }
+		}
 
-        // Retain the drawable here or it will be released when the
-        // autorelease pool goes out of scope.
-        if (Drawable != nullptr)
-        {
-            Drawable->retain();
-        }
+		// Retain the drawable here or it will be released when the
+		// autorelease pool goes out of scope.
+		if (Drawable != nullptr)
+		{
+			Drawable->retain();
+		}
 	}
-
+	
+#if WITH_MULTI_VIEW_SLATE_WINDOW_SUPPORT
+	DrawLayers[DrawLayerIndexValue].Drawable = Drawable;
+#endif
+	
 	return Drawable;
 }
 
-MTL::Texture* FMetalViewport::GetDrawableTexture(EMetalViewportAccessFlag Accessor)
+CA::MetalDrawable* FMetalViewport::GetDrawable(EMetalViewportAccessFlag Accessor,
+											   TOptional<int32> DrawLayerIndex)
+{
+	return GetDrawableInternal(Accessor, DrawLayerIndex, false);
+}
+
+MTL::Texture* FMetalViewport::GetDrawableTexture(EMetalViewportAccessFlag Accessor, TOptional<int32> DrawLayerIndex)
 {
-	CA::MetalDrawable* CurrentDrawable = GetDrawable(Accessor);
+	FScopeLock Lock(&Mutex);
+	int32 DrawLayerIndexValue = DrawLayerIndex.IsSet() ? DrawLayerIndex.GetValue() : DefaultNativeLayer;
+	CA::MetalDrawable* CurrentDrawable = GetDrawableInternal(Accessor, DrawLayerIndexValue, true);
+	
+	if(CurrentDrawable == nullptr)
+	{
+		UE_LOG(LogMetal, Warning, TEXT("CurrentDrawable is null with layer %d in Metal Viewport %p"), DrawLayerIndexValue, this);
+		return nullptr;
+	}
+	
     uint32 Index = GetViewportIndex(Accessor);
     
 #if METAL_DEBUG_OPTIONS
     MTL_SCOPED_AUTORELEASE_POOL;
 
+#if WITH_MULTI_VIEW_SLATE_WINDOW_SUPPORT
+    CAMetalLayer* CurrentLayer = (CAMetalLayer*)[DrawLayers[DrawLayerIndexValue].View layer];
+	TRefCountPtr<FMetalSurface> BackBufferPtr = DrawLayers[DrawLayerIndexValue].BackBuffer[Index];
+#else
 #if PLATFORM_MAC
-    CAMetalLayer* CurrentLayer = (CAMetalLayer*)[View layer];
+	CAMetalLayer* CurrentLayer = (CAMetalLayer*)[View layer];
 #else
     CAMetalLayer* CurrentLayer = (CAMetalLayer*)[[IOSAppDelegate GetDelegate].IOSView layer];
 #endif
-    
+	TRefCountPtr<FMetalSurface> BackBufferPtr = BackBuffer[Index];
+#endif
+	
+	if(BackBufferPtr == nullptr)
+	{
+		UE_LOG(LogMetal, Warning, TEXT("BackBufferPtr is null with layer %d in Metal Viewport %p"), DrawLayerIndexValue, this);
+		return nullptr;
+	}
+	
     CGSize Size = CurrentLayer.drawableSize;
-    if (CurrentDrawable->texture()->width() != BackBuffer[Index]->GetSizeX() || CurrentDrawable->texture()->height() != BackBuffer[Index]->GetSizeY())
-    {
-        UE_LOG(LogMetal, Display, TEXT("Viewport Size Mismatch: Drawable W:%f H:%f, Texture W:%llu H:%llu, Viewport W:%u H:%u"), Size.width, Size.height, CurrentDrawable->texture()->height(), CurrentDrawable->texture()->height(), BackBuffer[Index]->GetSizeX(), BackBuffer[Index]->GetSizeY());
+	CGSize ViewportSize = CGSizeMake(BackBufferPtr->GetSizeX(), BackBufferPtr->GetSizeY());
+    if (CurrentDrawable->texture()->width() != ViewportSize.width || 
+		CurrentDrawable->texture()->height() != ViewportSize.height)
+    {		
+        UE_LOG(LogMetal, Display, TEXT("Viewport Size Mismatch: Drawable W:%f H:%f, Texture W:%llu H:%llu, Viewport W:%f H:%f, layer frame W:%f, H:%f"), 
+		    Size.width, Size.height, 
+		    CurrentDrawable->texture()->height(), CurrentDrawable->texture()->height(), 
+		    ViewportSize.width, ViewportSize.height, CurrentLayer.frame.size.width, CurrentLayer.frame.size.height);
+		
+#if WITH_MULTI_VIEW_SLATE_WINDOW_SUPPORT
+		DrawLayers[DrawLayerIndexValue].View.frame.size = ViewportSize;
+#endif
     }
 #endif
     
-	DrawableTextures[Index] = CurrentDrawable->texture();
-	return CurrentDrawable->texture();
+	MTL::Texture* ReturnValue = CurrentDrawable->texture();
+	
+	DrawableTextures[Index] = ReturnValue;
+
+	return ReturnValue;
+}
+
+MTL::Texture* FMetalViewport::GetDrawableTexture(EMetalViewportAccessFlag Accessor, FMetalSurface* SurfaceWithDrawLayer)
+{
+#if WITH_MULTI_VIEW_SLATE_WINDOW_SUPPORT
+	for(TPair<int32, DrawLayer> &DrawLayer : DrawLayers)
+	{
+		for(int i = 0; 2 > i; ++i)
+		{
+			if(SurfaceWithDrawLayer == DrawLayer.Value.BackBuffer[i])
+			{
+				return GetDrawableTexture(Accessor, DrawLayer.Key);
+			}
+		}
+	}
+	
+	return nullptr;
+#else
+	return GetDrawableTexture(Accessor, TOptional<int32>(std::numeric_limits<int32>::min()));
+#endif
 }
 
 MTL::Texture* FMetalViewport::GetCurrentTexture(EMetalViewportAccessFlag Accessor)
@@ -437,37 +763,236 @@ MTL::Texture* FMetalViewport::GetCurrentTexture(EMetalViewportAccessFlag Accesso
 	return DrawableTextures[Index];
 }
 
-void FMetalViewport::ReleaseDrawable()
+void FMetalViewport::ReleaseDrawable(int32 DrawLayerKey)
 {
+#if WITH_MULTI_VIEW_SLATE_WINDOW_SUPPORT
+	TRefCountPtr<FMetalSurface> BackBufferTexture = DrawLayers[DrawLayerKey].BackBuffer[GetViewportIndex(EMetalViewportAccessRHI)];
+	CA::MetalDrawable* Drawable = DrawLayers[DrawLayerKey].Drawable;
+#else
+	TRefCountPtr<FMetalSurface> BackBufferTexture = BackBuffer[GetViewportIndex(EMetalViewportAccessRHI)];
+#endif
+	
+	
 	if (Drawable != nullptr)
 	{
 		Drawable->release();
+#if WITH_MULTI_VIEW_SLATE_WINDOW_SUPPORT
+		DrawLayers[DrawLayerKey].Drawable = nullptr;
+#else
 		Drawable = nullptr;
+#endif
 	}
 
-	if (!GMetalSupportsIntermediateBackBuffer && IsValidRef(BackBuffer[GetViewportIndex(EMetalViewportAccessRHI)]))
+	if(!GMetalSupportsIntermediateBackBuffer)
 	{
-		BackBuffer[GetViewportIndex(EMetalViewportAccessRHI)]->ReleaseDrawableTexture();
+		if (IsValidRef(BackBufferTexture))
+		{
+			BackBufferTexture->ReleaseDrawableTexture();
+		}
 	}
 }
 
+void FMetalViewport::ReleaseDrawable()
+{
+#if WITH_MULTI_VIEW_SLATE_WINDOW_SUPPORT
+	for(TPair<int32, DrawLayer>& DrawLayer : DrawLayers)
+	{
+		ReleaseDrawable(DrawLayer.Key);
+	}
+#else
+	ReleaseDrawable(std::numeric_limits<int32>::min());
+#endif
+}
+
 #if PLATFORM_MAC
 NSWindow* FMetalViewport::GetWindow() const
 {
+#if WITH_MULTI_VIEW_SLATE_WINDOW_SUPPORT
+	return [DrawLayers[std::numeric_limits<int32>::min()].View window];
+#else
 	return [View window];
+#endif
+}
+#endif
+
+#if WITH_MULTI_VIEW_SLATE_WINDOW_SUPPORT
+FMetalViewport::CommittedDrawable FMetalViewport::GetCommittedDrawableToPresent(int32 DrawLayer, FMetalCommandBuffer* CommandBuffer,
+																bool& bLayersToDraw, FMetalView** TheView)
+#else
+FMetalViewport::CommittedDrawable FMetalViewport::GetCommittedDrawableToPresent(int32 DrawLayer, FMetalCommandBuffer* CommandBuffer,
+																				bool& bLayersToDraw)
+#endif
+{
+	CA::MetalDrawable* LocalDrawable = GetDrawable(EMetalViewportAccessDisplayLink, DrawLayer);
+	LocalDrawable->retain();
+	MTL::Texture* DrawableTexture = GetDrawableTexture(EMetalViewportAccessDisplayLink, DrawLayer);
+	
+	if(DrawableTexture)
+	{
+		check(CommandBuffer);
+		
+		if (GMetalSupportsIntermediateBackBuffer)
+		{
+			TRefCountPtr<FMetalSurface> Texture = GetBackBuffer(EMetalViewportAccessRHI, DrawLayer);
+			check(IsValidRef(Texture));
+			
+			MTLTexturePtr Src = Texture->Texture;
+			MTLTexturePtr Dst = NS::RetainPtr(DrawableTexture);
+			
+			NS::UInteger Width = FMath::Min(Src->width(), Dst->width());
+			NS::UInteger Height = FMath::Min(Src->height(), Dst->height());
+			
+			MTLBlitCommandEncoderPtr Encoder = NS::RetainPtr(CommandBuffer->GetMTLCmdBuffer()->blitCommandEncoder());
+			check(Encoder);
+
+			Encoder->copyFromTexture(Src.get(), 0, 0, MTL::Origin(0, 0, 0), MTL::Size(Width, Height, 1), Dst.get(), 0, 0, MTL::Origin(0, 0, 0));
+			Encoder->endEncoding();
+			
+			LocalDrawable->release();
+			LocalDrawable = nullptr;
+		}
+		
+		if(!bLayersToDraw)
+		{
+#if WITH_MULTI_VIEW_SLATE_WINDOW_SUPPORT
+			check(*TheView == nullptr);
+			*TheView = DrawLayers[DrawLayer].View;
+#endif
+		}
+		
+		bLayersToDraw = true;
+	}
+	
+	return CommittedDrawable
+	{
+		.Drawable = LocalDrawable,
+		.Texture = DrawableTexture,
+		.Key = DrawLayer
+	};
 }
+
+void FMetalViewport::PresentDrawLayers(FMetalCommandBuffer* CommandBuffer)
+{
+	FPlatformAtomics::InterlockedDecrement(&FrameAvailable);
+	bool bLayersToDraw = false;
+	
+	TArray<CommittedDrawable> CommittedDrawables;
+	
+#if PLATFORM_MAC
+#if WITH_MULTI_VIEW_SLATE_WINDOW_SUPPORT
+	FMetalView* TheView = nullptr;
+#else
+	FMetalView* TheView = View;
+#endif
+#endif
+	
+#if WITH_MULTI_VIEW_SLATE_WINDOW_SUPPORT
+	for(TPair<int32, DrawLayer> &DrawLayer : DrawLayers)
+	{
+		CommittedDrawables.Add(GetCommittedDrawableToPresent(DrawLayer.Key, CommandBuffer, 
+															 bLayersToDraw, &TheView));
+	}
+#else
+	CA::MetalDrawable* LocalDrawable = GetDrawable(EMetalViewportAccessDisplayLink);
+	LocalDrawable->retain();
+	
+	CommittedDrawables.Add(GetCommittedDrawableToPresent(std::numeric_limits<int32>::min(), CommandBuffer, bLayersToDraw));
+#endif
+	
+	if(!bLayersToDraw)
+	{
+		return;
+	}
+
+	dispatch_semaphore_t& FrameSemaphore = Device.GetFrameSemaphore();
+	dispatch_retain(FrameSemaphore);
+	
+#if PLATFORM_MAC
+	MTL::HandlerFunction CommandBufferHandler = [TheView, CommittedDrawables, FrameSemaphore](MTL::CommandBuffer* cmd_buf)
+#else
+	MTL::HandlerFunction CommandBufferHandler = [CommittedDrawables, FrameSemaphore](MTL::CommandBuffer* cmd_buf)
 #endif
+	{
+		dispatch_semaphore_signal(FrameSemaphore);
+		dispatch_release(FrameSemaphore);
+
+		for(CommittedDrawable CommittedDrawableInstance : CommittedDrawables)
+		{
+			if(CommittedDrawableInstance.Texture && CommittedDrawableInstance.Drawable)
+			{
+				CommittedDrawableInstance.Drawable->release();
+			}
+		}
+#if PLATFORM_MAC
+		MainThreadCall(^{
+			FCocoaWindow* Window = (FCocoaWindow*)[TheView window];
+			[Window startRendering];
+		}, false);
+#endif
+	};
+	
+#if PLATFORM_MAC		// Mac needs the older way to present otherwise we end up with bad behaviour of the completion handlers that causes GPU timeouts.
+	MTL::HandlerFunction ScheduledHandler = [CommittedDrawables](MTL::CommandBuffer*)
+	{
+		int DrawablesPresented = 0;
+		for(CommittedDrawable CommittedDrawableInstance : CommittedDrawables)
+		{
+			if(CommittedDrawableInstance.Texture && CommittedDrawableInstance.Drawable)
+			{
+				CommittedDrawableInstance.Drawable->present();
+				++DrawablesPresented;
+			}
+		}
+	};
+			
+	CommandBuffer->GetMTLCmdBuffer()->addCompletedHandler(CommandBufferHandler);
+	CommandBuffer->GetMTLCmdBuffer()->addScheduledHandler(ScheduledHandler);
+
+#else // PLATFORM_MAC
+	CommandBuffer->GetMTLCmdBuffer()->addCompletedHandler(CommandBufferHandler);
+
+	{
+		uint32 FramePace = FPlatformRHIFramePacer::GetFramePace();
+		float MinPresentDuration = FramePace ? (1.0f / (float)FramePace) : 0.0f;
+		
+		// Queue this on the current command buffer to ensure that all work is committed prior to the present, present only knows about dependencies on committed work.
+		if (MinPresentDuration && GEnablePresentPacing)
+		{
+			CommandBuffer->GetMTLCmdBuffer()->presentDrawableAfterMinimumDuration(LocalDrawable, 1.0f/(float)FramePace);
+		}
+		else
+		{
+			CommandBuffer->GetMTLCmdBuffer()->presentDrawable(LocalDrawable);
+		}
+	}
+#endif // PLATFORM_MAC
+	
+	// Wait for the frame semaphore
+	dispatch_semaphore_wait(Device.GetFrameSemaphore(), DISPATCH_TIME_FOREVER);
+}
 
 void FMetalViewport::Present(FMetalCommandBuffer* CommandBuffer, bool bLockToVsync)
 {
 	FScopeLock Lock(&Mutex);
 	
 #if PLATFORM_MAC
+#if WITH_MULTI_VIEW_SLATE_WINDOW_SUPPORT
+	NSNumber* ScreenId = [DrawLayers[std::numeric_limits<int32>::min()].View.window.screen.deviceDescription objectForKey:@"NSScreenNumber"];
+#else
 	NSNumber* ScreenId = [View.window.screen.deviceDescription objectForKey:@"NSScreenNumber"];
+#endif
 	DisplayID = ScreenId.unsignedIntValue;
 	{
+#if WITH_MULTI_VIEW_SLATE_WINDOW_SUPPORT
+		for(TPair<int32, DrawLayer>& DrawLayer : DrawLayers)
+		{
+			FCAMetalLayer* CurrentLayer = (FCAMetalLayer*)[DrawLayer.Value.View layer];
+			CurrentLayer.displaySyncEnabled = bLockToVsync || (!(IsRunningGame() && bIsFullScreen));
+		}
+#else
 		FCAMetalLayer* CurrentLayer = (FCAMetalLayer*)[View layer];
 		CurrentLayer.displaySyncEnabled = bLockToVsync || (!(IsRunningGame() && bIsFullScreen));
+#endif
 	}
 #endif
 
@@ -493,7 +1018,6 @@ void FMetalViewport::Present(FMetalCommandBuffer* CommandBuffer, bool bLockToVsy
 	}
 #endif
 
-	LastCompleteFrame = GetBackBuffer(EMetalViewportAccessRHI);
 	FPlatformAtomics::InterlockedExchange(&FrameAvailable, 1);
 	
 #if !PLATFORM_MAC
@@ -504,86 +1028,7 @@ void FMetalViewport::Present(FMetalCommandBuffer* CommandBuffer, bool bLockToVsy
 
 	if (FrameAvailable > 0)
 	{
-		FPlatformAtomics::InterlockedDecrement(&FrameAvailable);
-		CA::MetalDrawable* LocalDrawable = GetDrawable(EMetalViewportAccessDisplayLink);
-		LocalDrawable->retain();
-		MTL::Texture* DrawableTexture = GetDrawableTexture(EMetalViewportAccessDisplayLink);
-		
-		if (DrawableTexture)
-		{
-			check(CommandBuffer);
-			
-			if (GMetalSupportsIntermediateBackBuffer)
-			{
-				TRefCountPtr<FMetalSurface> Texture = LastCompleteFrame;
-				check(IsValidRef(Texture));
-				
-				MTLTexturePtr Src = Texture->Texture;
-				MTLTexturePtr Dst = NS::RetainPtr(DrawableTexture);
-				
-				NS::UInteger Width = FMath::Min(Src->width(), Dst->width());
-				NS::UInteger Height = FMath::Min(Src->height(), Dst->height());
-				
-				MTLBlitCommandEncoderPtr Encoder = NS::RetainPtr(CommandBuffer->GetMTLCmdBuffer()->blitCommandEncoder());
-				check(Encoder);
-
-				Encoder->copyFromTexture(Src.get(), 0, 0, MTL::Origin(0, 0, 0), MTL::Size(Width, Height, 1), Dst.get(), 0, 0, MTL::Origin(0, 0, 0));
-				Encoder->endEncoding();
-				
-				Drawable->release();
-				Drawable = nullptr;
-			}
-			
-			dispatch_semaphore_t& FrameSemaphore = Device.GetFrameSemaphore();
-			dispatch_retain(FrameSemaphore);
-			
-#if PLATFORM_MAC
-			FMetalView* theView = View;
-			MTL::HandlerFunction CommandBufferHandler = [LocalDrawable, theView, FrameSemaphore](MTL::CommandBuffer* cmd_buf)
-#else
-			MTL::HandlerFunction CommandBufferHandler = [LocalDrawable, FrameSemaphore](MTL::CommandBuffer* cmd_buf)
-#endif
-			{
-				dispatch_semaphore_signal(FrameSemaphore);
-				dispatch_release(FrameSemaphore);
-
-				LocalDrawable->release();
-#if PLATFORM_MAC
-				MainThreadCall(^{
-					FCocoaWindow* Window = (FCocoaWindow*)[theView window];
-					[Window startRendering];
-				}, false);
-#endif
-			};
-			
-#if PLATFORM_MAC		// Mac needs the older way to present otherwise we end up with bad behaviour of the completion handlers that causes GPU timeouts.
-			MTL::HandlerFunction ScheduledHandler = [LocalDrawable](MTL::CommandBuffer*)
-			{
-				LocalDrawable->present();
-			};
-					
-			CommandBuffer->GetMTLCmdBuffer()->addCompletedHandler(CommandBufferHandler);
-			CommandBuffer->GetMTLCmdBuffer()->addScheduledHandler(ScheduledHandler);
-
-#else // PLATFORM_MAC
-			CommandBuffer->GetMTLCmdBuffer()->addCompletedHandler(CommandBufferHandler);
-
-			{
-				// Queue this on the current command buffer to ensure that all work is committed prior to the present, present only knows about dependencies on committed work.
-				if (MinPresentDuration && GEnablePresentPacing)
-				{
-					CommandBuffer->GetMTLCmdBuffer()->presentDrawableAfterMinimumDuration(LocalDrawable, 1.0f/(float)FramePace);
-				}
-				else
-				{
-					CommandBuffer->GetMTLCmdBuffer()->presentDrawable(LocalDrawable);
-				}
-			}
-#endif // PLATFORM_MAC
-			
-			// Wait for the frame semaphore
-			dispatch_semaphore_wait(Device.GetFrameSemaphore(), DISPATCH_TIME_FOREVER);
-		}
+		PresentDrawLayers(CommandBuffer);
 	}
 }
 
@@ -811,7 +1256,15 @@ FTextureRHIRef FMetalDynamicRHI::RHIGetViewportBackBuffer(FRHIViewport* Viewport
     MTL_SCOPED_AUTORELEASE_POOL;
     
 	FMetalViewport* Viewport = ResourceCast(ViewportRHI);
-	return FTextureRHIRef(Viewport->GetBackBuffer(EMetalViewportAccessRenderer).GetReference());
+	return FTextureRHIRef(Viewport->GetBackBuffer(EMetalViewportAccessRenderer, TOptional<int32>()).GetReference());
+}
+
+FTextureRHIRef FMetalDynamicRHI::RHIGetViewportBackBufferAtLayerIndex(FRHIViewport* ViewportRHI, uint32 DrawLayerIndex)
+{
+	MTL_SCOPED_AUTORELEASE_POOL;
+	
+	FMetalViewport* Viewport = ResourceCast(ViewportRHI);
+	return FTextureRHIRef(Viewport->GetBackBuffer(EMetalViewportAccessRenderer, DrawLayerIndex).GetReference());
 }
 
 void FMetalDynamicRHI::RHIAdvanceFrameForGetViewportBackBuffer(FRHIViewport* ViewportRHI, bool bPresent)
diff --git a/Engine/Source/Runtime/Apple/MetalRHI/Public/MetalViewport.h b/Engine/Source/Runtime/Apple/MetalRHI/Public/MetalViewport.h
--- a/Engine/Source/Runtime/Apple/MetalRHI/Public/MetalViewport.h
+++ b/Engine/Source/Runtime/Apple/MetalRHI/Public/MetalViewport.h
@@ -21,6 +21,9 @@
 #if PLATFORM_MAC
 #include "Mac/CocoaTextView.h"
 @interface FMetalView : FCocoaTextView
+{
+	bool bAllowMouseActions;
+}
 @end
 #endif
 #include "HAL/PlatformFramePacer.h"
@@ -29,6 +32,7 @@
 #import <CompositorServices/CompositorServices.h>
 #endif
 
+@class FCocoaWindow;
 class FMetalSurface;
 class FMetalDevice;
 
@@ -60,14 +64,25 @@ public:
 
 	void Resize(uint32 InSizeX, uint32 InSizeY, bool bInIsFullscreen,EPixelFormat Format);
 	
-	TRefCountPtr<FMetalSurface> GetBackBuffer(EMetalViewportAccessFlag Accessor) const;
-	CA::MetalDrawable* GetDrawable(EMetalViewportAccessFlag Accessor);
-	MTL::Texture* GetDrawableTexture(EMetalViewportAccessFlag Accessor);
+	TRefCountPtr<FMetalSurface> GetBackBuffer(EMetalViewportAccessFlag Accessor, 
+											  TOptional<int32> DrawLayerIndex = TOptional<int32>()) const;
+	
+	CA::MetalDrawable* GetDrawable(EMetalViewportAccessFlag Accessor, TOptional<int32> DrawLayerIndex = TOptional<int32>());
+	MTL::Texture* GetDrawableTexture(EMetalViewportAccessFlag Accessor, TOptional<int32> DrawLayerIndex = TOptional<int32>());
+	MTL::Texture* GetDrawableTexture(EMetalViewportAccessFlag Accessor, FMetalSurface* SurfaceWithDrawLayer);
 	MTL::Texture* GetCurrentTexture(EMetalViewportAccessFlag Accessor);
+	
 	void ReleaseDrawable(void);
+	void ReleaseDrawable(int32 DrawLayerKey);
+	
+#if WITH_MULTI_VIEW_SLATE_WINDOW_SUPPORT
+	virtual void CreateNativeLayer(int32 NewNativeLayer, void* NativeViewHandle) override;
+	virtual void DeleteNativeLayer(int32 OldNativeLayer) override;
+	virtual void SetDefaultNativeLayer(int32 NativeLayer) override;
+#endif
 
 	// supports pulling the raw MTLTexture
-	virtual void* GetNativeBackBufferTexture() const override { return GetBackBuffer(EMetalViewportAccessRenderer).GetReference(); }
+	virtual void* GetNativeBackBufferTexture() const override { return GetBackBuffer(EMetalViewportAccessRenderer, TOptional<int32>()).GetReference(); }
 	virtual void* GetNativeBackBufferRT() const override { return (const_cast<FMetalViewport *>(this))->GetDrawableTexture(EMetalViewportAccessRenderer); }
 	
 #if PLATFORM_MAC
@@ -91,38 +106,88 @@ public:
 #endif
 	
 private:
+#if PLATFORM_MAC
+	FMetalView* CreateNewMetalView(bool bIsContentView, NSWindowOrderingMode ViewOrder, NSView* ViewRelativeTo);
+#endif
 	uint32 GetViewportIndex(EMetalViewportAccessFlag Accessor) const;
+	void PresentDrawLayers(FMetalCommandBuffer* CommandBuffer);
+	
+	CA::MetalDrawable* GetDrawableInternal(EMetalViewportAccessFlag Accessor, TOptional<int32> DrawLayerIndex,
+										   bool bLockAlreadyAquired);
 
 private:
+#if WITH_MULTI_VIEW_SLATE_WINDOW_SUPPORT
+	struct DrawLayer
+	{
+		TRefCountPtr<FMetalSurface> BackBuffer[2];
+		CA::MetalDrawable* Drawable;
+#if PLATFORM_MAC
+		FMetalView* View;
+#else
+	#error "Multi view slate window rendering mode not supported on non-mac platforms!"
+#endif
+	};
+#endif
+	
+	struct CommittedDrawable
+	{
+		CA::MetalDrawable* Drawable;
+		MTL::Texture* Texture;
+		int32 Key;
+	};
+	
+#if WITH_MULTI_VIEW_SLATE_WINDOW_SUPPORT
+	CommittedDrawable GetCommittedDrawableToPresent(int32 DrawLayer, FMetalCommandBuffer* CommandBuffer,
+													bool& bLayersToDraw, FMetalView** TheView);
+#else
+	CommittedDrawable GetCommittedDrawableToPresent(int32 DrawLayer, FMetalCommandBuffer* CommandBuffer,
+													bool& bLayersToDraw);
+#endif
+	
+	void ResizeBackBuffer(int32 DrawLayerKey, uint32 Index, uint32 InSizeX, uint32 InSizeY, EPixelFormat InFormat);
+	
 #if PLATFORM_VISIONOS
 	CP_OBJECT_cp_layer_renderer* SwiftLayer = nullptr;
 #endif
 	
 	FMetalDevice& Device;
-	CA::MetalDrawable* Drawable;
-	TRefCountPtr<FMetalSurface> BackBuffer[2];
 	mutable FCriticalSection Mutex;
 	
-    MTL::Texture* DrawableTextures[2];
-	
 	uint32 DisplayID;
 	FMetalViewportPresentHandler Block;
 	volatile int32 FrameAvailable;
-	TRefCountPtr<FMetalSurface> LastCompleteFrame;
 	bool bIsFullScreen;
 
-#if PLATFORM_IOS && !PLATFORM_TVOS && !PLATFORM_VISIONOS	
+
+#if !PLATFORM_TVOS && !PLATFORM_VISIONOS
 	uint32 SizeX;
 	uint32 SizeY;
 	EPixelFormat Format;
+#endif
+#if PLATFORM_IOS && !PLATFORM_TVOS && !PLATFORM_VISIONOS	
 	UIInterfaceOrientationMask OrientationMask;
 	UIInterfaceOrientationMask OldOrientationMask;
     FDelegateHandle OrientationChangedHandle;
 #endif
 
+	int32 DefaultNativeLayer;
+	
+#if PLATFORM_MAC
+	FCocoaWindow* WindowHandle;
+#endif
+	
+#if WITH_MULTI_VIEW_SLATE_WINDOW_SUPPORT
+	// all views with slate layers >= the key get drawn in the value metal view
+	TMap<int32, DrawLayer> DrawLayers;
+	TArray<NSView*> ViewsAdded;
+#else
+	CA::MetalDrawable* Drawable;
+	TRefCountPtr<FMetalSurface> BackBuffer[2];
 #if PLATFORM_MAC
 	FMetalView* View;
 #endif
+#endif
+	MTL::Texture* DrawableTextures[2];
     
 #if PLATFORM_MAC || PLATFORM_VISIONOS
     FCustomPresentRHIRef CustomPresent;
diff --git a/Engine/Source/Runtime/RHI/Public/DynamicRHI.h b/Engine/Source/Runtime/RHI/Public/DynamicRHI.h
--- a/Engine/Source/Runtime/RHI/Public/DynamicRHI.h
+++ b/Engine/Source/Runtime/RHI/Public/DynamicRHI.h
@@ -647,6 +647,11 @@ public:
 	// With RHI thread, this is the current backbuffer from the perspective of the render thread.
 	// FlushType: Thread safe
 	virtual FTextureRHIRef RHIGetViewportBackBuffer(FRHIViewport* Viewport) = 0;
+	
+	virtual FTextureRHIRef RHIGetViewportBackBufferAtLayerIndex(FRHIViewport* Viewport, uint32 DrawLayerIndex) 
+	{
+		return nullptr;
+	}
 
 	virtual uint32 RHIGetHTilePlatformConfig(uint32 DepthWidth, uint32 DepthHeight) const
 	{
@@ -1252,6 +1257,11 @@ inline FTextureRHIRef RHIGetViewportBackBuffer(FRHIViewport* Viewport)
 	return GDynamicRHI->RHIGetViewportBackBuffer(Viewport);
 }
 
+inline FTextureRHIRef RHIGetViewportBackBufferAtLayerIndex(FRHIViewport* Viewport, uint32 DrawLayerIndex)
+{
+	return GDynamicRHI->RHIGetViewportBackBufferAtLayerIndex(Viewport, DrawLayerIndex);
+}
+
 UE_DEPRECATED(5.8, "RHIGetViewportBackBufferUAV has been deprecated and always returns null. Use the standard UAV creation functions instead. The platform must support back buffer UAVs, which is indicated by the TexCreate_UAV flag in the back buffer texture descriptor.")
 inline FUnorderedAccessViewRHIRef RHIGetViewportBackBufferUAV(FRHIViewport* Viewport)
 {
diff --git a/Engine/Source/Runtime/RHI/Public/RHIResources.h b/Engine/Source/Runtime/RHI/Public/RHIResources.h
--- a/Engine/Source/Runtime/RHI/Public/RHIResources.h
+++ b/Engine/Source/Runtime/RHI/Public/RHIResources.h
@@ -2509,6 +2509,26 @@ public:
 	 * Sets custom Present handler on the viewport
 	 */
 	virtual void SetCustomPresent(class FRHICustomPresent*) {}
+	
+	/**
+	 * Creates a platform-level native slate rendering layer for handling rendering order with platform-native views.
+	 * This means that everything on top of the specified layer will be rendered to a different platform-level view than 
+	 * everything on the bottom of the specified layer.
+	 * Right now this is only implemented in MetalRHI for macOS NSView interop with Slate.
+	 */
+	virtual void CreateNativeLayer(int32 NewNativeLayer, void* NativeViewHandle) {}
+	
+	/**
+	 * Deletes a platform-level native layer created by CreateNativeLayer.
+	 * Right now this is only implemented in MetalRHI for macOS NSView interop with Slate.
+	 */
+	virtual void DeleteNativeLayer(int32 OldNativeLayer) {}
+	
+	/**
+	 * Sets the default native layer for items to render into. If the WITH_MULTI_VIEW_SLATE_WINDOW_SUPPORT
+	 * compile flag is false this function will have no effect. 
+	 */
+	virtual void SetDefaultNativeLayer(int32 NativeLayer) {}
 
 	/**
 	 * Returns currently set custom present handler.
