Ensure that ResizeMetalLayer and SizeX/SizeY MetalViewport properties do not break tvOS and visionOS compilation

#Synced-CL 49810091

#Synced-CL 49810091

diff --git a/Engine/Source/Runtime/Apple/MetalRHI/Private/MetalViewport.cpp b/Engine/Source/Runtime/Apple/MetalRHI/Private/MetalViewport.cpp
--- a/Engine/Source/Runtime/Apple/MetalRHI/Private/MetalViewport.cpp
+++ b/Engine/Source/Runtime/Apple/MetalRHI/Private/MetalViewport.cpp
@@ -399,6 +399,7 @@ void FMetalViewport::ResizeBackBuffer(int32 DrawLayerKey, uint32 Index, uint32 I
 #endif
 }
 
+#if PLATFORM_MAC
 void ResizeMetalLayer(CAMetalLayer* MetalLayer, MTL::PixelFormat MetalFormat, bool bUseHDR,
 					  uint32 InSizeX, uint32 InSizeY)
 {	
@@ -414,6 +415,7 @@ void ResizeMetalLayer(CAMetalLayer* MetalLayer, MTL::PixelFormat MetalFormat, bo
 		MetalLayer.wantsExtendedDynamicRangeContent = bUseHDR;
 	}
 }
+#endif
 
 void FMetalViewport::Resize(uint32 InSizeX, uint32 InSizeY, bool bInIsFullscreen, EPixelFormat InFormat)
 {
@@ -504,12 +506,14 @@ void FMetalViewport::Resize(uint32 InSizeX, uint32 InSizeY, bool bInIsFullscreen
 	});
 
 		// Cache new size and orientation
+#if !PLATFORM_TVOS && !PLATFORM_VISIONOS
 		SizeX = InSizeX;
 		SizeY = InSizeY;
-#if PLATFORM_IOS && !PLATFORM_TVOS && !PLATFORM_VISIONOS
+#if PLATFORM_IOS
 		Format = InFormat;
 		OrientationMask = CachedOrientationMask;
 		FPlatformMisc::LowLevelOutputDebugStringf(TEXT("[rotation] FMetalViewport::Resize: OrientationMask %d"), OrientationMask);
+#endif
 #endif
 	}
 #endif
