Adding RHI Global to expose which descriptor types can be allocated together.

#Synced-CL 49956609

diff --git a/Engine/Source/Runtime/Apple/MetalRHI/Private/MetalBindlessDescriptors.cpp b/Engine/Source/Runtime/Apple/MetalRHI/Private/MetalBindlessDescriptors.cpp
--- a/Engine/Source/Runtime/Apple/MetalRHI/Private/MetalBindlessDescriptors.cpp
+++ b/Engine/Source/Runtime/Apple/MetalRHI/Private/MetalBindlessDescriptors.cpp
@@ -194,6 +194,10 @@ void FMetalBindlessDescriptorManager::Init()
 	
 	StandardResources.Init(RHICmdList, GBindlessResourceDescriptorHeapSize);
 	SamplerResources.Init(RHICmdList, GBindlessSamplerDescriptorHeapSize);
+
+	// Add the types to the global groups
+	GRHIGlobals.DescriptorTypeGroups.Emplace(StandardResources.GetTypeMask());
+	GRHIGlobals.DescriptorTypeGroups.Emplace(SamplerResources.GetTypeMask());
 	
 	bIsSupported = true;
 }
diff --git a/Engine/Source/Runtime/D3D12RHI/Private/D3D12BindlessDescriptors.cpp b/Engine/Source/Runtime/D3D12RHI/Private/D3D12BindlessDescriptors.cpp
--- a/Engine/Source/Runtime/D3D12RHI/Private/D3D12BindlessDescriptors.cpp
+++ b/Engine/Source/Runtime/D3D12RHI/Private/D3D12BindlessDescriptors.cpp
@@ -159,6 +159,10 @@ void FD3D12BindlessDescriptorAllocator::Init()
 		};
 
 		SamplerAllocator = new FRHIHeapDescriptorAllocator(D3D12DescriptorTypeMaskFromHeapType(ERHIDescriptorHeapType::Sampler), NumSamplerDescriptors, SamplerStats);
+
+		// Add the types to the global groups
+		GRHIGlobals.DescriptorTypeGroups.Emplace(ResourceAllocator->GetTypeMask());
+		GRHIGlobals.DescriptorTypeGroups.Emplace(SamplerAllocator->GetTypeMask());
 	}
 }
 
diff --git a/Engine/Source/Runtime/RHI/Public/RHIGlobals.h b/Engine/Source/Runtime/RHI/Public/RHIGlobals.h
--- a/Engine/Source/Runtime/RHI/Public/RHIGlobals.h
+++ b/Engine/Source/Runtime/RHI/Public/RHIGlobals.h
@@ -685,6 +685,9 @@ PRAGMA_ENABLE_DEPRECATION_WARNINGS
 	/** Whether dynamic (bindless) resources are supported */
 	bool bSupportsBindless = false;
 
+	/** Array of descriptor type masks showing which descriptor types can be allocated together. */
+	TArray<ERHIDescriptorTypeMask> DescriptorTypeGroups;
+
 	struct FReservedResources
 	{
 		/**
diff --git a/Engine/Source/Runtime/VulkanRHI/Private/VulkanBindlessDescriptorManager.cpp b/Engine/Source/Runtime/VulkanRHI/Private/VulkanBindlessDescriptorManager.cpp
--- a/Engine/Source/Runtime/VulkanRHI/Private/VulkanBindlessDescriptorManager.cpp
+++ b/Engine/Source/Runtime/VulkanRHI/Private/VulkanBindlessDescriptorManager.cpp
@@ -563,6 +563,9 @@ void FVulkanBindlessDescriptorManager::Init()
 				OutBufferBindingInfo.usage = BufferUsageFlags;
 			}
 
+			// Add the types to the global groups
+			GRHIGlobals.DescriptorTypeGroups.Emplace(InOutState.Allocator->GetTypeMask());
+
 			return IsSamplerSet ? 0u : DescriptorBufferSize;
 		};
 
