Metal RHI

Fixed error when updating Bindless descriptors from Bottom of Pipe.

diff --git a/Engine/Source/Runtime/Apple/MetalRHI/Private/MetalBindlessDescriptors.cpp b/Engine/Source/Runtime/Apple/MetalRHI/Private/MetalBindlessDescriptors.cpp
--- a/Engine/Source/Runtime/Apple/MetalRHI/Private/MetalBindlessDescriptors.cpp
+++ b/Engine/Source/Runtime/Apple/MetalRHI/Private/MetalBindlessDescriptors.cpp
@@ -273,7 +273,10 @@ void FMetalBindlessDescriptorManager::UpdateDescriptor(FRHICommandListBase& RHIC
 		}
 	});
 
-	RHICmdList.RHIThreadFence(true);
+	if (RHICmdList.IsTopOfPipe())
+	{
+		RHICmdList.RHIThreadFence(true);
+	}
 }
 
 void FMetalBindlessDescriptorManager::UpdateDescriptor(FRHICommandListBase& RHICmdList, FRHIDescriptorHandle DescriptorHandle, MTL::Texture* Texture, EMetalDescriptorUpdateType UpdateType)
