Metal RHI

Fix issue with thread fence at bottom of pipe

diff --git a/Engine/Source/Runtime/Apple/MetalRHI/Private/MetalUAV.cpp b/Engine/Source/Runtime/Apple/MetalRHI/Private/MetalUAV.cpp
--- a/Engine/Source/Runtime/Apple/MetalRHI/Private/MetalUAV.cpp
+++ b/Engine/Source/Runtime/Apple/MetalRHI/Private/MetalUAV.cpp
@@ -377,7 +377,10 @@ FMetalUnorderedAccessView::FMetalUnorderedAccessView(FMetalDevice& InDevice, FRH
 		UpdateView(&Context, true);
 	});
 	
-	RHICmdList.RHIThreadFence(true);
+	if (RHICmdList.IsTopOfPipe())
+	{
+		RHICmdList.RHIThreadFence(true);
+	}
 }
 
 FMetalUnorderedAccessView::~FMetalUnorderedAccessView()
