[BackgroundHttpIOS]
* Added bRequiresPlatformSDK = true to build.cs

#Synced-CL 50090217

diff --git a/Engine/Source/Runtime/Apple/BackgroundHttpIOS/BackgroundHttpIOS.build.cs b/Engine/Source/Runtime/Apple/BackgroundHttpIOS/BackgroundHttpIOS.build.cs
--- a/Engine/Source/Runtime/Apple/BackgroundHttpIOS/BackgroundHttpIOS.build.cs
+++ b/Engine/Source/Runtime/Apple/BackgroundHttpIOS/BackgroundHttpIOS.build.cs
@@ -7,6 +7,8 @@ public class BackgroundHttpIOS : ModuleRules
 {
     public BackgroundHttpIOS(ReadOnlyTargetRules Target) : base(Target)
     {
+		bRequiresPlatformSDK = true;
+
         PublicDependencyModuleNames.AddRange(
             new string[]
             {
