Metal - Remove dead code using function constants

#tests Preflight & local cook & build

#Synced-CL 50166143

diff --git a/Engine/Shaders/Public/Platform/Metal/MetalCommon.ush b/Engine/Shaders/Public/Platform/Metal/MetalCommon.ush
--- a/Engine/Shaders/Public/Platform/Metal/MetalCommon.ush
+++ b/Engine/Shaders/Public/Platform/Metal/MetalCommon.ush
@@ -7,7 +7,7 @@
 #pragma once
 
 // Update this GUID to invalidate shader recompilation for only Metal shaders
-#pragma message("UESHADERMETADATA_VERSION 607C4A09-452C-40D7-9C9C-3DBDEFFB50AD")
+#pragma message("UESHADERMETADATA_VERSION B86865F6-202A-4A70-931B-B344063E4B08")
 
 #if VERTEXSHADER && HULLSHADER
 #ifndef TESSELLATIONSHADER
diff --git a/Engine/Source/Developer/Apple/MetalShaderFormat/Private/MetalShaderCompiler.cpp b/Engine/Source/Developer/Apple/MetalShaderFormat/Private/MetalShaderCompiler.cpp
--- a/Engine/Source/Developer/Apple/MetalShaderFormat/Private/MetalShaderCompiler.cpp
+++ b/Engine/Source/Developer/Apple/MetalShaderFormat/Private/MetalShaderCompiler.cpp
@@ -595,7 +595,6 @@ void BuildMetalShaderOutput(
 		Header.RayTracing.InstanceIndexBuffer = CCHeader.RayTracingInstanceIndexBuffer;
 	}
 
-	Header.bDeviceFunctionConstants = (FCStringAnsi::Strstr(USFSource, "#define __METAL_DEVICE_CONSTANT_INDEX__ 1") != nullptr) ? 1 : 0;
 	Header.SideTable = CCHeader.SideTable;
 	Header.Bindings.ArgumentBufferMasks = CCHeader.ArgumentBuffers;
 	Header.Bindings.ArgumentBuffers = 0;
diff --git a/Engine/Source/Runtime/Apple/MetalRHI/Private/Shaders/MetalCompiledShaderKey.h b/Engine/Source/Runtime/Apple/MetalRHI/Private/Shaders/MetalCompiledShaderKey.h
--- a/Engine/Source/Runtime/Apple/MetalRHI/Private/Shaders/MetalCompiledShaderKey.h
+++ b/Engine/Source/Runtime/Apple/MetalRHI/Private/Shaders/MetalCompiledShaderKey.h
@@ -11,25 +11,23 @@
 
 struct FMetalCompiledShaderKey
 {
-	FMetalCompiledShaderKey(uint32 InCodeSize, uint32 InCodeCRC, uint32 InConstants)
+	FMetalCompiledShaderKey(uint32 InCodeSize, uint32 InCodeCRC)
 		: CodeSize(InCodeSize)
 		, CodeCRC(InCodeCRC)
-		, Constants(InConstants)
 	{
 		// VOID
 	}
 
 	friend bool operator ==(const FMetalCompiledShaderKey& A, const FMetalCompiledShaderKey& B)
 	{
-		return A.CodeSize == B.CodeSize && A.CodeCRC == B.CodeCRC && A.Constants == B.Constants;
+		return A.CodeSize == B.CodeSize && A.CodeCRC == B.CodeCRC;
 	}
 
 	friend uint32 GetTypeHash(const FMetalCompiledShaderKey &Key)
 	{
-		return HashCombine(HashCombine(GetTypeHash(Key.CodeSize), GetTypeHash(Key.CodeCRC)), GetTypeHash(Key.Constants));
+		return HashCombine(HashCombine(GetTypeHash(Key.CodeSize), GetTypeHash(Key.CodeCRC)));
 	}
 
 	uint32 CodeSize;
 	uint32 CodeCRC;
-	uint32 Constants;
 };
diff --git a/Engine/Source/Runtime/Apple/MetalRHI/Private/Shaders/Types/Templates/MetalBaseShader.h b/Engine/Source/Runtime/Apple/MetalRHI/Private/Shaders/Types/Templates/MetalBaseShader.h
--- a/Engine/Source/Runtime/Apple/MetalRHI/Private/Shaders/Types/Templates/MetalBaseShader.h
+++ b/Engine/Source/Runtime/Apple/MetalRHI/Private/Shaders/Types/Templates/MetalBaseShader.h
@@ -63,9 +63,6 @@ struct FMetalShaderData
     /** CRC & Len for name disambiguation */
     uint32 SourceLen = 0;
     uint32 SourceCRC = 0;
-
-    /** Hash for the shader/material permutation constants */
-    uint32 ConstantValueHash = 0;
     
     // this is the compiler shader
     MTLFunctionPtr Function;
@@ -81,10 +78,6 @@ struct FMetalShaderData
     /** The uncompressed text source size */
     uint32 CodeSize = 0;
 
-    // Function constant states
-    bool bHasFunctionConstants = false;
-    bool bDeviceFunctionConstants = false;
-
     /** Index of the function (in the library) pointing to the function requested by the user (when GetCompiledFunction() is called with an explicit index). */
     uint32 LibraryFunctionIndex = -1;
 };
@@ -154,8 +147,6 @@ void TMetalBaseShader<BaseResourceType, ShaderType>::Init(TArrayView<const uint8
 	// If this triggers than a level above us has failed to provide valid shader data and the cook is probably bogus
 	UE_CLOG(Header.SourceLen == 0 || Header.SourceCRC == 0, LogMetal, Fatal, TEXT("Invalid Shader Bytecode provided."));
 
-	bDeviceFunctionConstants = (Header.bDeviceFunctionConstants == 0 ? false : true);
-
 	// remember where the header ended and code (precompiled or source) begins
 	int32 CodeOffset = Ar.Tell();
 	uint32 BufferSize = ShaderCode.GetActualShaderCodeSize() - CodeOffset;
@@ -220,17 +211,12 @@ void TMetalBaseShader<BaseResourceType, ShaderType>::Init(TArrayView<const uint8
 		GlslCodeNSString->retain();
 	}
 
-	bHasFunctionConstants = (Header.bDeviceFunctionConstants == 0 ? false : true);
-
-	ConstantValueHash = 0;
-
 	Library = InLibrary;
 
 	bool bNeedsCompiling = false;
 
 	// Find the existing compiled shader in the cache.
-	uint32 FunctionConstantHash = ConstantValueHash;
-	FMetalCompiledShaderKey Key(Header.SourceLen, Header.SourceCRC, FunctionConstantHash);
+	FMetalCompiledShaderKey Key(Header.SourceLen, Header.SourceCRC);
 
 	Function = GetMetalCompiledShaderCache().FindRef(Key);
 	if (!Library && Function)
@@ -435,8 +421,7 @@ MTLFunctionPtr TMetalBaseShader<BaseResourceType, ShaderType>::GetCompiledFuncti
     if (!Func || bNeedToRecreateFunction)
 	{
 		// Find the existing compiled shader in the cache.
-		uint32 FunctionConstantHash = ConstantValueHash;
-		FMetalCompiledShaderKey Key(SourceLen, SourceCRC, FunctionConstantHash);
+		FMetalCompiledShaderKey Key(SourceLen, SourceCRC);
 		Func = Function = GetMetalCompiledShaderCache().FindRef(Key);
 
         if (bNeedToRecreateFunction)
@@ -448,70 +433,17 @@ MTLFunctionPtr TMetalBaseShader<BaseResourceType, ShaderType>::GetCompiledFuncti
 
 		if (!Func)
 		{
-			// Get the function from the library - the function name is "Main" followed by the CRC32 of the source MTLSL as 0-padded hex.
+			// Get the function from th¡e library - the function name is "Main" followed by the CRC32 of the source MTLSL as 0-padded hex.
 			// This ensures that even if we move to a unified library that the function names will be unique - duplicates will only have one entry in the library.
             NS::String* Name = (LibraryFunctionIndex != -1)
             ? (NS::String*)Library->functionNames()->object(LibraryFunctionIndex) : FStringToNSString(FString::Printf(TEXT("Main_%0.8x_%0.8x"), SourceLen, SourceCRC));
-			MTL::FunctionConstantValues* ConstantValues = nullptr;
-			ON_SCOPE_EXIT { if (ConstantValues){ ConstantValues->release(); } };
-			if (bHasFunctionConstants)
-			{
-				ConstantValues = MTL::FunctionConstantValues::alloc()->init();
-                check(ConstantValues);
+			
+			Function = NS::TransferPtr(Library->newFunction(Name));
 
-				if (bDeviceFunctionConstants)
-				{
-					// Index 33 is the device vendor id constant
-					ConstantValues->setConstantValue((void*)&GRHIVendorId, MTL::DataTypeUInt, NS::String::string("GMetalDeviceManufacturer", NS::UTF8StringEncoding));
-				}
-			}
-
-#if PLATFORM_SUPPORTS_GEOMETRY_SHADERS
-            // TODO: Make those constants optional?
-            static constexpr bool bEnableTessellation = false;
-            
-            ConstantValues = MTL::FunctionConstantValues::alloc()->init();
-            bHasFunctionConstants = true;
-            
-            ConstantValues->setConstantValue(&bEnableTessellation, MTL::DataTypeBool,
-											 NS::String::string("tessellationEnabled", NS::UTF8StringEncoding));
-			ConstantValues->setConstantValue(&Bindings.OutputSizeVS, MTL::DataTypeInt,
-											 NS::String::string("vertex_shader_output_size_fc", NS::UTF8StringEncoding));
-#endif
+			check(Function);
+			GetMetalCompiledShaderCache().Add(Key, Library, Function);
 
-			if (!bHasFunctionConstants || !bAsync)
-			{
-				if (!bHasFunctionConstants)
-				{
-					Function = NS::TransferPtr(Library->newFunction(Name));
-				}
-				else
-				{
-					NS::Error* Error;
-					Function = NS::TransferPtr(Library->newFunction(Name, ConstantValues, &Error));
-					UE_CLOG(Function.get() == nullptr, LogMetal, Error, TEXT("Failed to create function: %s"), *NSStringToFString(Error->description()));
-					UE_CLOG(Function.get() == nullptr, LogMetal, Fatal, TEXT("*********** Error\n%s"), *NSStringToFString(GetSourceCode()));
-				}
-
-				check(Function);
-				GetMetalCompiledShaderCache().Add(Key, Library, Function);
-
-				Func = Function;
-			}
-			else
-			{
-				const std::function<void(MTL::Function* pFunction, NS::Error* pError)> Handler = [Key, this](MTL::Function* NewFunction, NS::Error* Error)
-				{
-					UE_CLOG(NewFunction == nullptr, LogMetal, Error, TEXT("Failed to create function: %s"), *NSStringToFString(Error->description()));
-					UE_CLOG(NewFunction == nullptr, LogMetal, Fatal, TEXT("*********** Error\n%s"), *NSStringToFString(GetSourceCode()));
-
-					GetMetalCompiledShaderCache().Add(Key, Library, NS::RetainPtr(NewFunction));
-				};
-                    
-                Library->newFunction(Name, ConstantValues, Handler);
-
-				return MTLFunctionPtr();
-			}
+			Func = Function;
 		}
 	}
 
diff --git a/Engine/Source/Runtime/Apple/MetalRHI/Public/MetalShaderResources.h b/Engine/Source/Runtime/Apple/MetalRHI/Public/MetalShaderResources.h
--- a/Engine/Source/Runtime/Apple/MetalRHI/Public/MetalShaderResources.h
+++ b/Engine/Source/Runtime/Apple/MetalRHI/Public/MetalShaderResources.h
@@ -235,7 +235,6 @@ struct FMetalCodeHeader
 	FMetalRayTracingHeader RayTracing;
 	uint8 Frequency;
 	int8 SideTable;
-	uint8 bDeviceFunctionConstants;
 	
 	FMetalCodeHeader()
 	: SourceLen(0)
@@ -247,7 +246,6 @@ struct FMetalCodeHeader
 	, CompileFlags(0)
 	, Frequency(0)
 	, SideTable(-1)
-	, bDeviceFunctionConstants(0)
 	{
 	}
 
@@ -272,7 +270,6 @@ inline FArchive& FMetalCodeHeader::Serialize(FArchive& Ar, FShaderResourceTable&
 	}
 	Ar << CompileFlags;
 	Ar << SideTable;
-	Ar << bDeviceFunctionConstants;
 	
     return Ar;
 }
