Metal Bindless - Adding missing GetBindlessHandle override.

#Synced-CL 50267093

diff --git a/Engine/Source/Runtime/Apple/MetalRHI/Private/MetalUniformBuffer.cpp b/Engine/Source/Runtime/Apple/MetalRHI/Private/MetalUniformBuffer.cpp
--- a/Engine/Source/Runtime/Apple/MetalRHI/Private/MetalUniformBuffer.cpp
+++ b/Engine/Source/Runtime/Apple/MetalRHI/Private/MetalUniformBuffer.cpp
@@ -77,6 +77,13 @@ FMetalSuballocatedUniformBuffer::~FMetalSuballocatedUniformBuffer()
 	FMemory::Free(Shadow);
 }
 
+#if PLATFORM_SUPPORTS_BINDLESS_RENDERING
+FRHIDescriptorHandle FMetalSuballocatedUniformBuffer::GetBindlessHandle() const
+{
+	return BindlessHandle;
+}
+#endif
+
 void FMetalSuballocatedUniformBuffer::Update(FRHICommandListBase& RHICmdList, const void* Contents)
 {
     UE::RHICore::UpdateUniformBufferConstants(Shadow, Contents, GetLayout());
diff --git a/Engine/Source/Runtime/Apple/MetalRHI/Public/MetalUniformBuffer.h b/Engine/Source/Runtime/Apple/MetalRHI/Public/MetalUniformBuffer.h
--- a/Engine/Source/Runtime/Apple/MetalRHI/Public/MetalUniformBuffer.h
+++ b/Engine/Source/Runtime/Apple/MetalRHI/Public/MetalUniformBuffer.h
@@ -45,6 +45,10 @@ public:
 									EUniformBufferUsage Usage, EUniformBufferValidation Validation);
     ~FMetalSuballocatedUniformBuffer();
 
+#if PLATFORM_SUPPORTS_BINDLESS_RENDERING
+	virtual FRHIDescriptorHandle GetBindlessHandle() const;
+#endif
+
 	void Update(FRHICommandListBase& RHICmdList, const void* Contents);
 
 private:
