Fix double input glitch with NativeLayers

#Synced-CL 50664794

diff --git a/Engine/Source/Runtime/Apple/MetalRHI/MetalRHI.Build.cs b/Engine/Source/Runtime/Apple/MetalRHI/MetalRHI.Build.cs
--- a/Engine/Source/Runtime/Apple/MetalRHI/MetalRHI.Build.cs
+++ b/Engine/Source/Runtime/Apple/MetalRHI/MetalRHI.Build.cs
@@ -19,6 +19,8 @@ public class MetalRHI : ModuleRules
 				"RenderCore",
 				"Projects",
 				"TraceLog",
+				"Slate",
+				"SlateCore"
 			}
 			);
 
diff --git a/Engine/Source/Runtime/Apple/MetalRHI/Private/MetalViewport.cpp b/Engine/Source/Runtime/Apple/MetalRHI/Private/MetalViewport.cpp
--- a/Engine/Source/Runtime/Apple/MetalRHI/Private/MetalViewport.cpp
+++ b/Engine/Source/Runtime/Apple/MetalRHI/Private/MetalViewport.cpp
@@ -19,6 +19,7 @@
 #include "Mac/MacSystemIncludes.h"
 #include "Mac/CocoaWindow.h"
 #include "Mac/CocoaThread.h"
+#include "Mac/MacApplication.h"
 #else
 #include "IOS/IOSSystemIncludes.h"
 #include "IOS/IOSApplication.h"
@@ -29,6 +30,9 @@
 #include "RenderUtils.h"
 #include "Engine/RendererSettings.h"
 
+#include "Input/HittestGrid.h"
+#include "Framework/Application/SlateApplication.h"
+
 extern int32 GMetalSupportsIntermediateBackBuffer;
 extern float GMetalPresentFramePacing;
 
@@ -59,28 +63,69 @@ static FAutoConsoleVariableRef CVarMetalNonBlockingPresent(
 
 @implementation FMetalView
 
-- (id)initWithFrame:(NSRect)frameRect
+@synthesize LayerRange;
+
+- (id)initWithFrame:(NSRect)frameRect LayerRange:(TInterval<int32>)layerRange MetalViewport:(FMetalViewport*)metalViewport
 {
 	self = [super initWithFrame:frameRect];
 	if (self)
 	{
+		LayerRange = layerRange;
+		MetalViewport = metalViewport;
+		bSlateWidgetHit = false;
 	}
 	return self;
 }
 
 - (void)setAllowMouseActions: (BOOL)allowMouseActions
 {
-	bAllowMouseActions = allowMouseActions;
+	bAllowMouseActionsPastLayerRange = allowMouseActions;
 }
 
 - (NSView *) hitTest:(NSPoint) point
 {
-	return bAllowMouseActions ? [super hitTest: point] : nil;
+	// we unconditionally allow the bottommost layer to go through:
+	// there's no way that a WKWebView/native view could incorrectly
+	// preempt its input, and we need to have it go through due to the
+	// parent child relationship with the lowest layer being the parent
+	// of all the subsequent higher layers/native views
+	
+	if(LayerRange.Min != std::numeric_limits<int32>::min())
+	{
+		// want to convert to screen coords as that's what 
+		// FMacApplication::ConvertCocoaPositionToSlate bases its calculations on
+		NSPoint ConvertedPoint = [[self superview] convertPoint: point toView: nil];
+		ConvertedPoint = [[self window] convertPointToScreen: ConvertedPoint];
+		
+		GameThreadCall(^
+	    {
+			TSharedPtr<SWindow> Window = MetalViewport->GetBackingSlateWindow();
+			
+			if(Window.IsValid())
+			{
+				FHittestGrid& HittestGrid = Window->GetHittestGrid();
+				FVector2D SlatePoint = FMacApplication::ConvertCocoaPositionToSlate(ConvertedPoint.x, ConvertedPoint.y);
+				
+				// we want to test [LayerRange.Min + 1, LayerRange.Max) so as to NOT include the widgets that represent the 
+				// native layer (e.g. SAppleWebBrowser)
+				TArray<FWidgetAndPointer> BubblePath = HittestGrid.GetBubblePath(SlatePoint, 1.0f, true, INDEX_NONE, 
+																				 TInterval<int32>(LayerRange.Min + 1, LayerRange.Max));
+				
+				bSlateWidgetHit = BubblePath.Num() > 1;
+			}
+		}, false);
+		
+		return bSlateWidgetHit ? [super hitTest: point] : nil;
+	}
+	else
+	{
+		return [super hitTest: point];
+	}
 }
 
 - (BOOL)isOpaque
 {
-	return !bAllowMouseActions;
+	return !bAllowMouseActionsPastLayerRange;
 }
 
 - (BOOL)mouseDownCanMoveWindow
@@ -135,7 +180,7 @@ FMetalViewport::FMetalViewport(FMetalDevice& InDevice, void* WindowHandle, uint3
 	// instead of a view variable
 #if PLATFORM_MAC
 	MainThreadCall(^{
-		View = CreateNewMetalView(true, NSWindowAbove, nil);
+		View = CreateNewMetalView(true, NSWindowAbove, nil, TInterval<int32>(std::numeric_limits<int32>::min(), std::numeric_limits<int32>::max()));
 	});
 #endif
 	
@@ -185,11 +230,12 @@ FMetalViewport::~FMetalViewport()
 }
 
 #if PLATFORM_MAC
-FMetalView* FMetalViewport::CreateNewMetalView(bool bIsContentView, NSWindowOrderingMode ViewOrder, NSView* ViewRelativeTo)
+FMetalView* FMetalViewport::CreateNewMetalView(bool bIsContentView, NSWindowOrderingMode ViewOrder, NSView* ViewRelativeTo,
+											   TInterval<int32> LayerRange)
 {	
-	const NSRect ContentRect = NSMakeRect(0, 0, SizeX, SizeY);
+	const NSRect ContentRect = NSMakeRect(0, 0, WindowHandle.openGLFrame.size.width, WindowHandle.openGLFrame.size.height);
 	
-	FMetalView* MetalView = [[FMetalView alloc] initWithFrame:ContentRect];
+	FMetalView* MetalView = [[FMetalView alloc] initWithFrame:ContentRect LayerRange:LayerRange MetalViewport:this];
 	[MetalView setAutoresizingMask:NSViewWidthSizable | NSViewHeightSizable];
 	[MetalView setWantsLayer:YES];
 	[MetalView setAllowMouseActions: bIsContentView];
@@ -259,7 +305,7 @@ void FMetalViewport::DeleteNativeLayer(int32 OldNativeLayer)
 		});
 }
 
-void FMetalViewport::CreateNativeLayer(int32 NewNativeLayer, void* NativeViewHandle)
+void FMetalViewport::CreateNativeLayer(int32 NewNativeLayer, void* NativeViewHandle, const TArray<TInterval<int32>>& LayerRanges)
 {
 	if(DrawLayers.Contains(NewNativeLayer))
 	{
@@ -275,6 +321,7 @@ void FMetalViewport::CreateNativeLayer(int32 NewNativeLayer, void* NativeViewHan
 		UE_LOG(LogMetal, Display, TEXT("Creating NativeLayer %d in Metal Viewport %p"), NewNativeLayer, this);
 		
 		TOptional<int32> LayerRightBeforeNewLayer;
+		TInterval<int32> LayerRange(NewNativeLayer, std::numeric_limits<int32>::max());
 		FScopeLock Lock(&Mutex);
 		
 		if(NewNativeLayer != std::numeric_limits<int32>::min())
@@ -289,6 +336,7 @@ void FMetalViewport::CreateNativeLayer(int32 NewNativeLayer, void* NativeViewHan
 					LayerRightBeforeNewLayer = DrawLayer.Key;
 					break;
 				}
+				LayerRange.Max = DrawLayer.Key;
 			}
 		}
 		
@@ -315,11 +363,11 @@ void FMetalViewport::CreateNativeLayer(int32 NewNativeLayer, void* NativeViewHan
 			
 			[ContentView addSubview: NativeView positioned: bPlaceBelowNativeView ? NSWindowBelow : NSWindowAbove 
 						 relativeTo: bPlaceRelToSpecificView ? DrawLayers[LayerRightBeforeNewLayer.GetValue()].View : nil];
-			View = CreateNewMetalView(false, NSWindowAbove, NativeView);
+			View = CreateNewMetalView(false, NSWindowAbove, NativeView, LayerRange);
 		}
 		else
 		{
-			View = CreateNewMetalView(NewNativeLayer == std::numeric_limits<int32>::min(), NSWindowAbove, LayerRightBeforeNewLayer.IsSet() ? 					  DrawLayers[LayerRightBeforeNewLayer.GetValue()].View : nil);
+			View = CreateNewMetalView(NewNativeLayer == std::numeric_limits<int32>::min(), NSWindowAbove, LayerRightBeforeNewLayer.IsSet() ? 					  DrawLayers[LayerRightBeforeNewLayer.GetValue()].View : nil, LayerRange);
 		}
 		
 		DrawLayers.Add(NewNativeLayer, DrawLayer
@@ -333,6 +381,14 @@ void FMetalViewport::CreateNativeLayer(int32 NewNativeLayer, void* NativeViewHan
 		{
 			return A > B;
 		});
+		
+		for(size_t i = 0; LayerRanges.Num() > i; ++i)
+		{
+			if(DrawLayers.Contains(LayerRanges[i].Min))
+			{
+				DrawLayers[LayerRanges[i].Min].View.LayerRange = LayerRanges[i];
+			}
+		}
 	});
 #endif
 	
@@ -368,6 +424,32 @@ uint32 FMetalViewport::GetViewportIndex(EMetalViewportAccessFlag Accessor) const
 	}
 }
 
+#if PLATFORM_MAC
+TSharedPtr<SWindow> FMetalViewport::GetBackingSlateWindow()
+{
+	TSharedPtr<SWindow> BackingSlateWindowStrong = BackingSlateWindow.Pin();
+	
+	if (!BackingSlateWindowStrong.IsValid() && MacApplication)
+	{
+		TSharedPtr<FMacWindow> Window = MacApplication->FindWindowByNSWindow(WindowHandle);
+		
+		if (Window.IsValid())
+		{
+			for(TSharedRef<SWindow>& SlateWindow : FSlateApplication::Get().GetInteractiveTopLevelWindows())
+			{
+				if(Window == SlateWindow->GetNativeWindow())
+				{
+					BackingSlateWindow = SlateWindow;
+					return SlateWindow;
+				}
+			}
+		}
+	}
+	
+	return BackingSlateWindowStrong;
+}
+#endif
+
 void FMetalViewport::ResizeBackBuffer(int32 DrawLayerKey, uint32 Index, uint32 InSizeX, uint32 InSizeY,
 									  EPixelFormat InFormat)
 {
@@ -539,7 +621,7 @@ TRefCountPtr<FMetalSurface> FMetalViewport::GetBackBuffer(EMetalViewportAccessFl
 	uint32 Index = GetViewportIndex(Accessor);
 	
 #if WITH_MULTI_VIEW_SLATE_WINDOW_SUPPORT
-	uint32 DrawLayerIndexValue = DrawLayerIndex.IsSet() ? DrawLayerIndex.GetValue() : DefaultNativeLayer;
+	int32 DrawLayerIndexValue = DrawLayerIndex.IsSet() ? DrawLayerIndex.GetValue() : DefaultNativeLayer;
 	
 	check(DrawLayers.Contains(DrawLayerIndexValue));
 	check(IsValidRef(DrawLayers[DrawLayerIndexValue].BackBuffer[Index]));
diff --git a/Engine/Source/Runtime/Apple/MetalRHI/Public/MetalViewport.h b/Engine/Source/Runtime/Apple/MetalRHI/Public/MetalViewport.h
--- a/Engine/Source/Runtime/Apple/MetalRHI/Public/MetalViewport.h
+++ b/Engine/Source/Runtime/Apple/MetalRHI/Public/MetalViewport.h
@@ -12,6 +12,8 @@
 #include "RHIResources.h"
 #include "MetalResources.h"
 
+#include "Widgets/SWindow.h"
+
 #if PLATFORM_MAC
 #include "Mac/MacSystemIncludes.h"
 #elif PLATFORM_IOS
@@ -22,8 +24,13 @@
 #include "Mac/CocoaTextView.h"
 @interface FMetalView : FCocoaTextView
 {
-	bool bAllowMouseActions;
+	bool bAllowMouseActionsPastLayerRange;
+	bool bSlateWidgetHit;
+	
+	FMetalViewport* MetalViewport;
 }
+
+@property TInterval<int32> LayerRange;
 @end
 #endif
 #include "HAL/PlatformFramePacer.h"
@@ -76,7 +83,7 @@ public:
 	void ReleaseDrawable(int32 DrawLayerKey);
 	
 #if WITH_MULTI_VIEW_SLATE_WINDOW_SUPPORT
-	virtual void CreateNativeLayer(int32 NewNativeLayer, void* NativeViewHandle) override;
+	virtual void CreateNativeLayer(int32 NewNativeLayer, void* NativeViewHandle, const TArray<TInterval<int32>>& LayerRanges) override;
 	virtual void DeleteNativeLayer(int32 OldNativeLayer) override;
 	virtual void SetDefaultNativeLayer(int32 NativeLayer) override;
 #endif
@@ -105,9 +112,14 @@ public:
     void PresentImmersive(const MetalRHIVisionOS::PresentImmersiveParams& Params);
 #endif
 	
+#if PLATFORM_MAC
+	TSharedPtr<SWindow> GetBackingSlateWindow();
+#endif
+	
 private:
 #if PLATFORM_MAC
-	FMetalView* CreateNewMetalView(bool bIsContentView, NSWindowOrderingMode ViewOrder, NSView* ViewRelativeTo);
+	FMetalView* CreateNewMetalView(bool bIsContentView, NSWindowOrderingMode ViewOrder, NSView* ViewRelativeTo,
+								   TInterval<int32> LayerRange);
 #endif
 	uint32 GetViewportIndex(EMetalViewportAccessFlag Accessor) const;
 	void PresentDrawLayers(FMetalCommandBuffer* CommandBuffer);
@@ -192,6 +204,8 @@ private:
 #if PLATFORM_MAC || PLATFORM_VISIONOS
     FCustomPresentRHIRef CustomPresent;
 #endif
+	
+	TWeakPtr<SWindow> BackingSlateWindow;
 };
 
 template<>
diff --git a/Engine/Source/Runtime/RHI/Public/RHIResources.h b/Engine/Source/Runtime/RHI/Public/RHIResources.h
--- a/Engine/Source/Runtime/RHI/Public/RHIResources.h
+++ b/Engine/Source/Runtime/RHI/Public/RHIResources.h
@@ -2542,7 +2542,7 @@ public:
 	 * everything on the bottom of the specified layer.
 	 * Right now this is only implemented in MetalRHI for macOS NSView interop with Slate.
 	 */
-	virtual void CreateNativeLayer(int32 NewNativeLayer, void* NativeViewHandle) {}
+	virtual void CreateNativeLayer(int32 NewNativeLayer, void* NativeViewHandle, const TArray<TInterval<int32>>& LayerRanges) {}
 	
 	/**
 	 * Deletes a platform-level native layer created by CreateNativeLayer.
diff --git a/Engine/Source/Runtime/SlateCore/Private/Input/HittestGrid.cpp b/Engine/Source/Runtime/SlateCore/Private/Input/HittestGrid.cpp
--- a/Engine/Source/Runtime/SlateCore/Private/Input/HittestGrid.cpp
+++ b/Engine/Source/Runtime/SlateCore/Private/Input/HittestGrid.cpp
@@ -178,7 +178,15 @@ FHittestGrid::FHittestGrid()
 {
 }
 
-TArray<FWidgetAndPointer> FHittestGrid::GetBubblePath(UE::Slate::FDeprecateVector2DParameter DesktopSpaceCoordinate, float CursorRadius, bool bIgnoreEnabledStatus, int32 UserIndex)
+bool FHittestGrid::IsWidgetDataInLayerRange(const FWidgetData& WidgetData, const TInterval<int32>& LayerRange)
+{
+	// top 32 bits are InBatchPriorityGroup, which we don't worry about
+	int32 WidgetLayer = static_cast<int32>(WidgetData.PrimarySort);
+	
+	return LayerRange.Min <= WidgetLayer && LayerRange.Max > WidgetLayer;
+}
+
+TArray<FWidgetAndPointer> FHittestGrid::GetBubblePath(UE::Slate::FDeprecateVector2DParameter DesktopSpaceCoordinate, float CursorRadius, bool bIgnoreEnabledStatus, int32 UserIndex, TInterval<int32> LayerRange)
 {
 	checkSlow(IsInGameThread());
 
@@ -209,11 +217,15 @@ TArray<FWidgetAndPointer> FHittestGrid::GetBubblePath(UE::Slate::FDeprecateVecto
 					FGeometry DesktopSpaceGeometry = CurWidget->GetPaintSpaceGeometry();
 					DesktopSpaceGeometry.AppendTransform(FSlateLayoutTransform(GridOrigin - GridWindowOrigin));
 
-					Path.Emplace(FArrangedWidget(CurWidget.ToSharedRef(), DesktopSpaceGeometry));
+					if(IsWidgetDataInLayerRange(BestHitWidgetData, LayerRange))
+					{
+						Path.Emplace(FArrangedWidget(CurWidget.ToSharedRef(), DesktopSpaceGeometry));
+					}
+					
 					CurWidget = CurWidget->Advanced_GetPaintParentWidget();
 				}
 
-				if (!Path.Last().Widget->Advanced_IsWindow())
+				if (Path.IsEmpty() || !Path.Last().Widget->Advanced_IsWindow())
 				{
 					return TArray<FWidgetAndPointer>();
 				}
diff --git a/Engine/Source/Runtime/SlateCore/Public/Input/HittestGrid.h b/Engine/Source/Runtime/SlateCore/Public/Input/HittestGrid.h
--- a/Engine/Source/Runtime/SlateCore/Public/Input/HittestGrid.h
+++ b/Engine/Source/Runtime/SlateCore/Public/Input/HittestGrid.h
@@ -38,9 +38,10 @@ public:
 
 	/**
 	 * Given a Slate Units coordinate in virtual desktop space, perform a hittest
-	 * and return the path along which the corresponding event would be bubbled.
+	 * and return the path along which the corresponding event would be bubbled, optionally filtering
+	 * the hittest to widgets in a range of layers.
 	 */
-	SLATECORE_API TArray<FWidgetAndPointer> GetBubblePath(UE::Slate::FDeprecateVector2DParameter DesktopSpaceCoordinate, float CursorRadius, bool bIgnoreEnabledStatus, int32 UserIndex = INDEX_NONE);
+	SLATECORE_API TArray<FWidgetAndPointer> GetBubblePath(UE::Slate::FDeprecateVector2DParameter DesktopSpaceCoordinate, float CursorRadius, bool bIgnoreEnabledStatus, int32 UserIndex = INDEX_NONE, TInterval<int32> LayerRange = TInterval<int32>(std::numeric_limits<int32>::min(), std::numeric_limits<int32>::max()));
 
 	/**
 	 * Set the position and size of the hittest area in desktop coordinates
@@ -311,6 +312,8 @@ private:
 	TInstancedStruct<FNavigationMethod> GetNavigationMethod(const FArrangedWidget& StartingWidget);
 	TSharedPtr<SWidget> FindNextFocusableWidgetDefault(const FArrangedWidget& StartingWidget, const EUINavigation Direction, const FNavigationReply& NavigationReply, const FArrangedWidget& RuleWidget, int32 UserIndex);
 	TSharedPtr<SWidget> FindNextFocusableWidgetExperimental(const FArrangedWidget& StartingWidget, const EUINavigation Direction, const FNavigationReply& NavigationReply, const FArrangedWidget& RuleWidget, int32 UserIndex);
+	
+	bool IsWidgetDataInLayerRange(const FWidgetData& WidgetData, const TInterval<int32>& LayerRange);
 
 private:
 	/** Map of all the widgets currently in the hit test grid to their stable index. */
diff --git a/Engine/Source/Runtime/SlateRHIRenderer/Private/SlateRHIRenderer.cpp b/Engine/Source/Runtime/SlateRHIRenderer/Private/SlateRHIRenderer.cpp
--- a/Engine/Source/Runtime/SlateRHIRenderer/Private/SlateRHIRenderer.cpp
+++ b/Engine/Source/Runtime/SlateRHIRenderer/Private/SlateRHIRenderer.cpp
@@ -1293,7 +1293,8 @@ void FSlateRHIRenderer::CreateNativeLayer(int32 NewNativeLayer, SWindow& InWindo
 		for(int32_t NativeLayer : WindowLayerInfos[&InWindow].NativeLayers)
 		{
 			// create the native layer within metal if it doesn't exist yet (see MetalViewport)
-			ViewportInfo->ViewportRHI->CreateNativeLayer(NativeLayer, NativeViewHandle);
+			ViewportInfo->ViewportRHI->CreateNativeLayer(NativeLayer, NativeViewHandle,
+														 WindowLayerInfos[&InWindow].LayerRanges);
 		}
 	}		
 	
