Missing RHIUtilities.h include

#Synced-CL 50759986

diff --git a/Engine/Source/Runtime/Apple/MetalRHI/Private/MetalViewport.cpp b/Engine/Source/Runtime/Apple/MetalRHI/Private/MetalViewport.cpp
--- a/Engine/Source/Runtime/Apple/MetalRHI/Private/MetalViewport.cpp
+++ b/Engine/Source/Runtime/Apple/MetalRHI/Private/MetalViewport.cpp
@@ -6,6 +6,7 @@
 
 #include "MetalViewport.h"
 #include "Apple/ScopeAutoreleasePool.h"
+#include "RHIUtilities.h"
 #include "MetalDynamicRHI.h"
 #include "MetalRHIPrivate.h"
 #include "MetalCommandBuffer.h"
