Metal - Fix double free with MTL::CompileOptions when exiting.

#tests Local test & preflight

#Synced-CL 50707634

diff --git a/Engine/Source/Runtime/Apple/MetalRHI/Private/Shaders/Types/Templates/MetalBaseShader.h b/Engine/Source/Runtime/Apple/MetalRHI/Private/Shaders/Types/Templates/MetalBaseShader.h
--- a/Engine/Source/Runtime/Apple/MetalRHI/Private/Shaders/Types/Templates/MetalBaseShader.h
+++ b/Engine/Source/Runtime/Apple/MetalRHI/Private/Shaders/Types/Templates/MetalBaseShader.h
@@ -398,7 +398,6 @@ void TMetalBaseShader<BaseResourceType, ShaderType>::Init(TArrayView<const uint8
 				UE_LOG(LogRHI, Warning, TEXT("Created shader with warnings: %s"), *NSStringToFString(Error->description()));
 			}
 
-			CompileOptions->release();
 			GlslCodeNSString = NewShaderString;
 			GlslCodeNSString->retain();
 		}
