Fix for native layer dpi bug causing top bar menus' pop up windows to cut off items and/or have broken dpi

#Synced-CL 50718145

diff --git a/Engine/Source/Runtime/Apple/MetalRHI/Private/MetalViewport.cpp b/Engine/Source/Runtime/Apple/MetalRHI/Private/MetalViewport.cpp
--- a/Engine/Source/Runtime/Apple/MetalRHI/Private/MetalViewport.cpp
+++ b/Engine/Source/Runtime/Apple/MetalRHI/Private/MetalViewport.cpp
@@ -314,8 +314,10 @@ void FMetalViewport::CreateNativeLayer(int32 NewNativeLayer, void* NativeViewHan
 		return;
 	}
 	
-	SizeX = WindowHandle.openGLFrame.size.width;
-	SizeY = WindowHandle.openGLFrame.size.height;
+	float ContentScaleFactor = FSlateApplication::Get().GetApplicationScale() * MacApplication->FindWindowByNSWindow(WindowHandle)->GetDPIScaleFactor();
+	
+	SizeX = ContentScaleFactor * WindowHandle.openGLFrame.size.width;
+	SizeY = ContentScaleFactor * WindowHandle.openGLFrame.size.height;
 	
 #if WITH_MULTI_VIEW_SLATE_WINDOW_SUPPORT
 	MainThreadCall(^{
@@ -547,6 +549,11 @@ void FMetalViewport::Resize(uint32 InSizeX, uint32 InSizeY, bool bInIsFullscreen
     Fence.Wait();
     
 #if PLATFORM_MAC
+	float ContentScaleFactor = FSlateApplication::Get().GetApplicationScale() * MacApplication->FindWindowByNSWindow(WindowHandle)->GetDPIScaleFactor();
+	
+	SizeX = ContentScaleFactor * InSizeX;
+	SizeY = ContentScaleFactor * InSizeY;
+	
 	MainThreadCall(^
 	{
 #if WITH_MULTI_VIEW_SLATE_WINDOW_SUPPORT
