MetalRHI

- Aggressively release fully deallocated heap blocks
- Reduce Heap size to 512mb per heap

diff --git a/Engine/Source/Runtime/Apple/MetalRHI/Private/MetalBuffer.cpp b/Engine/Source/Runtime/Apple/MetalRHI/Private/MetalBuffer.cpp
--- a/Engine/Source/Runtime/Apple/MetalRHI/Private/MetalBuffer.cpp
+++ b/Engine/Source/Runtime/Apple/MetalRHI/Private/MetalBuffer.cpp
@@ -1576,6 +1576,8 @@ TDoubleLinkedList<FMetalResourceHeap::MemoryBlock>::TIterator FMetalResourceHeap
 	return BlockItA;
 }
 
+static constexpr uint32 METALHEAP_DEFAULT_BLOCKSIZE = 1024 << 19; // 512mb
+
 void FMetalResourceHeap::FreeBlock(const uint32 ResourceAllocationHandle)
 {
 	FScopeLock ScopeLock(&FreeListCS);
@@ -1601,7 +1603,7 @@ void FMetalResourceHeap::FreeBlock(const uint32 ResourceAllocationHandle)
 		NextBlockIt++;
 	}
 	
-	// Put the block into the free list
+	// Put the block into the free list	
 	UsedList->RemoveNode(BlockIt.GetNode(), false);
 	FreeList->InsertNode(BlockIt.GetNode(), NextBlockIt ? NextBlockIt.GetNode() : nullptr);
 	
@@ -1662,6 +1664,16 @@ void FMetalResourceHeap::FreeBlock(const uint32 ResourceAllocationHandle)
 			}
 		}
 	}
+	
+	// If heap is completely deallocated then remove it
+	MemoryBlock& CurrentBlock = *BlockIt;
+	
+	if(CurrentBlock.Size == CurrentBlock.Heap->size())
+	{
+		FMetalDynamicRHI::Get().DeferredDelete([HeapPtr = CurrentBlock.Heap](){});
+		CurrentBlock.Heap = nullptr;
+		FreeList->RemoveNode(BlockIt.GetNode(), false);
+	}
 }
 
 TDoubleLinkedList<FMetalResourceHeap::MemoryBlock>::TIterator FMetalResourceHeap::FindOrAllocateBlock(uint32 Size, uint32 Alignment, MTL::ResourceOptions Options)
@@ -1708,11 +1720,9 @@ TDoubleLinkedList<FMetalResourceHeap::MemoryBlock>::TIterator FMetalResourceHeap
 			}
 		}
 	}
-
-	// Allocate a fresh block to use
-	static constexpr uint32 DefaultBlockSize = 1024 << 20; // 1GB
+	
 	{
-		uint64 NewBlockSize = FMath::Max(Size, DefaultBlockSize);
+		uint64 NewBlockSize = FMath::Max(Size, METALHEAP_DEFAULT_BLOCKSIZE);
 
 		MTL::StorageMode StorageMode = (MTL::StorageMode)(((NS::UInteger)Options & MTL::ResourceStorageModeMask) >> MTL::ResourceStorageModeShift);
 		MTL::CPUCacheMode CpuMode = (MTL::CPUCacheMode)(((NS::UInteger)Options & MTL::ResourceCpuCacheModeMask) >> MTL::ResourceCpuCacheModeShift);
@@ -1731,7 +1741,7 @@ TDoubleLinkedList<FMetalResourceHeap::MemoryBlock>::TIterator FMetalResourceHeap
 		MemoryBlock NewBlock;
 		NewBlock.Heap = BlockHeap;
 		NewBlock.Offset = 0;
-		NewBlock.Size = NewBlockSize;
+		NewBlock.Size = BlockHeap->size();
 		NewBlock.Resource = nullptr;
 		NewBlock.Options = Options;
 
