Fix iOSCrashReporter screenshot

#Synced-CL 51024481

diff --git a/Engine/Source/Runtime/Apple/MetalRHI/Private/MetalDynamicRHI.h b/Engine/Source/Runtime/Apple/MetalRHI/Private/MetalDynamicRHI.h
--- a/Engine/Source/Runtime/Apple/MetalRHI/Private/MetalDynamicRHI.h
+++ b/Engine/Source/Runtime/Apple/MetalRHI/Private/MetalDynamicRHI.h
@@ -166,7 +166,7 @@ public:
 	virtual void RHIReadSurfaceData(FRHITexture* Texture,FIntRect Rect,TArray<FColor>& OutData,FReadSurfaceDataFlags InFlags) final override;
 	virtual void RHIReadSurfaceData(FRHITexture* TextureRHI, FIntRect InRect, TArray<FLinearColor>& OutData, FReadSurfaceDataFlags InFlags) final override;
 #if PLATFORM_IOS
-	virtual void RHIReadSurfaceDataDirect(FRHITexture* TextureRHI, FIntRect Rect, TArray<FColor>& OutData) final override;
+	virtual bool RHIReadViewportDataDirect(uint32& OutWidth, uint32& OutHeight, TArray<FColor>& OutData) final override;
 #endif // PLATFORM_IOS
 	virtual void RHIMapStagingSurface(FRHITexture* Texture, FRHIGPUFence* Fence, void*& OutData, int32& OutWidth, int32& OutHeight, uint32 GPUIndex = 0) final override;
 	virtual void RHIUnmapStagingSurface(FRHITexture* Texture, uint32 GPUIndex = 0) final override;
diff --git a/Engine/Source/Runtime/Apple/MetalRHI/Private/MetalRenderTarget.cpp b/Engine/Source/Runtime/Apple/MetalRHI/Private/MetalRenderTarget.cpp
--- a/Engine/Source/Runtime/Apple/MetalRHI/Private/MetalRenderTarget.cpp
+++ b/Engine/Source/Runtime/Apple/MetalRHI/Private/MetalRenderTarget.cpp
@@ -253,21 +253,45 @@ void FMetalDynamicRHI::RHIReadSurfaceData(FRHITexture* TextureRHI, FIntRect Rect
 }
 
 #if PLATFORM_IOS
-void FMetalDynamicRHI::RHIReadSurfaceDataDirect(FRHITexture* TextureRHI, FIntRect Rect, TArray<FColor>& OutData)
-{   
-	if (!ensure(TextureRHI)) return;
+bool FMetalDynamicRHI::RHIReadViewportDataDirect(uint32& OutWidth, uint32& OutHeight, TArray<FColor>& OutData)
+{
+	check(IsInGameThread());
 
-	FMetalSurface* Surface = GetMetalSurfaceFromRHITexture(TextureRHI);
+	FMetalViewport* Viewport = nullptr;
+	{
+		FScopeLock Lock(&FMetalViewport::ViewportsMutex);
+		for (FMetalViewport* ViewportPtr : FMetalViewport::Viewports)
+		{
+			if (ViewportPtr)
+			{
+				Viewport = ViewportPtr;
+				break;
+			}
+		}
+	}
+
+	if (!Viewport)
+	{
+		// No viewport available to capture from.
+		return false;
+	}
+
+	FMetalSurface* BackBuffer = Viewport->GetBackBuffer(EMetalViewportAccessGame, {});
+	if (!BackBuffer)
+	{
+		// Could not retrieve the back buffer to capture from.
+		return false;
+	}
 
 	// allocate output space
-	const uint32 SizeX = Rect.Width();
-	const uint32 SizeY = Rect.Height();
-	OutData.SetNumUninitialized(SizeX * SizeY);
+	OutWidth = BackBuffer->GetDesc().Extent.X;
+	OutHeight = BackBuffer->GetDesc().Extent.Y;
+	OutData.SetNumUninitialized(OutWidth * OutHeight);
 
-	MTLTexturePtr Texture = Surface->Texture;
-    if(!Texture && EnumHasAnyFlags(Surface->GetDesc().Flags, TexCreate_Presentable))
+	MTLTexturePtr Texture = BackBuffer->Texture;
+    if(!Texture && EnumHasAnyFlags(BackBuffer->GetDesc().Flags, TexCreate_Presentable))
     {
-        Texture = Surface->GetCurrentTexture();
+        Texture = BackBuffer->GetCurrentTexture();
     }
 	
     if(!Texture)
@@ -275,15 +299,16 @@ void FMetalDynamicRHI::RHIReadSurfaceDataDirect(FRHITexture* TextureRHI, FIntRec
 		OutData.Empty();
 
 		UE_LOG(LogRHI, Error, TEXT("Trying to read from an uninitialised texture."));
-		
-        return;
+        return false;
     }
 
-	const uint32 Stride = GPixelFormats[Surface->GetDesc().Format].BlockBytes * SizeX;
-	const uint32 BytesPerImage = Stride * SizeY;
+	const uint32 Stride = GPixelFormats[BackBuffer->GetDesc().Format].BlockBytes * OutWidth;
+	const uint32 BytesPerImage = Stride * OutHeight;
 
-	MTL::Region Region(Rect.Min.X, Rect.Min.Y, SizeX, SizeY);
+	MTL::Region Region(0, 0, OutWidth, OutHeight);
 	Texture->getBytes(OutData.GetData(), Stride, BytesPerImage, Region, 0, 0);
+
+	return true;
 }
 #endif // PLATFORM_IOS
 
diff --git a/Engine/Source/Runtime/Apple/MetalRHI/Private/MetalViewport.cpp b/Engine/Source/Runtime/Apple/MetalRHI/Private/MetalViewport.cpp
--- a/Engine/Source/Runtime/Apple/MetalRHI/Private/MetalViewport.cpp
+++ b/Engine/Source/Runtime/Apple/MetalRHI/Private/MetalViewport.cpp
@@ -138,8 +138,8 @@ static FAutoConsoleVariableRef CVarMetalNonBlockingPresent(
 
 #endif
 
-static FCriticalSection ViewportsMutex;
-static TSet<FMetalViewport*> Viewports;
+FCriticalSection FMetalViewport::ViewportsMutex;
+TSet<FMetalViewport*> FMetalViewport::Viewports;
 
 FMetalViewport::FMetalViewport(FMetalDevice& InDevice, void* WindowHandle, uint32 InSizeX, uint32 InSizeY, bool bInIsFullscreen, EPixelFormat InFormat)
 	: Device(InDevice)
diff --git a/Engine/Source/Runtime/Apple/MetalRHI/Public/IMetalDynamicRHI.h b/Engine/Source/Runtime/Apple/MetalRHI/Public/IMetalDynamicRHI.h
--- a/Engine/Source/Runtime/Apple/MetalRHI/Public/IMetalDynamicRHI.h
+++ b/Engine/Source/Runtime/Apple/MetalRHI/Public/IMetalDynamicRHI.h
@@ -22,6 +22,11 @@ struct IMetalDynamicRHI : public FDynamicRHI
 
 	virtual void 			  RHIRunOnQueue(TFunction<void(MTL::CommandQueue*)>&& CodeToRun, bool bWaitForSubmission) = 0;
 	virtual FTextureRHIRef	  RHICreateTexture2DFromCVMetalTexture(EPixelFormat Format, ETextureCreateFlags TexCreateFlags, const FClearValueBinding& ClearValueBinding, CVMetalTextureRef Resource) = 0;
+
+#if PLATFORM_IOS
+	/** Used for getting the last back buffer content after a crash. */
+	virtual bool RHIReadViewportDataDirect(uint32& OutWidth, uint32& OutHeight, TArray<FColor>& OutData) = 0;
+#endif // PLATFORM_IOS
 };
 
 inline bool IsRHIMetal()
diff --git a/Engine/Source/Runtime/Apple/MetalRHI/Public/MetalViewport.h b/Engine/Source/Runtime/Apple/MetalRHI/Public/MetalViewport.h
--- a/Engine/Source/Runtime/Apple/MetalRHI/Public/MetalViewport.h
+++ b/Engine/Source/Runtime/Apple/MetalRHI/Public/MetalViewport.h
@@ -66,6 +66,9 @@ typedef void (^FMetalViewportPresentHandler)(uint32 CGDirectDisplayID, double Ou
 class FMetalViewport : public FRHIViewport
 {
 public:
+	static FCriticalSection ViewportsMutex;
+	static TSet<FMetalViewport*> Viewports;
+
 	FMetalViewport(FMetalDevice& InDevice, void* WindowHandle, uint32 InSizeX, uint32 InSizeY, bool bInIsFullscreen,EPixelFormat InFormat);
 	~FMetalViewport();
 
diff --git a/Engine/Source/Runtime/RHI/Public/DynamicRHI.h b/Engine/Source/Runtime/RHI/Public/DynamicRHI.h
--- a/Engine/Source/Runtime/RHI/Public/DynamicRHI.h
+++ b/Engine/Source/Runtime/RHI/Public/DynamicRHI.h
@@ -606,15 +606,6 @@ public:
 	// Default fallback; will not work for non-8-bit surfaces and it's extremely slow.
 	RHI_API virtual void RHIReadSurfaceData(FRHITexture* Texture, FIntRect Rect, TArray<FLinearColor>& OutData, FReadSurfaceDataFlags InFlags);
 
-	/** 
-	* Implemented by the platform to read directly from the texture. 
-	* Useful for getting the last back buffer content after a crash.
-	* As the RHIReadSurfaceData implementation will copy the texture's content to a new temporary buffer, that won't work for an unknown crash state.
-	*/
-#if PLATFORM_IOS
-	virtual void RHIReadSurfaceDataDirect(FRHITexture* Texture, FIntRect Rect, TArray<FColor>& OutData) {}
-#endif // PLATFORM_IOS
-
 	/** Watch out for OutData to be 0 (can happen on DXGI_ERROR_DEVICE_REMOVED), don't call RHIUnmapStagingSurface in that case. */
 	// FlushType: Flush Immediate (seems wrong)
 	virtual void RHIMapStagingSurface(FRHITexture* Texture, FRHIGPUFence* Fence, void*& OutData, int32& OutWidth, int32& OutHeight, uint32 GPUIndex = 0) = 0;
