Fix CIS issue in Mac Menu

The MacMenu module needs to only be built in Mac Platforms

#Synced-CL 51141141

diff --git a/Engine/Source/Runtime/Apple/MacMenu/MacMenu.Build.cs b/Engine/Source/Runtime/Apple/MacMenu/MacMenu.Build.cs
--- a/Engine/Source/Runtime/Apple/MacMenu/MacMenu.Build.cs
+++ b/Engine/Source/Runtime/Apple/MacMenu/MacMenu.Build.cs
@@ -3,6 +3,7 @@
 using UnrealBuildTool;
 using System.IO;
 
+[SupportedPlatforms("Mac")]
 public class MacMenu : ModuleRules
 {
     public MacMenu(ReadOnlyTargetRules Target) : base(Target)
