fixing race conditions related to rapidly creating/deleting nativelayers on macos

#Synced-CL 51217571

diff --git a/Engine/Source/Runtime/Apple/MetalRHI/Private/MetalViewport.cpp b/Engine/Source/Runtime/Apple/MetalRHI/Private/MetalViewport.cpp
--- a/Engine/Source/Runtime/Apple/MetalRHI/Private/MetalViewport.cpp
+++ b/Engine/Source/Runtime/Apple/MetalRHI/Private/MetalViewport.cpp
@@ -319,7 +319,10 @@ void FMetalViewport::CreateNativeLayer(int32 NewNativeLayer, void* NativeViewHan
 	SizeX = ContentScaleFactor * WindowHandle.openGLFrame.size.width;
 	SizeY = ContentScaleFactor * WindowHandle.openGLFrame.size.height;
 	
+	bool bWaitForRenderingThread = true;
 #if WITH_MULTI_VIEW_SLATE_WINDOW_SUPPORT
+	bWaitForRenderingThread = false;
+	
 	MainThreadCall(^{
 		UE_LOG(LogMetal, Display, TEXT("Creating NativeLayer %d in Metal Viewport %p"), NewNativeLayer, this);
 		
@@ -395,7 +398,7 @@ void FMetalViewport::CreateNativeLayer(int32 NewNativeLayer, void* NativeViewHan
 	});
 #endif
 	
-	Resize(SizeX, SizeY, bIsFullScreen, Format);
+	Resize(SizeX, SizeY, bIsFullScreen, Format, bWaitForRenderingThread);
 }
 #endif
 
@@ -502,7 +505,7 @@ void ResizeMetalLayer(CAMetalLayer* MetalLayer, MTL::PixelFormat MetalFormat, bo
 }
 #endif
 
-void FMetalViewport::Resize(uint32 InSizeX, uint32 InSizeY, bool bInIsFullscreen, EPixelFormat InFormat)
+void FMetalViewport::Resize(uint32 InSizeX, uint32 InSizeY, bool bInIsFullscreen, EPixelFormat InFormat, bool bWaitForRenderingThread)
 {
 	bIsFullScreen = bInIsFullscreen;
 	uint32 Index = GetViewportIndex(EMetalViewportAccessGame);
@@ -543,10 +546,13 @@ void FMetalViewport::Resize(uint32 InSizeX, uint32 InSizeY, bool bInIsFullscreen
 			});
 	}
     
-    // Issue a fence command to the rendering thread and wait for it to complete.
-    FRenderCommandFence Fence;
-    Fence.BeginFence();
-    Fence.Wait();
+	if(bWaitForRenderingThread)
+	{
+		// Issue a fence command to the rendering thread and wait for it to complete.
+		FRenderCommandFence Fence;
+		Fence.BeginFence();
+		Fence.Wait();
+	}
     
 #if PLATFORM_MAC
 	float ContentScaleFactor = FSlateApplication::Get().GetApplicationScale() * MacApplication->FindWindowByNSWindow(WindowHandle)->GetDPIScaleFactor();
diff --git a/Engine/Source/Runtime/Apple/MetalRHI/Public/MetalViewport.h b/Engine/Source/Runtime/Apple/MetalRHI/Public/MetalViewport.h
--- a/Engine/Source/Runtime/Apple/MetalRHI/Public/MetalViewport.h
+++ b/Engine/Source/Runtime/Apple/MetalRHI/Public/MetalViewport.h
@@ -72,7 +72,7 @@ public:
 	FMetalViewport(FMetalDevice& InDevice, void* WindowHandle, uint32 InSizeX, uint32 InSizeY, bool bInIsFullscreen,EPixelFormat InFormat);
 	~FMetalViewport();
 
-	void Resize(uint32 InSizeX, uint32 InSizeY, bool bInIsFullscreen,EPixelFormat Format);
+	void Resize(uint32 InSizeX, uint32 InSizeY, bool bInIsFullscreen,EPixelFormat Format,bool bWaitForRenderingThread = true);
 	
 	TRefCountPtr<FMetalSurface> GetBackBuffer(EMetalViewportAccessFlag Accessor, 
 											  TOptional<int32> DrawLayerIndex = TOptional<int32>()) const;
diff --git a/Engine/Source/Runtime/ApplicationCore/Private/Mac/CocoaWindow.cpp b/Engine/Source/Runtime/ApplicationCore/Private/Mac/CocoaWindow.cpp
--- a/Engine/Source/Runtime/ApplicationCore/Private/Mac/CocoaWindow.cpp
+++ b/Engine/Source/Runtime/ApplicationCore/Private/Mac/CocoaWindow.cpp
@@ -471,7 +471,7 @@ NSString* NSPerformDragOperation = @"NSPerformDragOperation";
 	SCOPED_AUTORELEASE_POOL;
 	if (MacApplication && sender == self && !bIsBeingOrderedFront) // Skip informing Slate if we're simply changing the z order of windows
 	{
-		const bool bWaitForThread = false;
+		const bool bWaitForThread = true;
 		
 		GameThreadCall(^{
 			if (MacApplication) // Another check because game thread may destroy MacApplication before it gets here
diff --git a/Engine/Source/Runtime/SlateRHIRenderer/Private/SlateRHIRenderer.cpp b/Engine/Source/Runtime/SlateRHIRenderer/Private/SlateRHIRenderer.cpp
--- a/Engine/Source/Runtime/SlateRHIRenderer/Private/SlateRHIRenderer.cpp
+++ b/Engine/Source/Runtime/SlateRHIRenderer/Private/SlateRHIRenderer.cpp
@@ -353,6 +353,10 @@ void FSlateViewportInfo::CreateTexture2D()
 
 void FSlateViewportInfo::CreateNativeLayer(int32 NativeLayer, void* NativeViewHandle)
 {
+	// this flush is needed to make sure that the rendering thread is finished and not
+	// using any of the NativeLayer resources to prevent race conditions.
+	FlushRenderingCommands();
+	
 	bool bAlreadyCreated = false;
 	Layers.NativeLayers.Add(NativeLayer, &bAlreadyCreated);
 
@@ -403,6 +407,10 @@ void FSlateViewportInfo::CreateNativeLayer(int32 NativeLayer, void* NativeViewHa
 
 void FSlateViewportInfo::DeleteNativeLayer(int32 NativeLayer)
 {
+	// this flush is needed to make sure that the rendering thread is finished and not
+	// using any of the NativeLayer resources to prevent race conditions.
+	FlushRenderingCommands();
+	
 	// Make sure we aren't deleting the bottom most layer
 	check(NativeLayer != std::numeric_limits<int32>::min());
 
@@ -1387,17 +1395,19 @@ struct FSlateDrawWindowsCommand : public TConcurrentLinearObject<FSlateDrawWindo
 void FSlateRHIRenderer::CreateNativeLayer(int32 NewNativeLayer, SWindow& InWindow, void* NativeViewHandle)
 {
 	FSlateViewportInfo* ViewportInfo = WindowToViewportInfo.FindRef(&InWindow);
-	check(ViewportInfo != nullptr);
-
-	ViewportInfo->CreateNativeLayer(NewNativeLayer, NativeViewHandle);
+	if(ViewportInfo != nullptr)
+	{
+		ViewportInfo->CreateNativeLayer(NewNativeLayer, NativeViewHandle);
+	}
 }
 
 void FSlateRHIRenderer::DeleteNativeLayer(int32 OldNativeLayer, SWindow& InWindow)
 {
 	FSlateViewportInfo* ViewportInfo = WindowToViewportInfo.FindRef(&InWindow);
-	check(ViewportInfo != nullptr);
-
-	ViewportInfo->DeleteNativeLayer(OldNativeLayer);
+	if(ViewportInfo != nullptr)
+	{
+		ViewportInfo->DeleteNativeLayer(OldNativeLayer);
+	}
 }
 #endif
 
