Metal - Fix corruption happening with Metal Native Shader Library after CL 51296110
It was caused by the hash not beeing properly set anymore in MetalShaderLibrary (a mistake on my end)

#tests Local test + preflight
#Synced-CL 51373489

diff --git a/Engine/Source/Runtime/Apple/MetalRHI/Private/Shaders/MetalShaderLibrary.cpp b/Engine/Source/Runtime/Apple/MetalRHI/Private/Shaders/MetalShaderLibrary.cpp
--- a/Engine/Source/Runtime/Apple/MetalRHI/Private/Shaders/MetalShaderLibrary.cpp
+++ b/Engine/Source/Runtime/Apple/MetalRHI/Private/Shaders/MetalShaderLibrary.cpp
@@ -255,5 +255,10 @@ TRefCountPtr<FRHIShader> FMetalShaderLibrary::CreateShader(int32 Index, bool bRe
 			break;
 	}
 
+	if (Shader)
+	{
+		Shader->SetHash(SerializedShaders.GetShaderHashes()[Index]);
+	}
+
 	return Shader;
 }
