Unshelved from pending changelist '51720269':

Metal RHI

Depth was included twice in size calculation for Textures

diff --git a/Engine/Source/Runtime/Apple/MetalRHI/Private/MetalTexture.cpp b/Engine/Source/Runtime/Apple/MetalRHI/Private/MetalTexture.cpp
--- a/Engine/Source/Runtime/Apple/MetalRHI/Private/MetalTexture.cpp
+++ b/Engine/Source/Runtime/Apple/MetalRHI/Private/MetalTexture.cpp
@@ -1252,7 +1252,8 @@ uint32 FMetalSurface::GetMemorySize()
 	uint32 TotalSize = 0;
 	for (uint32 MipIndex = 0; MipIndex < Texture->mipmapLevelCount(); MipIndex++)
 	{
-		TotalSize += GetMipSize(MipIndex, NULL, false) * GetNumFaces();
+		uint32 NumFaces = (GetDesc().IsTextureCube() ? 6 : 1);
+		TotalSize += GetMipSize(MipIndex, NULL, false) * GetDesc().ArraySize * NumFaces;
 	}
 	
 	return TotalSize;
