Metal RHI

Fix Bindless descriptor OOB
Fix issue where viewport attemps to set OOB

#Synced-CL 51619317

diff --git a/Engine/Source/Runtime/Apple/MetalRHI/Private/MetalBindlessDescriptors.cpp b/Engine/Source/Runtime/Apple/MetalRHI/Private/MetalBindlessDescriptors.cpp
--- a/Engine/Source/Runtime/Apple/MetalRHI/Private/MetalBindlessDescriptors.cpp
+++ b/Engine/Source/Runtime/Apple/MetalRHI/Private/MetalBindlessDescriptors.cpp
@@ -12,7 +12,7 @@
 #include "MetalCommandEncoder.h"
 #include "MetalDynamicRHI.h"
 
-int32 GBindlessResourceDescriptorHeapSize = 2048 * 1024;
+int32 GBindlessResourceDescriptorHeapSize = 1 << 23; // 8MB
 static FAutoConsoleVariableRef CVarBindlessResourceDescriptorHeapSize(
 	TEXT("Metal.Bindless.ResourceDescriptorHeapSize"),
 	GBindlessResourceDescriptorHeapSize,
@@ -20,7 +20,7 @@ static FAutoConsoleVariableRef CVarBindlessResourceDescriptorHeapSize(
 	ECVF_ReadOnly
 );
 
-int32 GBindlessSamplerDescriptorHeapSize = 64 << 10; // TODO: We should be able to reduce the size of the sampler heap if we fix static sampler creation.
+int32 GBindlessSamplerDescriptorHeapSize = 1 << 16; // 64k
 static FAutoConsoleVariableRef CVarBindlessSamplerDescriptorHeapSize(
 	TEXT("Metal.Bindless.SamplerDescriptorHeapSize"),
 	GBindlessSamplerDescriptorHeapSize,
@@ -134,7 +134,7 @@ void FMetalDescriptorHeap::Init(FRHICommandListBase& RHICmdList, int32 HeapSize)
 
 	ResourceHeap = new FMetalRHIBuffer(RHICmdList, Device, CreateDesc, nullptr);
 	
-	Manager = new FRHIHeapDescriptorAllocator(TypeMask, HeapSize, {});
+	Manager = new FRHIHeapDescriptorAllocator(TypeMask, HeapSize / sizeof(IRDescriptorTableEntry), {});
 }
 
 FRHIDescriptorHandle FMetalDescriptorHeap::AllocateDescriptor(ERHIDescriptorType DescriptorType)
diff --git a/Engine/Source/Runtime/Apple/MetalRHI/Private/MetalStateCache.cpp b/Engine/Source/Runtime/Apple/MetalRHI/Private/MetalStateCache.cpp
--- a/Engine/Source/Runtime/Apple/MetalRHI/Private/MetalStateCache.cpp
+++ b/Engine/Source/Runtime/Apple/MetalRHI/Private/MetalStateCache.cpp
@@ -361,17 +361,13 @@ void FMetalStateCache::SetScissorRect(bool const bEnable, MTL::ScissorRect const
 		Scissor[0].height = Viewport[0].height;
 	}
 	
-	// Clamp to framebuffer size - Metal doesn't allow scissor to be larger.
-	Scissor[0].x = Scissor[0].x;
-	Scissor[0].y = Scissor[0].y;
-	
 	// FrameBufferSize will be 0 if using RT-less rasterization, so ignore
 	if(FrameBufferSize.width != 0 && FrameBufferSize.height != 0)
 	{
-		Scissor[0].width = FMath::Max((Scissor[0].x + Scissor[0].width <= FMath::RoundToInt32(FrameBufferSize.width)) ? Scissor[0].width : FMath::RoundToInt32(FrameBufferSize.width) - Scissor[0].x, (NS::UInteger)1u);
-		Scissor[0].height = FMath::Max((Scissor[0].y + Scissor[0].height <= FMath::RoundToInt32(FrameBufferSize.height)) ? Scissor[0].height : FMath::RoundToInt32(FrameBufferSize.height) - Scissor[0].y, (NS::UInteger)1u);
+		Scissor[0].width = FMath::Max((Scissor[0].x + Scissor[0].width <= FMath::RoundToInt32(FrameBufferSize.width)) ? Scissor[0].width : FMath::RoundToInt32(FrameBufferSize.width) - Scissor[0].x, (NS::UInteger)0u);
+		Scissor[0].height = FMath::Max((Scissor[0].y + Scissor[0].height <= FMath::RoundToInt32(FrameBufferSize.height)) ? Scissor[0].height : FMath::RoundToInt32(FrameBufferSize.height) - Scissor[0].y, (NS::UInteger)0u);
 	}
-	
+
 	RasterBits |= EMetalRenderFlagScissorRect;
 	
 	ActiveScissors = 1;
