Metal RHI

- Fixed bug where the metal buffer lock was trying to access the context while executing.

#Synced-CL 51619317

diff --git a/Engine/Source/Runtime/Apple/MetalRHI/Private/MetalRHIBuffer.cpp b/Engine/Source/Runtime/Apple/MetalRHI/Private/MetalRHIBuffer.cpp
--- a/Engine/Source/Runtime/Apple/MetalRHI/Private/MetalRHIBuffer.cpp
+++ b/Engine/Source/Runtime/Apple/MetalRHI/Private/MetalRHIBuffer.cpp
@@ -380,11 +380,13 @@ void* FMetalRHIBuffer::Lock(FRHICommandListBase& RHICmdList, EResourceLockMode I
 			AllocTransferBuffer(true, Len, RLM_WriteOnly);
 			check(TransferBuffer->GetLength() >= InSize);
 			
-			FRHICommandListImmediate& ImmediateCmdList = FRHICommandListImmediate::Get();
-			FMetalRHICommandContext& Context = FMetalRHICommandContext::Get(RHICmdList);
+			FRHICommandListImmediate& ImmediateCmdList = RHICmdList.GetAsImmediate();
 			
-			// Synchronise the buffer with the CPU
-			Context.CopyFromBufferToBuffer(GetCurrentBuffer(), 0, TransferBuffer, 0, GetCurrentBuffer()->GetLength());
+			ImmediateCmdList.EnqueueLambda([this, InTransferBuffer=TransferBuffer](FRHICommandListBase& ExecutingCmdList)
+			{
+				FMetalRHICommandContext& Context = FMetalRHICommandContext::Get(ExecutingCmdList);
+				Context.CopyFromBufferToBuffer(GetCurrentBuffer(), 0, InTransferBuffer, 0, GetCurrentBuffer()->GetLength());
+			});
 			
 			//kick the current command buffer.
 			ImmediateCmdList.SubmitAndBlockUntilGPUIdle();
@@ -396,11 +398,14 @@ void* FMetalRHIBuffer::Lock(FRHICommandListBase& RHICmdList, EResourceLockMode I
 		{
 			SCOPE_CYCLE_COUNTER(STAT_MetalBufferPageOffTime);
 			
-			FRHICommandListImmediate& ImmediateCmdList = FRHICommandListImmediate::Get();
-			FMetalRHICommandContext& Context = FMetalRHICommandContext::Get(RHICmdList);
+			FRHICommandListImmediate& ImmediateCmdList = RHICmdList.GetAsImmediate();
 			
-			// Synchronise the buffer with the CPU
-			Context.SynchronizeResource(GetCurrentBuffer()->GetMTLBuffer());
+			ImmediateCmdList.EnqueueLambda([this](FRHICommandListBase& ExecutingCmdList)
+			{
+				FMetalRHICommandContext& Context = FMetalRHICommandContext::Get(ExecutingCmdList);
+				// Synchronise the buffer with the CPU
+				Context.SynchronizeResource(GetCurrentBuffer()->GetMTLBuffer());
+			});
 			
 			//kick the current command buffer.
 			ImmediateCmdList.SubmitAndBlockUntilGPUIdle();
@@ -639,10 +644,6 @@ void* FMetalDynamicRHI::RHILockBuffer(class FRHICommandListBase& RHICmdList, FRH
 		void* Result = nullptr;
 		if (LockMode != RLM_WriteOnly)
 		{
-			QUICK_SCOPE_CYCLE_COUNTER(STAT_RHIMETHOD_LockBuffer_FlushAndLock);
-			CSV_SCOPED_TIMING_STAT(RHITFlushes, LockBuffer_BottomOfPipe);
-
-			FRHICommandListScopedFlushAndExecute Flush(RHICmdList.GetAsImmediate());
 			Result = Buffer->Lock(RHICmdList, LockMode, Offset, SizeRHI);
 
 			GBufferLockTracker.Lock(Buffer, nullptr, Offset, SizeRHI, LockMode);
@@ -686,10 +687,6 @@ void FMetalDynamicRHI::RHIUnlockBuffer(class FRHICommandListBase& RHICmdList, FR
 
 		if (Params.LockMode != RLM_WriteOnly)
 		{
-			QUICK_SCOPE_CYCLE_COUNTER(STAT_RHIMETHOD_UnlockBuffer_FlushAndUnlock);
-			CSV_SCOPED_TIMING_STAT(RHITFlushes, UnlockBuffer_BottomOfPipe);
-			
-			FRHICommandListScopedFlushAndExecute Flush(RHICmdList.GetAsImmediate());
 			Buffer->Unlock(RHICmdList);
 		}
 		else
