Metal RHI:

Fixed multiple bugs impacting Metal Ray Tracing:

- Using incorrect metal buffer from FRHIUniformBuffer causing flickering
- Bug in parallel workers when enqueuing resources, now has array per worker

#Synced-CL 52187599

diff --git a/Engine/Source/Runtime/Apple/MetalRHI/Private/MetalRayTracing.cpp b/Engine/Source/Runtime/Apple/MetalRHI/Private/MetalRayTracing.cpp
--- a/Engine/Source/Runtime/Apple/MetalRHI/Private/MetalRayTracing.cpp
+++ b/Engine/Source/Runtime/Apple/MetalRHI/Private/MetalRayTracing.cpp
@@ -367,16 +367,11 @@ inline t_A RoundUpToNextMultiple(const t_A& a, const t_B& b)
 
 class FMetalRayTracingPipelineState;
 
-static void ProcessUniformBuffer(FRHIUniformBuffer* UniformBuffer, FMetalStateCache& StateCache)
+static void ProcessUniformBuffer(FMetalBufferPtr BackingBuffer, const TArray<TRefCountPtr<FRHIResource>>& CapturedResourceTable, FMetalStateCache& StateCache)
 {
-	const TArray<TRefCountPtr<FRHIResource>>& ResourceTable = UniformBuffer->GetResourceTable();
-	
-	FMetalUniformBuffer *UB = ResourceCast(UniformBuffer);
-	FMetalBufferPtr BackingBuffer = UB->BackingBuffer;
-	
 	StateCache.CacheOrSkipResourceResidencyUpdate(BackingBuffer->GetMTLBuffer(), EMetalShaderStages::Compute, true);
-	
-	for(const TRefCountPtr<FRHIResource>& RHIResourceRef : ResourceTable)
+
+	for(const TRefCountPtr<FRHIResource>& RHIResourceRef : CapturedResourceTable)
 	{
 		FRHIResource* Resource = RHIResourceRef.GetReference();
 	
@@ -422,6 +417,13 @@ static void ProcessUniformBuffer(FRHIUniformBuffer* UniformBuffer, FMetalStateCa
 	}
 }
 
+static void ProcessUniformBuffer(FRHIUniformBuffer* UniformBuffer, FMetalStateCache& StateCache)
+{
+	FMetalUniformBuffer* UB = ResourceCast(UniformBuffer);
+	StateCache.CacheOrSkipResourceResidencyUpdate(UB->BackingBuffer->GetMTLBuffer(), EMetalShaderStages::Compute, true);
+	ProcessUniformBuffer(UB->BackingBuffer, UniformBuffer->GetResourceTable(), StateCache);
+}
+
 class FMetalRayTracingShaderBindingTable : public FRHIShaderBindingTable
 {
 private:
@@ -666,17 +668,25 @@ public:
 		}
 	}
 	
-	void SetUniformBindings(FMetalRayTracingPipelineState* Pipeline, uint32 Offset, uint32 RecordIndex, const uint32 NumUniformBuffers,
+	void SetUniformBindings(FMetalRayTracingPipelineState* Pipeline, uint32 WorkerIndex, uint32 Offset, uint32 RecordIndex, const uint32 NumUniformBuffers,
 									FRHIUniformBuffer* const* UniformBuffers, const void* LooseParameterData, uint32 LooseParameterDataSize);
-	
-	void SetHitGroupUniformBindings(FMetalRayTracingPipelineState* Pipeline, uint32 RecordIndex, const uint32 NumUniformBuffers,
+
+	void SetHitGroupUniformBindings(FMetalRayTracingPipelineState* Pipeline, uint32 WorkerIndex, uint32 RecordIndex, const uint32 NumUniformBuffers,
 									FRHIUniformBuffer* const* UniformBuffers, const void* LooseParameterData, uint32 LooseParameterDataSize);
-	
-	void SetMissUniformBindings(FMetalRayTracingPipelineState* Pipeline, uint32 RecordIndex, const uint32 NumUniformBuffers, 
+
+	void SetMissUniformBindings(FMetalRayTracingPipelineState* Pipeline, uint32 WorkerIndex, uint32 RecordIndex, const uint32 NumUniformBuffers,
 									FRHIUniformBuffer* const* UniformBuffers, const void* LooseParameterData, uint32 LooseParameterDataSize);
-	
-	void SetCallableUniformBindings(FMetalRayTracingPipelineState* Pipeline, uint32 RecordIndex, const uint32 NumUniformBuffers,
+
+	void SetCallableUniformBindings(FMetalRayTracingPipelineState* Pipeline, uint32 WorkerIndex, uint32 RecordIndex, const uint32 NumUniformBuffers,
 									FRHIUniformBuffer* const* UniformBuffers, const void* LooseParameterData, uint32 LooseParameterDataSize);
+
+	void ResetResourceBindings()
+	{
+		for (uint32 i = 0; i < MaxBindingWorkers; ++i)
+		{
+			ResourceBindings[i].Reset();
+		}
+	}
 	
 	const IRShaderIdentifier &ShaderIdentifierAtOffset(uint32 Offset) const
 	{
@@ -716,9 +726,9 @@ public:
 		}
 		StateCache.CacheOrSkipResourceResidencyUpdate(Buffer->GetMTLBuffer(), EMetalShaderStages::Compute, true);
 		
-		for (FRHIUniformBuffer* UB : BoundUniformBuffers)
+		for (const FBoundUniformBuffer& BoundUB : BoundUniformBuffers)
 		{
-			ProcessUniformBuffer(UB, StateCache);	
+			ProcessUniformBuffer(BoundUB.BackingBuffer, BoundUB.CapturedResourceTable, StateCache);
 		}
 	}
 	
@@ -813,8 +823,14 @@ private:
 	uint32 LocalRecordStride = 0; // size of shader identifier + local root parameters, aligned to SHADER_RECORD_BYTE_ALIGNMENT (same for hit groups and callable shaders)
 	TResourceArray<uint8, METAL_RAYTRACING_SHADER_TABLE_BYTE_ALIGNMENT> Data;
 	
-	TArray<FRHIUniformBuffer*> BoundUniformBuffers;
-	
+	struct FBoundUniformBuffer
+	{
+		FRHIUniformBuffer* UniformBuffer;
+		FMetalBufferPtr BackingBuffer;
+		TArray<TRefCountPtr<FRHIResource>> CapturedResourceTable;
+	};
+	TArray<FBoundUniformBuffer> BoundUniformBuffers;
+
 	TSet<MTL::Heap*> BindingTableHeaps[MaxBindingWorkers];
 	
 	FMetalBufferPtr Buffer;
@@ -829,13 +845,15 @@ private:
 	{
 		static constexpr uint32 MaxUniformBuffers = 8;
 		uint64 LooseParametersVA = 0;
-		FRHIUniformBuffer *UniformBuffers[MaxUniformBuffers];
+		FRHIUniformBuffer* UniformBuffers[MaxUniformBuffers];
+		FMetalBufferPtr BackingBuffers[MaxUniformBuffers];
+		TArray<TRefCountPtr<FRHIResource>> CapturedResourceTables[MaxUniformBuffers];
 		uint32 NumUniformBuffers = 0;
 		uint32 BindingOffset;
 	};
 
 	FCriticalSection ResourceCS;
-	TQueue<FResourceBinder> Resources;
+	TArray<FResourceBinder> ResourceBindings[MaxBindingWorkers];
 };
 
 class FMetalRayTracingPipelineState : public FRHIRayTracingPipelineState
@@ -1162,7 +1180,8 @@ public:
 		
 		PipelineStateObject = FMetalShaderPipelinePtr(new FMetalShaderPipeline(Device));
 		PipelineStateObject->ComputePipelineState = NS::TransferPtr(MTLDevice->newComputePipelineState(ComputePSODescriptor, 0, nullptr, &Error));
-		check(PipelineStateObject->ComputePipelineState);
+		
+		UE_CLOG((PipelineStateObject->ComputePipelineState.get() == nullptr), LogMetal, Error, TEXT("Invalid Compute Pipeline State, Descriptor: %s"), *NSStringToFString(ComputePSODescriptor->description()));
 		
 		LinkedFunctions->release();
 
@@ -1304,43 +1323,50 @@ void FMetalRayTracingShaderBindingTable::CopyToGPU(FMetalRHICommandContext& Cont
 	
 	{
 		FScopeLock Lock(&ResourceCS);
-		FResourceBinder Binder;
-		
+
 		BoundUniformBuffers.Empty();
-		
-		while (Resources.Dequeue(Binder))
+
+		for (uint32 WorkerIdx = 0; WorkerIdx < MaxBindingWorkers; ++WorkerIdx)
 		{
-			const bool bHasLooseParameters = (Binder.LooseParametersVA != 0);
-			
-			uint64 Entries[FResourceBinder::MaxUniformBuffers];
-			uint32 EntryCount = Binder.NumUniformBuffers;
-			uint32 CurEntryIdx = 0;
-			if (bHasLooseParameters)
+			for (FResourceBinder& Binder : ResourceBindings[WorkerIdx])
 			{
-				EntryCount++;
-				CurEntryIdx++;
-			}
-			
-			FMemory::Memset(&Entries[0], 0, sizeof(Entries[0]) * EntryCount);
-			if (bHasLooseParameters)
-			{
-				Entries[0] = Binder.LooseParametersVA;
-			}
-			
-			for (uint32 i = 0; i < Binder.NumUniformBuffers; ++i)
-			{
-				FMetalUniformBuffer* UniformBuffer = ResourceCast((FRHIUniformBuffer*)Binder.UniformBuffers[i]);
-				if (UniformBuffer)
+				const bool bHasLooseParameters = (Binder.LooseParametersVA != 0);
+
+				// +1 to accommodate loose parameters entry in slot 0 when present
+				uint64 Entries[FResourceBinder::MaxUniformBuffers + 1];
+				uint32 EntryCount = Binder.NumUniformBuffers;
+				uint32 CurEntryIdx = 0;
+				if (bHasLooseParameters)
 				{
-					FMetalBufferPtr BackingBuffer = UniformBuffer->BackingBuffer;
-					Entries[CurEntryIdx] = BackingBuffer->GetGPUAddress();
-					BoundUniformBuffers.Add(UniformBuffer);
-					
+					EntryCount++;
 					CurEntryIdx++;
 				}
+
+				FMemory::Memset(&Entries[0], 0, sizeof(Entries[0]) * EntryCount);
+				if (bHasLooseParameters)
+				{
+					Entries[0] = Binder.LooseParametersVA;
+				}
+
+				for (uint32 i = 0; i < Binder.NumUniformBuffers; ++i)
+				{
+					FMetalBufferPtr BackingBuffer = Binder.BackingBuffers[i];
+					if (BackingBuffer)
+					{
+						StateCache.CacheOrSkipResourceResidencyUpdate(BackingBuffer->GetMTLBuffer(), EMetalShaderStages::Compute, true);
+
+						Entries[CurEntryIdx] = BackingBuffer->GetGPUAddress();
+						BoundUniformBuffers.Add({Binder.UniformBuffers[i], BackingBuffer, MoveTemp(Binder.CapturedResourceTables[i])});
+
+						CurEntryIdx++;
+					}
+				}
+				WriteData(Binder.BindingOffset, &Entries[0], sizeof(Entries[0]) * EntryCount);
 			}
-			WriteData(Binder.BindingOffset, &Entries[0], sizeof(Entries[0]) * EntryCount);
 		}
+
+		// Clear per-worker arrays after processing (ready for next frame)
+		ResetResourceBindings();
 	}
 
 	if (bBufferDirty)
@@ -1365,6 +1391,7 @@ void FMetalRayTracingShaderBindingTable::CopyToGPU(FMetalRHICommandContext& Cont
 }
 
 void FMetalRayTracingShaderBindingTable::SetUniformBindings(FMetalRayTracingPipelineState* Pipeline,
+															uint32 WorkerIndex,
 															uint32 TableOffset,
 															uint32 RecordIndex,
 															const uint32 NumUniformBuffers,
@@ -1406,31 +1433,41 @@ void FMetalRayTracingShaderBindingTable::SetUniformBindings(FMetalRayTracingPipe
 	{
 		ResourceBinder.UniformBuffers[i] = UniformBuffers[i];
 		checkf(UniformBuffers[i] == nullptr || (uint64_t)UniformBuffers[i] >= 1000, TEXT("UniformBuffer %u of %u is %p"), i, NumUniformBuffers, UniformBuffers[i]);
+
+		if (UniformBuffers[i])
+		{
+			FMetalUniformBuffer* UB = ResourceCast(UniformBuffers[i]);
+
+			ResourceBinder.BackingBuffers[i] = UB->BackingBuffer;
+			ResourceBinder.CapturedResourceTables[i] = UniformBuffers[i]->GetResourceTable();
+		}
 	}
 	ResourceBinder.NumUniformBuffers = NumUniformBuffers;
 	
 	ResourceBinder.BindingOffset = ShaderTableOffset + sizeof(IRShaderIdentifier) + sizeof(FMetalHitGroupSystemParameters);
-	Resources.Enqueue(ResourceBinder);
+
+	checkf(WorkerIndex < MaxBindingWorkers, TEXT("WorkerIndex %u exceeds MaxBindingWorkers %u"), WorkerIndex, MaxBindingWorkers);
+	ResourceBindings[WorkerIndex].Add(MoveTemp(ResourceBinder));
 
 	check(ShaderTableOffset + sizeof(IRShaderIdentifier) + sizeof(FMetalHitGroupSystemParameters) + NumUniformBuffers * sizeof(uint64 /*IRDescriptorTableEntry*/) <= ShaderTableOffset + LocalRecordStride);
 }
 
-void FMetalRayTracingShaderBindingTable::SetHitGroupUniformBindings(FMetalRayTracingPipelineState* Pipeline, uint32 RecordIndex, const uint32 NumUniformBuffers,
+void FMetalRayTracingShaderBindingTable::SetHitGroupUniformBindings(FMetalRayTracingPipelineState* Pipeline, uint32 WorkerIndex, uint32 RecordIndex, const uint32 NumUniformBuffers,
 																	FRHIUniformBuffer* const* UniformBuffers, const void* LooseParameterData, uint32 LooseParameterDataSize)
 {
-	SetUniformBindings(Pipeline, HitGroupShaderTableOffset, RecordIndex, NumUniformBuffers, UniformBuffers, LooseParameterData, LooseParameterDataSize);
+	SetUniformBindings(Pipeline, WorkerIndex, HitGroupShaderTableOffset, RecordIndex, NumUniformBuffers, UniformBuffers, LooseParameterData, LooseParameterDataSize);
 }
 
-void FMetalRayTracingShaderBindingTable::SetMissUniformBindings(FMetalRayTracingPipelineState* Pipeline, uint32 RecordIndex, const uint32 NumUniformBuffers,
+void FMetalRayTracingShaderBindingTable::SetMissUniformBindings(FMetalRayTracingPipelineState* Pipeline, uint32 WorkerIndex, uint32 RecordIndex, const uint32 NumUniformBuffers,
 																FRHIUniformBuffer* const* UniformBuffers, const void* LooseParameterData, uint32 LooseParameterDataSize)
 {
-	SetUniformBindings(Pipeline, MissShaderTableOffset, RecordIndex, NumUniformBuffers, UniformBuffers, LooseParameterData, LooseParameterDataSize);
+	SetUniformBindings(Pipeline, WorkerIndex, MissShaderTableOffset, RecordIndex, NumUniformBuffers, UniformBuffers, LooseParameterData, LooseParameterDataSize);
 }
 
-void FMetalRayTracingShaderBindingTable::SetCallableUniformBindings(FMetalRayTracingPipelineState* Pipeline, uint32 RecordIndex, const uint32 NumUniformBuffers,
+void FMetalRayTracingShaderBindingTable::SetCallableUniformBindings(FMetalRayTracingPipelineState* Pipeline, uint32 WorkerIndex, uint32 RecordIndex, const uint32 NumUniformBuffers,
 																FRHIUniformBuffer* const* UniformBuffers, const void* LooseParameterData, uint32 LooseParameterDataSize)
 {
-	SetUniformBindings(Pipeline, CallableShaderTableOffset, RecordIndex, NumUniformBuffers, UniformBuffers, LooseParameterData, LooseParameterDataSize);
+	SetUniformBindings(Pipeline, WorkerIndex, CallableShaderTableOffset, RecordIndex, NumUniformBuffers, UniformBuffers, LooseParameterData, LooseParameterDataSize);
 }
 
 FMetalAccelerationStructure::~FMetalAccelerationStructure()
@@ -1987,6 +2024,7 @@ void FMetalRayTracingGeometry::RebuildDescriptors()
 
 static void SetRayTracingHitGroup(
 	FMetalRayTracingShaderBindingTable* ShaderTable,
+	uint32 WorkerIndex,
 	uint32 RecordIndex,
 	FMetalRayTracingPipelineState* Pipeline,
 	const FMetalRayTracingGeometry* Geometry,
@@ -2025,7 +2063,7 @@ static void SetRayTracingHitGroup(
 			
 			GeometryParameters.RootConstants.UserData = UserData;
 			ShaderTable->SetHitGroupSystemParameters(RecordIndex, GeometryParameters);
-			ShaderTable->SetHitGroupUniformBindings(Pipeline, RecordIndex, NumUniformBuffers, UniformBuffers, LooseParameterData, LooseParameterDataSize);
+			ShaderTable->SetHitGroupUniformBindings(Pipeline, WorkerIndex, RecordIndex, NumUniformBuffers, UniformBuffers, LooseParameterData, LooseParameterDataSize);
 		}
 
 		if (EnumHasAnyFlags(ShaderBindingMode, ERayTracingShaderBindingMode::Inline))
@@ -2043,8 +2081,9 @@ static void SetRayTracingHitGroup(
 
 static void SetRayTracingCallableShader(
 	FMetalRayTracingShaderBindingTable* ShaderTable,
+	uint32 WorkerIndex,
 	FMetalRayTracingPipelineState* Pipeline,
-	uint32 RecordIndex, 
+	uint32 RecordIndex,
 	uint32 ShaderIndexInPipeline,
 	uint32 NumUniformBuffers,
 	FRHIUniformBuffer* const* UniformBuffers,
@@ -2061,8 +2100,8 @@ static void SetRayTracingCallableShader(
 	check(ShaderIndexInPipeline < Pipeline->CallableShaders.Identifiers.Num());
 
 	ShaderTable->SetCallableIdentifier(RecordIndex, *Pipeline->CallableShaders.Identifiers[ShaderIndexInPipeline]);
-	ShaderTable->SetCallableUniformBindings(Pipeline, RecordIndex, NumUniformBuffers, UniformBuffers, LooseParameterData, LooseParameterDataSize);
-	
+	ShaderTable->SetCallableUniformBindings(Pipeline, WorkerIndex, RecordIndex, NumUniformBuffers, UniformBuffers, LooseParameterData, LooseParameterDataSize);
+
 	UE_LOG(LogMetal, Log, TEXT("Binding callable shader: RecordIndex=%u, ShaderIndex=%u, Handle=%llu"),
 		  RecordIndex, ShaderIndexInPipeline,
 		  Pipeline->CallableShaders.Identifiers[ShaderIndexInPipeline]->ShaderIdentifier.shaderHandle);
@@ -2070,8 +2109,9 @@ static void SetRayTracingCallableShader(
 
 static void SetRayTracingMissShader(
 	FMetalRayTracingShaderBindingTable* ShaderTable,
+	uint32 WorkerIndex,
 	FMetalRayTracingPipelineState* Pipeline,
-	uint32 RecordIndex, 
+	uint32 RecordIndex,
 	uint32 ShaderIndexInPipeline,
 	uint32 NumUniformBuffers,
 	FRHIUniformBuffer* const* UniformBuffers,
@@ -2088,7 +2128,7 @@ static void SetRayTracingMissShader(
 	check(ShaderIndexInPipeline < Pipeline->MissShaders.Identifiers.Num());
 
 	ShaderTable->SetMissIdentifier(RecordIndex, *Pipeline->MissShaders.Identifiers[ShaderIndexInPipeline]);
-	ShaderTable->SetMissUniformBindings(Pipeline, RecordIndex, NumUniformBuffers, UniformBuffers, LooseParameterData, LooseParameterDataSize);
+	ShaderTable->SetMissUniformBindings(Pipeline, WorkerIndex, RecordIndex, NumUniformBuffers, UniformBuffers, LooseParameterData, LooseParameterDataSize);
 }
 
 void FMetalRHICommandContext::RHICommitShaderBindingTable(FRHIShaderBindingTable* InSBT, FRHIBuffer* InlineBindingDataBuffer)
@@ -2145,6 +2185,7 @@ void FMetalRHICommandContext::RHISetBindingsOnShaderBindingTable(
 			if (Binding.BindingType != ERayTracingLocalShaderBindingType::Clear)
 			{
 				SetRayTracingHitGroup(ShaderTable,
+									  Context.WorkerIndex,
 									  Binding.RecordIndex,
 									  Pipeline,
 									  Geometry,
@@ -2197,6 +2238,7 @@ void FMetalRHICommandContext::RHISetBindingsOnShaderBindingTable(
 		{
 			SetRayTracingCallableShader(
 							ShaderTable,
+							Context.WorkerIndex,
 							Pipeline,
 							Binding.RecordIndex,
 							Binding.ShaderIndexInPipeline,
@@ -2210,6 +2252,7 @@ void FMetalRHICommandContext::RHISetBindingsOnShaderBindingTable(
 		{
 			SetRayTracingMissShader(
 							ShaderTable,
+							Context.WorkerIndex,
 							Pipeline,
 							Binding.RecordIndex,
 							Binding.ShaderIndexInPipeline,
@@ -2732,9 +2775,7 @@ static void DispatchRays(FMetalRHICommandContext& Context,
 		if(GlobalResourceBindings.Textures[i])
 		{
 			FMetalSurface* Surface = GetMetalSurfaceFromRHITexture(GlobalResourceBindings.Textures[i]);
-			MTLTexturePtr Texture = Surface->GetCurrentTexture();
-			
-			StateCache.IRMakeTextureResident(EMetalShaderStages::Compute, Texture.get());
+			StateCache.IRMakeTextureResident(EMetalShaderStages::Compute, Surface->Texture.get());
 		}
 	}
 	
