Metal RHI

Fixed error calculating the geometry size when building acceleration structures

#Synced-CL 52187599

diff --git a/Engine/Source/Runtime/Apple/MetalRHI/Private/MetalRayTracing.cpp b/Engine/Source/Runtime/Apple/MetalRHI/Private/MetalRayTracing.cpp
--- a/Engine/Source/Runtime/Apple/MetalRHI/Private/MetalRayTracing.cpp
+++ b/Engine/Source/Runtime/Apple/MetalRHI/Private/MetalRayTracing.cpp
@@ -2619,7 +2619,11 @@ void FMetalRHICommandContext::RHIBuildAccelerationStructures(TConstArrayView<FRa
 		FMetalRayTracingGeometry* const Geometry = ResourceCast(P.Geometry.GetReference());
 		const bool bIsUpdate = P.BuildMode == EAccelerationStructureBuildMode::Update;
 
+		Geometry->RebuildDescriptors();
+		Geometry->SizeInfo = CalcRayTracingGeometrySize(Device, Geometry->AccelerationStructureDescriptor);
+		
 		uint64 ScratchBufferRequiredSize = bIsUpdate ? Geometry->SizeInfo.UpdateScratchSize : Geometry->SizeInfo.BuildScratchSize;
+		
 		checkf(ScratchBufferRequiredSize + ScratchBufferOffset <= ScratchBufferSize,
 			TEXT("BLAS scratch buffer size is %ld bytes with offset %ld (%ld bytes available), but the build requires %lld bytes. "),
 			ScratchBufferSize, ScratchBufferOffset, ScratchBufferSize - ScratchBufferOffset, ScratchBufferRequiredSize);
@@ -2634,9 +2638,6 @@ void FMetalRHICommandContext::RHIBuildAccelerationStructures(TConstArrayView<FRa
 		}
 
 		ScratchBufferOffset = Align(ScratchBufferOffset + ScratchBufferRequiredSize, ScratchAlignment);
-
-		// We must update the descriptor if any segments have changed.
-		Geometry->RebuildDescriptors();
 	}
 
 	FMetalBufferPtr ScratchBufferRes = ScratchBuffer->GetCurrentBuffer();
@@ -2646,8 +2647,6 @@ void FMetalRHICommandContext::RHIBuildAccelerationStructures(TConstArrayView<FRa
 	{
 		FMetalRayTracingGeometry* Geometry = BuildRequest.Geometry;
 		
-		FRayTracingAccelerationStructureSize TestSize = CalcRayTracingGeometrySize(Device, Geometry->AccelerationStructureDescriptor);
-																		
 		BuildAccelerationStructure(Geometry->GetAccelerationStructure()->GetAccelerationStructure(), Geometry->AccelerationStructureDescriptor,
 										   ScratchBufferRes, BuildRequest.Offset);
 	}
