Metal libraries/shaders LLM tagging fix

A single shader creation via some asset related callstack could lead to Metal library creation (as that is lazy and happens on the first shader creation from the library), attributing a ton of actually "shared" memory allocations (library itself & driver bookkeping, etc.) on that asset. That is misleading and also changes form run to run -- depends on the asset loading order.

Updated shader creation codepath to be tagged as Shaders.
Updated lazy Metal library creation path to also be attributed to Shaders and not the asset.

#rnx

#Synced-CL 52463319

#Synced-CL 52463319

diff --git a/Engine/Source/Runtime/Apple/MetalRHI/Private/Shaders/MetalShaderLibrary.cpp b/Engine/Source/Runtime/Apple/MetalRHI/Private/Shaders/MetalShaderLibrary.cpp
--- a/Engine/Source/Runtime/Apple/MetalRHI/Private/Shaders/MetalShaderLibrary.cpp
+++ b/Engine/Source/Runtime/Apple/MetalRHI/Private/Shaders/MetalShaderLibrary.cpp
@@ -159,6 +159,11 @@ TRefCountPtr<FRHIShader> FMetalShaderLibrary::CreateShader(int32 Index, bool bRe
 		
 		if (LazyLibrary->Library.get() == nullptr)
 		{
+			LLM_SCOPE(ELLMTag::Shaders);
+			LLM_TAGSET_SCOPE_CLEAR(ELLMTagSet::Assets);
+			LLM_TAGSET_SCOPE_CLEAR(ELLMTagSet::AssetClasses);
+			UE_TRACE_METADATA_CLEAR_SCOPE();
+
 			check(LazyLibrary->Data.Get() != nullptr);
 			
 			NS::Error* Error;
diff --git a/Engine/Source/Runtime/RenderCore/Private/ShaderLibrary/ShaderLibraryInstance.cpp b/Engine/Source/Runtime/RenderCore/Private/ShaderLibrary/ShaderLibraryInstance.cpp
--- a/Engine/Source/Runtime/RenderCore/Private/ShaderLibrary/ShaderLibraryInstance.cpp
+++ b/Engine/Source/Runtime/RenderCore/Private/ShaderLibrary/ShaderLibraryInstance.cpp
@@ -308,6 +308,8 @@ bool FShaderLibraryInstance::TryRemoveResource(FShaderMapResource_SharedCode* Re
 
 TRefCountPtr<FRHIShader> FShaderLibraryInstance::GetOrCreateShader(int32 ShaderIndex, bool bRequired)
 {
+	LLM_SCOPE(ELLMTag::Shaders);
+
 	const int32 BucketIndex = ShaderIndex % NumShaderLocks;
 	TRefCountPtr<FRHIShader> RHIShader;
 	{
