Metal RHI

Fix possible UB buffer overflow with CBV's and added checks.

#Synced-CL 52187599

diff --git a/Engine/Source/Runtime/Apple/MetalRHI/Private/MetalStateCache.cpp b/Engine/Source/Runtime/Apple/MetalRHI/Private/MetalStateCache.cpp
--- a/Engine/Source/Runtime/Apple/MetalRHI/Private/MetalStateCache.cpp
+++ b/Engine/Source/Runtime/Apple/MetalRHI/Private/MetalStateCache.cpp
@@ -1620,6 +1620,7 @@ void FMetalStateCache::IRBindPackedUniforms(EMetalShaderStages const Frequency,
 	}
 	
 	CacheOrSkipResourceResidencyUpdate(Buffer->GetMTLBuffer(), Frequency, true);
+	check((uint32)Index < TopLevelABNumEntry);
 	CBVTable[Frequency][Index] = Buffer->GetGPUAddress();
 }
 
@@ -1628,6 +1629,7 @@ void FMetalStateCache::IRBindUniformBuffer(EMetalShaderStages const Frequency, i
 	FMetalBufferPtr Buffer = UB->BackingBuffer;
 	CacheOrSkipResourceResidencyUpdate(Buffer->GetMTLBuffer(), Frequency, true);
 
+	check((uint32)Index < TopLevelABNumEntry);
 	CBVTable[Frequency][Index] = Buffer->GetGPUAddress();
 }
 #endif
@@ -1941,6 +1943,7 @@ void FMetalStateCache::IRBindResourcesToEncoder(ShaderType Shader, FMetalCommand
 	FMetalStaticSamplers* StaticSamplersTable = Device.GetStaticSamplers();
 	const uint64 StaticSamplersTableVA = StaticSamplersTable->GetGPUAddress();
 
+	check(Shader->Bindings.RSNumCBVs < TopLevelABNumEntry);
 	CBVTable[Frequency][Shader->Bindings.RSNumCBVs] = StaticSamplersTableVA;
 	
     // Bind CBV Table
diff --git a/Engine/Source/Runtime/Apple/MetalRHI/Private/MetalStateCache.h b/Engine/Source/Runtime/Apple/MetalRHI/Private/MetalStateCache.h
--- a/Engine/Source/Runtime/Apple/MetalRHI/Private/MetalStateCache.h
+++ b/Engine/Source/Runtime/Apple/MetalRHI/Private/MetalStateCache.h
@@ -313,7 +313,7 @@ private:
     MTL::StoreAction StencilStore;
 		
 #if METAL_USE_METAL_SHADER_CONVERTER
-    static constexpr uint32 TopLevelABNumEntry = 16;
+    static constexpr uint32 TopLevelABNumEntry = METAL_MAX_BUFFERS + 1; // One slot per buffer index (0..METAL_MAX_BUFFERS-1) plus one for the static sampler VA
     uint64 CBVTable[EMetalShaderStages::Num][TopLevelABNumEntry];
 	
 	IRRuntimeVertexBuffer VertexBufferVAs[31];
